[Mutable] Fixed some bugs related to "None" values of Mutable Tables.
+Now the "None" row is always generated but only will be shown in parameters that have the "Add None Option" selected.
+Now the "None" row is generated in the UE compilation part instead of the Core compilation part.

#rnx

diff --git a/Engine/Plugins/Mutable/Source/CustomizableObject/Internal/MuCO/CustomizableObjectCustomVersion.h b/Engine/Plugins/Mutable/Source/CustomizableObject/Internal/MuCO/CustomizableObjectCustomVersion.h
--- a/Engine/Plugins/Mutable/Source/CustomizableObject/Internal/MuCO/CustomizableObjectCustomVersion.h
+++ b/Engine/Plugins/Mutable/Source/CustomizableObject/Internal/MuCO/CustomizableObjectCustomVersion.h
@@ -168,6 +168,8 @@ struct CUSTOMIZABLEOBJECT_API FCustomizableObjectCustomVersion
 
 		FixedNoPinsInRerouteNodes,
 
+		TableNoneOptionsMovedToUnrealCode,
+
 		// -----<new versions can be added above this line>-------------------------------------------------
 		VersionPlusOne,
 		LatestVersion = VersionPlusOne - 1
diff --git a/Engine/Plugins/Mutable/Source/CustomizableObject/Internal/MuCO/CustomizableObjectPrivate.h b/Engine/Plugins/Mutable/Source/CustomizableObject/Internal/MuCO/CustomizableObjectPrivate.h
--- a/Engine/Plugins/Mutable/Source/CustomizableObject/Internal/MuCO/CustomizableObjectPrivate.h
+++ b/Engine/Plugins/Mutable/Source/CustomizableObject/Internal/MuCO/CustomizableObjectPrivate.h
@@ -1816,6 +1816,8 @@ public:
 
 		ImageHiResLODsUseLODGroupInfo,
 
+		MovedTableRowNoneGenerationToUnreal,
+
 		// -----<new versions can be added above this line>--------
 		LastCustomizableObjectVersion
 	};
diff --git a/Engine/Plugins/Mutable/Source/CustomizableObject/Private/MuCO/CustomizableObjectInstance.cpp b/Engine/Plugins/Mutable/Source/CustomizableObject/Private/MuCO/CustomizableObjectInstance.cpp
--- a/Engine/Plugins/Mutable/Source/CustomizableObject/Private/MuCO/CustomizableObjectInstance.cpp
+++ b/Engine/Plugins/Mutable/Source/CustomizableObject/Private/MuCO/CustomizableObjectInstance.cpp
@@ -6918,7 +6918,8 @@ void UCustomizableInstancePrivate::BuildMaterials(const TSharedRef<FUpdateContex
 								{
 									FLinearColor Color = Param.Vector;
 
-									if (FVector4f(Color) == mu::DefaultMutableColorValue)
+									// Decoding "None" color parameters that use the material color
+									if (FMath::IsNaN(Color.R))
 									{
 										FMaterialParameterInfo ParameterInfo(Param.ParamName);
 										MaterialTemplate->GetVectorParameterValue(ParameterInfo, Color);
diff --git a/Engine/Plugins/Mutable/Source/CustomizableObjectEditor/Private/MuCOE/GenerateMutableSource/GenerateMutableSource.h b/Engine/Plugins/Mutable/Source/CustomizableObjectEditor/Private/MuCOE/GenerateMutableSource/GenerateMutableSource.h
--- a/Engine/Plugins/Mutable/Source/CustomizableObjectEditor/Private/MuCOE/GenerateMutableSource/GenerateMutableSource.h
+++ b/Engine/Plugins/Mutable/Source/CustomizableObjectEditor/Private/MuCOE/GenerateMutableSource/GenerateMutableSource.h
@@ -912,6 +912,9 @@ public:
 	/** Version Bridge of the root object */
 	TObjectPtr<UObject> RootVersionBridge;
 
+	/** Index of the Referenced Material being generated */
+	int32 CurrentReferencedMaterialIndex = -1;
+
 private:
 
 	/** Two operation modes:
diff --git a/Engine/Plugins/Mutable/Source/CustomizableObjectEditor/Private/MuCOE/GenerateMutableSource/GenerateMutableSourceSurface.cpp b/Engine/Plugins/Mutable/Source/CustomizableObjectEditor/Private/MuCOE/GenerateMutableSource/GenerateMutableSourceSurface.cpp
--- a/Engine/Plugins/Mutable/Source/CustomizableObjectEditor/Private/MuCOE/GenerateMutableSource/GenerateMutableSourceSurface.cpp
+++ b/Engine/Plugins/Mutable/Source/CustomizableObjectEditor/Private/MuCOE/GenerateMutableSource/GenerateMutableSourceSurface.cpp
@@ -267,13 +267,12 @@ mu::Ptr<mu::NodeSurface> GenerateMutableSourceSurface(const UEdGraphPin * Pin, F
 		FMutableGraphGenerationContext::FSharedSurface& SharedSurface = SharedSurfaces.Add_GetRef(FMutableGraphGenerationContext::FSharedSurface(GenerationContext.CurrentLOD, SurfNode));
 		SharedSurface.bMakeUnique = !TypedNodeMat->IsReuseMaterialBetweenLODs();
 
-		int32 ReferencedMaterialsIndex = -1;
 		uint32 SurfaceMetadataUniqueHash = 0; // Value 0 is used as invalid hash.
 		if (TypedNodeMat->GetMaterial())
 		{
 			GenerationContext.AddParticipatingObject(*TypedNodeMat->GetMaterial());
+			GenerationContext.CurrentReferencedMaterialIndex = GenerationContext.ReferencedMaterials.AddUnique(TypedNodeMat->GetMaterial());
 
-			ReferencedMaterialsIndex = GenerationContext.ReferencedMaterials.AddUnique(TypedNodeMat->GetMaterial());
 			if (ConnectedMaterialPin)
 			{
 				if (const UEdGraphPin* SkeletalMeshPin = FindMeshBaseSource(*ConnectedMaterialPin, false))
@@ -872,6 +871,9 @@ mu::Ptr<mu::NodeSurface> GenerateMutableSourceSurface(const UEdGraphPin * Pin, F
 			FString VectorName = TypedNodeMat->GetParameterName(EMaterialParameterType::Vector, VectorIndex).ToString();
 			FNodeMaterialParameterId VectorId = TypedNodeMat->GetParameterId(EMaterialParameterType::Vector, VectorIndex);
 
+			GenerationContext.CurrentMaterialTableParameter = VectorName;
+			GenerationContext.CurrentMaterialTableParameterId = VectorId.ParameterId.ToString();
+
 			if (bVectorPinConnected)
 			{				
 				if (const UEdGraphPin* ConnectedPin = FollowInputPin(*VectorPin))
@@ -900,6 +902,9 @@ mu::Ptr<mu::NodeSurface> GenerateMutableSourceSurface(const UEdGraphPin * Pin, F
 			FString ScalarName = TypedNodeMat->GetParameterName(EMaterialParameterType::Scalar, ScalarIndex).ToString();
 			FNodeMaterialParameterId ScalarId = TypedNodeMat->GetParameterId(EMaterialParameterType::Scalar, ScalarIndex);
 
+			GenerationContext.CurrentMaterialTableParameter = ScalarName;
+			GenerationContext.CurrentMaterialTableParameterId = ScalarId.ParameterId.ToString();
+
 			if (bScalarPinConnected)
 			{
 				if (const UEdGraphPin* ConnectedPin = FollowInputPin(*ScalarPin))
@@ -934,7 +939,7 @@ mu::Ptr<mu::NodeSurface> GenerateMutableSourceSurface(const UEdGraphPin * Pin, F
 				if (const UEdGraphPin* ConnectedPin = FollowInputPin(*MaterialPin))
 				{
 					// Adding pin ID to support multiple material columns
-					MaterialName += ConnectedPin->PinId.ToString();
+					MaterialName += MaterialPin->PinId.ToString();
 
 					GenerationContext.CurrentMaterialTableParameterId = MaterialName;
 					mu::NodeScalarPtr ScalarNode = GenerateMutableSourceFloat(ConnectedPin, GenerationContext);
@@ -946,7 +951,7 @@ mu::Ptr<mu::NodeSurface> GenerateMutableSourceSurface(const UEdGraphPin * Pin, F
 			else
 			{
 				mu::NodeScalarConstantPtr ScalarNode = new mu::NodeScalarConstant();
-				ScalarNode->SetValue(ReferencedMaterialsIndex);
+				ScalarNode->SetValue(GenerationContext.CurrentReferencedMaterialIndex);
 
 				SurfNode->Scalars[MaterialIndex].Scalar = ScalarNode;
 				SurfNode->Scalars[MaterialIndex].Name = MaterialName;
diff --git a/Engine/Plugins/Mutable/Source/CustomizableObjectEditor/Private/MuCOE/GenerateMutableSource/GenerateMutableSourceTable.cpp b/Engine/Plugins/Mutable/Source/CustomizableObjectEditor/Private/MuCOE/GenerateMutableSource/GenerateMutableSourceTable.cpp
--- a/Engine/Plugins/Mutable/Source/CustomizableObjectEditor/Private/MuCOE/GenerateMutableSource/GenerateMutableSourceTable.cpp
+++ b/Engine/Plugins/Mutable/Source/CustomizableObjectEditor/Private/MuCOE/GenerateMutableSource/GenerateMutableSourceTable.cpp
@@ -17,6 +17,7 @@
 #include "MuCO/CustomizableObjectUIData.h"
 #include "MuCO/UnrealPortabilityHelpers.h"
 #include "MuCOE/CustomizableObjectCompiler.h"
+#include "MuCOE/EdGraphSchema_CustomizableObject.h"
 #include "MuCOE/GenerateMutableSource/GenerateMutableSourceMesh.h"
 #include "MuCOE/Nodes/CustomizableObjectNodeTable.h"
 #include "MuCOE/Nodes/CustomizableObjectNodeAnimationPose.h"
@@ -68,10 +69,9 @@ bool FillTableColumn(const UCustomizableObjectNodeTable* TableNode, mu::Ptr<mu::
 				return true;
 			}
 
-			// Getting Animation Blueprint and Animation Slot
+			// Getting Animation information (Anim Blueprints, Animation Slot and Anim Tags)
 			FString AnimBP, AnimSlot, GameplayTag, AnimBPAssetTag;
 			TArray<FGameplayTag> GameplayTags;
-
 			TableNode->GetAnimationColumns(ColumnName, AnimBP, AnimSlot, GameplayTag);
 
 			if (!AnimBP.IsEmpty())
@@ -676,7 +676,7 @@ bool FillTableColumn(const UCustomizableObjectNodeTable* TableNode, mu::Ptr<mu::
 
 	else
 	{
-		// Unsuported Variable Type		
+		// Unsuported Variable Type
 		return false;
 	}
 
@@ -964,13 +964,29 @@ bool GenerateTableColumn(const UCustomizableObjectNodeTable* TableNode, const UE
 		}
 	}
 
-	for (int32 RowIndex = 0; RowIndex < RowNames.Num(); ++RowIndex)
+	// Variable to check if something failed during a Cell generation
+	bool bCellGenerated = true;
+	const UEdGraphSchema_CustomizableObject* Schema = GetDefault<UEdGraphSchema_CustomizableObject>();
+	check(Schema);
+
+	// At some point we may want more default values types
+	TSet<FName> ValidNonePinTypes({ Schema->PC_Color, Schema->PC_MaterialAsset });
+	
+	// Only set the "None" value of Material and Color columns
+	if (ValidNonePinTypes.Contains(Pin->PinType.PinCategory))
+	{
+		bCellGenerated = GenerateNoneRow(*TableNode, Pin, DataTableColumnName, MutableTable, GenerationContext);
+	}
+
+	// Set the value of each row
+	for (int32 RowIndex = 0; RowIndex < RowNames.Num() && bCellGenerated; ++RowIndex)
 	{
 		if (uint8* CellData = GetCellData(RowNames[RowIndex], *DataTable, *ColumnProperty))
 		{
-			bool bCellGenerated = FillTableColumn(TableNode, MutableTable, DataTableColumnName, RowNames[RowIndex].ToString(), RowIds[RowIndex], CellData, ColumnProperty,
+			bCellGenerated = FillTableColumn(TableNode, MutableTable, DataTableColumnName, RowNames[RowIndex].ToString(), RowIds[RowIndex], CellData, ColumnProperty,
 				LODIndexConnected, SectionIndexConnected, LODIndex, SectionIndex, SectionMetadataId, bOnlyConnectedLOD, GenerationContext);
 
+			// Stop the compilation if something fails
 			if (!bCellGenerated)
 			{
 				return false;
@@ -1130,9 +1146,18 @@ mu::Ptr<mu::FTable> GenerateMutableSourceTable(const UDataTable* DataTable, cons
 		TArray<uint32> RowIds;
 		TArray<FName> RowNames = GetRowsToCompile(*DataTable, *TableNode, GenerationContext, RowIds);
 
-		// Adding and filling Name Column
+		// Adding the Name Column
 		MutableTable->AddColumn("Name", mu::ETableColumnType::String);
 
+		// Always generate "None" row
+		{
+			FString NoneName = "None";
+			uint32 RowId = CityHash32(reinterpret_cast<const char*>(*NoneName), NoneName.Len() * sizeof(FString::ElementType));
+			MutableTable->AddRow(RowId);
+			MutableTable->SetCell(0, RowId, NoneName);
+		}
+
+		// Adding name rows
 		for (int32 RowIndex = 0; RowIndex < RowNames.Num(); ++RowIndex)
 		{
 			MutableTable->AddRow(RowIds[RowIndex]);
@@ -1142,6 +1167,7 @@ mu::Ptr<mu::FTable> GenerateMutableSourceTable(const UDataTable* DataTable, cons
 		// Generating Parameter Metadata for new table parameters
 		GenerateTableParameterUIData(DataTable, TableNode, GenerationContext);
 
+		// Generating data for Table Cache
 		FMutableGraphGenerationContext::FGeneratedDataTablesData GeneratedTable;
 		GeneratedTable.GeneratedTable = MutableTable;
 		GeneratedTable.bDisableCheckedRows = TableNode->bDisableCheckedRows;
@@ -1351,4 +1377,60 @@ void LogRowGenerationMessage(const UCustomizableObjectNodeTable* TableNode, cons
 }
 
 
+bool GenerateNoneRow(const UCustomizableObjectNodeTable& TableNode, const UEdGraphPin* Pin, const FString& ColumnName, mu::Ptr<mu::FTable> MutableTable, FMutableGraphGenerationContext& GenerationContext)
+{
+	if (!Pin || !MutableTable)
+	{
+		FText Message = FText::Format(LOCTEXT("NoneRowWarning_NullPinTable","Error creating [None] row for Table Node {0}. Null pin or Table."), FText::FromString(TableNode.ParameterName));
+		GenerationContext.Log(Message, &TableNode, EMessageSeverity::Error);
+
+		return false;
+	}
+
+	FString NoneName = "None";
+	uint32 RowId = CityHash32(reinterpret_cast<const char*>(*NoneName), NoneName.Len() * sizeof(FString::ElementType));
+	int32 ColumnId = MutableTable->FindColumn(ColumnName);
+	
+	const UEdGraphSchema_CustomizableObject* Schema = GetDefault<UEdGraphSchema_CustomizableObject>();
+	check(Schema);
+
+	//Ensure that we are generating a material column and not a Texture column
+	FString EncodedSwitchParameterName = "__MutableMaterialId";
+
+	if (Pin->PinType.PinCategory == Schema->PC_MaterialAsset && ColumnName.Contains(EncodedSwitchParameterName))
+	{
+		// Add a scalar column if it has not been added yet
+		if (ColumnId == INDEX_NONE)
+		{
+			ColumnId = MutableTable->AddColumn(ColumnName, mu::ETableColumnType::Scalar);
+		}
+
+		// Get the material of the linked mesh section
+		if (!GenerationContext.ReferencedMaterials.IsValidIndex(GenerationContext.CurrentReferencedMaterialIndex))
+		{
+			FText Message = FText::Format(LOCTEXT("NoneRowWarning_InvalidMaterialIndex", "Error creating [None] row for Table Node {0}. Invalid Material Index."), FText::FromString(TableNode.ParameterName));
+			GenerationContext.Log(Message, &TableNode, EMessageSeverity::Error);
+
+			return false;
+		}
+
+		MutableTable->SetCell(ColumnId, RowId, GenerationContext.CurrentReferencedMaterialIndex);
+	}
+	else if (Pin->PinType.PinCategory == Schema->PC_Color)
+	{
+		// Add a color column if it has not been added yet
+		if (ColumnId == INDEX_NONE)
+		{
+			ColumnId = MutableTable->AddColumn(ColumnName, mu::ETableColumnType::Color);
+		}
+
+		// Encoding an invalid value. QuietNaN does not trigger errors nor checks. 
+		// It's checked in the moment that the material of the COI is generated.
+		FLinearColor InvalidColor(std::numeric_limits<float>::quiet_NaN(), 0.0f, 0.0f);
+		MutableTable->SetCell(ColumnId, RowId, TableNode.bUseMaterialColor ? InvalidColor : FLinearColor::Black);
+	}
+
+	return true;
+}
+
 #undef LOCTEXT_NAMESPACE
diff --git a/Engine/Plugins/Mutable/Source/CustomizableObjectEditor/Private/MuCOE/GenerateMutableSource/GenerateMutableSourceTable.h b/Engine/Plugins/Mutable/Source/CustomizableObjectEditor/Private/MuCOE/GenerateMutableSource/GenerateMutableSourceTable.h
--- a/Engine/Plugins/Mutable/Source/CustomizableObjectEditor/Private/MuCOE/GenerateMutableSource/GenerateMutableSourceTable.h
+++ b/Engine/Plugins/Mutable/Source/CustomizableObjectEditor/Private/MuCOE/GenerateMutableSource/GenerateMutableSourceTable.h
@@ -73,6 +73,8 @@ TArray<FName> GetRowsToCompile(const UDataTable& DataTable, const UCustomizableO
 /** Generates a Data Table from the Data Tables referenced in a Script Struct */
 UDataTable* GenerateDataTableFromStruct(const UCustomizableObjectNodeTable* TableNode, FMutableGraphGenerationContext& GenerationContext);
 
+/** Generates and populates the "None" row for each column. */
+bool GenerateNoneRow(const UCustomizableObjectNodeTable& TableNode, const UEdGraphPin* Pin, const FString& ColumnName, mu::Ptr<mu::FTable> MutableTable, FMutableGraphGenerationContext& GenerationContext);
+
 /** This method adds the original Data Table(s) of the processed row at the end of the log message when the data table is a Composite Data table. */
 void LogRowGenerationMessage(const UCustomizableObjectNodeTable* TableNode, const UDataTable* DataTable, FMutableGraphGenerationContext& GenerationContext, const FString& Message, const FString& RowName);
-
diff --git a/Engine/Plugins/Mutable/Source/CustomizableObjectEditor/Private/MuCOE/Nodes/CustomizableObjectNodeTable.cpp b/Engine/Plugins/Mutable/Source/CustomizableObjectEditor/Private/MuCOE/Nodes/CustomizableObjectNodeTable.cpp
--- a/Engine/Plugins/Mutable/Source/CustomizableObjectEditor/Private/MuCOE/Nodes/CustomizableObjectNodeTable.cpp
+++ b/Engine/Plugins/Mutable/Source/CustomizableObjectEditor/Private/MuCOE/Nodes/CustomizableObjectNodeTable.cpp
@@ -256,6 +256,12 @@ void UCustomizableObjectNodeTable::BackwardsCompatibleFixup(int32 CustomizableOb
 			}
 		}
 	}
+	
+	if (CustomizableObjectCustomVersion == FCustomizableObjectCustomVersion::TableNoneOptionsMovedToUnrealCode)
+	{
+		// set this option to true for already generated pins with the none option.
+		bUseMaterialColor = bAddNoneOption;
+	}
 }
 
 
diff --git a/Engine/Plugins/Mutable/Source/CustomizableObjectEditor/Private/MuCOE/Nodes/CustomizableObjectNodeTableDetails.cpp b/Engine/Plugins/Mutable/Source/CustomizableObjectEditor/Private/MuCOE/Nodes/CustomizableObjectNodeTableDetails.cpp
--- a/Engine/Plugins/Mutable/Source/CustomizableObjectEditor/Private/MuCOE/Nodes/CustomizableObjectNodeTableDetails.cpp
+++ b/Engine/Plugins/Mutable/Source/CustomizableObjectEditor/Private/MuCOE/Nodes/CustomizableObjectNodeTableDetails.cpp
@@ -67,6 +67,12 @@ void FCustomizableObjectNodeTableDetails::CustomizeDetails(const TSharedPtr<IDet
 		TSharedPtr<FString> CurrentThumbnailColumn = GenerateThumbnailColumnComboBoxOptions();
 
 		CustomizableObjectCategory.AddProperty("ParameterName");
+		TSharedRef<IPropertyHandle> AddNoneOptionProperty = DetailBuilder->GetProperty("bAddNoneOption");
+		TSharedRef<IPropertyHandle> UseMaterialColorProperty = DetailBuilder->GetProperty("bUseMaterialColor");
+
+		IDetailGroup& AddNoneGroup = CustomizableObjectCategory.AddGroup(TEXT("TableNode_NoneOptionGroup"), LOCTEXT("TableNode_NoneOptionGroup", "Add None Option"), false, true);
+		AddNoneGroup.HeaderProperty(AddNoneOptionProperty);
+		AddNoneGroup.AddPropertyRow(UseMaterialColorProperty);
 
 		CompilationRestrictionsCategory.AddCustomRow(LOCTEXT("VersionColumn_Selector", "VersionColumn"))
 			.NameContent()
diff --git a/Engine/Plugins/Mutable/Source/CustomizableObjectEditor/Public/MuCOE/Nodes/CustomizableObjectNodeTable.h b/Engine/Plugins/Mutable/Source/CustomizableObjectEditor/Public/MuCOE/Nodes/CustomizableObjectNodeTable.h
--- a/Engine/Plugins/Mutable/Source/CustomizableObjectEditor/Public/MuCOE/Nodes/CustomizableObjectNodeTable.h
+++ b/Engine/Plugins/Mutable/Source/CustomizableObjectEditor/Public/MuCOE/Nodes/CustomizableObjectNodeTable.h
@@ -243,6 +243,12 @@ public:
 	UPROPERTY()
 	TMap<FString, FTableNodeColumnData> PinColumnDataMap;
 
+	/** If true, the "None" colors will use the color of the material parameter.
+		If false, the "None" colors will be black.
+		Note: if this option is true and the colors are used to generate textures, "None" option colors will be set to black.*/
+	UPROPERTY(EditAnywhere, Category = TableProperties, meta = (DisplayAfter = "bAddNoneOption", EditCondition = "bAddNoneOption==true", EditConditionHides))
+	bool bUseMaterialColor = false;
+
 public:
 
 	// UObject interface
diff --git a/Engine/Plugins/Mutable/Source/MutableRuntime/Public/MuR/Parameters.h b/Engine/Plugins/Mutable/Source/MutableRuntime/Public/MuR/Parameters.h
--- a/Engine/Plugins/Mutable/Source/MutableRuntime/Public/MuR/Parameters.h
+++ b/Engine/Plugins/Mutable/Source/MutableRuntime/Public/MuR/Parameters.h
@@ -18,9 +18,6 @@ namespace mu
     class FRangeIndex;
 	class FModel;
 
-	//! Default parameter values definition
-	static const FVector4f DefaultMutableColorValue(-UE_BIG_NUMBER);
-
 
     /** Model parameter types. */
     enum class EParameterType : uint32
diff --git a/Engine/Plugins/Mutable/Source/MutableTools/Private/MuT/CodeGenerator.cpp b/Engine/Plugins/Mutable/Source/MutableTools/Private/MuT/CodeGenerator.cpp
--- a/Engine/Plugins/Mutable/Source/MutableTools/Private/MuT/CodeGenerator.cpp
+++ b/Engine/Plugins/Mutable/Source/MutableTools/Private/MuT/CodeGenerator.cpp
@@ -310,75 +310,60 @@ namespace mu
 
 	Ptr<ASTOp> CodeGenerator::GenerateTableVariable( Ptr<const Node> InNode, const FTableCacheKey& CacheKey, bool bAddNoneOption, const FString& DefaultRowName)
 	{
-		Ptr<ASTOp> result;
+        FParameterDesc Param;
+		Param.Name = CacheKey.ParameterName;
 
-        FParameterDesc param;
-        param.Name = CacheKey.ParameterName;
-
-        if (param.Name.Len() == 0)
+        if (Param.Name.Len() == 0)
         {
-            param.Name = CacheKey.Table->GetName();
+			Param.Name = CacheKey.Table->GetName();
         }
         
-		param.Type = EParameterType::Int;
-        param.DefaultValue.Set<FParamIntType>(0);
+		Param.Type = EParameterType::Int;
+		Param.DefaultValue.Set<FParamIntType>(0);
 
 		// Add the possible values
 		{
 			// See if there is a string column. If there is one, we will use it as names for the
 			// options. Only the first string column will be used.
-			int32 nameCol = -1;
-			int32 cols = CacheKey.Table->GetPrivate()->Columns.Num();
-			for (int32 c = 0; c < cols && nameCol < 0; ++c)
+			int32 NameColumn = -1;
+			int32 Columns = CacheKey.Table->GetPrivate()->Columns.Num();
+			for (int32 ColumnIndex = 0; ColumnIndex < Columns && NameColumn < 0; ++ColumnIndex)
 			{
-				if (CacheKey.Table->GetPrivate()->Columns[c].Type == ETableColumnType::String)
+				if (CacheKey.Table->GetPrivate()->Columns[ColumnIndex].Type == ETableColumnType::String)
 				{
-					nameCol = c;
+					NameColumn = ColumnIndex;
 				}
 			}
 
-			if (bAddNoneOption)
-			{
-				FParameterDesc::FIntValueDesc nullValue;
-				nullValue.Value = -1;
-				nullValue.Name = "None";
-				param.PossibleValues.Add(nullValue);
-				param.DefaultValue.Set<FParamIntType>(nullValue.Value);
-			}
-
-			// Add every row
 			int32 RowCount = CacheKey.Table->GetPrivate()->Rows.Num();
 			check(RowCount < MAX_int16) // max FIntValueDesc allows
 
-			for (int32 RowIndex = 0; RowIndex < RowCount; ++RowIndex)
+			// Skip "None" option (first row) if it's not required
+			int32 StarCount = bAddNoneOption ? 0 : 1;
+
+			for (int32 RowIndex = StarCount; RowIndex < RowCount; ++RowIndex)
 			{
-				FParameterDesc::FIntValueDesc value;
-				value.Value = RowIndex;
+				FParameterDesc::FIntValueDesc Value;
+				Value.Value = RowIndex;
 
-				if (nameCol > -1)
+				if (NameColumn > -1)
 				{
-					value.Name = CacheKey.Table->GetPrivate()->Rows[RowIndex].Values[nameCol].String;
+					Value.Name = CacheKey.Table->GetPrivate()->Rows[RowIndex].Values[NameColumn].String;
 				}
 
-				param.PossibleValues.Add(value);
+				Param.PossibleValues.Add(Value);
 
-				// Set the first row as the default one (if there is none option)
-				if (RowIndex == 0 && !bAddNoneOption)
+				// Set the first row or the selected row as the default one.
+				if(RowIndex == StarCount || Value.Name == DefaultRowName)
 				{
-					param.DefaultValue.Set<FParamIntType>(value.Value);
-				}
-				
-				// Set the selected row as default (if exists)
-				if (value.Name == DefaultRowName)
-				{
-					param.DefaultValue.Set<FParamIntType>(value.Value);
+					Param.DefaultValue.Set<FParamIntType>(Value.Value);
 				}
             }
         }
 
 		Ptr<ASTOpParameter> op = new ASTOpParameter();
 		op->Type = EOpType::NU_PARAMETER;
-		op->Parameter = param;
+		op->Parameter = Param;
 
 		FirstPass.ParameterNodes.Add(InNode, op);
 
@@ -2873,53 +2858,4 @@ namespace mu
 			}
 		}
 	}
-
-
-	Ptr<ASTOp> CodeGenerator::GenerateDefaultTableValue(ETableColumnType NodeType)
-	{
-		switch (NodeType)
-		{
-		case mu::ETableColumnType::Scalar:
-		{
-			//TODO(Max):MTBL-1660
-			//Ptr<NodeScalarConstant> pNode = new NodeScalarConstant();
-			//pNode->SetValue(-UE_MAX_FLT);
-			//
-			//return Generate(pNode);
-			return nullptr;
-		}
-		case mu::ETableColumnType::Color:
-		{
-			mu::Ptr<mu::NodeColourConstant> pNode = new NodeColourConstant();
-			pNode->Value = mu::DefaultMutableColorValue;
-
-			FColorGenerationResult ChildResult;
-			FGenericGenerationOptions Options;
-			GenerateColor(ChildResult, Options, pNode);
-			return ChildResult.op;
-		}
-		case mu::ETableColumnType::Image:
-		{
-			//TODO(Max):MTBL-1660
-			//FImageGenerationOptions DummyOptions;
-			//FImageGenerationResult DefaultValue;
-			//
-			//mu::Ptr<mu::NodeImageReference> ImageNode = new mu::NodeImageReference();
-			//ImageNode->SetImageReference(-1);
-			//
-			//GenerateImage(DummyOptions, DefaultValue, ImageNode);
-			//
-			//return DefaultValue.op;
-			return nullptr;
-		}
-		case mu::ETableColumnType::Mesh:
-			/* The default mesh is always null */
-			return nullptr;
-		default:
-			break;
-		}
-
-		return nullptr;
-	}
-
 }
diff --git a/Engine/Plugins/Mutable/Source/MutableTools/Private/MuT/CodeGenerator.h b/Engine/Plugins/Mutable/Source/MutableTools/Private/MuT/CodeGenerator.h
--- a/Engine/Plugins/Mutable/Source/MutableTools/Private/MuT/CodeGenerator.h
+++ b/Engine/Plugins/Mutable/Source/MutableTools/Private/MuT/CodeGenerator.h
@@ -140,6 +140,9 @@ namespace mu
 
 			int32 State = -1;
 			TArray<FString> ActiveTags;
+
+			// Variable to identify if an operation will end up generating an image.
+			bool bIsImage = false;
 		};
 
 		struct FComponentGenerationOptions : public FGenericGenerationOptions
@@ -382,6 +385,7 @@ namespace mu
 			FImageGenerationOptions(int32 InComponentId)
 				: ComponentId(InComponentId)
 			{
+				bIsImage = true;
 			}
 
 			/** The id of the component that we are currently generating. */
@@ -804,11 +808,6 @@ namespace mu
 
         void GenerateSurface( FSurfaceGenerationResult&, const FSurfaceGenerationOptions&, Ptr<const NodeSurfaceNew> );
 
-		//-----------------------------------------------------------------------------------------
-		// Default Table FParameters
-		Ptr<ASTOp> GenerateDefaultTableValue(ETableColumnType NodeType);
-
-
 		//-----------------------------------------------------------------------------------------
 
 		struct FLayoutBlockDesc
@@ -913,7 +912,7 @@ namespace mu
         Ptr<ASTOpSwitch> SwitchOp = new ASTOpSwitch();
 		SwitchOp->type = OPTYPE;
 		SwitchOp->variable = Variable;
-		SwitchOp->def = GenerateDefaultTableValue(TYPE);
+		SwitchOp->def = nullptr;
 
 		for (int32 RowIndex = 0; RowIndex < NumRows; ++RowIndex)
         {
diff --git a/Engine/Plugins/Mutable/Source/MutableTools/Private/MuT/CodeGenerator_Color.cpp b/Engine/Plugins/Mutable/Source/MutableTools/Private/MuT/CodeGenerator_Color.cpp
--- a/Engine/Plugins/Mutable/Source/MutableTools/Private/MuT/CodeGenerator_Color.cpp
+++ b/Engine/Plugins/Mutable/Source/MutableTools/Private/MuT/CodeGenerator_Color.cpp
@@ -492,6 +492,13 @@ namespace mu
 			{
 				Ptr<NodeColourConstant> CellData = new NodeColourConstant();
 				FVector4f Colour = node.Table->GetPrivate()->Rows[row].Values[colIndex].Color;
+
+				// Break the NaN codification for colors that will end up generating an image
+				if (Options.bIsImage && FMath::IsNaN(node.Table->GetPrivate()->Rows[row].Values[colIndex].Color[0]))
+				{
+					Colour = FVector4f(0.0f, 0.0f, 0.0f, 1.0f);
+				}
+
 				CellData->Value = Colour;
 				FColorGenerationResult BranchResults;
 				GenerateColor(BranchResults, Options, CellData );
diff --git a/Engine/Plugins/Mutable/Source/MutableTools/Private/MuT/CodeGenerator_Image.cpp b/Engine/Plugins/Mutable/Source/MutableTools/Private/MuT/CodeGenerator_Image.cpp
--- a/Engine/Plugins/Mutable/Source/MutableTools/Private/MuT/CodeGenerator_Image.cpp
+++ b/Engine/Plugins/Mutable/Source/MutableTools/Private/MuT/CodeGenerator_Image.cpp
@@ -1901,8 +1901,11 @@ namespace mu
 
 				if (!pImage)
 				{
-					FString Msg = FString::Printf(TEXT("Table has a missing image in column %d, row %d."), colIndex, row);
-					pErrorLog->GetPrivate()->Add(Msg, ELMT_ERROR, InNode->GetMessageContext());
+					if (row != 0) // "None" option (row 0) is always a null image, do not trigger this error.
+					{
+						FString Msg = FString::Printf(TEXT("Table has a missing image in column %d, row %d."), colIndex, row);
+						pErrorLog->GetPrivate()->Add(Msg, ELMT_ERROR, InNode->GetMessageContext());
+					}
 
 					return ImageOp;
 				}
