[Mutable] Fixed name collision.

diff --git a/Engine/Plugins/Mutable/Source/CustomizableObjectEditor/Private/MuCOE/GenerateMutableSource/GenerateMutableSourceTable.cpp b/Engine/Plugins/Mutable/Source/CustomizableObjectEditor/Private/MuCOE/GenerateMutableSource/GenerateMutableSourceTable.cpp
--- a/Engine/Plugins/Mutable/Source/CustomizableObjectEditor/Private/MuCOE/GenerateMutableSource/GenerateMutableSourceTable.cpp
+++ b/Engine/Plugins/Mutable/Source/CustomizableObjectEditor/Private/MuCOE/GenerateMutableSource/GenerateMutableSourceTable.cpp
@@ -1151,10 +1151,10 @@ mu::Ptr<mu::FTable> GenerateMutableSourceTable(const UDataTable* DataTable, cons
 
 		// Always generate "None" row
 		{
-			FString NoneName = "None";
-			uint32 RowId = CityHash32(reinterpret_cast<const char*>(*NoneName), NoneName.Len() * sizeof(FString::ElementType));
+			FString NoneRowName = "None";
+			uint32 RowId = CityHash32(reinterpret_cast<const char*>(*NoneRowName), NoneRowName.Len() * sizeof(FString::ElementType));
 			MutableTable->AddRow(RowId);
-			MutableTable->SetCell(0, RowId, NoneName);
+			MutableTable->SetCell(0, RowId, NoneRowName);
 		}
 
 		// Adding name rows
@@ -1387,8 +1387,8 @@ bool GenerateNoneRow(const UCustomizableObjectNodeTable& TableNode, const UEdGra
 		return false;
 	}
 
-	FString NoneName = "None";
-	uint32 RowId = CityHash32(reinterpret_cast<const char*>(*NoneName), NoneName.Len() * sizeof(FString::ElementType));
+	FString NoneRowName = "None";
+	uint32 RowId = CityHash32(reinterpret_cast<const char*>(*NoneRowName), NoneRowName.Len() * sizeof(FString::ElementType));
 	int32 ColumnId = MutableTable->FindColumn(ColumnName);
 	
 	const UEdGraphSchema_CustomizableObject* Schema = GetDefault<UEdGraphSchema_CustomizableObject>();
