[Mutable] Fix crash when executing consecutive Game Thread tasks.

#rnx

diff --git a/Engine/Plugins/Mutable/Source/CustomizableObject/Private/MuCO/FMutableTaskGraph.cpp b/Engine/Plugins/Mutable/Source/CustomizableObject/Private/MuCO/FMutableTaskGraph.cpp
--- a/Engine/Plugins/Mutable/Source/CustomizableObject/Private/MuCO/FMutableTaskGraph.cpp
+++ b/Engine/Plugins/Mutable/Source/CustomizableObject/Private/MuCO/FMutableTaskGraph.cpp
@@ -282,35 +282,39 @@ void FMutableTaskGraph::TryLaunchGameThreadTask()
 {
 	MUTABLE_CPUPROFILER_SCOPE(AdvanceCurrentOperation);
 
-	FGameThreadTask* PendingTask = GameThreadTasks.Peek();
-	if (!PendingTask)
+	// Scope where PendingTask is valid.
 	{
-		return;
-	}
-
-	if (!PendingTask->AreDependenciesComplete())
-	{
-		return;
-	}
-	
-	if (PendingTask->bLockMutableThread)
-	{
-		AsyncLockMutableThread();
+		FGameThreadTask* PendingTask = GameThreadTasks.Peek();
+		if (!PendingTask)
+		{
+			return;
+		}
 
-		if (!IsMutableThreadLocked())
+		if (!PendingTask->AreDependenciesComplete())
 		{
 			return;
 		}
-	}
+	
+		if (PendingTask->bLockMutableThread)
+		{
+			AsyncLockMutableThread();
 
-	PendingTask->Body();
+			if (!IsMutableThreadLocked())
+			{
+				return;
+			}
+		}
+	}
+	
+	FGameThreadTask CurrentTask;
+	GameThreadTasks.Dequeue(CurrentTask);
 	
-	if (PendingTask->bLockMutableThread)
+	CurrentTask.Body();
+	
+	if (CurrentTask.bLockMutableThread)
 	{
 		UnlockMutableThread();
 	}
-
-	GameThreadTasks.Pop();
 }
 
 
