[Mutable] Hidden the Customizable Skeletal Component property "Component Index".
+ Fixed problem where the component index was set in the wrong function.

#rnx

diff --git a/Engine/Plugins/Mutable/Source/CustomizableObject/Private/MuCO/CustomizableSkeletalComponent.cpp b/Engine/Plugins/Mutable/Source/CustomizableObject/Private/MuCO/CustomizableSkeletalComponent.cpp
--- a/Engine/Plugins/Mutable/Source/CustomizableObject/Private/MuCO/CustomizableSkeletalComponent.cpp
+++ b/Engine/Plugins/Mutable/Source/CustomizableObject/Private/MuCO/CustomizableSkeletalComponent.cpp
@@ -17,26 +17,6 @@ UCustomizableSkeletalComponent::UCustomizableSkeletalComponent()
 }
 
 
-#if WITH_EDITOR
-void UCustomizableSkeletalComponent::PostEditChangeProperty(FPropertyChangedEvent& PropertyChangedEvent)
-{
-	Super::PostEditChangeProperty(PropertyChangedEvent);
-
-	if (!PropertyChangedEvent.MemberProperty)
-	{
-		return;
-	}
-	
-	const FName PropertyName = PropertyChangedEvent.MemberProperty->GetFName();
-
-	if (PropertyName == GET_MEMBER_NAME_CHECKED(UCustomizableSkeletalComponent, ComponentName))
-	{
-		SetComponentName(ComponentName);
-	}
-}
-#endif
-
-
 UCustomizableSkeletalComponentPrivate::UCustomizableSkeletalComponentPrivate()
 {
 	InstanceUsage = CreateDefaultSubobject<UCustomizableObjectInstanceUsage>(TEXT("InstanceUsage"));
@@ -103,14 +83,13 @@ USkeletalMesh* UCustomizableSkeletalComponentPrivate::GetAttachedSkeletalMesh()
 
 void UCustomizableSkeletalComponent::SetComponentName(const FName& Name)
 {
-	ComponentIndex = INDEX_NONE;
 	ComponentName = Name;
 }
 
 
 FName UCustomizableSkeletalComponent::GetComponentName() const
 {
-	if (ComponentIndex == INDEX_NONE)
+	if (!ComponentName.IsNone())
 	{
 		return ComponentName;	
 	}
diff --git a/Engine/Plugins/Mutable/Source/CustomizableObject/Public/MuCO/CustomizableSkeletalComponent.h b/Engine/Plugins/Mutable/Source/CustomizableObject/Public/MuCO/CustomizableSkeletalComponent.h
--- a/Engine/Plugins/Mutable/Source/CustomizableObject/Public/MuCO/CustomizableSkeletalComponent.h
+++ b/Engine/Plugins/Mutable/Source/CustomizableObject/Public/MuCO/CustomizableSkeletalComponent.h
@@ -28,7 +28,7 @@ public:
 	/** This component index refers to the object list of components 
 	* DEPRECATED: Use FNames instead with Get/SetComponentName
 	*/
-	UPROPERTY(BlueprintReadWrite, EditAnywhere, Category = CustomizableSkeletalComponent)
+	UPROPERTY(BlueprintReadWrite, Category = CustomizableSkeletalComponent)
 	int32 ComponentIndex;
 
 private:
@@ -54,13 +54,10 @@ private:
 	TObjectPtr<UCustomizableSkeletalComponentPrivate> Private;
 
 public:
+	
 	UCustomizableSkeletalComponent();
 	
-	// UObject interface
 protected:
-#if WITH_EDITOR
-	virtual void PostEditChangeProperty(FPropertyChangedEvent& PropertyChangedEvent) override;
-#endif
 
 	// USceneComponent interface
 	virtual void OnAttachmentChanged() override;
