[Mutable] Allow NodeModifierExtendMeshSection to be applied to LODs from FirstLOD onwards.

#rnx

diff --git a/Engine/Plugins/Mutable/Source/CustomizableObjectEditor/Private/MuCOE/GenerateMutableSource/GenerateMutableSourceModifier.cpp b/Engine/Plugins/Mutable/Source/CustomizableObjectEditor/Private/MuCOE/GenerateMutableSource/GenerateMutableSourceModifier.cpp
--- a/Engine/Plugins/Mutable/Source/CustomizableObjectEditor/Private/MuCOE/GenerateMutableSource/GenerateMutableSourceModifier.cpp
+++ b/Engine/Plugins/Mutable/Source/CustomizableObjectEditor/Private/MuCOE/GenerateMutableSource/GenerateMutableSourceModifier.cpp
@@ -251,9 +251,9 @@ mu::Ptr<mu::NodeModifier> GenerateMutableSourceModifier(const UEdGraphPin * Pin,
 		int32 LODCount = GenerationContext.NumLODs[GenerationContext.CurrentMeshComponent];
 		SurfNode->LODs.SetNum(LODCount);
 
-		for (int32 LODIndex=0; LODIndex<LODCount; ++LODIndex)
+		for (int32 LODIndex= TypedNodeExt->FirstLOD; LODIndex<LODCount; ++LODIndex)
 		{
-			GenerationContext.FromLOD = 0;
+			GenerationContext.FromLOD = TypedNodeExt->FirstLOD;
 			GenerationContext.CurrentLOD = LODIndex;
 
 			mu::Ptr<mu::NodeMesh> AddMeshNode;
diff --git a/Engine/Plugins/Mutable/Source/CustomizableObjectEditor/Public/MuCOE/Nodes/CustomizableObjectNodeModifierExtendMeshSection.h b/Engine/Plugins/Mutable/Source/CustomizableObjectEditor/Public/MuCOE/Nodes/CustomizableObjectNodeModifierExtendMeshSection.h
--- a/Engine/Plugins/Mutable/Source/CustomizableObjectEditor/Public/MuCOE/Nodes/CustomizableObjectNodeModifierExtendMeshSection.h
+++ b/Engine/Plugins/Mutable/Source/CustomizableObjectEditor/Public/MuCOE/Nodes/CustomizableObjectNodeModifierExtendMeshSection.h
@@ -30,6 +30,10 @@ public:
 	UPROPERTY(EditAnywhere, BlueprintReadWrite, Category = EnableTags)
 	TArray<FString> Tags;
 
+	/** First LOD where the modifier is applied to. */
+	UPROPERTY(EditAnywhere, Category = LOD, meta = (ClampMin = 0, ClampMax = 255))
+	uint32 FirstLOD = 0;
+
 public:
 	
 	// Begin EdGraphNode interface
