[mutable] Updated mutable node and pin colors for a better CO graph reading
- Changed name of "Table Material" node to "Material"
- Changed "Object" named pins to the actual type of the pin (Import and Export nodes)
- Updated color of some pins (pin types)
- Removed unnecesary includes
- Fixed nodes using the default or wrong color due to incorrect type

#rnx

diff --git a/Engine/Plugins/Mutable/Source/CustomizableObject/Internal/MuCO/CustomizableObjectCustomVersion.h b/Engine/Plugins/Mutable/Source/CustomizableObject/Internal/MuCO/CustomizableObjectCustomVersion.h
--- a/Engine/Plugins/Mutable/Source/CustomizableObject/Internal/MuCO/CustomizableObjectCustomVersion.h
+++ b/Engine/Plugins/Mutable/Source/CustomizableObject/Internal/MuCO/CustomizableObjectCustomVersion.h
@@ -170,6 +170,8 @@ struct CUSTOMIZABLEOBJECT_API FCustomizableObjectCustomVersion
 
 		TableNoneOptionsMovedToUnrealCode,
 
+		UpdatedNodesPinName,	
+
 		// -----<new versions can be added above this line>-------------------------------------------------
 		VersionPlusOne,
 		LatestVersion = VersionPlusOne - 1
diff --git a/Engine/Plugins/Mutable/Source/CustomizableObjectEditor/Private/MuCOE/EdGraphSchema_CustomizableObject.cpp b/Engine/Plugins/Mutable/Source/CustomizableObjectEditor/Private/MuCOE/EdGraphSchema_CustomizableObject.cpp
--- a/Engine/Plugins/Mutable/Source/CustomizableObjectEditor/Private/MuCOE/EdGraphSchema_CustomizableObject.cpp
+++ b/Engine/Plugins/Mutable/Source/CustomizableObjectEditor/Private/MuCOE/EdGraphSchema_CustomizableObject.cpp
@@ -12,7 +12,6 @@
 #include "Materials/MaterialInterface.h"
 #include "MuCO/CustomizableObjectExtension.h"
 #include "MuCO/ICustomizableObjectModule.h"
-#include "MuCOE/CustomizableObjectEditorNodeContextCommands.h"
 #include "MuCOE/ICustomizableObjectEditor.h"
 #include "MuCOE/Nodes/CustomizableObjectNodeAnimationPose.h"
 #include "MuCOE/Nodes/CustomizableObjectNodeColorArithmeticOp.h"
@@ -32,7 +31,6 @@
 #include "MuCOE/Nodes/CustomizableObjectNodeFloatVariation.h"
 #include "MuCOE/Nodes/CustomizableObjectNodeFloatArithmeticOp.h"
 #include "MuCOE/Nodes/CustomizableObjectNodeGroupProjectorParameter.h"
-#include "MuCOE/Nodes/CustomizableObjectNodeLayoutBlocks.h"
 #include "MuCOE/Nodes/CustomizableObjectNodeMaterialVariation.h"
 #include "MuCOE/Nodes/CustomizableObjectNodeMaterialSwitch.h"
 #include "MuCOE/Nodes/CustomizableObjectNodeModifierExtendMeshSection.h"
@@ -45,7 +43,6 @@
 #include "MuCOE/Nodes/CustomizableObjectNodeModifierClipWithMesh.h"
 #include "MuCOE/Nodes/CustomizableObjectNodeModifierClipWithUVMask.h"
 #include "MuCOE/Nodes/CustomizableObjectNodeModifierTransformInMesh.h"
-#include "MuCOE/Nodes/CustomizableObjectNodeMeshGeometryOperation.h"
 #include "MuCOE/Nodes/CustomizableObjectNodeMeshMorph.h"
 #include "MuCOE/Nodes/CustomizableObjectNodeMeshMorphStackApplication.h"
 #include "MuCOE/Nodes/CustomizableObjectNodeMeshMorphStackDefinition.h"
@@ -88,7 +85,6 @@
 #include "MuCOE/Nodes/CustomizableObjectNodeReroute.h"
 #include "Toolkits/ToolkitManager.h"
 #include "PropertyEditorModule.h"
-#include "MuCO/CustomizableObjectCustomVersion.h"
 #include "MuCOE/GraphTraversal.h"
 #include "MuCOE/Nodes/CustomizableObjectNodeComponentMeshAddTo.h"
 #include "MuCOE/Nodes/CustomizableObjectNodeTransformConstant.h"
@@ -709,7 +705,7 @@ FLinearColor UEdGraphSchema_CustomizableObject::GetPinTypeColor(const FName& Typ
 {
 	if (TypeString == PC_Enum)
 	{
-		return FLinearColor(0.357667f, 0.500000f, 0.060000f, 1.000000f); // Light green
+		return FLinearColor(0.004f, 0.42f, 0.384f, 1.000000f); // Light green
 	}
 	else if (TypeString == PC_Float)
 	{
@@ -725,27 +721,23 @@ FLinearColor UEdGraphSchema_CustomizableObject::GetPinTypeColor(const FName& Typ
 	}
 	else if (TypeString == PC_Projector)
 	{
-		return FLinearColor(1.000000f, 0.500000f, 1.000000f, 0.600000f); // Light pink
+		return FLinearColor(FColorList::Aquamarine);
 	}
 	else if (TypeString == PC_GroupProjector)
 	{
-		return FLinearColor(1.000000f, 0.172585f, 0.000000f, 1.000000f); // Orange
+		return FLinearColor(FColorList::DarkTan);
 	}
 	else if (TypeString == PC_Mesh)
 	{
-		return FLinearColor(0.100000f, 0.000000f, 0.500000f, 1.000000f); // Purple
+		return FLinearColor(FColorList::MediumOrchid); 
 	}
 	else if (TypeString == PC_Layout)
 	{
 		return FLinearColor(0.500000f, 0.500000f, 0.100000f, 1.000000f); // Light yellow
 	}
-	else if (TypeString == PC_Image)
-	{
-		return FLinearColor(0.353393f, 0.454175f, 1.000000f, 1.000000f); // Grey
-	}
-	else if (TypeString == PC_PassThroughImage)
+	else if (TypeString == PC_Image || TypeString == PC_PassThroughImage)
 	{
-		return FLinearColor(0.353393f, 0.454175f, 0.353393f, 1.000000f); // Dark green
+		return FLinearColor(0.75294f, 0.08098f, 0.08098f, 1.000000f); // Soft Red
 	}
 	else if (TypeString == PC_Material)
 	{
@@ -753,7 +745,7 @@ FLinearColor UEdGraphSchema_CustomizableObject::GetPinTypeColor(const FName& Typ
 	}
 	else if (TypeString == PC_Modifier)
 	{
-		return FColor::Silver;
+		return FLinearColor(FColorList::LightGrey); 
 	}
 	else if (TypeString == PC_Object)
 	{
@@ -761,7 +753,7 @@ FLinearColor UEdGraphSchema_CustomizableObject::GetPinTypeColor(const FName& Typ
 	}
 	else if (TypeString == PC_Component)
 	{
-		return FLinearColor(FColorList::MediumOrchid); 
+		return FLinearColor(FColorList::DarkOrchid); 
 	}
 	else if (TypeString == PC_Stack)
 	{
@@ -781,7 +773,7 @@ FLinearColor UEdGraphSchema_CustomizableObject::GetPinTypeColor(const FName& Typ
 	}
 	else if (TypeString == PC_Transform)
 	{
-		return FLinearColor(0.5f, 1.0f, 1.0f, 1.0f); // Cyan
+		return FLinearColor(FColorList::Orange); 
 	}
 
 	for (const FRegisteredCustomizableObjectPinType& PinType : ICustomizableObjectModule::Get().GetExtendedPinTypes())
@@ -1416,7 +1408,7 @@ FName UEdGraphSchema_CustomizableObject::GetPinCategoryName(const FName& PinCate
 	}
 	else if (PinCategory == UEdGraphSchema_CustomizableObject::PC_MaterialAsset)
 	{
-		return TEXT("Table Material");
+		return TEXT("Material");
 	}
 	else if (PinCategory == UEdGraphSchema_CustomizableObject::PC_Wildcard)
 	{
@@ -1449,7 +1441,7 @@ FName UEdGraphSchema_CustomizableObject::GetPinCategoryName(const FName& PinCate
 
 FText UEdGraphSchema_CustomizableObject::GetPinCategoryFriendlyName(const FName& PinCategory)
 {
-if (PinCategory == UEdGraphSchema_CustomizableObject::PC_Object)
+	if (PinCategory == UEdGraphSchema_CustomizableObject::PC_Object)
 	{
 		return LOCTEXT("Object_Pin_Category", "Object");
 	}
@@ -1511,7 +1503,7 @@ if (PinCategory == UEdGraphSchema_CustomizableObject::PC_Object)
 	}
 	else if (PinCategory == UEdGraphSchema_CustomizableObject::PC_MaterialAsset)
 	{
-		return LOCTEXT("Material_Asset_Pin_Category", "Table Material");
+		return LOCTEXT("Material_Asset_Pin_Category", "Material");
 	}
 	else if (PinCategory == UEdGraphSchema_CustomizableObject::PC_Wildcard)
 	{
diff --git a/Engine/Plugins/Mutable/Source/CustomizableObjectEditor/Private/MuCOE/Nodes/CustomizableObjectNodeExposePin.cpp b/Engine/Plugins/Mutable/Source/CustomizableObjectEditor/Private/MuCOE/Nodes/CustomizableObjectNodeExposePin.cpp
--- a/Engine/Plugins/Mutable/Source/CustomizableObjectEditor/Private/MuCOE/Nodes/CustomizableObjectNodeExposePin.cpp
+++ b/Engine/Plugins/Mutable/Source/CustomizableObjectEditor/Private/MuCOE/Nodes/CustomizableObjectNodeExposePin.cpp
@@ -2,6 +2,8 @@
 
 #include "MuCOE/Nodes/CustomizableObjectNodeExposePin.h"
 
+#include "MuCO/CustomizableObjectCustomVersion.h"
+#include "MuCO/ICustomizableObjectModule.h"
 #include "MuCOE/EdGraphSchema_CustomizableObject.h"
 #include "MuCOE/Nodes/CustomizableObjectNodeExternalPin.h"
 
@@ -22,8 +24,12 @@ void UCustomizableObjectNodeExposePin::PostEditChangeProperty(FPropertyChangedEv
 
 void UCustomizableObjectNodeExposePin::AllocateDefaultPins(UCustomizableObjectNodeRemapPins* RemapPins)
 {
+	const FName PinName = UEdGraphSchema_CustomizableObject::GetPinCategoryName(PinType);
+	const FText PinFriendlyName = UEdGraphSchema_CustomizableObject::GetPinCategoryFriendlyName(PinType);
+	
 	const bool bIsArrayPinCategory = PinType == UEdGraphSchema_CustomizableObject::PC_GroupProjector;
-	UEdGraphPin* InputPin = CustomCreatePin(EGPD_Input, PinType, FName(TEXT("Object")), bIsArrayPinCategory);
+	UEdGraphPin* InputPin = CustomCreatePin(EGPD_Input, PinType, PinName, bIsArrayPinCategory);
+	InputPin->PinFriendlyName = PinFriendlyName;
 }
 
 
@@ -73,9 +79,142 @@ void UCustomizableObjectNodeExposePin::OnRenameNode(const FString& NewName)
 }
 
 
+void UCustomizableObjectNodeExposePin::BackwardsCompatibleFixup(int32 CustomizableObjectCustomVersion)
+{
+	Super::BackwardsCompatibleFixup(CustomizableObjectCustomVersion);
+
+	// Rename the Pin named "Object" with the friendly name that represents its category : "Color", "Transform" , "Enum"...
+	if (CustomizableObjectCustomVersion == FCustomizableObjectCustomVersion::UpdatedNodesPinName)
+	{
+		if (UEdGraphPin* Pin = FindPin(TEXT("Object")))
+		{
+			// Set the name of the pin based on the type of pin it is (extracted from UEdGraphSchema_CustomizableObject::GetPinCategoryName and
+			// UEdGraphSchema_CustomizableObject::GetPinCategoryFriendlyName to keep this fixup immutable)
+			if (PinType == UEdGraphSchema_CustomizableObject::PC_Object)
+			{
+				Pin->PinName = TEXT("Object");
+				Pin->PinFriendlyName = LOCTEXT("Object_Pin_Category", "Object");
+			}
+			else if (PinType == UEdGraphSchema_CustomizableObject::PC_Component)
+			{
+				Pin->PinName = TEXT("Component");
+				Pin->PinFriendlyName = LOCTEXT("Component_Pin_Category", "Component");
+			}
+			else if (PinType == UEdGraphSchema_CustomizableObject::PC_Material)
+			{
+				Pin->PinName = TEXT("Material");
+				Pin->PinFriendlyName = LOCTEXT("MeshSection_Pin_Category", "Mesh Section");
+			}
+			else if (PinType == UEdGraphSchema_CustomizableObject::PC_Modifier)
+			{
+				Pin->PinName = TEXT("Modifier");
+				Pin->PinFriendlyName = LOCTEXT("Modifier_Pin_Category", "Modifier");
+			}
+			else if (PinType == UEdGraphSchema_CustomizableObject::PC_Mesh)
+			{
+				Pin->PinName = TEXT("Mesh");
+				Pin->PinFriendlyName = LOCTEXT("Mesh_Pin_Category", "Mesh");
+			}
+			else if (PinType == UEdGraphSchema_CustomizableObject::PC_Layout)
+			{
+				Pin->PinName = TEXT("Layout");
+				Pin->PinFriendlyName = LOCTEXT("Layout_Pin_Category", "Layout");
+			}
+			else if (PinType == UEdGraphSchema_CustomizableObject::PC_Image)
+			{
+				Pin->PinName = TEXT("Texture");
+				Pin->PinFriendlyName = LOCTEXT("Image_Pin_Category", "Texture");
+			}
+			else if (PinType == UEdGraphSchema_CustomizableObject::PC_PassThroughImage)
+			{
+				Pin->PinName = TEXT("PassThrough Texture");
+				Pin->PinFriendlyName = LOCTEXT("PassThrough_Image_Pin_Category", "PassThrough Texture");
+			}
+			else if (PinType == UEdGraphSchema_CustomizableObject::PC_Projector)
+			{
+				Pin->PinName = TEXT("Projector");
+				Pin->PinFriendlyName = LOCTEXT("Projector_Pin_Category", "Projector");
+			}
+			else if (PinType == UEdGraphSchema_CustomizableObject::PC_GroupProjector)
+			{
+				Pin->PinName = TEXT("Group Projector");
+				Pin->PinFriendlyName = LOCTEXT("Group_Projector_Pin_Category", "Group Projector");
+			}
+			else if (PinType == UEdGraphSchema_CustomizableObject::PC_Color)
+			{
+				Pin->PinName = TEXT("Color");
+				Pin->PinFriendlyName = LOCTEXT("Color_Pin_Category", "Color");
+			}
+			else if (PinType == UEdGraphSchema_CustomizableObject::PC_Float)
+			{
+				Pin->PinName = TEXT("Float");
+				Pin->PinFriendlyName = LOCTEXT("Float_Pin_Category", "Float");
+			}
+			else if (PinType == UEdGraphSchema_CustomizableObject::PC_Bool)
+			{
+				Pin->PinName = TEXT("Bool");
+				Pin->PinFriendlyName = LOCTEXT("Bool_Pin_Category", "Bool");
+			}
+			else if (PinType == UEdGraphSchema_CustomizableObject::PC_Enum)
+			{
+				Pin->PinName = TEXT("Enum");
+				Pin->PinFriendlyName = LOCTEXT("Enum_Pin_Category", "Enum");
+			}
+			else if (PinType == UEdGraphSchema_CustomizableObject::PC_Stack)
+			{
+				Pin->PinName = TEXT("Stack");
+				Pin->PinFriendlyName = LOCTEXT("Stack_Pin_Category", "Stack");
+			}
+			else if (PinType == UEdGraphSchema_CustomizableObject::PC_MaterialAsset)
+			{
+				Pin->PinName = TEXT("Material");
+				Pin->PinFriendlyName = LOCTEXT("Material_Asset_Pin_Category", "Material");
+			}
+			else if (PinType == UEdGraphSchema_CustomizableObject::PC_Wildcard)
+			{
+				Pin->PinName = TEXT("Wildcard");
+				Pin->PinFriendlyName = LOCTEXT("Wildcard_Pin_Category", "Wildcard");
+			}
+			else if (PinType == UEdGraphSchema_CustomizableObject::PC_PoseAsset)
+			{
+				Pin->PinName = TEXT("PoseAsset");
+				Pin->PinFriendlyName = LOCTEXT("Pose_Pin_Category", "PoseAsset");
+			}
+			else if (PinType == UEdGraphSchema_CustomizableObject::PC_Transform)
+			{
+				Pin->PinName = TEXT("Transform");
+				Pin->PinFriendlyName = LOCTEXT("Transform_Pin_Category", "Transform");
+			}
+			else
+			{
+				bool bPinTypeWasFound = false;
+				for (const FRegisteredCustomizableObjectPinType& RegisteredCustomizableObjectPin : ICustomizableObjectModule::Get().GetExtendedPinTypes())
+				{
+					if (RegisteredCustomizableObjectPin.PinType.Name == PinType)
+					{
+						Pin->PinName = RegisteredCustomizableObjectPin.PinType.Name;
+						Pin->PinFriendlyName = RegisteredCustomizableObjectPin.PinType.DisplayName;
+						bPinTypeWasFound = true;
+						break;
+					}
+				}
+
+				// Need to fail gracefully here in case a plugin that was active when this graph was
+				// created is no longer loaded.
+				if (!bPinTypeWasFound)
+				{
+					Pin->PinName = TEXT("Unknown");
+					Pin->PinFriendlyName =  LOCTEXT("Unknown_Pin_Category", "Unknown");
+				}
+			}
+		}
+	}
+}
+
+
 bool UCustomizableObjectNodeExposePin::CanConnect(const UEdGraphPin* InOwnedInputPin, const UEdGraphPin* InOutputPin, bool& bOutIsOtherNodeBlocklisted, bool& bOutArePinsCompatible) const
 {
-	// Chech the pin types do match
+	// Check the pin types do match
 	bOutArePinsCompatible = Super::CanConnect(InOwnedInputPin, InOutputPin, bOutIsOtherNodeBlocklisted, bOutArePinsCompatible);
 
 	// Check the type of the other node to make sure it is not one we do not want to allow the connection with
diff --git a/Engine/Plugins/Mutable/Source/CustomizableObjectEditor/Private/MuCOE/Nodes/CustomizableObjectNodeExternalPin.cpp b/Engine/Plugins/Mutable/Source/CustomizableObjectEditor/Private/MuCOE/Nodes/CustomizableObjectNodeExternalPin.cpp
--- a/Engine/Plugins/Mutable/Source/CustomizableObjectEditor/Private/MuCOE/Nodes/CustomizableObjectNodeExternalPin.cpp
+++ b/Engine/Plugins/Mutable/Source/CustomizableObjectEditor/Private/MuCOE/Nodes/CustomizableObjectNodeExternalPin.cpp
@@ -3,6 +3,7 @@
 #include "MuCOE/Nodes/CustomizableObjectNodeExternalPin.h"
 
 #include "MuCO/CustomizableObjectCustomVersion.h"
+#include "MuCO/ICustomizableObjectModule.h"
 #include "MuCOE/CustomizableObjectGraph.h"
 #include "MuCOE/CustomizableObjectPin.h"
 #include "MuCOE/EdGraphSchema_CustomizableObject.h"
@@ -33,6 +34,133 @@ void UCustomizableObjectNodeExternalPin::BackwardsCompatibleFixup(int32 Customiz
 			PinType = Pins[0]->PinType.PinCategory;
 		}
 	}
+
+	// All pins named "Object" will be updated to use the friendly name of the pin category. 
+	if (CustomizableObjectCustomVersion == FCustomizableObjectCustomVersion::UpdatedNodesPinName)
+	{
+		if (UEdGraphPin* Pin = FindPin(TEXT("Object")))
+		{
+			// Set the name of the pin based on the type of pin it is (extracted from UEdGraphSchema_CustomizableObject::GetPinCategoryName and
+			// UEdGraphSchema_CustomizableObject::GetPinCategoryFriendlyName to keep this fixup immutable)
+			if (PinType == UEdGraphSchema_CustomizableObject::PC_Object)
+			{
+				Pin->PinName = TEXT("Object");
+				Pin->PinFriendlyName = LOCTEXT("Object_Pin_Category", "Object");
+			}
+			else if (PinType == UEdGraphSchema_CustomizableObject::PC_Component)
+			{
+				Pin->PinName = TEXT("Component");
+				Pin->PinFriendlyName = LOCTEXT("Component_Pin_Category", "Component");
+			}
+			else if (PinType == UEdGraphSchema_CustomizableObject::PC_Material)
+			{
+				Pin->PinName = TEXT("Material");
+				Pin->PinFriendlyName = LOCTEXT("MeshSection_Pin_Category", "Mesh Section");
+			}
+			else if (PinType == UEdGraphSchema_CustomizableObject::PC_Modifier)
+			{
+				Pin->PinName = TEXT("Modifier");
+				Pin->PinFriendlyName = LOCTEXT("Modifier_Pin_Category", "Modifier");
+			}
+			else if (PinType == UEdGraphSchema_CustomizableObject::PC_Mesh)
+			{
+				Pin->PinName = TEXT("Mesh");
+				Pin->PinFriendlyName = LOCTEXT("Mesh_Pin_Category", "Mesh");
+			}
+			else if (PinType == UEdGraphSchema_CustomizableObject::PC_Layout)
+			{
+				Pin->PinName = TEXT("Layout");
+				Pin->PinFriendlyName = LOCTEXT("Layout_Pin_Category", "Layout");
+			}
+			else if (PinType == UEdGraphSchema_CustomizableObject::PC_Image)
+			{
+				Pin->PinName = TEXT("Texture");
+				Pin->PinFriendlyName = LOCTEXT("Image_Pin_Category", "Texture");
+			}
+			else if (PinType == UEdGraphSchema_CustomizableObject::PC_PassThroughImage)
+			{
+				Pin->PinName = TEXT("PassThrough Texture");
+				Pin->PinFriendlyName = LOCTEXT("PassThrough_Image_Pin_Category", "PassThrough Texture");
+			}
+			else if (PinType == UEdGraphSchema_CustomizableObject::PC_Projector)
+			{
+				Pin->PinName = TEXT("Projector");
+				Pin->PinFriendlyName = LOCTEXT("Projector_Pin_Category", "Projector");
+			}
+			else if (PinType == UEdGraphSchema_CustomizableObject::PC_GroupProjector)
+			{
+				Pin->PinName = TEXT("Group Projector");
+				Pin->PinFriendlyName = LOCTEXT("Group_Projector_Pin_Category", "Group Projector");
+			}
+			else if (PinType == UEdGraphSchema_CustomizableObject::PC_Color)
+			{
+				Pin->PinName = TEXT("Color");
+				Pin->PinFriendlyName = LOCTEXT("Color_Pin_Category", "Color");
+			}
+			else if (PinType == UEdGraphSchema_CustomizableObject::PC_Float)
+			{
+				Pin->PinName = TEXT("Float");
+				Pin->PinFriendlyName = LOCTEXT("Float_Pin_Category", "Float");
+			}
+			else if (PinType == UEdGraphSchema_CustomizableObject::PC_Bool)
+			{
+				Pin->PinName = TEXT("Bool");
+				Pin->PinFriendlyName = LOCTEXT("Bool_Pin_Category", "Bool");
+			}
+			else if (PinType == UEdGraphSchema_CustomizableObject::PC_Enum)
+			{
+				Pin->PinName = TEXT("Enum");
+				Pin->PinFriendlyName = LOCTEXT("Enum_Pin_Category", "Enum");
+			}
+			else if (PinType == UEdGraphSchema_CustomizableObject::PC_Stack)
+			{
+				Pin->PinName = TEXT("Stack");
+				Pin->PinFriendlyName = LOCTEXT("Stack_Pin_Category", "Stack");
+			}
+			else if (PinType == UEdGraphSchema_CustomizableObject::PC_MaterialAsset)
+			{
+				Pin->PinName = TEXT("Material");
+				Pin->PinFriendlyName = LOCTEXT("Material_Asset_Pin_Category", "Material");
+			}
+			else if (PinType == UEdGraphSchema_CustomizableObject::PC_Wildcard)
+			{
+				Pin->PinName = TEXT("Wildcard");
+				Pin->PinFriendlyName = LOCTEXT("Wildcard_Pin_Category", "Wildcard");
+			}
+			else if (PinType == UEdGraphSchema_CustomizableObject::PC_PoseAsset)
+			{
+				Pin->PinName = TEXT("PoseAsset");
+				Pin->PinFriendlyName = LOCTEXT("Pose_Pin_Category", "PoseAsset");
+			}
+			else if (PinType == UEdGraphSchema_CustomizableObject::PC_Transform)
+			{
+				Pin->PinName = TEXT("Transform");
+				Pin->PinFriendlyName = LOCTEXT("Transform_Pin_Category", "Transform");
+			}
+			else
+			{
+				bool bPinTypeWasFound = false;
+				for (const FRegisteredCustomizableObjectPinType& RegisteredCustomizableObjectPin : ICustomizableObjectModule::Get().GetExtendedPinTypes())
+				{
+					if (RegisteredCustomizableObjectPin.PinType.Name == PinType)
+					{
+						Pin->PinName = RegisteredCustomizableObjectPin.PinType.Name;
+						Pin->PinFriendlyName = RegisteredCustomizableObjectPin.PinType.DisplayName;
+						bPinTypeWasFound = true;
+						break;
+					}
+				}
+
+				// Need to fail gracefully here in case a plugin that was active when this graph was
+				// created is no longer loaded.
+				if (!bPinTypeWasFound)
+				{
+					Pin->PinName = TEXT("Unknown");
+					Pin->PinFriendlyName =  LOCTEXT("Unknown_Pin_Category", "Unknown");
+				}
+			}
+		}
+	}
 }
 
 
@@ -77,7 +205,8 @@ void UCustomizableObjectNodeExternalPin::SetExternalObjectNodeId(FGuid Guid)
 
 UEdGraphPin* UCustomizableObjectNodeExternalPin::GetExternalPin() const
 {
-	return Pins[0];
+	const FName TargetPinName = GeneratePinName();
+	return FindPin(TargetPinName);
 }
 
 
@@ -89,27 +218,34 @@ UCustomizableObjectNodeExposePin* UCustomizableObjectNodeExternalPin::GetNodeExp
 
 void UCustomizableObjectNodeExternalPin::PrePropagateConnectionChanged()
 {
-	PropagatePreviousPin = ReverseFollowPinArray(*Pins[0]);
+	PropagatePreviousPin = ReverseFollowPinArray(*GetExternalPin());
 }
 
 
 void UCustomizableObjectNodeExternalPin::PropagateConnectionChanged()
 {
 	// Propagate new left.
-	PropagatePreviousPin.Append(ReverseFollowPinArray(*Pins[0])); // Notify old connections and new connections.
+	PropagatePreviousPin.Append(ReverseFollowPinArray(*GetExternalPin())); // Notify old connections and new connections.
 	NodePinConnectionListChanged(PropagatePreviousPin); // This function avoids double notifications.
 	
 	// Propagate right.
-	NodePinConnectionListChanged(FollowOutputPinArray(*Pins[0]));
+	NodePinConnectionListChanged(FollowOutputPinArray(*GetExternalPin()));
+}
+
+FName UCustomizableObjectNodeExternalPin::GeneratePinName() const
+{
+	return GetNodeExposePin() ? FName(GetNodeExposePin()->GetNodeName()) : UEdGraphSchema_CustomizableObject::GetPinCategoryName(PinType);
 }
 
 
 void UCustomizableObjectNodeExternalPin::AllocateDefaultPins(UCustomizableObjectNodeRemapPins* RemapPins)
 {
-	FName PinName = GetNodeExposePin() ? FName(GetNodeExposePin()->GetNodeName()) : FName("Object");
-	const bool bIsArrayPinCategory = PinType == UEdGraphSchema_CustomizableObject::PC_GroupProjector;
+	const FName PinName = GeneratePinName();
+	const FText PinFriendlyName = UEdGraphSchema_CustomizableObject::GetPinCategoryFriendlyName(PinType);
 
-	CustomCreatePin(EGPD_Output, PinType, PinName, bIsArrayPinCategory);
+	const bool bIsArrayPinCategory = PinType == UEdGraphSchema_CustomizableObject::PC_GroupProjector;
+	UEdGraphPin* OutputPin = CustomCreatePin(EGPD_Output, PinType, PinName, bIsArrayPinCategory);
+	OutputPin->PinFriendlyName = PinFriendlyName;
 }
 
 
diff --git a/Engine/Plugins/Mutable/Source/CustomizableObjectEditor/Private/MuCOE/Nodes/CustomizableObjectNodeMaterial.cpp b/Engine/Plugins/Mutable/Source/CustomizableObjectEditor/Private/MuCOE/Nodes/CustomizableObjectNodeMaterial.cpp
--- a/Engine/Plugins/Mutable/Source/CustomizableObjectEditor/Private/MuCOE/Nodes/CustomizableObjectNodeMaterial.cpp
+++ b/Engine/Plugins/Mutable/Source/CustomizableObjectEditor/Private/MuCOE/Nodes/CustomizableObjectNodeMaterial.cpp
@@ -327,7 +327,7 @@ void UCustomizableObjectNodeMaterial::AllocateDefaultPins(UCustomizableObjectNod
 	}
 	
 	{
-		const FString PinFriendlyName = TEXT("Table Material");
+		const FString PinFriendlyName = UEdGraphSchema_CustomizableObject::GetPinCategoryFriendlyName(Schema->PC_MaterialAsset).ToString();
 		const FString PinName = PinFriendlyName + FString(TEXT("_Input_Pin"));
 		UEdGraphPin* TableMaterialPin = CustomCreatePin(EGPD_Input, Schema->PC_MaterialAsset, FName(*PinName));
 		TableMaterialPin->PinFriendlyName = FText::FromString(PinFriendlyName);
@@ -733,7 +733,21 @@ void UCustomizableObjectNodeMaterial::BackwardsCompatibleFixup(int32 Customizabl
 			MaterialPin->PinName = FName(PinName);
 			MaterialPin->PinFriendlyName = FText::FromString(PinFriendlyName);
 		}
+	}
 
+	if (CustomizableObjectCustomVersion == FCustomizableObjectCustomVersion::UpdatedNodesPinName)
+	{
+		const FString ExpectedPinFriendlyName = "Table Material";
+		const FString ExpectedPinName = ExpectedPinFriendlyName + FString(TEXT("_Input_Pin"));	
+		
+		if (UEdGraphPin* FoundPin = FindPin(ExpectedPinName,EGPD_Input))
+		{
+			const FString TargetPinFriendlyName = "Material";
+			const FString TargetPinName = TargetPinFriendlyName + FString(TEXT("_Input_Pin"));
+			
+			FoundPin->PinFriendlyName = FText::FromString(TargetPinFriendlyName);
+			FoundPin->PinName = FName(TargetPinName);
+		}
 	}
 }
 
@@ -792,11 +806,7 @@ FText UCustomizableObjectNodeMaterial::GetNodeTitle(ENodeTitleType::Type TitleTy
 }
 
 
-FLinearColor UCustomizableObjectNodeMaterial::GetNodeTitleColor() const
-{
-	const UEdGraphSchema_CustomizableObject* Schema = GetDefault<UEdGraphSchema_CustomizableObject>();
-	return Schema->GetPinTypeColor(Schema->PC_Material);
-}
+
 
 
 UEdGraphPin* UCustomizableObjectNodeMaterial::OutputPin() const
@@ -847,8 +857,8 @@ UEdGraphPin* UCustomizableObjectNodeMaterial::GetMeshPin() const
 
 UEdGraphPin* UCustomizableObjectNodeMaterial::GetMaterialAssetPin() const
 {
-	FString PinFriendlyName = TEXT("Table Material");
-	FString PinName = PinFriendlyName + FString(TEXT("_Input_Pin"));
+	const FString PinFriendlyName = UEdGraphSchema_CustomizableObject::GetPinCategoryFriendlyName(UEdGraphSchema_CustomizableObject::PC_MaterialAsset).ToString();
+	const FString PinName = PinFriendlyName + FString(TEXT("_Input_Pin"));
 
 	UEdGraphPin* Pin = FindPin(PinName);
 
@@ -858,7 +868,7 @@ UEdGraphPin* UCustomizableObjectNodeMaterial::GetMaterialAssetPin() const
 	}
 	else
 	{
-		return FindPin(TEXT("Table Material"));
+		return FindPin(PinFriendlyName);
 	}
 }
 
diff --git a/Engine/Plugins/Mutable/Source/CustomizableObjectEditor/Private/MuCOE/Nodes/CustomizableObjectNodeMaterialBase.cpp b/Engine/Plugins/Mutable/Source/CustomizableObjectEditor/Private/MuCOE/Nodes/CustomizableObjectNodeMaterialBase.cpp
--- a/Engine/Plugins/Mutable/Source/CustomizableObjectEditor/Private/MuCOE/Nodes/CustomizableObjectNodeMaterialBase.cpp
+++ b/Engine/Plugins/Mutable/Source/CustomizableObjectEditor/Private/MuCOE/Nodes/CustomizableObjectNodeMaterialBase.cpp
@@ -2,3 +2,9 @@
 
 #include "MuCOE/Nodes/CustomizableObjectNodeMaterialBase.h"
 
+#include "MuCOE/EdGraphSchema_CustomizableObject.h"
+
+FLinearColor UCustomizableObjectNodeMaterialBase::GetNodeTitleColor() const
+{
+	return UEdGraphSchema_CustomizableObject::GetPinTypeColor(UEdGraphSchema_CustomizableObject::PC_Material);
+}
diff --git a/Engine/Plugins/Mutable/Source/CustomizableObjectEditor/Private/MuCOE/Nodes/CustomizableObjectNodeMeshMorphStackDefinition.cpp b/Engine/Plugins/Mutable/Source/CustomizableObjectEditor/Private/MuCOE/Nodes/CustomizableObjectNodeMeshMorphStackDefinition.cpp
--- a/Engine/Plugins/Mutable/Source/CustomizableObjectEditor/Private/MuCOE/Nodes/CustomizableObjectNodeMeshMorphStackDefinition.cpp
+++ b/Engine/Plugins/Mutable/Source/CustomizableObjectEditor/Private/MuCOE/Nodes/CustomizableObjectNodeMeshMorphStackDefinition.cpp
@@ -48,7 +48,7 @@ FText UCustomizableObjectNodeMeshMorphStackDefinition::GetNodeTitle(ENodeTitleTy
 FLinearColor UCustomizableObjectNodeMeshMorphStackDefinition::GetNodeTitleColor() const
 {
 	const UEdGraphSchema_CustomizableObject* Schema = GetDefault<UEdGraphSchema_CustomizableObject>();
-	return Schema->GetPinTypeColor(Schema->PC_Mesh);
+	return Schema->GetPinTypeColor(Schema->PC_Stack);
 }
 
 
diff --git a/Engine/Plugins/Mutable/Source/CustomizableObjectEditor/Public/MuCOE/Nodes/CustomizableObjectNodeExposePin.h b/Engine/Plugins/Mutable/Source/CustomizableObjectEditor/Public/MuCOE/Nodes/CustomizableObjectNodeExposePin.h
--- a/Engine/Plugins/Mutable/Source/CustomizableObjectEditor/Public/MuCOE/Nodes/CustomizableObjectNodeExposePin.h
+++ b/Engine/Plugins/Mutable/Source/CustomizableObjectEditor/Public/MuCOE/Nodes/CustomizableObjectNodeExposePin.h
@@ -2,6 +2,7 @@
 
 #pragma once
 
+#include "MuCOE/EdGraphSchema_CustomizableObject.h"
 #include "MuCOE/Nodes/CustomizableObjectNode.h"
 
 #include "CustomizableObjectNodeExposePin.generated.h"
@@ -33,6 +34,7 @@ public:
 	FText GetTooltipText() const override;
 	virtual bool GetCanRenameNode() const override { return true; }
 	virtual void OnRenameNode(const FString& NewName) override;
+	virtual void BackwardsCompatibleFixup(int32 CustomizableObjectCustomVersion) override;
 
 	// UCustomizableObjectNode interface
 	virtual bool CanConnect(const UEdGraphPin* InOwnedInputPin, const UEdGraphPin* InOutputPin, bool& bOutIsOtherNodeBlocklisted, bool& bOutArePinsCompatible) const override;
@@ -47,10 +49,11 @@ public:
 	
 	UEdGraphPin* InputPin() const
 	{
-		return FindPin(TEXT("Object"));
+		const FName InputPinName = UEdGraphSchema_CustomizableObject::GetPinCategoryName(PinType);
+		return FindPin(InputPinName);
 	}
 
-	/** Boradcasted when the UPROPERTY Name changes. */
+	/** Broadcasted when the UPROPERTY Name changes. */
 	FOnNameChangedDelegate OnNameChangedDelegate;
 	
 private:
diff --git a/Engine/Plugins/Mutable/Source/CustomizableObjectEditor/Public/MuCOE/Nodes/CustomizableObjectNodeExternalPin.h b/Engine/Plugins/Mutable/Source/CustomizableObjectEditor/Public/MuCOE/Nodes/CustomizableObjectNodeExternalPin.h
--- a/Engine/Plugins/Mutable/Source/CustomizableObjectEditor/Public/MuCOE/Nodes/CustomizableObjectNodeExternalPin.h
+++ b/Engine/Plugins/Mutable/Source/CustomizableObjectEditor/Public/MuCOE/Nodes/CustomizableObjectNodeExternalPin.h
@@ -69,6 +69,13 @@ public:
 	TObjectPtr<UCustomizableObject> ExternalObject;
 
 private:
+	
+	/**
+	 * Computes the name to be used by the pin. 
+	 * @return the name that the pin should use
+	 */
+	FName GeneratePinName() const;
+	
 	/** Linked Node Expose Pin node guid. */
 	UPROPERTY()
 	FGuid ExternalObjectNodeId;
diff --git a/Engine/Plugins/Mutable/Source/CustomizableObjectEditor/Public/MuCOE/Nodes/CustomizableObjectNodeMaterial.h b/Engine/Plugins/Mutable/Source/CustomizableObjectEditor/Public/MuCOE/Nodes/CustomizableObjectNodeMaterial.h
--- a/Engine/Plugins/Mutable/Source/CustomizableObjectEditor/Public/MuCOE/Nodes/CustomizableObjectNodeMaterial.h
+++ b/Engine/Plugins/Mutable/Source/CustomizableObjectEditor/Public/MuCOE/Nodes/CustomizableObjectNodeMaterial.h
@@ -112,7 +112,6 @@ public:
 
 	// EdGraphNode interface
 	virtual FText GetNodeTitle(ENodeTitleType::Type TitleType) const override;
-	virtual FLinearColor GetNodeTitleColor() const override;
 	virtual FText GetTooltipText() const override;
 	virtual TSharedPtr<SGraphNode> CreateVisualWidget() override;
 	virtual void PinConnectionListChanged(UEdGraphPin* Pin) override;
diff --git a/Engine/Plugins/Mutable/Source/CustomizableObjectEditor/Public/MuCOE/Nodes/CustomizableObjectNodeMaterialBase.h b/Engine/Plugins/Mutable/Source/CustomizableObjectEditor/Public/MuCOE/Nodes/CustomizableObjectNodeMaterialBase.h
--- a/Engine/Plugins/Mutable/Source/CustomizableObjectEditor/Public/MuCOE/Nodes/CustomizableObjectNodeMaterialBase.h
+++ b/Engine/Plugins/Mutable/Source/CustomizableObjectEditor/Public/MuCOE/Nodes/CustomizableObjectNodeMaterialBase.h
@@ -142,12 +142,15 @@ public:
 	 * @return Return -1 if the UV Index is set to Ignore. Return >= 0 for a valid UV Index.  */ 
 	virtual int32 GetImageUVLayout(int32 ImageIndex) const PURE_VIRTUAL(UCustomizableObjectNodeMaterial::GetImageUVLayout, return {}; );
 
-	// This method should be overidden in all derived classes
+	// This method should be overridden in all derived classes
 	virtual UEdGraphPin* OutputPin() const PURE_VIRTUAL(UCustomizableObjectNodeMaterial::OutputPin, return {}; );
 	
 	virtual FPostImagePinModeChangedDelegate* GetPostImagePinModeChangedDelegate() PURE_VIRTUAL(UCustomizableObjectNodeMaterial::OutputPin, return {}; );
 	
 	/** Return true if a Material Parameter has changed on which we had a pin connected or data saved. */
 	virtual bool RealMaterialDataHasChanged() const PURE_VIRTUAL(UCustomizableObjectNodeMaterial::OutputPin, return {}; );
+
+	// UEdGraphNode interface
+	virtual FLinearColor GetNodeTitleColor() const override;
 };
 
