[Mutable, HairStrandsMutable] Fix Grooms not deataching when discarding an Instance.

diff --git a/Engine/Plugins/Experimental/HairStrandsMutable/Source/HairStrandsMutable/Private/HairStrandsMutableExtension.cpp b/Engine/Plugins/Experimental/HairStrandsMutable/Source/HairStrandsMutable/Private/HairStrandsMutableExtension.cpp
--- a/Engine/Plugins/Experimental/HairStrandsMutable/Source/HairStrandsMutable/Private/HairStrandsMutableExtension.cpp
+++ b/Engine/Plugins/Experimental/HairStrandsMutable/Source/HairStrandsMutable/Private/HairStrandsMutableExtension.cpp
@@ -66,15 +66,17 @@ FInstancedStruct UHairStrandsMutableExtension::GenerateExtensionInstanceData(
 }
 
 
-void UHairStrandsMutableExtension::OnCustomizableObjectInstanceUsageUpdated(const UCustomizableObjectInstanceUsage* InstanceUsage) const
+void UHairStrandsMutableExtension::OnCustomizableObjectInstanceUsageUpdated(UCustomizableObjectInstanceUsage& InstanceUsage) const
 {
-	UCustomizableObjectInstance* Instance = InstanceUsage->GetCustomizableObjectInstance();
+	Super::OnCustomizableObjectInstanceUsageUpdated(InstanceUsage);
+	
+	UCustomizableObjectInstance* Instance = InstanceUsage.GetCustomizableObjectInstance();
 	if (!Instance)
 	{
 		return;
 	}
 
-	USkeletalMeshComponent* AttachedParent = InstanceUsage->GetAttachParent();
+	USkeletalMeshComponent* AttachedParent = InstanceUsage.GetAttachParent();
 	if (!AttachedParent)
 	{
 		return;
@@ -87,13 +89,13 @@ void UHairStrandsMutableExtension::OnCustomizableObjectInstanceUsageUpdated(cons
 		return;
 	}
 
-	USceneComponent* AttachParent = InstanceUsage->GetAttachParent();
+	USceneComponent* AttachParent = InstanceUsage.GetAttachParent();
 	
 	TArray<UGroomComponent*> UpdateGroomComponents; // Grooms that must be present at the end of this update.
 
 	for (const FGroomPinData& GroomPinData : InstanceData->Grooms)
 	{
-		if (GroomPinData.ComponentName != InstanceUsage->GetComponentName())
+		if (GroomPinData.ComponentName != InstanceUsage.GetComponentName())
 		{
 			continue;
 		}
@@ -158,3 +160,26 @@ void UHairStrandsMutableExtension::OnCustomizableObjectInstanceUsageUpdated(cons
 		}
 	}
 }
+
+
+void UHairStrandsMutableExtension::OnCustomizableObjectInstanceUsageDiscarded(UCustomizableObjectInstanceUsage& InstanceUsage) const
+{
+	Super::OnCustomizableObjectInstanceUsageDiscarded(InstanceUsage);
+
+	USceneComponent* AttachParent = InstanceUsage.GetAttachParent();
+	if (!AttachParent)
+	{
+		return;
+	}
+
+	TArray<USceneComponent*> Children;
+	AttachParent->GetChildrenComponents(false, Children);
+	for (USceneComponent* Child : Children)
+	{
+		if (Child->IsA(UGroomComponent::StaticClass()) &&
+			Child->ComponentTags.Contains(GroomComponentTag))
+		{
+			Child->DestroyComponent();
+		}
+	}
+}
diff --git a/Engine/Plugins/Experimental/HairStrandsMutable/Source/HairStrandsMutable/Public/HairStrandsMutableExtension.h b/Engine/Plugins/Experimental/HairStrandsMutable/Source/HairStrandsMutable/Public/HairStrandsMutableExtension.h
--- a/Engine/Plugins/Experimental/HairStrandsMutable/Source/HairStrandsMutable/Public/HairStrandsMutableExtension.h
+++ b/Engine/Plugins/Experimental/HairStrandsMutable/Source/HairStrandsMutable/Public/HairStrandsMutableExtension.h
@@ -70,9 +70,9 @@ public:
 	/** UCustomizableObjectExtension interface */
 	virtual TArray<FCustomizableObjectPinType> GetPinTypes() const override;
 	virtual TArray<FObjectNodeInputPin> GetAdditionalObjectNodePins() const override;
-	virtual FInstancedStruct GenerateExtensionInstanceData(
-		const TArray<FInputPinDataContainer>& InputPinData) const override;
-	virtual void OnCustomizableObjectInstanceUsageUpdated(const UCustomizableObjectInstanceUsage* Usage) const override;
+	virtual FInstancedStruct GenerateExtensionInstanceData(const TArray<FInputPinDataContainer>& InputPinData) const override;
+	virtual void OnCustomizableObjectInstanceUsageUpdated(UCustomizableObjectInstanceUsage& Usage) const override;
+	virtual void OnCustomizableObjectInstanceUsageDiscarded(UCustomizableObjectInstanceUsage& Usage) const override;
 
 	HAIRSTRANDSMUTABLE_API static const FName GroomPinType;
 	HAIRSTRANDSMUTABLE_API static const FName GroomsBaseNodePinName;
diff --git a/Engine/Plugins/Mutable/Source/CustomizableObject/Internal/MuCO/CustomizableObjectInstanceUsagePrivate.h b/Engine/Plugins/Mutable/Source/CustomizableObject/Internal/MuCO/CustomizableObjectInstanceUsagePrivate.h
--- a/Engine/Plugins/Mutable/Source/CustomizableObject/Internal/MuCO/CustomizableObjectInstanceUsagePrivate.h
+++ b/Engine/Plugins/Mutable/Source/CustomizableObject/Internal/MuCO/CustomizableObjectInstanceUsagePrivate.h
@@ -30,7 +30,7 @@ public:
 	// Own interface
 	
 	/** Common end point of all updates. Even those which failed. */
-	void Callbacks() const;
+	void Callbacks();
 
 	void SetSkeletalMeshAndOverrideMaterials(USkeletalMeshComponent& Parent, USkeletalMesh* SkeletalMesh, const UCustomizableObjectInstance& CustomizableObjectInstance, bool* bOutSkeletalMeshUpdated, bool* bOutMaterialsUpdated);
 
diff --git a/Engine/Plugins/Mutable/Source/CustomizableObject/Private/MuCO/CustomizableObjectInstance.cpp b/Engine/Plugins/Mutable/Source/CustomizableObject/Private/MuCO/CustomizableObjectInstance.cpp
--- a/Engine/Plugins/Mutable/Source/CustomizableObject/Private/MuCO/CustomizableObjectInstance.cpp
+++ b/Engine/Plugins/Mutable/Source/CustomizableObject/Private/MuCO/CustomizableObjectInstance.cpp
@@ -3083,12 +3083,17 @@ void UCustomizableInstancePrivate::DiscardResources()
 void UCustomizableInstancePrivate::SetDefaultSkeletalMesh(bool bSetEmptyMesh) const
 {
 	UCustomizableObjectInstance* Instance = Cast<UCustomizableObjectInstance>(GetOuter());
-	if (!Instance || !Instance->GetCustomizableObject())
+	if (!Instance)
 	{
 		return;
 	}
 
 	UCustomizableObject* CustomizableObject = Instance->GetCustomizableObject();
+	if (!CustomizableObject)
+	{
+		return;
+	}
+	
 	const FModelResources& ModelResources = CustomizableObject->GetPrivate()->GetModelResources();
 	
 	for (TObjectIterator<UCustomizableObjectInstanceUsage> It; It; ++It)
diff --git a/Engine/Plugins/Mutable/Source/CustomizableObject/Private/MuCO/CustomizableObjectInstanceUsage.cpp b/Engine/Plugins/Mutable/Source/CustomizableObject/Private/MuCO/CustomizableObjectInstanceUsage.cpp
--- a/Engine/Plugins/Mutable/Source/CustomizableObject/Private/MuCO/CustomizableObjectInstanceUsage.cpp
+++ b/Engine/Plugins/Mutable/Source/CustomizableObject/Private/MuCO/CustomizableObjectInstanceUsage.cpp
@@ -30,11 +30,11 @@ ETickableTickType UCustomizableObjectInstanceUsagePrivate::GetTickableTickType()
 }
 
 
-void UCustomizableObjectInstanceUsagePrivate::Callbacks() const
+void UCustomizableObjectInstanceUsagePrivate::Callbacks()
 {
 	for (const UCustomizableObjectExtension* Extension : ICustomizableObjectModule::Get().GetRegisteredExtensions())
 	{
-		Extension->OnCustomizableObjectInstanceUsageUpdated(GetPublic());
+		Extension->OnCustomizableObjectInstanceUsageUpdated(*GetPublic());
 	}
 	
 	if (UCustomizableSkeletalComponent* CustomizableSkeletalComponent = GetCustomizableSkeletalComponent())
diff --git a/Engine/Plugins/Mutable/Source/CustomizableObject/Private/MuCO/CustomizableObjectSystem.cpp b/Engine/Plugins/Mutable/Source/CustomizableObject/Private/MuCO/CustomizableObjectSystem.cpp
--- a/Engine/Plugins/Mutable/Source/CustomizableObject/Private/MuCO/CustomizableObjectSystem.cpp
+++ b/Engine/Plugins/Mutable/Source/CustomizableObject/Private/MuCO/CustomizableObjectSystem.cpp
@@ -888,7 +888,7 @@ void FinishUpdateGlobal(const TSharedRef<FUpdateContextPrivate>& Context)
 		// Call CustomizableObjectInstanceUsages updated callbacks.
 		for (TObjectIterator<UCustomizableObjectInstanceUsage> It; It; ++It) // Since iterating objects is expensive, for now CustomizableObjectInstanceUsage does not have a FinishUpdate function.
 		{
-			const UCustomizableObjectInstanceUsage* InstanceUsage = *It;
+			UCustomizableObjectInstanceUsage* InstanceUsage = *It;
 			if (!IsValid(InstanceUsage))
 			{
 				continue;
@@ -3952,7 +3952,47 @@ void UCustomizableObjectSystem::DiscardInstances()
 				COIPrivateData->LastMinSquareDistFromComponentToPlayer > FMath::Square(GetPrivate()->CurrentInstanceLODManagement->GetOnlyUpdateCloseCustomizableObjectsDist()))
 			{
 				COIPrivateData->DiscardResources();
-				COIPrivateData->SetDefaultSkeletalMesh(!IsReplaceDiscardedWithReferenceMeshEnabled());
+
+				if (UCustomizableObject* CustomizableObject = COI->GetCustomizableObject())
+				{
+					const FModelResources& ModelResources = CustomizableObject->GetPrivate()->GetModelResources();
+				
+					for (TObjectIterator<UCustomizableObjectInstanceUsage> It; It; ++It)
+					{
+						UCustomizableObjectInstanceUsage* InstanceUsage = *It;
+						if (InstanceUsage->GetCustomizableObjectInstance() != COI)
+						{
+							continue;
+						}
+
+#if WITH_EDITOR
+						if (InstanceUsage->GetPrivate()->IsNetMode(NM_DedicatedServer))
+						{
+							continue;
+						}
+#endif
+
+						const FName& ComponentName = InstanceUsage->GetComponentName();
+						const int32 ObjectComponentIndex = COI->GetCustomizableObject()->GetPrivate()->GetModelResources().ComponentNamesPerObjectComponent.IndexOfByKey(ComponentName);
+						if (ModelResources.ReferenceSkeletalMeshesData.IsValidIndex(ObjectComponentIndex))
+						{
+							USkeletalMesh* SkeletalMesh = nullptr;
+					
+							if (IsReplaceDiscardedWithReferenceMeshEnabled())
+							{
+								// Force load the reference mesh if necessary. 
+								SkeletalMesh = ModelResources.ReferenceSkeletalMeshesData[ObjectComponentIndex].SoftSkeletalMesh.LoadSynchronous();
+							}
+		
+							InstanceUsage->GetPrivate()->SetSkeletalMesh(SkeletalMesh);
+						}
+					
+						for (const UCustomizableObjectExtension* Extension : ICustomizableObjectModule::Get().GetRegisteredExtensions())
+						{
+							Extension->OnCustomizableObjectInstanceUsageDiscarded(*InstanceUsage);
+						}
+					}
+				}
 			}
 		}
 
diff --git a/Engine/Plugins/Mutable/Source/CustomizableObject/Public/MuCO/CustomizableObjectExtension.h b/Engine/Plugins/Mutable/Source/CustomizableObject/Public/MuCO/CustomizableObjectExtension.h
--- a/Engine/Plugins/Mutable/Source/CustomizableObject/Public/MuCO/CustomizableObjectExtension.h
+++ b/Engine/Plugins/Mutable/Source/CustomizableObject/Public/MuCO/CustomizableObjectExtension.h
@@ -121,6 +121,14 @@ public:
 	 * that the large data itself isn't being copied.
 	 */
 	virtual FInstancedStruct GenerateExtensionInstanceData(const TArray<FInputPinDataContainer>& InputPinData) const { return FInstancedStruct(); }
-	
-	virtual void OnCustomizableObjectInstanceUsageUpdated(const UCustomizableObjectInstanceUsage* Usage) const {}
+
+	/** Called when a Customizable Object Instance Usage is being updated.
+	 * 
+	 * @param Usage Customizable Object Instance Usage being updated. */
+	virtual void OnCustomizableObjectInstanceUsageUpdated(UCustomizableObjectInstanceUsage& Usage) const {}
+
+	/** Called when a Customizable Object Instance Usage is being discarded.
+	 * 
+	 * @param Usage Customizable Object Instance Usage being discarded. */
+	virtual void OnCustomizableObjectInstanceUsageDiscarded(UCustomizableObjectInstanceUsage& Usage) const {}
 };
