[Mutable] Fix Bake modifying original Materials when using the "Generate Constant Material Instances" option.

- Fixed CopyAllMaterialParameters setting the parameters to the original Material (instead of the UMaterialInstanceConstant).
- Simplified CopyAllMaterialParameters by using a template (instead of runtime type switching).
- Inlined GetTextureParameterNames.

diff --git a/Engine/Plugins/Mutable/Source/CustomizableObjectEditor/Private/MuCOE/CustomizableObjectInstanceBakingUtils.cpp b/Engine/Plugins/Mutable/Source/CustomizableObjectEditor/Private/MuCOE/CustomizableObjectInstanceBakingUtils.cpp
--- a/Engine/Plugins/Mutable/Source/CustomizableObjectEditor/Private/MuCOE/CustomizableObjectInstanceBakingUtils.cpp
+++ b/Engine/Plugins/Mutable/Source/CustomizableObjectEditor/Private/MuCOE/CustomizableObjectInstanceBakingUtils.cpp
@@ -480,10 +480,15 @@ bool BakeCustomizableObjectInstance(
 				
 				if (Material != nullptr && Inst != nullptr)
 				{
-					TArray<FName> ParameterNames = FUnrealBakeHelpers::GetTextureParameterNames(Material);
-					for (int32 i = 0; i < ParameterNames.Num(); i++)
+					TArray<FGuid> Guids;
+					TArray<FMaterialParameterInfo> OutParameterInfo;
+					Material->GetAllTextureParameterInfo(OutParameterInfo, Guids);
+					
+					for (int32 Index = 0; Index < OutParameterInfo.Num(); Index++)
 					{
-						if (Inst->GetTextureParameterValue(ParameterNames[i], Texture))
+						const FName& ParameterName = OutParameterInfo[Index].Name;
+						
+						if (Inst->GetTextureParameterValue(ParameterName, Texture))
 						{
 							UTexture2D* SrcTex = Cast<UTexture2D>(Texture);
 							if (!SrcTex)
@@ -491,7 +496,7 @@ bool BakeCustomizableObjectInstance(
 								continue;
 							}
 
-							FString ParameterSanitized = ParameterNames[i].GetPlainNameString();
+							FString ParameterSanitized = ParameterName.GetPlainNameString();
 							RemoveRestrictedChars(ParameterSanitized);
 							ResourceName = ObjectName + "_" + MaterialName + "_" + ParameterSanitized;
 							if (!GetUniqueResourceName(SrcTex, ResourceName, ArrayCachedObject, ArrayCachedElement))
@@ -527,7 +532,7 @@ bool BakeCustomizableObjectInstance(
 								{
 									// Recover original name of the texture parameter value, now substituted by the generated Mutable texture
 									UTexture* OriginalTexture = nullptr;
-									Inst->Parent->GetTextureParameterValue(FName(*ParameterNames[i].GetPlainNameString()), OriginalTexture);
+									Inst->Parent->GetTextureParameterValue(FName(*ParameterName.GetPlainNameString()), OriginalTexture);
 
 									PackageName = AssetPath + FString("/") + ResourceName;
 									TMap<UObject*, UObject*> FakeReplacementMap;
@@ -540,7 +545,7 @@ bool BakeCustomizableObjectInstance(
 
 									if (OriginalTexture != nullptr)
 									{
-										TextureReplacementMaps[m].Add(i, DupTex);
+										TextureReplacementMaps[m].Add(Index, DupTex);
 									}
 								}
 							}
@@ -558,7 +563,7 @@ bool BakeCustomizableObjectInstance(
 								OutSavedPackages.Add(PackageToSave);
 
 								UTexture* DupTexture = Cast<UTexture>(DuplicatedObject);
-								TextureReplacementMaps[m].Add(i, DupTexture);
+								TextureReplacementMaps[m].Add(Index, DupTexture);
 							}
 						}
 					}
@@ -599,7 +604,10 @@ bool BakeCustomizableObjectInstance(
 					TPair<EPackageSaveResolutionType, UPackage*> PackageToSave {SaveType, DuplicatedObject->GetPackage()};
 					OutSavedPackages.Add(PackageToSave);
 
-					FUnrealBakeHelpers::CopyAllMaterialParameters(DuplicatedObject, Interface, TextureReplacementMaps[m]);
+					if (UMaterial* DuplicatedMaterial = Cast<UMaterial>(DuplicatedObject))
+					{
+						FUnrealBakeHelpers::CopyAllMaterialParameters<UMaterial>(*DuplicatedMaterial, *Interface, TextureReplacementMaps[m]);
+					}
 				}
 			}
 		}
diff --git a/Engine/Plugins/Mutable/Source/CustomizableObjectEditor/Private/MuCOE/UnrealBakeHelpers.cpp b/Engine/Plugins/Mutable/Source/CustomizableObjectEditor/Private/MuCOE/UnrealBakeHelpers.cpp
--- a/Engine/Plugins/Mutable/Source/CustomizableObjectEditor/Private/MuCOE/UnrealBakeHelpers.cpp
+++ b/Engine/Plugins/Mutable/Source/CustomizableObjectEditor/Private/MuCOE/UnrealBakeHelpers.cpp
@@ -36,29 +36,27 @@ UObject* FUnrealBakeHelpers::BakeHelper_DuplicateAsset(UObject* Object, const FS
 
 	UPackage* Package = CreatePackage(*FinalPkgName);
 	Package->FullyLoad();
-
-	FObjectDuplicationParameters Params = InitStaticDuplicateObjectParams(Object, Package, *FinalObjName, RF_AllFlags, nullptr, EDuplicateMode::Normal);
+	
+	UMaterialInstance* MatInstance = Cast<UMaterialInstance>(Object);
 
 	UObject* DupObject = nullptr;
-	UMaterialInterface* MatInterface = Cast<UMaterialInterface>(Object);
-	// Only generate constant material instances if the original material is actually an instance, so check it here. 
-	// Otherwise just duplicate
-	UMaterialInstance* MatInstance = Cast<UMaterialInstance>(Object);
 	
-	if (bGenerateConstantMaterialInstances && MatInterface && MatInstance)
+	// Only generate UMaterialInstanceConstant constant material instances if the original material is actually an instance, so check it here. Otherwise just duplicate
+	if (bGenerateConstantMaterialInstances && MatInstance)
 	{
 		UMaterialInstanceConstantFactoryNew* MaterialFactory = NewObject<UMaterialInstanceConstantFactoryNew>();
 		MaterialFactory->InitialParent = MatInstance->Parent;
-		FString MaterialInstanceName = FinalObjName;
-		DupObject = (UMaterialInstanceConstant*)MaterialFactory->FactoryCreateNew(UMaterialInstanceConstant::StaticClass(),
-			Package, FName(MaterialInstanceName), RF_NoFlags, NULL, GWarn);
-		ensure(DupObject);
 
-		TMap<int, UTexture*> EmptyTextureReplacementMap;
-		CopyAllMaterialParameters(DupObject, MatInterface, EmptyTextureReplacementMap);
+		UMaterialInstanceConstant* MatInstanceConst = CastChecked<UMaterialInstanceConstant>(MaterialFactory->FactoryCreateNew(UMaterialInstanceConstant::StaticClass(), Package, FName(FinalObjName), RF_NoFlags, nullptr, GWarn));
+
+		TMap<int32, UTexture*> EmptyTextureReplacementMap;
+		CopyAllMaterialParameters<UMaterialInstanceConstant>(*MatInstanceConst, *MatInstance, EmptyTextureReplacementMap);
+		
+		DupObject = MatInstanceConst;
 	}
 	else
 	{
+		FObjectDuplicationParameters Params = InitStaticDuplicateObjectParams(Object, Package, *FinalObjName, RF_AllFlags, nullptr, EDuplicateMode::Normal);
 		DupObject = StaticDuplicateObjectEx(Params);
 	}
 
@@ -272,111 +270,3 @@ UTexture2D* FUnrealBakeHelpers::BakeHelper_CreateAssetTexture(UTexture2D* SrcTex
 	return DupTex;
 }
 
-
-void FUnrealBakeHelpers::CopyAllMaterialParameters(UObject* DestMaterial, UMaterialInterface* OriginMaterial, const TMap<int, UTexture*>& TextureReplacementMap)
-{
-	UMaterial* Material = OriginMaterial ? OriginMaterial->GetMaterial() : nullptr;
-	UMaterialInterface* DupMaterialInterface = Cast<UMaterialInterface>(DestMaterial);
-	UMaterial* DupMaterial = DupMaterialInterface ? DupMaterialInterface->GetMaterial() : nullptr;
-	UMaterialInstanceConstant* DupMaterialInstanceConstant = Cast<UMaterialInstanceConstant>(DestMaterial);
-
-	if (Material && DupMaterial)
-	{
-		// copy scalar parameters
-		TArray<FMaterialParameterInfo> ScalarParameterInfoArray;
-		TArray<FGuid> GuidArray;
-		OriginMaterial->GetAllScalarParameterInfo(ScalarParameterInfoArray, GuidArray);
-		for (const FMaterialParameterInfo& Param : ScalarParameterInfoArray)
-		{
-			float Value = 0.f;
-			if (OriginMaterial->GetScalarParameterValue(Param, Value))
-			{
-				if (DupMaterialInstanceConstant)
-				{
-					DupMaterialInstanceConstant->SetScalarParameterValueEditorOnly(Param.Name, Value);
-				}
-				else
-				{
-					DupMaterial->SetScalarParameterValueEditorOnly(Param.Name, Value);
-				}
-			}
-		}
-
-		// copy vector parameters
-		TArray<FMaterialParameterInfo> VectorParameterInfoArray;
-		OriginMaterial->GetAllVectorParameterInfo(VectorParameterInfoArray, GuidArray);
-		for (const FMaterialParameterInfo& Param : VectorParameterInfoArray)
-		{
-			FLinearColor Value;
-			if (OriginMaterial->GetVectorParameterValue(Param, Value))
-			{
-				if (DupMaterialInstanceConstant)
-				{
-					DupMaterialInstanceConstant->SetVectorParameterValueEditorOnly(Param.Name, Value);
-				}
-				else
-				{
-					DupMaterial->SetVectorParameterValueEditorOnly(Param.Name, Value);
-				}
-			}
-		}
-
-		// copy switch parameters								
-		TArray<FMaterialParameterInfo> StaticSwitchParameterInfoArray;
-		OriginMaterial->GetAllStaticSwitchParameterInfo(StaticSwitchParameterInfoArray, GuidArray);
-		for (int i = 0; i < StaticSwitchParameterInfoArray.Num(); ++i)
-		{
-			bool Value = false;
-			if (OriginMaterial->GetStaticSwitchParameterValue(StaticSwitchParameterInfoArray[i].Name, Value, GuidArray[i]))
-			{
-				DupMaterial->SetStaticSwitchParameterValueEditorOnly(StaticSwitchParameterInfoArray[i].Name, Value, GuidArray[i]);
-			}
-		}
-
-		// Replace Textures
-		TArray<FName> ParameterNames = GetTextureParameterNames(Material);
-		for (const TPair<int, UTexture*>& it : TextureReplacementMap)
-		{
-			if (ParameterNames.IsValidIndex(it.Key))
-			{
-				DupMaterial->SetTextureParameterValueEditorOnly(ParameterNames[it.Key], it.Value);
-			}
-		}
-
-		// Fix potential errors compiling materials due to Sampler Types
-		for (const TObjectPtr<UMaterialExpression>& Expression : DupMaterial->GetExpressions())
-		{
-			if (UMaterialExpressionTextureBase* MatExpressionTexBase = Cast<UMaterialExpressionTextureBase>(Expression))
-			{
-				MatExpressionTexBase->AutoSetSampleType();
-			}
-		}
-
-		DestMaterial->PreEditChange(NULL);
-		DestMaterial->PostEditChange();
-	}
-	else
-	{
-		ensure(false);
-	}
-}
-
-
-TArray<FName> FUnrealBakeHelpers::GetTextureParameterNames(UMaterial* Material)
-{
-	TArray<FGuid> Guids;
-	TArray<FName> ParameterNames;
-
-	TArray<FMaterialParameterInfo> OutParameterInfo;
-	Material->GetAllTextureParameterInfo(OutParameterInfo, Guids);
-
-	const int32 MaxIndex = OutParameterInfo.Num();
-	ParameterNames.SetNum(MaxIndex);
-
-	for (int32 i = 0; i < MaxIndex; i++)
-	{
-		ParameterNames[i] = OutParameterInfo[i].Name;
-	}
-
-	return ParameterNames;
-}
diff --git a/Engine/Plugins/Mutable/Source/CustomizableObjectEditor/Private/MuCOE/UnrealBakeHelpers.h b/Engine/Plugins/Mutable/Source/CustomizableObjectEditor/Private/MuCOE/UnrealBakeHelpers.h
--- a/Engine/Plugins/Mutable/Source/CustomizableObjectEditor/Private/MuCOE/UnrealBakeHelpers.h
+++ b/Engine/Plugins/Mutable/Source/CustomizableObjectEditor/Private/MuCOE/UnrealBakeHelpers.h
@@ -2,8 +2,13 @@
 
 #pragma once
 
+#include "MaterialTypes.h"
 #include "HAL/Platform.h"
 #include "Containers/Map.h"
+#include "Materials/Material.h"
+#include "Materials/MaterialExpressionTextureBase.h"
+#include "Materials/MaterialInstanceConstant.h"
+#include "Materials/MaterialInterface.h"
 
 class UObject;
 class USkeletalMesh;
@@ -28,8 +33,94 @@ public:
 	 */
 	static UTexture2D* BakeHelper_CreateAssetTexture(UTexture2D* SrcTex, const FString& TexObjName, const FString& TexPkgName, const UTexture* OrgTex, bool ResetDuplicatedFlags, TMap<UObject*, UObject*>& ReplacementMap, bool OverwritePackage);
 
-	static void CopyAllMaterialParameters(UObject* DestMaterial, UMaterialInterface* OriginMaterial, const TMap<int, UTexture*>& TextureReplacementMaps);
+	template<typename MaterialType>
+	static void CopyAllMaterialParameters(MaterialType& DestMaterial, UMaterialInterface& OriginMaterial, const TMap<int32, UTexture*>& TextureReplacementMap)
+	{
+		// Copy scalar parameters
+		{
+			TArray<FMaterialParameterInfo> ScalarParameterInfoArray;
+			TArray<FGuid> GuidArray;
+			OriginMaterial.GetAllScalarParameterInfo(ScalarParameterInfoArray, GuidArray);
+			for (const FMaterialParameterInfo& Param : ScalarParameterInfoArray)
+			{
+				float Value = 0.f;
+				if (OriginMaterial.GetScalarParameterValue(Param, Value))
+				{
+					DestMaterial.SetScalarParameterValueEditorOnly(Param.Name, Value);
+				}
+			}
+			
+		}
 
-	/** Get associated texture parameter names for the material given as parameter */
-	static TArray<FName> GetTextureParameterNames(UMaterial* Material);
+		// Copy vector parameters
+		{
+			TArray<FMaterialParameterInfo> VectorParameterInfoArray;
+			TArray<FGuid> GuidArray;
+			OriginMaterial.GetAllVectorParameterInfo(VectorParameterInfoArray, GuidArray);
+			for (const FMaterialParameterInfo& Param : VectorParameterInfoArray)
+			{
+				FLinearColor Value;
+				if (OriginMaterial.GetVectorParameterValue(Param, Value))
+				{
+					DestMaterial.SetVectorParameterValueEditorOnly(Param.Name, Value);
+				}
+			}
+		}
+
+		// Copy switch parameters								
+		{
+			TArray<FMaterialParameterInfo> StaticSwitchParameterInfoArray;
+			TArray<FGuid> GuidArray;
+			OriginMaterial.GetAllStaticSwitchParameterInfo(StaticSwitchParameterInfoArray, GuidArray);
+			for (int i = 0; i < StaticSwitchParameterInfoArray.Num(); ++i)
+			{
+				bool Value = false;
+				FGuid ExpressionsGuid;
+				if (OriginMaterial.GetStaticSwitchParameterValue(StaticSwitchParameterInfoArray[i].Name, Value, ExpressionsGuid))
+				{
+					// For some reason UMaterialInstance::SetStaticSwitchParameterValueEditorOnly signature is different than UMaterial::SetStaticSwitchParameterValueEditorOnly
+					if constexpr (std::is_same_v<MaterialType, UMaterial>)
+					{
+						DestMaterial.SetStaticSwitchParameterValueEditorOnly(StaticSwitchParameterInfoArray[i].Name, Value, ExpressionsGuid);
+					}
+					else if (std::is_same_v<MaterialType, UMaterialInstanceConstant>)
+					{
+						DestMaterial.SetStaticSwitchParameterValueEditorOnly(StaticSwitchParameterInfoArray[i].Name, Value);
+					}
+					else
+					{
+						static_assert(
+							std::is_same_v<MaterialType, UMaterial> ||
+							std::is_same_v<MaterialType, UMaterialInstanceConstant>);
+					}
+				}
+			}
+		}
+
+		// Replace Textures
+		{
+			TArray<FMaterialParameterInfo> OutParameterInfo;
+			TArray<FGuid> Guids;
+			OriginMaterial.GetAllTextureParameterInfo(OutParameterInfo, Guids);
+			for (const TPair<int32, UTexture*>& It : TextureReplacementMap)
+			{
+				if (OutParameterInfo.IsValidIndex(It.Key))
+				{
+					DestMaterial.SetTextureParameterValueEditorOnly(OutParameterInfo[It.Key].Name, It.Value);
+				}
+			}			
+		}
+
+		// Fix potential errors compiling materials due to Sampler Types
+		for (const TObjectPtr<UMaterialExpression>& Expression : DestMaterial.GetMaterial()->GetExpressions())
+		{
+			if (UMaterialExpressionTextureBase* MatExpressionTexBase = Cast<UMaterialExpressionTextureBase>(Expression))
+			{
+				MatExpressionTexBase->AutoSetSampleType();
+			}
+		}
+
+		DestMaterial.PreEditChange(NULL);
+		DestMaterial.PostEditChange();
+	}
 };
