[Mutable] Merged Node and CO details tab into one.

#rnx

diff --git a/Engine/Plugins/Mutable/Source/CustomizableObjectEditor/Private/MuCOE/CustomizableObjectEditor.cpp b/Engine/Plugins/Mutable/Source/CustomizableObjectEditor/Private/MuCOE/CustomizableObjectEditor.cpp
--- a/Engine/Plugins/Mutable/Source/CustomizableObjectEditor/Private/MuCOE/CustomizableObjectEditor.cpp
+++ b/Engine/Plugins/Mutable/Source/CustomizableObjectEditor/Private/MuCOE/CustomizableObjectEditor.cpp
@@ -80,10 +80,9 @@ enum class EColorArithmeticOperation : uint8;
 DEFINE_LOG_CATEGORY_STATIC(LogCustomizableObjectEditor, Log, All);
 
 const FName FCustomizableObjectEditor::ViewportTabId( TEXT( "CustomizableObjectEditor_Viewport" ) );
-const FName FCustomizableObjectEditor::ObjectPropertiesTabId( TEXT( "CustomizableObjectEditor_ObjectProperties" ) );
+const FName FCustomizableObjectEditor::DetailsTabId( TEXT( "CustomizableObjectEditor_ObjectProperties" ) );
 const FName FCustomizableObjectEditor::InstancePropertiesTabId( TEXT( "CustomizableObjectEditor_InstanceProperties" ) );
 const FName FCustomizableObjectEditor::GraphTabId( TEXT( "CustomizableObjectEditor_Graph" ) );
-const FName FCustomizableObjectEditor::GraphNodePropertiesTabId( TEXT( "CustomizableObjectEditor_GraphNodeProperties" ) );
 const FName FCustomizableObjectEditor::AdvancedPreviewSettingsTabId(TEXT("CustomizableObjectEditor_AdvancedPreviewSettings"));
 const FName FCustomizableObjectEditor::TextureAnalyzerTabId(TEXT("CustomizableObjectEditor_TextureAnalyzer"));
 const FName FCustomizableObjectEditor::PerformanceAnalyzerTabId(TEXT("CustomizableObjectEditor_MewPerformanceReport"));
@@ -109,8 +108,8 @@ void FCustomizableObjectEditor::RegisterTabSpawners(const TSharedRef<class FTabM
 		.SetDisplayName(LOCTEXT("ViewportTab", "Viewport"))
 		.SetGroup(WorkspaceMenuCategoryRef);
 	
-	InTabManager->RegisterTabSpawner(ObjectPropertiesTabId, FOnSpawnTab::CreateSP(this, &FCustomizableObjectEditor::SpawnTab_ObjectProperties))
-		.SetDisplayName(LOCTEXT("ObjectPropertiesTab", "Object Properties"))
+	InTabManager->RegisterTabSpawner(DetailsTabId, FOnSpawnTab::CreateSP(this, &FCustomizableObjectEditor::SpawnTab_Details))
+		.SetDisplayName(LOCTEXT("DetailsTab", "Details"))
 		.SetGroup(WorkspaceMenuCategoryRef);
 	
 	InTabManager->RegisterTabSpawner(InstancePropertiesTabId, FOnSpawnTab::CreateSP(this, &FCustomizableObjectEditor::SpawnTab_InstanceProperties))
@@ -121,10 +120,6 @@ void FCustomizableObjectEditor::RegisterTabSpawners(const TSharedRef<class FTabM
 		.SetDisplayName(LOCTEXT("GraphTab", "Object Graph"))
 		.SetGroup(WorkspaceMenuCategoryRef);
 	
-	InTabManager->RegisterTabSpawner(GraphNodePropertiesTabId, FOnSpawnTab::CreateSP(this, &FCustomizableObjectEditor::SpawnTab_GraphNodeProperties))
-		.SetDisplayName(LOCTEXT("GraphNodePropertiesTab", "Object Graph Node Properties"))
-		.SetGroup(WorkspaceMenuCategoryRef);
-
 	InTabManager->RegisterTabSpawner(AdvancedPreviewSettingsTabId, FOnSpawnTab::CreateSP(this, &FCustomizableObjectEditor::SpawnTab_AdvancedPreviewSettings))
 		.SetDisplayName(LOCTEXT("AdvancedPreviewSettingsTab", "Advanced Preview Settings"))
 		.SetGroup(WorkspaceMenuCategoryRef);
@@ -146,10 +141,9 @@ void FCustomizableObjectEditor::RegisterTabSpawners(const TSharedRef<class FTabM
 void FCustomizableObjectEditor::UnregisterTabSpawners(const TSharedRef<class FTabManager>& InTabManager)
 {
 	InTabManager->UnregisterTabSpawner( ViewportTabId );
-	InTabManager->UnregisterTabSpawner( ObjectPropertiesTabId );
+	InTabManager->UnregisterTabSpawner( DetailsTabId );
 	InTabManager->UnregisterTabSpawner( InstancePropertiesTabId );
 	InTabManager->UnregisterTabSpawner( GraphTabId );
-	InTabManager->UnregisterTabSpawner( GraphNodePropertiesTabId );
 	InTabManager->UnregisterTabSpawner( AdvancedPreviewSettingsTabId );
 	InTabManager->UnregisterTabSpawner( TextureAnalyzerTabId );
 	InTabManager->UnregisterTabSpawner( PerformanceAnalyzerTabId );
@@ -167,7 +161,7 @@ FCustomizableObjectEditor::~FCustomizableObjectEditor()
 		}
 	}
 	
-	CustomizableObjectDetailsView.Reset();
+	ObjectDetailsView.Reset();
 	GEditor->UnregisterForUndo(this);
 
 	FCoreUObjectDelegates::OnObjectModified.RemoveAll(this);
@@ -211,18 +205,15 @@ void FCustomizableObjectEditor::InitCustomizableObjectEditor(const EToolkitMode:
 	DetailsViewArgs.NotifyHook = this;
 	DetailsViewArgs.NameAreaSettings = FDetailsViewArgs::ENameAreaSettings::HideNameArea;
 	DetailsViewArgs.bAllowSearch = true;
-	//DetailsViewArgs.bShowActorLabel = false;
 	DetailsViewArgs.bShowObjectLabel = false;
 	DetailsViewArgs.bShowScrollBar = false;
 
-	CustomizableObjectDetailsView = PropPlugin.CreateDetailView( DetailsViewArgs );
-
+	// Detail Panels
+	ObjectDetailsView = PropPlugin.CreateDetailView( DetailsViewArgs );
 	CustomizableInstanceDetailsView = PropPlugin.CreateDetailView( DetailsViewArgs );
-	GraphNodeDetailsView = PropPlugin.CreateDetailView(DetailsViewArgs);
-
-	Viewport = SNew(SCustomizableObjectEditorViewportTabBody)
-		.CustomizableObjectEditor(SharedThis(this));
 
+	// Viewport Panel
+	Viewport = SNew(SCustomizableObjectEditorViewportTabBody).CustomizableObjectEditor(SharedThis(this));
 	Viewport->SetCustomizableObject(CustomizableObject);
 	ViewportClient = Viewport->GetViewportClient();
 
@@ -264,15 +255,9 @@ void FCustomizableObjectEditor::InitCustomizableObjectEditor(const EToolkitMode:
 				(
 					FTabManager::NewStack()
 					->SetSizeCoefficient(0.4f)
-					->AddTab(ObjectPropertiesTabId, ETabState::OpenedTab)
+					->AddTab(DetailsTabId, ETabState::OpenedTab)
 					->AddTab(TagExplorerTabId, ETabState::OpenedTab)
-					->SetForegroundTab(ObjectPropertiesTabId)
-				)
-				->Split
-				(
-					FTabManager::NewStack()
-					->SetSizeCoefficient(0.6f)
-					->AddTab(GraphNodePropertiesTabId, ETabState::OpenedTab)
+					->SetForegroundTab(DetailsTabId)
 				)
 			)
 			->Split
@@ -302,7 +287,7 @@ void FCustomizableObjectEditor::InitCustomizableObjectEditor(const EToolkitMode:
 	const bool bCreateDefaultToolbar = true;
 	FAssetEditorToolkit::InitAssetEditor( Mode, InitToolkitHost, CustomizableObjectEditorAppIdentifier, StandaloneDefaultLayout, bCreateDefaultStandaloneMenu, bCreateDefaultToolbar, CustomizableObject);
 	
-	CustomizableObjectDetailsView->SetObject(CustomizableObject); // Can only be called after initializing the Asset Editor
+	ObjectDetailsView->SetObject(CustomizableObject); // Can only be called after initializing the Asset Editor
 
 	ExtendToolbar();
 	RegenerateMenusAndToolbars();
@@ -351,7 +336,7 @@ void FCustomizableObjectEditor::AddReferencedObjects( FReferenceCollector& Colle
 	Collector.AddReferencedObject( PreviewInstance );
 	Collector.AddReferencedObject( ProjectorParameter );
 	Collector.AddReferencedObject( CustomSettings );
-	Collector.AddReferencedObject(EditorProperties);
+	Collector.AddReferencedObject( EditorProperties );
 }
 
 
@@ -381,23 +366,22 @@ TSharedRef<SDockTab> FCustomizableObjectEditor::SpawnTab_Viewport( const FSpawnT
 	return DockTab;
 }
 
-TSharedRef<SDockTab> FCustomizableObjectEditor::SpawnTab_ObjectProperties( const FSpawnTabArgs& Args )
+TSharedRef<SDockTab> FCustomizableObjectEditor::SpawnTab_Details( const FSpawnTabArgs& Args )
 {
-	check( Args.GetTabId() == ObjectPropertiesTabId );
+	check(Args.GetTabId() == DetailsTabId);
 
-	TSharedRef<SScrollBox> ScrollBox = SNew(SScrollBox)
-		+ SScrollBox::Slot()
-		[
-			CustomizableObjectDetailsView.ToSharedRef()
-		];
+	TSharedRef<SScrollBox> ScrollBox = SNew(SScrollBox) 
+	+SScrollBox::Slot()
+	[
+		ObjectDetailsView.ToSharedRef()
+	];
 
 	ScrollBox->SetScrollBarRightClickDragAllowed(true);
 
-	TSharedRef<SDockTab> DockTab = SNew(SDockTab)
-		.Label( FText::FromString( GetTabPrefix() + LOCTEXT( "CustomizableObjectProperties_TabTitle", "Object Properties" ).ToString() ) )
-		[
-			ScrollBox
-		];
+	TSharedRef<SDockTab> DockTab = SNew(SDockTab).Label(FText::FromString(GetTabPrefix() + LOCTEXT("Details_TabTitle", "Details").ToString()))
+	[
+		ScrollBox
+	];
 
 	DockTab->SetTabIcon(FCustomizableObjectEditorStyle::Get().GetBrush("CustomizableObjectEditor.Tabs.CustomizableObjectProperties"));
 
@@ -549,31 +533,6 @@ TSharedRef<SDockTab> FCustomizableObjectEditor::SpawnTab_Graph( const FSpawnTabA
 }
 
 
-TSharedRef<SDockTab> FCustomizableObjectEditor::SpawnTab_GraphNodeProperties( const FSpawnTabArgs& Args )
-{
-	check( Args.GetTabId().TabType == GraphNodePropertiesTabId );
-
-	TSharedRef<SScrollBox> ScrollBox = SNew(SScrollBox)
-		+ SScrollBox::Slot()
-		[
-			GraphNodeDetailsView.ToSharedRef()
-		];
-
-	ScrollBox->SetScrollBarRightClickDragAllowed(true);
-
-	TSharedRef<SDockTab> DockTab = SNew(SDockTab)
-		.Label(FText::FromString(GetTabPrefix() + LOCTEXT("Graph Node Properties", "Node Properties").ToString()))
-		.TabColorScale(GetTabColorScale())
-		[
-			ScrollBox
-		];
-
-	DockTab->SetTabIcon(FCustomizableObjectEditorStyle::Get().GetBrush("CustomizableObjectEditor.Tabs.NodeProperties"));
-
-	return DockTab;
-}
-
-
 TSharedRef<SDockTab> FCustomizableObjectEditor::SpawnTab_AdvancedPreviewSettings(const FSpawnTabArgs& Args)
 {
 	check(Args.GetTabId() == AdvancedPreviewSettingsTabId);
@@ -1151,9 +1110,9 @@ void FCustomizableObjectEditor::PostUndo(bool bSuccess)
 {
 	if (bSuccess)
 	{
-		if (CustomizableObjectDetailsView.IsValid())
+		if (ObjectDetailsView.IsValid())
 		{
-			CustomizableObjectDetailsView->RemoveInvalidObjects();
+			ObjectDetailsView->RemoveInvalidObjects();
 		}
 
 		if (CustomizableInstanceDetailsView.IsValid())
@@ -1161,11 +1120,6 @@ void FCustomizableObjectEditor::PostUndo(bool bSuccess)
 			CustomizableInstanceDetailsView->RemoveInvalidObjects();
 		}
 
-		if (GraphNodeDetailsView.IsValid())
-		{
-			GraphNodeDetailsView->RemoveInvalidObjects();
-		}
-
 		GraphEditor->NotifyGraphChanged();
 		CustomizableObject->MarkPackageDirty();
 
@@ -1680,9 +1634,16 @@ void FCustomizableObjectEditor::OnSelectedGraphNodesChanged(const FGraphPanelSel
 	}
 
 	// Standard details
-	if (GraphNodeDetailsView.IsValid())
+	if (ObjectDetailsView.IsValid())
 	{
-		GraphNodeDetailsView->SetObjects(Objects);
+		if (Objects.Num())
+		{
+			ObjectDetailsView->SetObjects(Objects);
+		}
+		else
+		{
+			ObjectDetailsView->SetObject(CustomizableObject);
+		}
 	}		
 	
 	if (!bRecursionGuard) // Calling the following functions will unselect some nodes causing OnSelectedGraphNodesChanged to be called again
@@ -2203,15 +2164,6 @@ void FCustomizableObjectEditor::OnUpdatePreviewInstance(UCustomizableObjectInsta
 }
 
 
-void FCustomizableObjectEditor::UpdateObjectProperties()
-{
-	if (CustomizableObjectDetailsView.IsValid())
-	{
-		CustomizableObjectDetailsView->ForceRefresh();
-	}
-}
-
-
 void FCustomizableObjectEditor::OnNodeTitleCommitted(const FText& NewText, ETextCommit::Type CommitInfo, UEdGraphNode* NodeBeingChanged)
 {
 	if (NodeBeingChanged)
diff --git a/Engine/Plugins/Mutable/Source/CustomizableObjectEditor/Private/MuCOE/CustomizableObjectEditor.h b/Engine/Plugins/Mutable/Source/CustomizableObjectEditor/Private/MuCOE/CustomizableObjectEditor.h
--- a/Engine/Plugins/Mutable/Source/CustomizableObjectEditor/Private/MuCOE/CustomizableObjectEditor.h
+++ b/Engine/Plugins/Mutable/Source/CustomizableObjectEditor/Private/MuCOE/CustomizableObjectEditor.h
@@ -187,9 +187,6 @@ public:
 	/** Callback to notify the editor when the PreviewInstance has been updated */
 	void OnUpdatePreviewInstance(UCustomizableObjectInstance* Instance);
 
-	/** Called when the Object Properties needs to be updated */
-	virtual void UpdateObjectProperties() override;
-
 	/** Debug the object as a raw mutable data in the internal tools. */
 	void DebugObject() const;
 
@@ -198,10 +195,9 @@ public:
 	
 private:
 	TSharedRef<SDockTab> SpawnTab_Viewport(const FSpawnTabArgs& Args);
-	TSharedRef<SDockTab> SpawnTab_ObjectProperties(const FSpawnTabArgs& Args);
+	TSharedRef<SDockTab> SpawnTab_Details(const FSpawnTabArgs& Args);
 	TSharedRef<SDockTab> SpawnTab_InstanceProperties(const FSpawnTabArgs& Args);
 	TSharedRef<SDockTab> SpawnTab_Graph(const FSpawnTabArgs& Args);
-	TSharedRef<SDockTab> SpawnTab_GraphNodeProperties(const FSpawnTabArgs& Args);
 	TSharedRef<SDockTab> SpawnTab_AdvancedPreviewSettings(const FSpawnTabArgs& Args);
 	TSharedRef<SDockTab> SpawnTab_TextureAnalyzer(const FSpawnTabArgs& Args);
 	TSharedRef<SDockTab> SpawnTab_PerformanceAnalyzer(const FSpawnTabArgs& Args);
@@ -305,10 +301,9 @@ public:
 
 	/**	The tab ids for all the tabs used */
 	static const FName ViewportTabId;
-	static const FName ObjectPropertiesTabId;
+	static const FName DetailsTabId;
 	static const FName InstancePropertiesTabId;
 	static const FName GraphTabId;
-	static const FName GraphNodePropertiesTabId;
 	static const FName SystemPropertiesTabId;
 	static const FName AdvancedPreviewSettingsTabId;
 	static const FName TextureAnalyzerTabId;
@@ -332,12 +327,11 @@ private:
 	TSharedPtr<IDetailsView> CustomizableInstanceDetailsView;
 
 
-	/** Property View */
-	TSharedPtr<class IDetailsView> CustomizableObjectDetailsView;
+	/** Property View for Customizable Object and Nodes. */
+	TSharedPtr<class IDetailsView> ObjectDetailsView;
 
 	/** */
 	TSharedPtr<SGraphEditor> GraphEditor;
-	TSharedPtr<class IDetailsView> GraphNodeDetailsView;
 
 	/** Widget to select which node pins are visible. */
 	TSharedPtr<class SCustomizableObjectNodePinViewer> NodePinViewer;
diff --git a/Engine/Plugins/Mutable/Source/CustomizableObjectEditor/Private/MuCOE/Nodes/CustomizableObjectNodeObject.cpp b/Engine/Plugins/Mutable/Source/CustomizableObjectEditor/Private/MuCOE/Nodes/CustomizableObjectNodeObject.cpp
--- a/Engine/Plugins/Mutable/Source/CustomizableObjectEditor/Private/MuCOE/Nodes/CustomizableObjectNodeObject.cpp
+++ b/Engine/Plugins/Mutable/Source/CustomizableObjectEditor/Private/MuCOE/Nodes/CustomizableObjectNodeObject.cpp
@@ -695,13 +695,6 @@ void UCustomizableObjectNodeObject::SetParentObject(UCustomizableObject* Customi
 		if (CustomizableObject)
 		{
 			CustomizableObject->GetPrivate()->SetIsChildObject(ParentObject != nullptr);
-
-			TSharedPtr<ICustomizableObjectEditor> Editor = GetGraphEditor();
-
-			if (Editor.IsValid())
-			{
-				Editor->UpdateObjectProperties();
-			}
 		}
 	}
 }
diff --git a/Engine/Plugins/Mutable/Source/CustomizableObjectEditor/Public/MuCOE/ICustomizableObjectEditor.h b/Engine/Plugins/Mutable/Source/CustomizableObjectEditor/Public/MuCOE/ICustomizableObjectEditor.h
--- a/Engine/Plugins/Mutable/Source/CustomizableObjectEditor/Public/MuCOE/ICustomizableObjectEditor.h
+++ b/Engine/Plugins/Mutable/Source/CustomizableObjectEditor/Public/MuCOE/ICustomizableObjectEditor.h
@@ -29,8 +29,6 @@ public:
 
 	virtual void UpdateGraphNodeProperties() {}
 
-	virtual void UpdateObjectProperties() {}
-
 	// Make sure the graph editor selects a specific node
 	virtual void SelectNode(const UEdGraphNode* Node) {}
 
