[mutable] Make custom container for log messages thread-safe to add messages.
- Cleanup some old naming convention.

#rnx

diff --git a/Engine/Plugins/Mutable/Source/CustomizableObjectEditor/Private/MuCOE/CustomizableObjectCompileRunnable.cpp b/Engine/Plugins/Mutable/Source/CustomizableObjectEditor/Private/MuCOE/CustomizableObjectCompileRunnable.cpp
--- a/Engine/Plugins/Mutable/Source/CustomizableObjectEditor/Private/MuCOE/CustomizableObjectCompileRunnable.cpp
+++ b/Engine/Plugins/Mutable/Source/CustomizableObjectEditor/Private/MuCOE/CustomizableObjectCompileRunnable.cpp
@@ -249,8 +249,8 @@ uint32 FCustomizableObjectCompileRunnable::Run()
 	Model = Compiler->Compile(MutableRoot);
 
 	// Dump all the log messages from the compiler
-	mu::Ptr<const mu::ErrorLog> pLog = Compiler->GetLog();
-	for (int i = 0; i < pLog->GetMessageCount(); ++i)
+	TSharedPtr<const mu::FErrorLog> pLog = Compiler->GetLog();
+	for (int32 i = 0; i < pLog->GetMessageCount(); ++i)
 	{
 		const FString& Message = pLog->GetMessageText(i);
 		const mu::ErrorLogMessageType MessageType = pLog->GetMessageType(i);
@@ -269,11 +269,11 @@ uint32 FCustomizableObjectCompileRunnable::Run()
 			}
 			}();
 
-			if (MessageAttachedData.m_unassignedUVs && MessageAttachedData.m_unassignedUVsSize > 0) 
+			if (MessageAttachedData.UnassignedUVs && MessageAttachedData.UnassignedUVsSize > 0) 
 			{			
 				TSharedPtr<FErrorAttachedData> ErrorAttachedData = MakeShared<FErrorAttachedData>();
 				ErrorAttachedData->UnassignedUVs.Reset();
-				ErrorAttachedData->UnassignedUVs.Append(MessageAttachedData.m_unassignedUVs, MessageAttachedData.m_unassignedUVsSize);
+				ErrorAttachedData->UnassignedUVs.Append(MessageAttachedData.UnassignedUVs, MessageAttachedData.UnassignedUVsSize);
 				const UObject* Context = static_cast<const UObject*>(pLog->GetMessageContext(i));
 				ArrayErrors.Add(FError(Severity, FText::AsCultureInvariant(Message), ErrorAttachedData, Context, SpamBin));
 			}
diff --git a/Engine/Plugins/Mutable/Source/MutableTools/Internal/MuT/AST.h b/Engine/Plugins/Mutable/Source/MutableTools/Internal/MuT/AST.h
--- a/Engine/Plugins/Mutable/Source/MutableTools/Internal/MuT/AST.h
+++ b/Engine/Plugins/Mutable/Source/MutableTools/Internal/MuT/AST.h
@@ -537,7 +537,7 @@ namespace mu
 
         //! Operations referring to this one. They may be null: elements are never removed
         //! from this TArray.
-		TArray<ASTOp*, TInlineAllocator<4> > m_parents;
+		TArray<ASTOp*, TInlineAllocator<4> > Parents;
 
     public:
 
@@ -1143,9 +1143,7 @@ namespace mu
     };
 
 
-    //---------------------------------------------------------------------------------------------
-    //! This class of operation is used to ease the transition from legacy code generation.
-    //---------------------------------------------------------------------------------------------
+    /** This class of operation is used to ease the transition from legacy code generation. */
     class ASTOpFixed : public ASTOp
     {
     public:
@@ -1163,8 +1161,7 @@ namespace mu
 
         ~ASTOpFixed() override;
 
-		virtual EOpType GetOpType() const override { return op.type; }
-
+		virtual EOpType GetOpType() const override;
 		virtual Ptr<ASTOp> Clone( MapChildFuncRef mapChild ) const override;
 		virtual void ForEachChild( const TFunctionRef<void(ASTChild&)> f ) override;
 		virtual void Link( FProgram& program, FLinkerOptions* Options) override;
@@ -1182,9 +1179,7 @@ namespace mu
 		virtual Ptr<ImageSizeExpression> GetImageSizeExpression() const override;
 		virtual FSourceDataDescriptor GetSourceDataDescriptor(FGetSourceDataDescriptorContext* Context) const override;
 
-        //---------------------------------------------------------------------------------------------
-        //! Add child for OP fixed nodes.
-        //---------------------------------------------------------------------------------------------
+        /** Add child for OP fixed nodes. */
         inline void SetChild( OP::ADDRESS& at, const Ptr<ASTOp>& child )
         {
             // hack check
@@ -1255,22 +1250,69 @@ namespace mu
     };
 
 
+	/** This operation represents another operation that is still being generated. */
+	//class ASTOpFuture : public ASTOp
+	//{
+	//public:
+
+	//	//! Fixed size op structure. Any code ADDRESS in it, refers to the children TArray below.
+	//	OP op;
+
+	//public:
+	//	ASTOpFuture();
+
+	//	// ASTOp interface
+	//	virtual EOpType GetOpType() const override
+	//	{
+	//		EnsureReady();
+	//		return Real->GetOpType();
+	//	}
+	//	virtual Ptr<ASTOp> Clone(MapChildFuncRef mapChild) const override;
+	//	virtual void ForEachChild(const TFunctionRef<void(ASTChild&)> f) override;
+	//	virtual void Link(FProgram& program, FLinkerOptions* Options) override;
+	//	virtual bool IsEqual(const ASTOp& otherUntyped) const override;
+	//	virtual uint64 Hash() const override;
+	//	virtual FImageDesc GetImageDesc(bool returnBestOption, FGetImageDescContext* context) const override;
+	//	virtual void GetLayoutBlockSize(int* pBlockX, int* pBlockY) override;
+	//	virtual Ptr<ASTOp> OptimiseSize() const override;
+	//	virtual Ptr<ASTOp> OptimiseSemantic(const FModelOptimizationOptions&, int32 Pass) const override;
+	//	virtual Ptr<ASTOp> OptimiseSink(const FModelOptimizationOptions&, FOptimizeSinkContext&) const override;
+	//	virtual FBoolEvalResult EvaluateBool(ASTOpList& facts, FEvaluateBoolCache* cache) const override;
+	//	virtual int32 EvaluateInt(ASTOpList& facts, bool& unknown) const override;
+	//	virtual bool IsImagePlainConstant(FVector4f& colour) const override;
+	//	virtual bool IsColourConstant(FVector4f& colour) const override;
+	//	virtual Ptr<ImageSizeExpression> GetImageSizeExpression() const override;
+	//	virtual FSourceDataDescriptor GetSourceDataDescriptor(FGetSourceDataDescriptorContext* Context) const override;
+
+	//	// Own interface
+
+	//	/** Wait until the real operation has been completely generated. */
+	//	void EnsureReady() const
+	//	{
+
+	//	}
+
+	//private:
+
+	//};
+
+
     //---------------------------------------------------------------------------------------------
     //---------------------------------------------------------------------------------------------
     //---------------------------------------------------------------------------------------------
     inline void ASTChild::AddParent()
     {
-        ParentIndexInChild = Child->m_parents.Num();
-        Child->m_parents.Add(Parent);
+        ParentIndexInChild = Child->Parents.Num();
+        Child->Parents.Add(Parent);
     }
 
     //---------------------------------------------------------------------------------------------
     inline void ASTChild::ClearParent()
     {
-        check( ParentIndexInChild<Child->m_parents.Num() );
+        check( ParentIndexInChild<Child->Parents.Num() );
 		// Can't do this, because the indices are stored in children.
 		//Child->m_parents.RemoveAtSwap(ParentIndexInChild);
-		Child->m_parents[ParentIndexInChild] = nullptr;
+		Child->Parents[ParentIndexInChild] = nullptr;
 	}
 
 
diff --git a/Engine/Plugins/Mutable/Source/MutableTools/Internal/MuT/ErrorLogPrivate.h b/Engine/Plugins/Mutable/Source/MutableTools/Internal/MuT/ErrorLogPrivate.h
deleted file mode 100644
diff --git a/Engine/Plugins/Mutable/Source/MutableTools/Internal/MuT/TablePrivate.h b/Engine/Plugins/Mutable/Source/MutableTools/Internal/MuT/TablePrivate.h
--- a/Engine/Plugins/Mutable/Source/MutableTools/Internal/MuT/TablePrivate.h
+++ b/Engine/Plugins/Mutable/Source/MutableTools/Internal/MuT/TablePrivate.h
@@ -3,8 +3,6 @@
 #pragma once
 
 #include "MuT/Table.h"
-
-#include "MuT/ErrorLogPrivate.h"
 #include "MuR/MutableMath.h"
 
 
diff --git a/Engine/Plugins/Mutable/Source/MutableTools/Private/MuT/AST.cpp b/Engine/Plugins/Mutable/Source/MutableTools/Private/MuT/AST.cpp
--- a/Engine/Plugins/Mutable/Source/MutableTools/Private/MuT/AST.cpp
+++ b/Engine/Plugins/Mutable/Source/MutableTools/Private/MuT/AST.cpp
@@ -13,7 +13,6 @@
 #include "MuT/ASTOpMeshRemoveMask.h"
 #include "MuT/ASTOpImageMultiLayer.h"
 #include "MuT/ASTOpParameter.h"
-#include "MuT/ErrorLogPrivate.h"
 #include "Trace/Detail/Channel.h"
 
 #include <atomic>
@@ -89,11 +88,11 @@ ASTChild& ASTChild::operator=( ASTChild&& rhs )
 //-------------------------------------------------------------------------------------------------
 void ASTOp::ForEachParent(const TFunctionRef<void(ASTOp*)> f) const
 {
-	for (ASTOp* p : m_parents)
+	for (ASTOp* Parent : Parents)
 	{
-		if (p)
+		if (Parent)
 		{
-			f(p);
+			f(Parent);
 		}
 	}
 }
@@ -170,8 +169,8 @@ void ASTOp::Assert()
 //                }
 //            });
 //            check(found);
-            check( c.ParentIndexInChild < size_t(c.child()->m_parents.Num()) );
-            check( c.child()->m_parents[c.ParentIndexInChild]==this );
+            check( c.ParentIndexInChild < c.child()->Parents.Num() );
+            check( c.child()->Parents[c.ParentIndexInChild]==this );
         }
      });
 }
@@ -815,9 +814,9 @@ void ASTOp::Traverse_BottomUp_Unique
 //-------------------------------------------------------------------------------------------------
 //-------------------------------------------------------------------------------------------------
 //-------------------------------------------------------------------------------------------------
-int ASTOp::GetParentCount() const
+int32 ASTOp::GetParentCount() const
 {
-    int result=0;
+    int32 result=0;
     ForEachParent( [&](const ASTOp* p)
     {
         if (p!=nullptr) ++result;
@@ -833,13 +832,13 @@ void ASTOp::Replace( const Ptr<ASTOp>& node, const Ptr<ASTOp>& other )
         return;
     }
 
-	TArray<ASTOp*, TInlineAllocator<4> > parentsCopy = node->m_parents;
+	TArray<ASTOp*, TInlineAllocator<4> > ParentsCopy = node->Parents;
 
-    for(ASTOp* p:parentsCopy)
+    for(ASTOp* Parent: ParentsCopy)
     {
-        if(p)
+        if(Parent)
         {
-            p->ForEachChild( [=](ASTChild& c)
+			Parent->ForEachChild( [=](ASTChild& c)
             {
                 if(c.Child==node)
                 {
@@ -1044,6 +1043,12 @@ bool ASTOpFixed::IsEqual(const ASTOp& otherUntyped) const
 }
 
 
+EOpType ASTOpFixed::GetOpType() const 
+{ 
+	return op.type; 
+}
+
+
 mu::Ptr<ASTOp> ASTOpFixed::Clone( MapChildFuncRef mapChild ) const
 {
     Ptr<ASTOpFixed> n = new ASTOpFixed();
diff --git a/Engine/Plugins/Mutable/Source/MutableTools/Private/MuT/ASTOpReferenceResource.cpp b/Engine/Plugins/Mutable/Source/MutableTools/Private/MuT/ASTOpReferenceResource.cpp
--- a/Engine/Plugins/Mutable/Source/MutableTools/Private/MuT/ASTOpReferenceResource.cpp
+++ b/Engine/Plugins/Mutable/Source/MutableTools/Private/MuT/ASTOpReferenceResource.cpp
@@ -23,7 +23,7 @@ namespace mu
 		if (otherUntyped.GetOpType() == GetOpType())
 		{
 			const ASTOpReferenceResource* other = static_cast<const ASTOpReferenceResource*>(&otherUntyped);
-			return type == other->type && ID == other->ID && bForceLoad == other->bForceLoad && ImageDesc == other->ImageDesc;
+			return Type == other->Type && ID == other->ID && bForceLoad == other->bForceLoad && ImageDesc == other->ImageDesc;
 		}
 		return false;
 	}
@@ -33,7 +33,7 @@ namespace mu
 	mu::Ptr<ASTOp> ASTOpReferenceResource::Clone(MapChildFuncRef) const
 	{
 		Ptr<ASTOpReferenceResource> n = new ASTOpReferenceResource();
-		n->type = type;
+		n->Type = Type;
 		n->ID = ID;
 		n->bForceLoad = bForceLoad;
 		n->ImageDesc = ImageDesc;
@@ -44,7 +44,7 @@ namespace mu
 	//-------------------------------------------------------------------------------------------------
 	uint64 ASTOpReferenceResource::Hash() const
 	{
-		uint64 res = std::hash<uint64>()(uint64(type));
+		uint64 res = std::hash<uint64>()(uint64(Type));
 		hash_combine(res, ID);
 		return res;
 	}
@@ -56,14 +56,14 @@ namespace mu
 		if (!linkedAddress)
 		{
 			OP::ResourceReferenceArgs Args;
-			FMemory::Memset(&Args, 0, sizeof(Args));
+			FMemory::Memzero(Args);
 			Args.ID = ID;
 			Args.ForceLoad = bForceLoad ? 1 : 0;
 			Args.ImageDesc = ImageDesc;
 
 			linkedAddress = (OP::ADDRESS)program.OpAddress.Num();
 			program.OpAddress.Add((uint32)program.ByteCode.Num());
-			AppendCode(program.ByteCode, type);
+			AppendCode(program.ByteCode, Type);
 			AppendCode(program.ByteCode, Args);
 		}
 	}
@@ -79,7 +79,7 @@ namespace mu
 	//-------------------------------------------------------------------------------------------------
 	void ASTOpReferenceResource::GetBlockLayoutSize(uint64 BlockId, int32* pBlockX, int32* pBlockY, FBlockLayoutSizeCache*)
 	{
-		switch (type)
+		switch (Type)
 		{
 
 		case EOpType::IM_REFERENCE:
@@ -98,7 +98,7 @@ namespace mu
 	//-------------------------------------------------------------------------------------------------
 	void ASTOpReferenceResource::GetLayoutBlockSize(int32* pBlockX, int32* pBlockY)
 	{
-		switch (type)
+		switch (Type)
 		{
 
 		case EOpType::IM_REFERENCE:
diff --git a/Engine/Plugins/Mutable/Source/MutableTools/Private/MuT/ASTOpReferenceResource.h b/Engine/Plugins/Mutable/Source/MutableTools/Private/MuT/ASTOpReferenceResource.h
--- a/Engine/Plugins/Mutable/Source/MutableTools/Private/MuT/ASTOpReferenceResource.h
+++ b/Engine/Plugins/Mutable/Source/MutableTools/Private/MuT/ASTOpReferenceResource.h
@@ -24,7 +24,7 @@ namespace mu
 	public:
 
 		//! Type of switch
-		EOpType type = EOpType::NONE;
+		EOpType Type = EOpType::NONE;
 		
 		/** */
 		bool bForceLoad = false;
@@ -43,7 +43,7 @@ namespace mu
 		// Own interface
 
 		// ASTOp interface
-		virtual EOpType GetOpType() const override { return type; }
+		virtual EOpType GetOpType() const override { return Type; }
 		virtual void ForEachChild(const TFunctionRef<void(ASTChild&)>) override;
 		virtual bool IsEqual(const ASTOp& otherUntyped) const override;
 		virtual Ptr<ASTOp> Clone(MapChildFuncRef mapChild) const override;
diff --git a/Engine/Plugins/Mutable/Source/MutableTools/Private/MuT/CodeGenerator.cpp b/Engine/Plugins/Mutable/Source/MutableTools/Private/MuT/CodeGenerator.cpp
--- a/Engine/Plugins/Mutable/Source/MutableTools/Private/MuT/CodeGenerator.cpp
+++ b/Engine/Plugins/Mutable/Source/MutableTools/Private/MuT/CodeGenerator.cpp
@@ -46,7 +46,7 @@
 #include "MuT/CodeGenerator_SecondPass.h"
 #include "MuT/CodeOptimiser.h"
 #include "MuT/CompilerPrivate.h"
-#include "MuT/ErrorLogPrivate.h"
+#include "MuT/ErrorLog.h"
 #include "MuT/NodeColour.h"
 #include "MuT/NodeColourConstant.h"
 #include "MuT/NodeComponent.h"
@@ -104,7 +104,7 @@ namespace mu
 		CompilerOptions = options;
 
 		// Create the message log
-		ErrorLog = new mu::ErrorLog;
+		ErrorLog = MakeShared<FErrorLog>();
 
 		// Add the parent at the top of the hierarchy
 		CurrentParents.Add(FParentKey());
@@ -689,7 +689,7 @@ namespace mu
 
 			if (TagIndex < 0)
 			{
-				ErrorLog->GetPrivate()->Add(
+				ErrorLog->Add(
 					FString::Printf(TEXT("Unknown tag found in component variation [%s]."), *Tag),
 					ELMT_WARNING,
 					Node->GetMessageContext(),
@@ -1266,7 +1266,7 @@ namespace mu
 						if (LayoutIndex >= MeshResults.GeneratedLayouts.Num() ||
 							LayoutIndex >= MeshResults.LayoutOps.Num())
 						{
-							ErrorLog->GetPrivate()->Add("Missing layout in object, or its parent.", ELMT_ERROR, SurfaceNode->GetMessageContext());
+							ErrorLog->Add("Missing layout in object, or its parent.", ELMT_ERROR, SurfaceNode->GetMessageContext());
 						}
 						else
 						{
@@ -1892,7 +1892,7 @@ namespace mu
 				if (usedNames.Contains(strChildName))
 				{
 					FString Msg = FString::Printf(TEXT("Object group has more than one children with the same name [%s]."), *strChildName );
-					ErrorLog->GetPrivate()->Add(Msg, ELMT_WARNING, InNode->GetMessageContext());
+					ErrorLog->Add(Msg, ELMT_WARNING, InNode->GetMessageContext());
 				}
 				else
 				{
@@ -1908,7 +1908,7 @@ namespace mu
     {
         // Log a warning
 		FString Msg = FString::Printf(TEXT("Required connection not found: %s"), Where);
-        ErrorLog->GetPrivate()->Add( Msg, ELMT_ERROR, ErrorContext);
+        ErrorLog->Add( Msg, ELMT_ERROR, ErrorContext);
 
         // Create a constant node
         Ptr<NodeBoolConstant> pNode = new NodeBoolConstant();
@@ -2061,7 +2061,7 @@ namespace mu
 						{
 							// When extending a mesh section the added mesh sectio will use the layout strategy of the base one
 							FString Msg = FString::Printf(TEXT("Extended mesh section layout count is differenta than the mesh being extended."));
-							ErrorLog->GetPrivate()->Add(Msg, ELMT_INFO, Edit->GetMessageContext(), ErrorContext);
+							ErrorLog->Add(Msg, ELMT_INFO, Edit->GetMessageContext(), ErrorContext);
 						}
 
 						for (int32 LayoutIndex = 0; LayoutIndex < BaseMeshResult.GeneratedLayouts.Num(); ++LayoutIndex)
@@ -2083,7 +2083,7 @@ namespace mu
 							{
 								// When extending a mesh section the added mesh sectio will use the layout strategy of the base one
 								FString Msg = FString::Printf(TEXT("Extended mesh section layout [%d] is using a different strategy than the section being extended. The base strategy will be used."), LayoutIndex);
-								ErrorLog->GetPrivate()->Add(Msg, ELMT_INFO, Edit->GetMessageContext(), ErrorContext);
+								ErrorLog->Add(Msg, ELMT_INFO, Edit->GetMessageContext(), ErrorContext);
 							}
 						}
 					}
@@ -2307,7 +2307,7 @@ namespace mu
 
 				if (!op->clip)
 				{
-					ErrorLog->GetPrivate()->Add("Clip mesh has not been generated", ELMT_ERROR, ErrorContext);
+					ErrorLog->Add("Clip mesh has not been generated", ELMT_ERROR, ErrorContext);
 					continue;
 				}
 
@@ -2368,7 +2368,7 @@ namespace mu
 
 					if (!op->MaskImage)
 					{
-						ErrorLog->GetPrivate()->Add("Clip UV mask has not been generated", ELMT_ERROR, ErrorContext);
+						ErrorLog->Add("Clip UV mask has not been generated", ELMT_ERROR, ErrorContext);
 						continue;
 					}
 				}
@@ -2504,8 +2504,7 @@ namespace mu
 
 				if (!ClipOp->ClipShape)
 				{
-					ErrorLog->GetPrivate()->Add
-					("Clip shape mesh has not been generated", ELMT_ERROR, ErrorContext);
+					ErrorLog->Add("Clip shape mesh has not been generated", ELMT_ERROR, ErrorContext);
 				}
 				else
 				{
@@ -2560,7 +2559,7 @@ namespace mu
 
 						if (!TransformOp->boundingMesh)
 						{
-							ErrorLog->GetPrivate()->Add("Bounding mesh has not been generated", ELMT_ERROR, ErrorContext);
+							ErrorLog->Add("Bounding mesh has not been generated", ELMT_ERROR, ErrorContext);
 							continue;
 						}
 						else
@@ -2568,7 +2567,7 @@ namespace mu
 							ASTOp::EClosedMeshTest IsClosed = TransformOp->boundingMesh->IsClosedMesh();
 							if (IsClosed==ASTOp::EClosedMeshTest::No)
 							{
-								ErrorLog->GetPrivate()->Add("Mesh used for clipping is not closed.", ELMT_WARNING, ErrorContext);
+								ErrorLog->Add("Mesh used for clipping is not closed.", ELMT_WARNING, ErrorContext);
 							}
 						}
 					}
@@ -2720,7 +2719,7 @@ namespace mu
 					{
 						// When extending a mesh section it is mandatory to provide textures for all section textures handled by Mutable.
 						FString Msg = FString::Printf(TEXT("Required texture [%s] is missing when trying to extend a mesh section."), *ImageData.MaterialParameterName);
-						ErrorLog->GetPrivate()->Add(Msg, ELMT_INFO, Edit->GetMessageContext(), ModifiedNodeErrorContext);
+						ErrorLog->Add(Msg, ELMT_INFO, Edit->GetMessageContext(), ModifiedNodeErrorContext);
 					}
 
 					continue;
@@ -2730,7 +2729,7 @@ namespace mu
 
 				if (LayoutIndex >= ExtraLayouts.Num() || !ExtraLayouts[LayoutIndex].Layout)
 				{
-					ErrorLog->GetPrivate()->Add(TEXT("Trying to extend a layout that doesn't exist."), ELMT_WARNING, Edit->GetMessageContext(), ModifiedNodeErrorContext);
+					ErrorLog->Add(TEXT("Trying to extend a layout that doesn't exist."), ELMT_WARNING, Edit->GetMessageContext(), ModifiedNodeErrorContext);
 				}
 				else
 				{
@@ -2853,7 +2852,7 @@ namespace mu
 
 				if (!bAtLeastSomeTexture)
 				{
-					ErrorLog->GetPrivate()->Add(TEXT("A mesh section modifier applies to a section but no texture matches."), ELMT_WARNING, Edit->GetMessageContext(), Node.GetMessageContext());
+					ErrorLog->Add(TEXT("A mesh section modifier applies to a section but no texture matches."), ELMT_WARNING, Edit->GetMessageContext(), Node.GetMessageContext());
 				}
 			}
 		}
diff --git a/Engine/Plugins/Mutable/Source/MutableTools/Private/MuT/CodeGenerator.h b/Engine/Plugins/Mutable/Source/MutableTools/Private/MuT/CodeGenerator.h
--- a/Engine/Plugins/Mutable/Source/MutableTools/Private/MuT/CodeGenerator.h
+++ b/Engine/Plugins/Mutable/Source/MutableTools/Private/MuT/CodeGenerator.h
@@ -232,7 +232,7 @@ namespace mu
 		FirstPassGenerator FirstPass;
 
         //!
-        ErrorLogPtr ErrorLog;
+        TSharedPtr<FErrorLog> ErrorLog;
 
         //! After the entire code generation this contains the information about all the states
         typedef TArray< TPair<FObjectState, Ptr<ASTOp>> > StateList;
@@ -892,18 +892,18 @@ namespace mu
 
 		if (NumRows == 0)
 		{
-			ErrorLog->GetPrivate()->Add("The table has no rows.", ELMT_ERROR, node.GetMessageContext());
+			ErrorLog->Add("The table has no rows.", ELMT_ERROR, node.GetMessageContext());
 			return nullptr;
 		}
         else if (ColIndex < 0)
         {
-            ErrorLog->GetPrivate()->Add("Table column not found.", ELMT_ERROR, node.GetMessageContext());
+            ErrorLog->Add("Table column not found.", ELMT_ERROR, node.GetMessageContext());
             return nullptr;
         }
 
         if (NodeTable->GetPrivate()->Columns[ ColIndex ].Type != TYPE )
         {
-            ErrorLog->GetPrivate()->Add("Table column type is not the right type.", ELMT_ERROR, node.GetMessageContext());
+            ErrorLog->Add("Table column type is not the right type.", ELMT_ERROR, node.GetMessageContext());
             return nullptr;
         }
 
@@ -919,7 +919,7 @@ namespace mu
             check(RowIndex <= 0xFFFF);
 			auto Condition = (uint16)RowIndex;
 
-            Ptr<ASTOp> Branch = GenerateOption( node, ColIndex, (int)RowIndex, ErrorLog.get() );
+            Ptr<ASTOp> Branch = GenerateOption( node, ColIndex, RowIndex, ErrorLog.Get() );
 
 			if (Branch || TYPE != ETableColumnType::Mesh)
 			{
diff --git a/Engine/Plugins/Mutable/Source/MutableTools/Private/MuT/CodeGenerator_Color.cpp b/Engine/Plugins/Mutable/Source/MutableTools/Private/MuT/CodeGenerator_Color.cpp
--- a/Engine/Plugins/Mutable/Source/MutableTools/Private/MuT/CodeGenerator_Color.cpp
+++ b/Engine/Plugins/Mutable/Source/MutableTools/Private/MuT/CodeGenerator_Color.cpp
@@ -19,7 +19,6 @@
 #include "MuT/CodeGenerator.h"
 #include "MuT/CodeGenerator_FirstPass.h"
 #include "MuT/ErrorLog.h"
-#include "MuT/ErrorLogPrivate.h"
 #include "MuT/Node.h"
 #include "MuT/NodeColour.h"
 #include "MuT/NodeColourArithmeticOperation.h"
@@ -260,7 +259,7 @@ namespace mu
 			if (tagIndex < 0)
 			{
 				FString Msg = FString::Printf(TEXT("Unknown tag found in color variation [%s]."), *tag);
-				ErrorLog->GetPrivate()->Add(Msg, ELMT_WARNING, Typed->GetMessageContext());
+				ErrorLog->Add(Msg, ELMT_WARNING, Typed->GetMessageContext());
 				continue;
 			}
 
@@ -488,7 +487,7 @@ namespace mu
 		const NodeColourTable& node = *Typed;
 
 		result.op = GenerateTableSwitch<NodeColourTable, ETableColumnType::Color, EOpType::CO_SWITCH>(node,
-			[this, &Options](const NodeColourTable& node, int32 colIndex, int32 row, mu::ErrorLog* pErrorLog)
+			[this, &Options](const NodeColourTable& node, int32 colIndex, int32 row, FErrorLog*)
 			{
 				Ptr<NodeColourConstant> CellData = new NodeColourConstant();
 				FVector4f Colour = node.Table->GetPrivate()->Rows[row].Values[colIndex].Color;
@@ -508,11 +507,11 @@ namespace mu
 
 
 	//---------------------------------------------------------------------------------------------
-	mu::Ptr<ASTOp> CodeGenerator::GenerateMissingColourCode(const TCHAR* strWhere, const void* errorContext)
+	mu::Ptr<ASTOp> CodeGenerator::GenerateMissingColourCode(const TCHAR* StrWhere, const void* ErrorContext)
 	{
 		// Log a warning
-		FString Msg = FString::Printf(TEXT("Required connection not found: %s"), strWhere);
-		ErrorLog->GetPrivate()->Add(Msg, ELMT_ERROR, errorContext);
+		FString Msg = FString::Printf(TEXT("Required connection not found: %s"), StrWhere);
+		ErrorLog->Add(Msg, ELMT_ERROR, ErrorContext);
 
 		// Create a constant colour node
 		Ptr<NodeColourConstant> pNode = new NodeColourConstant();
diff --git a/Engine/Plugins/Mutable/Source/MutableTools/Private/MuT/CodeGenerator_ExtensionData.cpp b/Engine/Plugins/Mutable/Source/MutableTools/Private/MuT/CodeGenerator_ExtensionData.cpp
--- a/Engine/Plugins/Mutable/Source/MutableTools/Private/MuT/CodeGenerator_ExtensionData.cpp
+++ b/Engine/Plugins/Mutable/Source/MutableTools/Private/MuT/CodeGenerator_ExtensionData.cpp
@@ -17,7 +17,6 @@
 #include "MuT/CodeGenerator.h"
 #include "MuT/CompilerPrivate.h"
 #include "MuT/ErrorLog.h"
-#include "MuT/ErrorLogPrivate.h"
 #include "MuT/NodeScalar.h"
 #include "MuT/NodeExtensionData.h"
 #include "MuT/NodeExtensionDataConstant.h"
@@ -75,7 +74,7 @@ class Node;
 			Data = MakeShared<FExtensionData>();
 			
 			// Log an error message
-			ErrorLog->GetPrivate()->Add("Constant extension data not set", ELMT_WARNING, Constant->GetMessageContext());
+			ErrorLog->Add("Constant extension data not set", ELMT_WARNING, Constant->GetMessageContext());
 		}
 
 		Op->Value = Data;
@@ -157,7 +156,7 @@ class Node;
 			if (TagIndex == INDEX_NONE)
 			{
 				const FString Msg = FString::Printf(TEXT("Unknown tag found in Extension Data variation [%s]"), *Tag);
-				ErrorLog->GetPrivate()->Add(Msg, ELMT_WARNING, Variation->GetMessageContext());
+				ErrorLog->Add(Msg, ELMT_WARNING, Variation->GetMessageContext());
 				continue;
 			}
 
@@ -191,7 +190,7 @@ class Node;
 	{
 		// Log a warning
 		const FString Msg = FString::Printf(TEXT("Required connection not found: %s"), StrWhere);
-		ErrorLog->GetPrivate()->Add(Msg, ELMT_ERROR, ErrorContext);
+		ErrorLog->Add(Msg, ELMT_ERROR, ErrorContext);
 
 		// Create a constant extension data
 		Ptr<const NodeExtensionDataConstant> Node = new NodeExtensionDataConstant;
diff --git a/Engine/Plugins/Mutable/Source/MutableTools/Private/MuT/CodeGenerator_FirstPass.cpp b/Engine/Plugins/Mutable/Source/MutableTools/Private/MuT/CodeGenerator_FirstPass.cpp
--- a/Engine/Plugins/Mutable/Source/MutableTools/Private/MuT/CodeGenerator_FirstPass.cpp
+++ b/Engine/Plugins/Mutable/Source/MutableTools/Private/MuT/CodeGenerator_FirstPass.cpp
@@ -15,7 +15,6 @@
 #include "MuT/ASTOpParameter.h"
 #include "MuT/CompilerPrivate.h"
 #include "MuT/ErrorLog.h"
-#include "MuT/ErrorLogPrivate.h"
 #include "MuT/NodeComponent.h"
 #include "MuT/NodeModifierMeshClipDeform.h"
 #include "MuT/NodeModifierMeshClipMorphPlane.h"
@@ -46,7 +45,7 @@ namespace mu
 
 
 	//---------------------------------------------------------------------------------------------
-    void FirstPassGenerator::Generate( Ptr<mu::ErrorLog> InErrorLog,
+    void FirstPassGenerator::Generate( TSharedPtr<FErrorLog> InErrorLog,
                                        const Node* Root,
 									   bool bIgnoreStates,
 									   CodeGenerator* InGenerator )
@@ -575,7 +574,7 @@ namespace mu
                 FString Msg = FString::Printf( TEXT("State [%s] has more than %d runtime parameters. Their update may fail."), 
 					*s.Name,
                     MUTABLE_MAX_RUNTIME_PARAMETERS_PER_STATE);
-				ErrorLog->GetPrivate()->Add(Msg, ELMT_ERROR, InNode->GetMessageContext());
+				ErrorLog->Add(Msg, ELMT_ERROR, InNode->GetMessageContext());
             }
         }
 
@@ -656,7 +655,7 @@ namespace mu
 						{
 							FString Msg = FString::Printf(TEXT("The Group Node [%s] has type Toggle and its direct child is a Group node, which is not allowed. Change the type or add a Child Object node in between them."),
 								*Private->Name);
-							ErrorLog->GetPrivate()->Add(Msg, ELMT_ERROR, InNode->GetMessageContext());
+							ErrorLog->Add(Msg, ELMT_ERROR, InNode->GetMessageContext());
 						}
 						else
 						{
diff --git a/Engine/Plugins/Mutable/Source/MutableTools/Private/MuT/CodeGenerator_FirstPass.h b/Engine/Plugins/Mutable/Source/MutableTools/Private/MuT/CodeGenerator_FirstPass.h
--- a/Engine/Plugins/Mutable/Source/MutableTools/Private/MuT/CodeGenerator_FirstPass.h
+++ b/Engine/Plugins/Mutable/Source/MutableTools/Private/MuT/CodeGenerator_FirstPass.h
@@ -40,7 +40,7 @@ namespace mu
 
 		FirstPassGenerator();
 
-        void Generate(Ptr<ErrorLog>, const Node* Root, bool bIgnoreStates, class CodeGenerator*);
+        void Generate(TSharedPtr<FErrorLog>, const Node* Root, bool bIgnoreStates, class CodeGenerator*);
 
 	private:
 
@@ -216,7 +216,7 @@ namespace mu
 		CodeGenerator* Generator = nullptr;
 
         //!
-        Ptr<ErrorLog> ErrorLog;
+		TSharedPtr<FErrorLog> ErrorLog;
 	};
 
 }
diff --git a/Engine/Plugins/Mutable/Source/MutableTools/Private/MuT/CodeGenerator_Image.cpp b/Engine/Plugins/Mutable/Source/MutableTools/Private/MuT/CodeGenerator_Image.cpp
--- a/Engine/Plugins/Mutable/Source/MutableTools/Private/MuT/CodeGenerator_Image.cpp
+++ b/Engine/Plugins/Mutable/Source/MutableTools/Private/MuT/CodeGenerator_Image.cpp
@@ -40,7 +40,6 @@
 #include "MuT/Compiler.h"
 #include "MuT/CompilerPrivate.h"
 #include "MuT/ErrorLog.h"
-#include "MuT/ErrorLogPrivate.h"
 #include "MuT/Node.h"
 #include "MuT/NodeBool.h"
 #include "MuT/NodeColour.h"
@@ -182,14 +181,14 @@ namespace mu
             pImage = GenerateMissingImage(EImageFormat::RGB_UByte );
 
             // Log an error message
-            ErrorLog->GetPrivate()->Add( "Constant image not set.", ELMT_WARNING, InNode->GetMessageContext());
+            ErrorLog->Add( "Constant image not set.", ELMT_WARNING, InNode->GetMessageContext());
         }
 
 
 		if (pImage->IsReference())
 		{
 			Ptr<ASTOpReferenceResource> ReferenceOp = new ASTOpReferenceResource();
-			ReferenceOp->type = EOpType::IM_REFERENCE;
+			ReferenceOp->Type = EOpType::IM_REFERENCE;
 			ReferenceOp->ID = pImage->GetReferencedTexture();
 			ReferenceOp->bForceLoad = pImage->IsForceLoad();
 			ReferenceOp->SourceDataDescriptor = InNode->SourceDataDescriptor;
@@ -251,7 +250,7 @@ namespace mu
 					ensure(false);
 
 					// Log an error message
-					ErrorLog->GetPrivate()->Add("A layout block goes across different texture tiles. This is not supported yet.", ELMT_ERROR, InNode->GetMessageContext());
+					ErrorLog->Add("A layout block goes across different texture tiles. This is not supported yet.", ELMT_ERROR, InNode->GetMessageContext());
 				}
 				else
 				{
@@ -1443,7 +1442,7 @@ namespace mu
 			}
 			else
 			{
-				ErrorLog->GetPrivate()->Add("Layout or block index error.", ELMT_ERROR, InNode->GetMessageContext());
+				ErrorLog->Add("Layout or block index error.", ELMT_ERROR, InNode->GetMessageContext());
 			}
 
 			// TODO: 
@@ -1518,7 +1517,7 @@ namespace mu
             cop->Type = EOpType::ME_CONSTANT;
             cop->SetValue(TempMesh, CompilerOptions->OptimisationOptions.DiskCacheContext);
 			ProjectOp->SetChild(ProjectOp->op.args.MeshProject.mesh, cop );
-            ErrorLog->GetPrivate()->Add( "Projector mesh not set.", ELMT_ERROR, InNode->GetMessageContext() );
+            ErrorLog->Add( "Projector mesh not set.", ELMT_ERROR, InNode->GetMessageContext() );
         }
 
 
@@ -1819,7 +1818,7 @@ namespace mu
             {
 				FString Msg = FString::Printf(TEXT("Unknown tag found in image variation [%s]."), *tag );
 
-                ErrorLog->GetPrivate()->Add( Msg, ELMT_WARNING, InNode->GetMessageContext() );
+                ErrorLog->Add( Msg, ELMT_WARNING, InNode->GetMessageContext() );
                 continue;
             }
 
@@ -1887,7 +1886,7 @@ namespace mu
 		const NodeImageTable& node = *InNode;
 
 		Result.op = GenerateTableSwitch<NodeImageTable, ETableColumnType::Image, EOpType::IM_SWITCH>(node,
-			[this, InNode, Options](const NodeImageTable& node, int32 colIndex, int32 row, mu::ErrorLog* pErrorLog)
+			[this, InNode, Options](const NodeImageTable& node, int32 colIndex, int32 row, FErrorLog* ErrorLog)
 			{
 				const FTableValue& CellData = node.Table->GetPrivate()->Rows[row].Values[colIndex];
 				TSharedPtr<const FImage> pImage = nullptr;
@@ -1904,7 +1903,7 @@ namespace mu
 					if (row != 0) // "None" option (row 0) is always a null image, do not trigger this error.
 					{
 						FString Msg = FString::Printf(TEXT("Table has a missing image in column %d, row %d."), colIndex, row);
-						pErrorLog->GetPrivate()->Add(Msg, ELMT_ERROR, InNode->GetMessageContext());
+						ErrorLog->Add(Msg, ELMT_ERROR, InNode->GetMessageContext());
 					}
 
 					return ImageOp;
@@ -2037,11 +2036,11 @@ namespace mu
 	}
 
 
-	Ptr<ASTOp> CodeGenerator::GenerateMissingImageCode(const TCHAR* strWhere, EImageFormat format, const void* errorContext, const FImageGenerationOptions& Options )
+	Ptr<ASTOp> CodeGenerator::GenerateMissingImageCode(const TCHAR* strWhere, EImageFormat format, const void* ErrorContext, const FImageGenerationOptions& Options )
 	{
 		// Log an error message
 		FString Msg = FString::Printf(TEXT("Required connection not found: %s"), strWhere );
-		ErrorLog->GetPrivate()->Add( Msg, ELMT_ERROR, errorContext );
+		ErrorLog->Add( Msg, ELMT_ERROR, ErrorContext);
 
 		// Make a checkered debug image
 		TSharedPtr<FImage> GeneratedImage = GenerateMissingImage( format );
diff --git a/Engine/Plugins/Mutable/Source/MutableTools/Private/MuT/CodeGenerator_Mesh.cpp b/Engine/Plugins/Mutable/Source/MutableTools/Private/MuT/CodeGenerator_Mesh.cpp
--- a/Engine/Plugins/Mutable/Source/MutableTools/Private/MuT/CodeGenerator_Mesh.cpp
+++ b/Engine/Plugins/Mutable/Source/MutableTools/Private/MuT/CodeGenerator_Mesh.cpp
@@ -37,7 +37,6 @@
 #include "MuT/CodeGenerator_FirstPass.h"
 #include "MuT/CompilerPrivate.h"
 #include "MuT/ErrorLog.h"
-#include "MuT/ErrorLogPrivate.h"
 #include "MuT/NodeImageProject.h"
 #include "MuT/NodeLayout.h"
 #include "MuT/NodeMesh.h"
@@ -110,7 +109,7 @@ namespace mu
 	void CodeGenerator::PrepareMeshForLayout( const FGeneratedLayout& GeneratedLayout,
 		TSharedPtr<FMesh> Mesh,
 		int32 LayoutChannel,
-		const void* errorContext,
+		const void* ErrorContext,
 		const FMeshGenerationOptions& MeshOptions,
 		bool bUseAbsoluteBlockIds
 		)
@@ -193,7 +192,7 @@ namespace mu
 			FString Msg = FString::Printf(TEXT("Source mesh has %d layout block overlapping in LOD %d"),
 				OverlappingBlocks.Num() + 1, CurrentParents.Last().Lod
 			);
-			ErrorLog->GetPrivate()->Add(Msg, ELMT_WARNING, errorContext);
+			ErrorLog->Add(Msg, ELMT_WARNING, ErrorContext);
 		}
 
 		// Get the information about the texture coordinates channel
@@ -249,7 +248,7 @@ namespace mu
 			if (bNonNormalizedUVs && !bIsOverlayLayout)
 			{
 				FString Msg = FString::Printf(TEXT("Source mesh has non-normalized UVs in LOD %d"), CurrentParents.Last().Lod );
-				ErrorLog->GetPrivate()->Add(Msg, ELMT_WARNING, errorContext);
+				ErrorLog->Add(Msg, ELMT_WARNING, ErrorContext);
 			}
 		}
 
@@ -481,11 +480,11 @@ namespace mu
 			{
 				FString Msg = FString::Printf(TEXT("Source mesh has %d vertices not assigned to any layout block in LOD %d"), UnassignedUVs.Num(), CurrentParents.Last().Lod);
 
-				ErrorLogMessageAttachedDataView attachedDataView;
-				attachedDataView.m_unassignedUVs = UnassignedUVs.GetData();
-				attachedDataView.m_unassignedUVsSize = (size_t)UnassignedUVs.Num();
+				ErrorLogMessageAttachedDataView AttachedDataView;
+				AttachedDataView.UnassignedUVs = UnassignedUVs.GetData();
+				AttachedDataView.UnassignedUVsSize = UnassignedUVs.Num();
 
-				ErrorLog->GetPrivate()->Add(Msg, attachedDataView, ELMT_WARNING, errorContext);
+				ErrorLog->Add(Msg, AttachedDataView, ELMT_WARNING, ErrorContext);
 			}
 		}
 
@@ -656,8 +655,7 @@ namespace mu
         else
         {
             // This argument is required
-            ErrorLog->GetPrivate()->Add( "Mesh morph base node is not set.",
-                                            ELMT_ERROR, InMorphNode->GetMessageContext());
+            ErrorLog->Add( "Mesh morph base node is not set.", ELMT_ERROR, InMorphNode->GetMessageContext());
         }        
 
 		if (node.Morph)
@@ -755,8 +753,7 @@ namespace mu
         else
         {
             // This argument is required
-            ErrorLog->GetPrivate()->Add( "Mesh make morph base node is not set.",
-                                            ELMT_ERROR, InMakeMorphNode->GetMessageContext());
+            ErrorLog->Add( "Mesh make morph base node is not set.", ELMT_ERROR, InMakeMorphNode->GetMessageContext());
         }
 
         // Target
@@ -774,8 +771,7 @@ namespace mu
         else
         {
             // This argument is required
-            ErrorLog->GetPrivate()->Add( "Mesh make morph target node is not set.",
-                                            ELMT_ERROR, InMakeMorphNode->GetMessageContext());
+            ErrorLog->Add( "Mesh make morph target node is not set.", ELMT_ERROR, InMakeMorphNode->GetMessageContext());
         }
 
         OutResult.MeshOp = op;
@@ -813,7 +809,7 @@ namespace mu
         else
         {
             // This argument is required
-            ErrorLog->GetPrivate()->Add( "Mesh fragment source is not set.", ELMT_ERROR, Node->GetMessageContext());
+            ErrorLog->Add( "Mesh fragment source is not set.", ELMT_ERROR, Node->GetMessageContext());
         }
 
         OutResult.BaseMeshOp = BaseResult.BaseMeshOp;
@@ -900,9 +896,7 @@ namespace mu
         {
             // TODO
             //op.args.MeshInterpolate.target[0] = GenerateMissingImageCode( "First mesh", RGB_UByte );
-            ErrorLog->GetPrivate()->Add
-                ( "Mesh interpolation: at least the first mesh is required.",
-                  ELMT_ERROR, InterpolateNode->GetMessageContext());
+            ErrorLog->Add( "Mesh interpolation: at least the first mesh is required.", ELMT_ERROR, InterpolateNode->GetMessageContext());
         }
     }
 
@@ -976,7 +970,7 @@ namespace mu
 		bool bFirstRowGenerated = false;
 
 		Ptr<ASTOp> Op = GenerateTableSwitch<NodeMeshTable, ETableColumnType::Mesh, EOpType::ME_SWITCH>(*TableNode,
-			[this, &NewResult, &bFirstRowGenerated, &InOptions] (const NodeMeshTable& node, int32 colIndex, int32 row, mu::ErrorLog* pErrorLog)
+			[this, &NewResult, &bFirstRowGenerated, &InOptions] (const NodeMeshTable& node, int32 colIndex, int32 row, FErrorLog*)
 			{
 				TSharedPtr<mu::FMesh> pMesh = node.Table->GetPrivate()->Rows[row].Values[colIndex].Mesh;
 				FMeshGenerationResult BranchResults;
@@ -1066,7 +1060,7 @@ namespace mu
 
             if ( tagIndex < 0 )
             {
-                ErrorLog->GetPrivate()->Add( 
+                ErrorLog->Add( 
 					FString::Printf(TEXT("Unknown tag found in mesh variation [%s]."), *tag),
 					ELMT_WARNING,
 					VariationNode->GetMessageContext(),
@@ -1136,7 +1130,7 @@ namespace mu
 			EmptyMesh->MeshIDPrefix = ConstantOp->GetValueHash();
 
 			// Log an error message
-			ErrorLog->GetPrivate()->Add("Constant mesh not set.", ELMT_WARNING, InNode->GetMessageContext());
+			ErrorLog->Add("Constant mesh not set.", ELMT_WARNING, InNode->GetMessageContext());
 
 			return;
 		}
@@ -1144,7 +1138,7 @@ namespace mu
 		if (pMesh->IsReference())
 		{
 			Ptr<ASTOpReferenceResource> ReferenceOp = new ASTOpReferenceResource();
-			ReferenceOp->type = EOpType::ME_REFERENCE;
+			ReferenceOp->Type = EOpType::ME_REFERENCE;
 			ReferenceOp->ID = pMesh->GetReferencedMesh();
 			ReferenceOp->bForceLoad = pMesh->IsForceLoad();
 
@@ -1399,7 +1393,7 @@ namespace mu
         else
         {
             // This argument is required
-            ErrorLog->GetPrivate()->Add("Mesh transform base node is not set.", ELMT_ERROR, TransformNode->GetMessageContext() );
+            ErrorLog->Add("Mesh transform base node is not set.", ELMT_ERROR, TransformNode->GetMessageContext() );
         }
 
         op->matrix = node.Transform;
@@ -1426,8 +1420,7 @@ namespace mu
         else
         {
             // This argument is required
-            ErrorLog->GetPrivate()->Add("Mesh clip-morph-plane source node is not set.",
-                ELMT_ERROR, ClipNode->GetMessageContext());
+            ErrorLog->Add("Mesh clip-morph-plane source node is not set.", ELMT_ERROR, ClipNode->GetMessageContext());
         }
 
         // Morph to an ellipse
@@ -1491,8 +1484,7 @@ namespace mu
         else
         {
             // This argument is required
-            ErrorLog->GetPrivate()->Add("Mesh clip-with-mesh source node is not set.",
-                ELMT_ERROR, ClipNode->GetMessageContext());
+            ErrorLog->Add("Mesh clip-with-mesh source node is not set.", ELMT_ERROR, ClipNode->GetMessageContext());
         }
 
         // Clipping mesh
@@ -1510,8 +1502,7 @@ namespace mu
         else
         {
             // This argument is required
-            ErrorLog->GetPrivate()->Add("Mesh clip-with-mesh clipping mesh node is not set.",
-                ELMT_ERROR, ClipNode->GetMessageContext());
+            ErrorLog->Add("Mesh clip-with-mesh clipping mesh node is not set.", ELMT_ERROR, ClipNode->GetMessageContext());
         }
 
         OutResult.MeshOp = op;
@@ -1532,7 +1523,7 @@ namespace mu
 		else
 		{
 			// This argument is required
-			ErrorLog->GetPrivate()->Add("Mesh Clip Deform base mesh node is not set.", ELMT_ERROR, ClipDeform->GetMessageContext());
+			ErrorLog->Add("Mesh Clip Deform base mesh node is not set.", ELMT_ERROR, ClipDeform->GetMessageContext());
 		}
 
 		// Base Shape
@@ -1571,8 +1562,7 @@ namespace mu
         else
         {
             // This argument is required
-            ErrorLog->GetPrivate()->Add("Mesh apply-pose base node is not set.",
-                ELMT_ERROR, PoseNode->GetMessageContext());
+            ErrorLog->Add("Mesh apply-pose base node is not set.", ELMT_ERROR, PoseNode->GetMessageContext());
         }
 
         // Pose mesh
@@ -1590,8 +1580,7 @@ namespace mu
         else
         {
             // This argument is required
-            ErrorLog->GetPrivate()->Add("Mesh apply-pose pose node is not set.",
-                ELMT_ERROR, PoseNode->GetMessageContext());
+            ErrorLog->Add("Mesh apply-pose pose node is not set.", ELMT_ERROR, PoseNode->GetMessageContext());
         }
 
         OutResult.MeshOp = op;
@@ -1614,8 +1603,7 @@ namespace mu
 		else
 		{
 			// This argument is required
-			ErrorLog->GetPrivate()->Add("Mesh geometric op mesh-a node is not set.",
-				ELMT_ERROR, GeomNode->GetMessageContext());
+			ErrorLog->Add("Mesh geometric op mesh-a node is not set.", ELMT_ERROR, GeomNode->GetMessageContext());
 		}
 
 		// Mesh B
@@ -1675,7 +1663,7 @@ namespace mu
 		else
 		{
 			// This argument is required
-			ErrorLog->GetPrivate()->Add("Mesh reshape base node is not set.", ELMT_ERROR, Reshape->GetMessageContext());
+			ErrorLog->Add("Mesh reshape base node is not set.", ELMT_ERROR, Reshape->GetMessageContext());
 		}
 
 		// Base and target shapes shouldn't have layouts or modifiers.
diff --git a/Engine/Plugins/Mutable/Source/MutableTools/Private/MuT/CodeGenerator_Projector.cpp b/Engine/Plugins/Mutable/Source/MutableTools/Private/MuT/CodeGenerator_Projector.cpp
--- a/Engine/Plugins/Mutable/Source/MutableTools/Private/MuT/CodeGenerator_Projector.cpp
+++ b/Engine/Plugins/Mutable/Source/MutableTools/Private/MuT/CodeGenerator_Projector.cpp
@@ -14,7 +14,6 @@
 #include "MuT/ASTOpParameter.h"
 #include "MuT/CodeGenerator.h"
 #include "MuT/ErrorLog.h"
-#include "MuT/ErrorLogPrivate.h"
 #include "MuT/NodeProjector.h"
 #include "MuT/NodeProjectorPrivate.h"
 #include "MuT/NodeRange.h"
@@ -151,8 +150,7 @@ namespace mu
 		const void* errorContext)
 	{
 		// Log an error message
-		ErrorLog->GetPrivate()->Add("Required projector connection not found.",
-			ELMT_ERROR, errorContext);
+		ErrorLog->Add("Required projector connection not found.", ELMT_ERROR, errorContext);
 
 		FProjector p;
 		p.type = EProjectorType::Planar;
diff --git a/Engine/Plugins/Mutable/Source/MutableTools/Private/MuT/CodeGenerator_Scalar.cpp b/Engine/Plugins/Mutable/Source/MutableTools/Private/MuT/CodeGenerator_Scalar.cpp
--- a/Engine/Plugins/Mutable/Source/MutableTools/Private/MuT/CodeGenerator_Scalar.cpp
+++ b/Engine/Plugins/Mutable/Source/MutableTools/Private/MuT/CodeGenerator_Scalar.cpp
@@ -19,7 +19,6 @@
 #include "MuT/CodeGenerator.h"
 #include "MuT/CodeGenerator_FirstPass.h"
 #include "MuT/ErrorLog.h"
-#include "MuT/ErrorLogPrivate.h"
 #include "MuT/NodeImage.h"
 #include "MuT/NodeMeshGeometryOperation.h"
 #include "MuT/NodeMeshTable.h"
@@ -182,8 +181,11 @@ namespace mu
 		{
 			FParameterDesc param;
 			param.Name = node.Name;
-			bool bParseOk = FGuid::Parse(node.UID, param.UID);
-			check(bParseOk);
+			if (!node.UID.IsEmpty())
+			{
+				bool bParseOk = FGuid::Parse(node.UID, param.UID);
+				check(bParseOk);
+			}
 			param.Type = EParameterType::Int;
 			param.DefaultValue.Set<FParamIntType>(node.DefaultValue);
 
@@ -304,7 +306,7 @@ namespace mu
 			{
 				FString Msg = FString::Printf(TEXT("Unknown tag found in image variation [%s]."), *tag);
 
-				ErrorLog->GetPrivate()->Add(Msg, ELMT_WARNING, Typed->GetMessageContext());
+				ErrorLog->Add(Msg, ELMT_WARNING, Typed->GetMessageContext());
 				continue;
 			}
 
@@ -410,7 +412,7 @@ namespace mu
 		const NodeScalarTable& node = *Typed;
 
 		Ptr<ASTOp> Op = GenerateTableSwitch<NodeScalarTable, ETableColumnType::Scalar, EOpType::SC_SWITCH>(node,
-			[this,&Options](const NodeScalarTable& node, int32 colIndex, int32 row, mu::ErrorLog* pErrorLog)
+			[this,&Options](const NodeScalarTable& node, int32 colIndex, int32 row, FErrorLog*)
 			{
 				NodeScalarConstantPtr pCell = new NodeScalarConstant();
 				float scalar = node.Table->GetPrivate()->Rows[row].Values[colIndex].Scalar;
@@ -427,7 +429,7 @@ namespace mu
 	{
 		// Log a warning
 		FString Msg = FString::Printf(TEXT("Required connection not found: %s"), strWhere );
-		ErrorLog->GetPrivate()->Add(Msg, ELMT_ERROR, errorContext);
+		ErrorLog->Add(Msg, ELMT_ERROR, errorContext);
 
 		// Create a constant node
 		NodeScalarConstantPtr pNode = new NodeScalarConstant();
diff --git a/Engine/Plugins/Mutable/Source/MutableTools/Private/MuT/CodeGenerator_SecondPass.cpp b/Engine/Plugins/Mutable/Source/MutableTools/Private/MuT/CodeGenerator_SecondPass.cpp
--- a/Engine/Plugins/Mutable/Source/MutableTools/Private/MuT/CodeGenerator_SecondPass.cpp
+++ b/Engine/Plugins/Mutable/Source/MutableTools/Private/MuT/CodeGenerator_SecondPass.cpp
@@ -317,7 +317,7 @@ namespace mu
 
 
 	//---------------------------------------------------------------------------------------------
-	bool SecondPassGenerator::Generate( Ptr<mu::ErrorLog> InErrorLog, const Node* root)
+	bool SecondPassGenerator::Generate(TSharedPtr<FErrorLog> InErrorLog, const Node* root)
 	{
 		MUTABLE_CPUPROFILER_SCOPE(SecondPassGenerate);
 
diff --git a/Engine/Plugins/Mutable/Source/MutableTools/Private/MuT/CodeGenerator_SecondPass.h b/Engine/Plugins/Mutable/Source/MutableTools/Private/MuT/CodeGenerator_SecondPass.h
--- a/Engine/Plugins/Mutable/Source/MutableTools/Private/MuT/CodeGenerator_SecondPass.h
+++ b/Engine/Plugins/Mutable/Source/MutableTools/Private/MuT/CodeGenerator_SecondPass.h
@@ -30,7 +30,7 @@ namespace mu
 		SecondPassGenerator( FirstPassGenerator*, const CompilerOptions::Private*  );
 
 		// Return true on success.
-        bool Generate( Ptr<ErrorLog>, const Node* Root );
+        bool Generate(TSharedPtr<FErrorLog>, const Node* Root );
 
 	private:
 
@@ -39,7 +39,7 @@ namespace mu
 
 
         //!
-        Ptr<ErrorLog> ErrorLog;
+		TSharedPtr<FErrorLog> ErrorLog;
 
         //!
         struct FConditionGenerationKey
diff --git a/Engine/Plugins/Mutable/Source/MutableTools/Private/MuT/CodeOptimiser.cpp b/Engine/Plugins/Mutable/Source/MutableTools/Private/MuT/CodeOptimiser.cpp
--- a/Engine/Plugins/Mutable/Source/MutableTools/Private/MuT/CodeOptimiser.cpp
+++ b/Engine/Plugins/Mutable/Source/MutableTools/Private/MuT/CodeOptimiser.cpp
@@ -2,7 +2,7 @@
 
 #include "MuT/CodeOptimiser.h"
 
-#include "MuT/ErrorLogPrivate.h"
+#include "MuT/ErrorLog.h"
 #include "MuT/AST.h"
 #include "MuT/ASTOpInstanceAdd.h"
 #include "MuT/ASTOpConditional.h"
diff --git a/Engine/Plugins/Mutable/Source/MutableTools/Private/MuT/Compiler.cpp b/Engine/Plugins/Mutable/Source/MutableTools/Private/MuT/Compiler.cpp
--- a/Engine/Plugins/Mutable/Source/MutableTools/Private/MuT/Compiler.cpp
+++ b/Engine/Plugins/Mutable/Source/MutableTools/Private/MuT/Compiler.cpp
@@ -10,7 +10,6 @@
 #include "MuT/CodeOptimiser.h"
 #include "MuT/CompilerPrivate.h"
 #include "MuT/ErrorLog.h"
-#include "MuT/ErrorLogPrivate.h"
 #include "MuT/Node.h"
 #include "MuT/NodePrivate.h"
 #include "MuT/Table.h"
@@ -303,7 +302,7 @@ namespace mu
         MUTABLE_CPUPROFILER_SCOPE(Compile);
 
         TArray< FStateCompilationData > states;
-        Ptr<ErrorLog> genErrorLog;
+        TSharedPtr<FErrorLog> genErrorLog;
 		TArray<FParameterDesc> Parameters;
         {
             CodeGenerator gen( m_pD->Options->GetPrivate() );
@@ -424,7 +423,7 @@ namespace mu
 						"safely dismissed in case of partial compilation."), 
 						*s.nodeState.Name,
 						*s.nodeState.RuntimeParams[RuntimeParameterIndex]);
-                    m_pD->ErrorLog->GetPrivate()->Add(Temp, ELMT_WARNING, pNode->GetMessageContext() );
+                    m_pD->ErrorLog->Add(Temp, ELMT_WARNING, pNode->GetMessageContext() );
                 }
             }
 
@@ -482,8 +481,8 @@ namespace mu
 		UE_LOG(LogMutableCore, Verbose, TEXT("(int) %s : %ld"), TEXT("Program size"), int64(Program.OpAddress.Num()));
 
         // Merge the log in the right order
-        genErrorLog->Merge( m_pD->ErrorLog.get() );
-        m_pD->ErrorLog = genErrorLog.get();
+        genErrorLog->Merge( m_pD->ErrorLog.Get() );
+        m_pD->ErrorLog = genErrorLog;
 
 		// Pack data
 		m_pD->GenerateRoms(pResult.Get(), LinkerOptions.AdditionalData );
@@ -506,7 +505,7 @@ namespace mu
 
 
     //---------------------------------------------------------------------------------------------
-    ErrorLogPtrConst Compiler::GetLog() const
+	TSharedPtr<FErrorLog> Compiler::GetLog() const
     {
         return m_pD->ErrorLog;
     }
diff --git a/Engine/Plugins/Mutable/Source/MutableTools/Private/MuT/CompilerPrivate.h b/Engine/Plugins/Mutable/Source/MutableTools/Private/MuT/CompilerPrivate.h
--- a/Engine/Plugins/Mutable/Source/MutableTools/Private/MuT/CompilerPrivate.h
+++ b/Engine/Plugins/Mutable/Source/MutableTools/Private/MuT/CompilerPrivate.h
@@ -5,9 +5,10 @@
 #include "MuT/Compiler.h"
 
 #include "MuT/AST.h"
-#include "MuT/ErrorLogPrivate.h"
+#include "MuT/ErrorLog.h"
 #include "MuT/NodeObjectPrivate.h"
 #include "MuR/Operations.h"
+#include "Templates/SharedPointer.h"
 
 
 namespace mu
@@ -88,31 +89,16 @@ namespace mu
 
         Private()
         {
-            ErrorLog = new mu::ErrorLog();
+            ErrorLog = MakeShared<FErrorLog>();
         }
 
-        Ptr<ErrorLog> ErrorLog;
+        TSharedPtr<FErrorLog> ErrorLog;
 
         /** */
         Ptr<CompilerOptions> Options;
 
-		/** Store for additional data generated during compilation, but not necessary for the runtime. */
-		struct FAdditionalData
-		{
-			/** Source data descriptor for every image constant that has been generated. 
-			* It must have the same size than the Program::ConstantImages array.
-			*/
-			TArray<FSourceDataDescriptor> SourceImagePerConstant;
-
-			/** Source data descriptor for every mesh constant that has been generated.
-			* It must have the same size than the Program::ConstantMeshes array.
-			*/
-			TArray<FSourceDataDescriptor> SourceMeshPerConstant;
-		};
-
-		//!
+		/** */
 		void GenerateRoms(FModel*, const FLinkerOptions::FAdditionalData&);
-
     };
 
 }
diff --git a/Engine/Plugins/Mutable/Source/MutableTools/Private/MuT/ErrorLog.cpp b/Engine/Plugins/Mutable/Source/MutableTools/Private/MuT/ErrorLog.cpp
--- a/Engine/Plugins/Mutable/Source/MutableTools/Private/MuT/ErrorLog.cpp
+++ b/Engine/Plugins/Mutable/Source/MutableTools/Private/MuT/ErrorLog.cpp
@@ -21,7 +21,6 @@
 #include "MuR/Skeleton.h"
 #include "MuR/System.h"
 #include "MuR/MutableRuntimeModule.h"
-#include "MuT/ErrorLogPrivate.h"
 #include "Templates/Tuple.h"
 #include "Trace/Detail/Channel.h"
 
@@ -29,40 +28,17 @@
 namespace mu
 {
 
-	ErrorLog::ErrorLog()
+    int32 FErrorLog::GetMessageCount() const
 	{
-		m_pD = new Private();
+		return Messages.Num();
 	}
 
 
-	ErrorLog::~ErrorLog()
+	const FString& FErrorLog::GetMessageText( int32 Index ) const
 	{
-        check( m_pD );
-		delete m_pD;
-		m_pD = 0;
-	}
-
-
-	ErrorLog::Private* ErrorLog::GetPrivate() const
-	{
-		return m_pD;
-	}
-
-
-    int32 ErrorLog::GetMessageCount( ErrorLogMessageType ) const
-	{
-        // \todo: by type argument
-		return (int32)m_pD->Messages.Num();
-	}
-
-
-	const FString& ErrorLog::GetMessageText( int32 index ) const
-	{
-		const char* strResult = "";
-
-		if (m_pD->Messages.IsValidIndex(index))
+		if (Messages.IsValidIndex(Index))
 		{
-			return m_pD->Messages[index].Text;
+			return Messages[Index].Text;
 		}
 
 		static FString Empty;
@@ -70,82 +46,84 @@ namespace mu
 	}
 
 
-	const void* ErrorLog::GetMessageContext(int32 index) const
+	const void* FErrorLog::GetMessageContext(int32 Index) const
 	{
-		const void* result = 0;
+		const void* Result = 0;
 
-		if (m_pD->Messages.IsValidIndex(index))
+		if (Messages.IsValidIndex(Index))
 		{
-			result = m_pD->Messages[index].Context;
+			Result = Messages[Index].Context;
 		}
 
-		return result;
+		return Result;
 	}
 
 
-	const void* ErrorLog::GetMessageContext2(int32 index) const
+	const void* FErrorLog::GetMessageContext2(int32 Index) const
 	{
-		const void* result = 0;
+		const void* Result = 0;
 
-		if (m_pD->Messages.IsValidIndex(index))
+		if (Messages.IsValidIndex(Index))
 		{
-			result = m_pD->Messages[index].Context2;
+			Result = Messages[Index].Context2;
 		}
 
-		return result;
+		return Result;
 	}
 
 
-	ErrorLogMessageType ErrorLog::GetMessageType( int32 index ) const
+	ErrorLogMessageType FErrorLog::GetMessageType( int32 Index ) const
 	{
-		ErrorLogMessageType result = ELMT_NONE;
+		ErrorLogMessageType Result = ELMT_NONE;
 
-		if (m_pD->Messages.IsValidIndex(index))
+		if (Messages.IsValidIndex(Index))
 		{
-			result = m_pD->Messages[index].Type;
+			Result = Messages[Index].Type;
 		}
 
-		return result;
+		return Result;
 	}
 
 
-	ErrorLogMessageSpamBin ErrorLog::GetMessageSpamBin(int32 index) const
+	ErrorLogMessageSpamBin FErrorLog::GetMessageSpamBin(int32 Index) const
 	{
-		ErrorLogMessageSpamBin result = ELMSB_ALL;
+		ErrorLogMessageSpamBin Result = ELMSB_ALL;
 
-		if (m_pD->Messages.IsValidIndex(index))
+		if (Messages.IsValidIndex(Index))
 		{
-			result = m_pD->Messages[index].Spam;
+			Result = Messages[Index].Spam;
 		}
 
-		return result;
+		return Result;
 	}
 
 
-	ErrorLogMessageAttachedDataView ErrorLog::GetMessageAttachedData( int32 index ) const
+	ErrorLogMessageAttachedDataView FErrorLog::GetMessageAttachedData( int32 Index ) const
 	{
-        ErrorLogMessageAttachedDataView result;
+        ErrorLogMessageAttachedDataView Result;
 
-		if (m_pD->Messages.IsValidIndex(index))
+		if (Messages.IsValidIndex(Index))
 		{
-			const Private::FMessage& message = m_pD->Messages[index];
+			const FMessage& message = Messages[Index];
             
             if ( message.Data ) 
             {
-                result.m_unassignedUVs = message.Data->UnassignedUVs.GetData();
-			    result.m_unassignedUVsSize = message.Data->UnassignedUVs.Num();
+                Result.UnassignedUVs = message.Data->UnassignedUVs.GetData();
+			    Result.UnassignedUVsSize = message.Data->UnassignedUVs.Num();
             }
 		}
 
-		return result;
+		return Result;
 	}
 
 
-	void ErrorLog::Private::Add(const FString& InMessage,
+	void FErrorLog::Add(const FString& InMessage,
 		ErrorLogMessageType InType,
 		const void* InContext,
 		ErrorLogMessageSpamBin InSpamBin)
 	{
+		UE::TUniqueLock Lock(MessageMutex);
+
 		FMessage& Msg = Messages.Emplace_GetRef();
 		Msg.Type = InType;
 		Msg.Spam = InSpamBin;
@@ -154,12 +132,14 @@ namespace mu
 	}
 
 
-	void ErrorLog::Private::Add(const FString& InMessage,
+	void FErrorLog::Add(const FString& InMessage,
 		ErrorLogMessageType InType,
 		const void* InContext,
 		const void* InContext2,
 		ErrorLogMessageSpamBin InSpamBin)
 	{
+		UE::TUniqueLock Lock(MessageMutex);
+
 		FMessage& Msg = Messages.Emplace_GetRef();
 		Msg.Type = InType;
 		Msg.Spam = InSpamBin;
@@ -169,12 +149,14 @@ namespace mu
 	}
 
 
-	void ErrorLog::Private::Add(const FString& InMessage,
+	void FErrorLog::Add(const FString& InMessage,
                                 const ErrorLogMessageAttachedDataView& InDataView,
                                 ErrorLogMessageType InType, 
 								const void* InContext,
 								ErrorLogMessageSpamBin InSpamBin)
 	{
+		UE::TUniqueLock Lock(MessageMutex);
+
 		FMessage& Msg = Messages.Emplace_GetRef();
 		Msg.Type = InType;
 		Msg.Spam = InSpamBin;
@@ -182,19 +164,19 @@ namespace mu
 		Msg.Context = InContext;
 		Msg.Data = MakeShared<FErrorData>();
 
-        if ( InDataView.m_unassignedUVs && InDataView.m_unassignedUVsSize > 0 )
+        if ( InDataView.UnassignedUVs && InDataView.UnassignedUVsSize > 0 )
         {
 			// \TODO: Review
-			Msg.Data->UnassignedUVs.Append(InDataView.m_unassignedUVs, InDataView.m_unassignedUVsSize);
+			Msg.Data->UnassignedUVs.Append(InDataView.UnassignedUVs, InDataView.UnassignedUVsSize);
         }
 	}
 	
 	
-	void ErrorLog::Log() const
+	void FErrorLog::Log() const
 	{
 		UE_LOG(LogMutableCore, Log, TEXT(" Error Log :\n"));
 
-		for ( const Private::FMessage& msg : m_pD->Messages )
+		for ( const FMessage& msg : Messages )
 		{
 			switch ( msg.Type )
 			{
@@ -207,15 +189,11 @@ namespace mu
 	}
 
 
-	void ErrorLog::Merge( const ErrorLog* pOther )
+	void FErrorLog::Merge( const FErrorLog* Other )
 	{
-		m_pD->Messages.Append(pOther->GetPrivate()->Messages);
+		Messages.Append(Other->Messages);
 	}
 
-
-	//---------------------------------------------------------------------------------------------
-	//---------------------------------------------------------------------------------------------
-	//---------------------------------------------------------------------------------------------
     
     // clang-format off
 
@@ -353,10 +331,4 @@ namespace mu
 
     // clang-format on
 
-    const TCHAR* GetOpName( EOpType type )
-    {
-        //check( type>=0 && type<(int32)EOpType::COUNT );
-        return s_opNames[(int32)type];
-    }
-
 }
diff --git a/Engine/Plugins/Mutable/Source/MutableTools/Public/MuT/Compiler.h b/Engine/Plugins/Mutable/Source/MutableTools/Public/MuT/Compiler.h
--- a/Engine/Plugins/Mutable/Source/MutableTools/Public/MuT/Compiler.h
+++ b/Engine/Plugins/Mutable/Source/MutableTools/Public/MuT/Compiler.h
@@ -2,6 +2,7 @@
 
 #pragma once
 
+#include "MuT/ErrorLog.h"
 #include "MuR/Ptr.h"
 #include "MuR/RefCounted.h"
 #include "MuR/Image.h"
@@ -20,10 +21,6 @@ namespace mu
     using TransformPtr=Ptr<Transform>;
     using TransformPtrConst=Ptr<const Transform>;
 
-    class ErrorLog;
-    using ErrorLogPtr=Ptr<ErrorLog>;
-    using ErrorLogPtrConst=Ptr<const ErrorLog>;
-
     class Node;
     class NodeTransformedObject;
 
@@ -163,7 +160,7 @@ namespace mu
 		TSharedPtr<FModel> Compile( const Ptr<Node>& pNode );
 
         //! Return the log of messages of all the compile operations executed so far.
-        ErrorLogPtrConst GetLog() const;
+        TSharedPtr<FErrorLog> GetLog() const;
 
     protected:
 
diff --git a/Engine/Plugins/Mutable/Source/MutableTools/Public/MuT/ErrorLog.h b/Engine/Plugins/Mutable/Source/MutableTools/Public/MuT/ErrorLog.h
--- a/Engine/Plugins/Mutable/Source/MutableTools/Public/MuT/ErrorLog.h
+++ b/Engine/Plugins/Mutable/Source/MutableTools/Public/MuT/ErrorLog.h
@@ -5,19 +5,14 @@
 #include "HAL/PlatformCrt.h"
 #include "MuR/Ptr.h"
 #include "MuR/RefCounted.h"
+#include "Templates/SharedPointer.h"
+#include "Async/Mutex.h"
 
 
 namespace mu
 {
 
-	// Forward declarations
-	class ErrorLog;
-	typedef Ptr<ErrorLog> ErrorLogPtr;
-	typedef Ptr<const ErrorLog> ErrorLogPtrConst;
-
-
-	//! Types of message stored in the log
-	//! \ingroup tools
+	/** Types of message stored in the log. */
 	typedef enum
 	{
         ELMT_NONE=0,
@@ -26,8 +21,7 @@ namespace mu
 		ELMT_INFO
 	} ErrorLogMessageType;
 
-	//! Categories of message stored in the log for the purpose of limiiting duplication of non-identical messages
-	//! \ingroup tools
+	/** Categories of message stored in the log for the purpose of limiiting duplication of non - identical messages. */
 	typedef enum
 	{
 		ELMSB_ALL = 0,
@@ -36,21 +30,40 @@ namespace mu
 
 	struct ErrorLogMessageAttachedDataView 
 	{
-		const float* m_unassignedUVs = nullptr;
-		size_t m_unassignedUVsSize = 0;
+		const float* UnassignedUVs = nullptr;
+		int32 UnassignedUVsSize = 0;
 	};
 
-	//! Class used to store the error, warning and information messages from several processes
-	//! performed by the tools library, like model transformation or compilation.
-	//! \ingroup tools
-	class MUTABLETOOLS_API ErrorLog : public RefCounted
+	/** Storage for mutable error, warning and information messages from several processes performed by the tools library, like compilation. 
+	* It support concurrent addition of messages, but not retrieval.
+	*/	
+	class MUTABLETOOLS_API FErrorLog
 	{
 	public:
 
-		//-----------------------------------------------------------------------------------------
-		// Life cycle
-		//-----------------------------------------------------------------------------------------
-		ErrorLog();
+		struct FErrorData
+		{
+			TArray< float > UnassignedUVs;
+		};
+
+		struct FMessage
+		{
+			ErrorLogMessageType Type = ELMT_NONE;
+			ErrorLogMessageSpamBin Spam = ELMSB_ALL;
+			FString Text;
+			TSharedPtr<FErrorData> Data;
+			const void* Context = nullptr;
+			const void* Context2 = nullptr;
+		};
+
+	private:
+
+		/** This mutex is used to control the access to the messages array when adding them. */
+		UE::FMutex MessageMutex;
+
+		TArray<FMessage> Messages;
+
+	public:
 
 		//-----------------------------------------------------------------------------------------
 		// Own interface
@@ -58,14 +71,14 @@ namespace mu
 
 		//! Get the number of messages.
 		//! If a message type is provided, return the number of that type of message only.
-		int32 GetMessageCount( ErrorLogMessageType type = ELMT_NONE ) const;
+		int32 GetMessageCount() const;
 
 		//! Get the text of a message.
-		//! \param index index of the message from 0 to GetMessageCount(ELMT_NONE)-1
+		//! \param index index of the message from 0 to GetMessageCount()-1
 		const FString& GetMessageText( int32 index ) const;
 
 		//! Get the opaque context of a message.
-		//! \param index index of the message from 0 to GetMessageCount(ELMT_NONE)-1
+		//! \param index index of the message from 0 to GetMessageCount()-1
 		const void* GetMessageContext(int32 index) const;
 		const void* GetMessageContext2(int32 index) const;
 
@@ -83,22 +96,14 @@ namespace mu
 		void Log() const;
 
 		//!
-		void Merge( const ErrorLog* pOther );
-
-		//-----------------------------------------------------------------------------------------
-		// Interface pattern
-		//-----------------------------------------------------------------------------------------
-		class Private;
-		Private* GetPrivate() const;
-
-	protected:
+		void Merge( const FErrorLog* Other );
 
-		//! Forbidden. Manage with the Ptr<> template.
-		~ErrorLog();
-
-	private:
+		//!
+		void Add(const FString& Message, ErrorLogMessageType Type, const void* Context, ErrorLogMessageSpamBin SpamBin = ELMSB_ALL);
+		void Add(const FString& Message, ErrorLogMessageType Type, const void* Context, const void* Context2, ErrorLogMessageSpamBin SpamBin = ELMSB_ALL);
 
-		Private* m_pD;
+		//!
+		void Add(const FString& Message, const ErrorLogMessageAttachedDataView& Data, ErrorLogMessageType Type, const void* Context, ErrorLogMessageSpamBin SpamBin = ELMSB_ALL);
 
 	};
 	
