[mutable] Fix mac compilation error.
#rnx

diff --git a/Engine/Plugins/Mutable/Source/MutableTools/Private/MuT/CodeGenerator_Image.cpp b/Engine/Plugins/Mutable/Source/MutableTools/Private/MuT/CodeGenerator_Image.cpp
--- a/Engine/Plugins/Mutable/Source/MutableTools/Private/MuT/CodeGenerator_Image.cpp
+++ b/Engine/Plugins/Mutable/Source/MutableTools/Private/MuT/CodeGenerator_Image.cpp
@@ -1886,7 +1886,7 @@ namespace mu
 		const NodeImageTable& node = *InNode;
 
 		Result.op = GenerateTableSwitch<NodeImageTable, ETableColumnType::Image, EOpType::IM_SWITCH>(node,
-			[this, InNode, Options](const NodeImageTable& node, int32 colIndex, int32 row, FErrorLog* ErrorLog)
+			[this, InNode, Options](const NodeImageTable& node, int32 colIndex, int32 row, FErrorLog* LocalErrorLog)
 			{
 				const FTableValue& CellData = node.Table->GetPrivate()->Rows[row].Values[colIndex];
 				TSharedPtr<const FImage> pImage = nullptr;
@@ -1903,7 +1903,7 @@ namespace mu
 					if (row != 0) // "None" option (row 0) is always a null image, do not trigger this error.
 					{
 						FString Msg = FString::Printf(TEXT("Table has a missing image in column %d, row %d."), colIndex, row);
-						ErrorLog->Add(Msg, ELMT_ERROR, InNode->GetMessageContext());
+						LocalErrorLog->Add(Msg, ELMT_ERROR, InNode->GetMessageContext());
 					}
 
 					return ImageOp;
