[mutable] Remove support for the unused mesh interpolation and generic geometry operations.
#rnx

diff --git a/Engine/Plugins/Mutable/Source/CustomizableObject/Internal/MuCO/CustomizableObjectPrivate.h b/Engine/Plugins/Mutable/Source/CustomizableObject/Internal/MuCO/CustomizableObjectPrivate.h
--- a/Engine/Plugins/Mutable/Source/CustomizableObject/Internal/MuCO/CustomizableObjectPrivate.h
+++ b/Engine/Plugins/Mutable/Source/CustomizableObject/Internal/MuCO/CustomizableObjectPrivate.h
@@ -1818,6 +1818,8 @@ public:
 
 		MovedTableRowNoneGenerationToUnreal,
 
+		RemoveObsoletMeshInterpolateAndGeometryOp,
+
 		// -----<new versions can be added above this line>--------
 		LastCustomizableObjectVersion
 	};
diff --git a/Engine/Plugins/Mutable/Source/CustomizableObjectEditor/Private/MuCOE/GenerateMutableSource/GenerateMutableSourceMesh.cpp b/Engine/Plugins/Mutable/Source/CustomizableObjectEditor/Private/MuCOE/GenerateMutableSource/GenerateMutableSourceMesh.cpp
--- a/Engine/Plugins/Mutable/Source/CustomizableObjectEditor/Private/MuCOE/GenerateMutableSource/GenerateMutableSourceMesh.cpp
+++ b/Engine/Plugins/Mutable/Source/CustomizableObjectEditor/Private/MuCOE/GenerateMutableSource/GenerateMutableSourceMesh.cpp
@@ -26,7 +26,6 @@
 #include "MuCOE/Nodes/CustomizableObjectNodeFloatConstant.h"
 #include "MuCOE/Nodes/CustomizableObjectNodeFloatParameter.h"
 #include "MuCOE/Nodes/CustomizableObjectNodeMaterial.h"
-#include "MuCOE/Nodes/CustomizableObjectNodeMeshGeometryOperation.h"
 #include "MuCOE/Nodes/CustomizableObjectNodeMeshMorph.h"
 #include "MuCOE/Nodes/CustomizableObjectNodeMeshMorphStackApplication.h"
 #include "MuCOE/Nodes/CustomizableObjectNodeMeshMorphStackDefinition.h"
@@ -39,7 +38,6 @@
 #include "MuCOE/Nodes/CustomizableObjectNodeTable.h"
 #include "MuCOE/UnrealEditorPortabilityHelpers.h"
 #include "MuT/NodeMeshConstant.h"
-#include "MuT/NodeMeshGeometryOperation.h"
 #include "MuT/NodeMeshMakeMorph.h"
 #include "MuT/NodeMeshMorph.h"
 #include "MuT/NodeMeshReshape.h"
@@ -3974,73 +3972,6 @@ mu::Ptr<mu::NodeMesh> GenerateMutableSourceMesh(const UEdGraphPin* Pin,
 		}
 	}
 
-	else if (const UCustomizableObjectNodeMeshGeometryOperation* TypedNodeGeometry = Cast<const UCustomizableObjectNodeMeshGeometryOperation>(Node))
-	{
-		mu::Ptr<mu::NodeMeshGeometryOperation> MeshNode = new mu::NodeMeshGeometryOperation();
-		Result = MeshNode;
-
-		if (const UEdGraphPin* ConnectedPin = FollowInputPin(*TypedNodeGeometry->MeshAPin()))
-		{
-			FMutableGraphMeshGenerationData ChildMeshData;
-			mu::Ptr<mu::NodeMesh> ChildNode = GenerateMutableSourceMesh(ConnectedPin, GenerationContext, ChildMeshData, SurfaceMetadataId, false, bOnlyConnectedLOD);
-			if (ChildNode)
-			{
-				MeshNode->SetMeshA(ChildNode.get());
-				MeshData.Combine(ChildMeshData);
-			}
-			else
-			{
-				GenerationContext.Log(LOCTEXT("MeshGenerationFailed", "Mesh generation failed."), Node);
-			}
-		}
-		else
-		{
-			FText Text = FText::Format(LOCTEXT("MeshGeometryMissingDef", "Geometry Operation node requires the {0} value."), TypedNodeGeometry->MeshAPin()->PinFriendlyName);
-			GenerationContext.Log(Text, Node);
-		}
-
-		if (const UEdGraphPin* ConnectedPin = FollowInputPin(*TypedNodeGeometry->MeshBPin()))
-		{
-			FMutableGraphMeshGenerationData ChildMeshData;
-			mu::Ptr<mu::NodeMesh> ChildNode = GenerateMutableSourceMesh(ConnectedPin, GenerationContext, ChildMeshData, SurfaceMetadataId, false, bOnlyConnectedLOD);
-			if (ChildNode)
-			{
-				MeshNode->SetMeshB(ChildNode.get());
-				MeshData.Combine(ChildMeshData);
-			}
-			else
-			{
-				GenerationContext.Log(LOCTEXT("MeshGenerationFailed", "Mesh generation failed."), Node);
-			}
-		}
-
-		if (const UEdGraphPin* ConnectedPin = FollowInputPin(*TypedNodeGeometry->ScalarAPin()))
-		{
-			mu::NodeScalarPtr ChildNode = GenerateMutableSourceFloat(ConnectedPin, GenerationContext);
-			if (ChildNode)
-			{
-				MeshNode->SetScalarA(ChildNode.get());
-			}
-			else
-			{
-				GenerationContext.Log(LOCTEXT("ScalarGenerationFailed", "Scalar generation failed."), Node);
-			}
-		}
-
-		if (const UEdGraphPin* ConnectedPin = FollowInputPin(*TypedNodeGeometry->ScalarBPin()))
-		{
-			mu::NodeScalarPtr ChildNode = GenerateMutableSourceFloat(ConnectedPin, GenerationContext);
-			if (ChildNode)
-			{
-				MeshNode->SetScalarB(ChildNode.get());
-			}
-			else
-			{
-				GenerationContext.Log(LOCTEXT("ScalarGenerationFailed", "Scalar generation failed."), Node);
-			}
-		}
-	}
-
 	else if (const UCustomizableObjectNodeMeshReshape* TypedNodeReshape = Cast<const UCustomizableObjectNodeMeshReshape>(Node))
 	{
 		mu::Ptr<mu::NodeMeshReshape> MeshNode = new mu::NodeMeshReshape();
diff --git a/Engine/Plugins/Mutable/Source/CustomizableObjectEditor/Private/MuCOE/GraphTraversal.cpp b/Engine/Plugins/Mutable/Source/CustomizableObjectEditor/Private/MuCOE/GraphTraversal.cpp
--- a/Engine/Plugins/Mutable/Source/CustomizableObjectEditor/Private/MuCOE/GraphTraversal.cpp
+++ b/Engine/Plugins/Mutable/Source/CustomizableObjectEditor/Private/MuCOE/GraphTraversal.cpp
@@ -11,7 +11,6 @@
 #include "MuCOE/Nodes/CustomizableObjectNodeExternalPin.h"
 #include "MuCOE/Nodes/CustomizableObjectNodeMaterial.h"
 #include "MuCOE/Nodes/CustomizableObjectNodeMaterialVariation.h"
-#include "MuCOE/Nodes/CustomizableObjectNodeMeshGeometryOperation.h"
 #include "MuCOE/Nodes/CustomizableObjectNodeMeshMorph.h"
 #include "MuCOE/Nodes/CustomizableObjectNodeMeshMorphStackApplication.h"
 #include "MuCOE/Nodes/CustomizableObjectNodeMeshMorphStackDefinition.h"
@@ -516,14 +515,6 @@ const UEdGraphPin* FindMeshBaseSource(const UEdGraphPin& Pin, const bool bOnlyLo
 		return &Pin;
 	}
 
-	else if (const UCustomizableObjectNodeMeshGeometryOperation* TypedNodeGeom = Cast<UCustomizableObjectNodeMeshGeometryOperation>(Node))
-	{
-		if (const UEdGraphPin* ConnectedPin = FollowInputPin(*TypedNodeGeom->MeshAPin()))
-		{
-			return FindMeshBaseSource(*ConnectedPin, bOnlyLookForStaticMesh);
-		}
-	}
-
 	else if (const UCustomizableObjectNodeMeshReshape* TypedNodeReshape = Cast<UCustomizableObjectNodeMeshReshape>(Node))
 	{
 		if (const UEdGraphPin* ConnectedPin = FollowInputPin(*TypedNodeReshape->BaseMeshPin()))
diff --git a/Engine/Plugins/Mutable/Source/CustomizableObjectEditor/Private/MuCOE/Nodes/CustomizableObjectNodeMeshGeometryOperation.cpp b/Engine/Plugins/Mutable/Source/CustomizableObjectEditor/Private/MuCOE/Nodes/CustomizableObjectNodeMeshGeometryOperation.cpp
deleted file mode 100644
diff --git a/Engine/Plugins/Mutable/Source/CustomizableObjectEditor/Public/MuCOE/Nodes/CustomizableObjectNodeMeshGeometryOperation.h b/Engine/Plugins/Mutable/Source/CustomizableObjectEditor/Public/MuCOE/Nodes/CustomizableObjectNodeMeshGeometryOperation.h
deleted file mode 100644
diff --git a/Engine/Plugins/Mutable/Source/MutableRuntime/Internal/MuR/Operations.h b/Engine/Plugins/Mutable/Source/MutableRuntime/Internal/MuR/Operations.h
--- a/Engine/Plugins/Mutable/Source/MutableRuntime/Internal/MuR/Operations.h
+++ b/Engine/Plugins/Mutable/Source/MutableRuntime/Internal/MuR/Operations.h
@@ -232,9 +232,6 @@ namespace mu
         //! Merge a mesh to a mesh
         ME_MERGE,
 
-        //! Interpolate between several meshes.
-        ME_INTERPOLATE,
-
         //! Create a new mask mesh selecting all the faces of a source that are inside a given
         //! clip mesh.
         ME_MASKCLIPMESH,
@@ -273,10 +270,6 @@ namespace mu
         //! Deform a skinned mesh applying a skeletal pose
         ME_APPLYPOSE,
 
-		//! Apply a geometry core operation to a mesh.
-		//! \TODO: Deprecated?
-		ME_GEOMETRYOPERATION,
-
 		//! Calculate the binding of a mesh on a shape
 		ME_BINDSHAPE,
 
diff --git a/Engine/Plugins/Mutable/Source/MutableRuntime/Private/MuR/CodeRunner.cpp b/Engine/Plugins/Mutable/Source/MutableRuntime/Private/MuR/CodeRunner.cpp
--- a/Engine/Plugins/Mutable/Source/MutableRuntime/Private/MuR/CodeRunner.cpp
+++ b/Engine/Plugins/Mutable/Source/MutableRuntime/Private/MuR/CodeRunner.cpp
@@ -1387,226 +1387,6 @@ namespace mu
             break;
         }
 
-        case EOpType::ME_INTERPOLATE:
-        {
-			OP::MeshInterpolateArgs Args = Program.GetOpArgs<OP::MeshInterpolateArgs>(item.At);
-            switch (item.Stage)
-            {
-            case 0:
-			{
-				if (Args.base)
-				{
-					AddOp(FScheduledOp(item.At, item, 1),
-						FScheduledOp(Args.base, item),
-						FScheduledOp(Args.factor, item));
-				}
-				else
-				{
-					StoreMesh(item, nullptr);
-				}
-				break;
-			}
-            case 1:
-            {
-            	MUTABLE_CPUPROFILER_SCOPE(ME_INTERPOLATE_1)
-
-                int32 Count = 1;
-                for (int32 TargetIndex = 0; TargetIndex < MUTABLE_OP_MAX_INTERPOLATE_COUNT-1 && Args.targets[TargetIndex]; ++TargetIndex)
-                {
-                    ++Count;
-                }
-
-                // Factor from 0 to 1 across all targets
-                float Factor = LoadScalar(FCacheAddress(Args.factor, item));
-
-                float Delta = 1.0f/(Count-1);
-                int32 Min = FMath::FloorToInt(Factor/Delta);
-                int32 Max = FMath::CeilToInt(Factor/Delta);
-
-                // Factor from 0 to 1 between the two targets
-                float Bifactor = Factor/Delta - Min;
-
-                FScheduledOpData Data;
-                Data.Interpolate.Bifactor = Bifactor;
-				Data.Interpolate.Min = FMath::Clamp(Min, 0, Count - 1);
-				Data.Interpolate.Max = FMath::Clamp(Max, 0, Count - 1);
-				uint32 DataAddress = uint32(HeapData.Num());
-
-                // Just the first of the targets
-                if (Bifactor < UE_SMALL_NUMBER)
-                {
-                    if (Min == 0)
-                    {
-                        // Just the base
-						TSharedPtr<const FMesh> BaseMesh = LoadMesh(FCacheAddress(Args.base, item));
-						StoreMesh(item, BaseMesh);
-					}
-                    else
-                    {
-                        // Base with one full morph
-                        HeapData.Add(Data);
-						AddOp(FScheduledOp(item.At, item, 2, DataAddress),
-							FScheduledOp(Args.base, item),
-							FScheduledOp(Args.targets[Min-1], item));
-					}
-				}
-                // Just the second of the targets
-                else if (Bifactor > 1.0f-UE_SMALL_NUMBER)
-                {
-                    HeapData.Add(Data);
-					AddOp(FScheduledOp(item.At, item, 2, DataAddress),
-						FScheduledOp(Args.base, item),
-						FScheduledOp(Args.targets[Max-1], item));
-				}
-                // Mix the first target on the base
-                else if (Min == 0)
-                {
-                    HeapData.Add(Data);
-					AddOp(FScheduledOp(item.At, item, 2, DataAddress),
-						FScheduledOp(Args.base, item),
-						FScheduledOp(Args.targets[0], item));
-				}
-                // Mix two targets on the base
-                else
-                {
-                    HeapData.Add(Data);
-					AddOp(FScheduledOp(item.At, item, 2, DataAddress),
-						FScheduledOp(Args.base, item),
-						FScheduledOp(Args.targets[Min-1], item),
-						FScheduledOp(Args.targets[Max-1], item));
-				}
-
-                break;
-            }
-
-            case 2:
-            {
-				MUTABLE_CPUPROFILER_SCOPE(ME_INTERPOLATE_2)
-
-                int32 Count = 1;
-                for (int32 TargetIndex = 0; TargetIndex < MUTABLE_OP_MAX_INTERPOLATE_COUNT-1 && Args.targets[TargetIndex]; ++TargetIndex)
-                {
-                    ++Count;
-                }
-
-                const FScheduledOpData& Data = HeapData[ item.CustomState ];
-
-                // Factor from 0 to 1 between the two targets
-                float Bifactor = Data.Interpolate.Bifactor;
-                int32 Min = Data.Interpolate.Min;
-                int32 Max = Data.Interpolate.Max;
-
-                TSharedPtr<const FMesh> BaseMesh = LoadMesh(FCacheAddress(Args.base, item));
-
-                if (BaseMesh)
-                {
-                    // Just the first of the targets
-                    if (Bifactor < UE_SMALL_NUMBER)
-                    {
-                        if (Min == 0)
-                        {
-                            // Just the base. It should have been dealt with in the previous stage.
-                            check(false);
-                        }
-                        else
-                        {
-                            // Base with one full morph
-                            TSharedPtr<const FMesh> MorphMesh = LoadMesh(FCacheAddress(Args.targets[Min-1], item));
-						
-							if (MorphMesh)
-							{
-								TSharedPtr<FMesh> Result = CloneOrTakeOver(BaseMesh);
-
-								MeshMorph(Result.Get(), MorphMesh.Get());
-							
-								StoreMesh(item, Result);
-							}
-							else
-							{
-								StoreMesh(item, BaseMesh);
-							}
-							
-							Release(MorphMesh);
-                        }
-                    }
-                    // Just the second of the targets
-                    else if (Bifactor > 1.0f-UE_SMALL_NUMBER)
-                    {
-                        check(Max > 0);
-                        TSharedPtr<const FMesh> MorphMesh = LoadMesh(FCacheAddress(Args.targets[Max-1], item));
-
-                        if (MorphMesh)
-                        {
-							TSharedPtr<FMesh> Result = CloneOrTakeOver(BaseMesh);
-							
-                            MeshMorph(Result.Get(), MorphMesh.Get());
-
-							StoreMesh(item, Result);
-                        }
-                        else
-                        {
-							StoreMesh(item, BaseMesh);
-                        }
-						
-						Release(MorphMesh);
-                    }
-                    // Mix the first target on the base
-                    else if (Min == 0)
-                    {
-                        TSharedPtr<const FMesh> MorphMesh = LoadMesh(FCacheAddress(Args.targets[0], item));
-                        if (MorphMesh)
-                        {
-							TSharedPtr<FMesh> Result = CloneOrTakeOver(BaseMesh);
-
-                            MeshMorph(Result.Get(), MorphMesh.Get(), Bifactor);
-
-							StoreMesh(item, Result);
-                        }
-                        else
-                        {
-							StoreMesh(item, BaseMesh);
-                        }
-						
-						Release(MorphMesh);
-                    }
-                    // Mix two targets on the base
-                    else
-                    {
-                        TSharedPtr<const FMesh> MinMesh = LoadMesh(FCacheAddress(Args.targets[Min-1], item));
-                        TSharedPtr<const FMesh> MaxMesh = LoadMesh(FCacheAddress(Args.targets[Max-1], item));
-
-                        if (MinMesh && MaxMesh)
-                        {
-							TSharedPtr<FMesh> Result = CloneOrTakeOver(BaseMesh);
-
-                            MeshMorph2(Result.Get(), MinMesh.Get(), MaxMesh.Get(), Bifactor);
-
-							StoreMesh(item, Result);
-                        }
-                        else
-                        {
-							StoreMesh(item, BaseMesh);
-                        }
-
-						Release(MinMesh);
-						Release(MaxMesh);
-                    }
-                }
-				else
-				{
-					StoreMesh(item, BaseMesh);
-				}
-
-                break;
-            }
-
-            default:
-                check(false);
-            }
-
-            break;
-        }
-
         case EOpType::ME_MASKCLIPMESH:
         {
 			OP::MeshMaskClipMeshArgs args = Program.GetOpArgs<OP::MeshMaskClipMeshArgs>(item.At);
@@ -2332,66 +2112,6 @@ namespace mu
             break;
         }
 
-
-		case EOpType::ME_GEOMETRYOPERATION:
-		{
-			OP::MeshGeometryOperationArgs args = Program.GetOpArgs<OP::MeshGeometryOperationArgs>(item.At);
-			switch (item.Stage)
-			{
-			case 0:
-			{
-				if (args.meshA)
-				{
-					AddOp(FScheduledOp(item.At, item, 1),
-						FScheduledOp(args.meshA, item),
-						FScheduledOp(args.meshB, item),
-						FScheduledOp(args.scalarA, item),
-						FScheduledOp(args.scalarB, item));
-				}
-				else
-				{
-					StoreMesh(item, nullptr);
-				}
-				break;
-			}
-			case 1:
-			{
-				MUTABLE_CPUPROFILER_SCOPE(ME_GEOMETRYOPERATION_1)
-				
-				TSharedPtr<const FMesh> MeshA = LoadMesh(FCacheAddress(args.meshA, item));
-				TSharedPtr<const FMesh> MeshB = LoadMesh(FCacheAddress(args.meshB, item));
-				float ScalarA = LoadScalar(FCacheAddress(args.scalarA, item));
-				float ScalarB = LoadScalar(FCacheAddress(args.scalarB, item));
-
-				TSharedPtr<FMesh> Result = CreateMesh(MeshA ? MeshA->GetDataSize() : 0);
-
-				bool bOutSuccess = false;
-				MeshGeometryOperation(Result.Get(), MeshA.Get(), MeshB.Get(), ScalarA, ScalarB, bOutSuccess);
-
-				Release(MeshA);
-				Release(MeshB);
-
-				if (!bOutSuccess)
-				{
-					Release(Result);
-					StoreMesh(item, nullptr);
-				}
-				else
-				{
-					StoreMesh(item, Result);
-				}
-
-				break;
-			}
-
-			default:
-				check(false);
-			}
-
-			break;
-		}
-
-
 		case EOpType::ME_BINDSHAPE:
 		{
 			OP::MeshBindShapeArgs Args = Program.GetOpArgs<OP::MeshBindShapeArgs>(item.At);
diff --git a/Engine/Plugins/Mutable/Source/MutableRuntime/Private/MuR/CodeRunnerTasked.cpp b/Engine/Plugins/Mutable/Source/MutableRuntime/Private/MuR/CodeRunnerTasked.cpp
--- a/Engine/Plugins/Mutable/Source/MutableRuntime/Private/MuR/CodeRunnerTasked.cpp
+++ b/Engine/Plugins/Mutable/Source/MutableRuntime/Private/MuR/CodeRunnerTasked.cpp
@@ -170,7 +170,6 @@ namespace mu
 		{   0, -20,   0,   0 },	// ME_DIFFERENCE
 		{   0,   0, -10,   0 },	// ME_MORPH2
 		{   0,   0,   0,   0 },	// ME_MERGE
-		{   0,   0, -20,   0 },	// ME_INTERPOLATE
 		{   0, -10,   0,   0 },	// ME_MASKCLIPMESH
 		{   0, -10,   0,   0 },	// ME_MASKCLIPUVMASK
 		{   0, -10,   0,   0 },	// ME_MASKDIFF
@@ -183,7 +182,6 @@ namespace mu
 		{   0,   0,   0,   0 },	// ME_SETSKELETON
 		{   0,   0,   0,   0 },	// ME_PROJECT
 		{   0,   0,   0,   0 },	// ME_APPLYPOSE
-		{   0,   0,   0,   0 },	// ME_GEOMETRYOPERATION
 		{   0, -20,   0,   0 },	// ME_BINDSHAPE
 		{   0, -10,   0,   0 },	// ME_APPLYSHAPE
 		{   0, -10,   0,   0 },	// ME_CLIPDEFORM
diff --git a/Engine/Plugins/Mutable/Source/MutableRuntime/Private/MuR/Operations.cpp b/Engine/Plugins/Mutable/Source/MutableRuntime/Private/MuR/Operations.cpp
--- a/Engine/Plugins/Mutable/Source/MutableRuntime/Private/MuR/Operations.cpp
+++ b/Engine/Plugins/Mutable/Source/MutableRuntime/Private/MuR/Operations.cpp
@@ -109,7 +109,6 @@ namespace mu
 		{ DT_MESH,			true,	{ 0, 0, 0, 0, 0, 0, 0, 0, 0, 0, 0, 0, 0, 0, 0, 0, 0, 0, 0, 0, 0, 0, 0, 0, 0, 0, 0, 0, 0 }		},	// ME_DIFFERENCE
 		{ DT_MESH,			true,	{ 0, 0, 0, 0, 0, 0, 0, 0, 0, 0, 0, 0, 0, 0, 0, 0, 0, 0, 0, 0, 0, 0, 0, 0, 0, 0, 0, 0, 0 }		},	// ME_MORPH
 		{ DT_MESH,			true,	{ 0, 0, 0, 0, 0, 0, 0, 0, 0, 0, 0, 0, 0, 0, 0, 0, 0, 0, 0, 0, 0, 0, 0, 0, 0, 0, 0, 0, 0 }		},	// ME_MERGE
-		{ DT_MESH,			true,	{ 0, 0, 0, 0, 0, 0, 0, 0, 0, 0, 0, 0, 0, 0, 0, 0, 0, 0, 0, 0, 0, 0, 0, 0, 0, 0, 0, 0, 0 }		},	// ME_INTERPOLATE
 		{ DT_MESH,			true,	{ 0, 0, 0, 0, 0, 0, 0, 0, 0, 0, 0, 0, 0, 0, 0, 0, 0, 0, 0, 0, 0, 0, 0, 0, 0, 0, 0, 0, 0 }		},	// ME_MASKCLIPMESH
 		{ DT_MESH,			true,	{ 0, 0, 0, 0, 0, 0, 0, 0, 0, 0, 0, 0, 0, 0, 0, 0, 0, 0, 0, 0, 0, 0, 0, 0, 0, 0, 0, 0, 0 }		},	// ME_MASKCLIPUVMASK
 		{ DT_MESH,			true,	{ 0, 0, 0, 0, 0, 0, 0, 0, 0, 0, 0, 0, 0, 0, 0, 0, 0, 0, 0, 0, 0, 0, 0, 0, 0, 0, 0, 0, 0 }		},	// ME_MASKDIFF
@@ -122,7 +121,6 @@ namespace mu
 		{ DT_MESH,			true,	{ 0, 0, 0, 0, 0, 0, 0, 0, 0, 0, 0, 0, 0, 0, 0, 0, 0, 0, 0, 0, 0, 0, 0, 0, 0, 0, 0, 0, 0 }		},	// ME_SETSKELETON
 		{ DT_MESH,			true,	{ 0, 0, 0, 0, 0, 0, 0, 0, 0, 0, 0, 0, 0, 0, 0, 0, 0, 0, 0, 0, 0, 0, 0, 0, 0, 0, 0, 0, 0 }		},	// ME_PROJECT
 		{ DT_MESH,			true,	{ 0, 0, 0, 0, 0, 0, 0, 0, 0, 0, 0, 0, 0, 0, 0, 0, 0, 0, 0, 0, 0, 0, 0, 0, 0, 0, 0, 0, 0 }		},	// ME_APPLYPOSE
-		{ DT_MESH,			true,	{ 0, 0, 0, 0, 0, 0, 0, 0, 0, 0, 0, 0, 0, 0, 0, 0, 0, 0, 0, 0, 0, 0, 0, 0, 0, 0, 0, 0, 0 }		},	// ME_GEOMETRYOPERATION
 		{ DT_MESH,			true,	{ 0, 0, 0, 0, 0, 0, 0, 0, 0, 0, 0, 0, 0, 0, 0, 0, 0, 0, 0, 0, 0, 0, 0, 0, 0, 0, 0, 0, 0 }		},	// ME_BINDSHAPE
 		{ DT_MESH,			true,	{ 0, 0, 0, 0, 0, 0, 0, 0, 0, 0, 0, 0, 0, 0, 0, 0, 0, 0, 0, 0, 0, 0, 0, 0, 0, 0, 0, 0, 0 }		},	// ME_APPLYSHAPE
 		{ DT_MESH,			true,	{ 0, 0, 0, 0, 0, 0, 0, 0, 0, 0, 0, 0, 0, 0, 0, 0, 0, 0, 0, 0, 0, 0, 0, 0, 0, 0, 0, 0, 0 }		},	// ME_CLIPDEFORM
@@ -319,16 +317,6 @@ namespace mu
             f(&op.args.MeshMerge.added );
             break;
 
-        case EOpType::ME_INTERPOLATE:
-            f(&op.args.MeshInterpolate.factor );
-            f(&op.args.MeshInterpolate.base );
-
-            for (int t=0;t<MUTABLE_OP_MAX_INTERPOLATE_COUNT-1;++t)
-            {
-                f(&op.args.MeshInterpolate.targets[t] );
-            }
-            break;
-
         case EOpType::ME_MASKDIFF:
             f(&op.args.MeshMaskDiff.source );
             f(&op.args.MeshMaskDiff.fragment );
@@ -829,19 +817,6 @@ namespace mu
             break;
         }
 
-        case EOpType::ME_INTERPOLATE:
-        {
-			OP::MeshInterpolateArgs args = program.GetOpArgs<OP::MeshInterpolateArgs>(at);
-            f(args.factor );
-            f(args.base );
-
-            for (int t=0;t<MUTABLE_OP_MAX_INTERPOLATE_COUNT-1;++t)
-            {
-                f(args.targets[t] );
-            }
-            break;
-        }
-
 		case EOpType::ME_MASKCLIPMESH:
 		{
 			OP::MeshMaskClipMeshArgs args = program.GetOpArgs<OP::MeshMaskClipMeshArgs>(at);
@@ -982,16 +957,6 @@ namespace mu
 			break;
 		}
 
-		case EOpType::ME_GEOMETRYOPERATION:
-		{
-			OP::MeshGeometryOperationArgs args = program.GetOpArgs<OP::MeshGeometryOperationArgs>(at);
-			f(args.meshA);
-			f(args.meshB);
-			f(args.scalarA);
-			f(args.scalarB);
-			break;
-		}
-
 		case EOpType::ME_BINDSHAPE:
 		{
 			OP::MeshBindShapeArgs args = program.GetOpArgs<OP::MeshBindShapeArgs>(at);
diff --git a/Engine/Plugins/Mutable/Source/MutableTools/Private/MuT/AST.cpp b/Engine/Plugins/Mutable/Source/MutableTools/Private/MuT/AST.cpp
--- a/Engine/Plugins/Mutable/Source/MutableTools/Private/MuT/AST.cpp
+++ b/Engine/Plugins/Mutable/Source/MutableTools/Private/MuT/AST.cpp
@@ -1717,9 +1717,6 @@ FSourceDataDescriptor ASTOpFixed::GetSourceDataDescriptor(FGetSourceDataDescript
 	case EOpType::ME_MERGE:
 		Result = GetSourceDataDescriptor(op.args.MeshMerge.base, Context);
 		break;	
-	case EOpType::ME_INTERPOLATE:
-		Result = GetSourceDataDescriptor(op.args.MeshInterpolate.base, Context);
-		break;	
 	case EOpType::ME_MASKDIFF:
 		Result = GetSourceDataDescriptor(op.args.MeshMaskDiff.source, Context);
 		break;	
diff --git a/Engine/Plugins/Mutable/Source/MutableTools/Private/MuT/ASTOpImageRasterMesh.cpp b/Engine/Plugins/Mutable/Source/MutableTools/Private/MuT/ASTOpImageRasterMesh.cpp
--- a/Engine/Plugins/Mutable/Source/MutableTools/Private/MuT/ASTOpImageRasterMesh.cpp
+++ b/Engine/Plugins/Mutable/Source/MutableTools/Private/MuT/ASTOpImageRasterMesh.cpp
@@ -260,16 +260,6 @@ namespace mu
 			break;
 		}
 
-		case EOpType::ME_INTERPOLATE:
-		{
-			// TODO: should be sink only if no imageAt?
-			auto typedSource = static_cast<const ASTOpFixed*>(sourceAt.get());
-			Ptr<ASTOpImageRasterMesh> rasterOp = mu::Clone<ASTOpImageRasterMesh>(this);
-			rasterOp->mesh = typedSource->children[typedSource->op.args.MeshInterpolate.base].child();
-			at = rasterOp;
-			break;
-		}
-
 		case EOpType::ME_MORPH:
 		{
 			// TODO: should be sink only if no imageAt?
diff --git a/Engine/Plugins/Mutable/Source/MutableTools/Private/MuT/ASTOpLayoutFromMesh.cpp b/Engine/Plugins/Mutable/Source/MutableTools/Private/MuT/ASTOpLayoutFromMesh.cpp
--- a/Engine/Plugins/Mutable/Source/MutableTools/Private/MuT/ASTOpLayoutFromMesh.cpp
+++ b/Engine/Plugins/Mutable/Source/MutableTools/Private/MuT/ASTOpLayoutFromMesh.cpp
@@ -194,14 +194,6 @@ namespace mu
 					break;
 				}
 
-				case EOpType::ME_INTERPOLATE:
-				{
-					// Sink, ignoring the op
-					const ASTOpFixed* Typed = static_cast<const ASTOpFixed*>(at.get());
-					newAt = Visit(Typed->children[Typed->op.args.MeshInterpolate.base].child(), CurrentSinkingOp);
-					break;
-				}
-
 				case EOpType::ME_CONDITIONAL:
 				{
 					Ptr<ASTOpConditional> NewConditional = mu::Clone<ASTOpConditional>(at);
diff --git a/Engine/Plugins/Mutable/Source/MutableTools/Private/MuT/ASTOpMeshExtractLayoutBlocks.cpp b/Engine/Plugins/Mutable/Source/MutableTools/Private/MuT/ASTOpMeshExtractLayoutBlocks.cpp
--- a/Engine/Plugins/Mutable/Source/MutableTools/Private/MuT/ASTOpMeshExtractLayoutBlocks.cpp
+++ b/Engine/Plugins/Mutable/Source/MutableTools/Private/MuT/ASTOpMeshExtractLayoutBlocks.cpp
@@ -206,23 +206,6 @@ namespace mu
 			break;
 		}
 
-		case EOpType::ME_INTERPOLATE:
-		{
-			Ptr<ASTOpFixed> newOp = mu::Clone<ASTOpFixed>(at);
-			newOp->SetChild(newOp->op.args.MeshInterpolate.base, Visit(newOp->children[newOp->op.args.MeshInterpolate.base].child(), currentSinkOp));
-
-			for (int32 t = 0; t < MUTABLE_OP_MAX_INTERPOLATE_COUNT - 1; ++t)
-			{
-				if (newOp->children[newOp->op.args.MeshInterpolate.targets[t]])
-				{
-					newOp->SetChild(newOp->op.args.MeshInterpolate.targets[t], Visit(newOp->children[newOp->op.args.MeshInterpolate.targets[t]].child(), currentSinkOp));
-				}
-			}
-
-			newAt = newOp;
-			break;
-		}
-
 		case EOpType::ME_REMOVEMASK:
 		{
 			// TODO: Make mask smaller?
diff --git a/Engine/Plugins/Mutable/Source/MutableTools/Private/MuT/ASTOpMeshFormat.cpp b/Engine/Plugins/Mutable/Source/MutableTools/Private/MuT/ASTOpMeshFormat.cpp
--- a/Engine/Plugins/Mutable/Source/MutableTools/Private/MuT/ASTOpMeshFormat.cpp
+++ b/Engine/Plugins/Mutable/Source/MutableTools/Private/MuT/ASTOpMeshFormat.cpp
@@ -394,37 +394,6 @@ mu::Ptr<ASTOp> Sink_MeshFormatAST::Visit(const mu::Ptr<ASTOp>& at, const ASTOpMe
 		break;
 	}
 
-	case EOpType::ME_INTERPOLATE:
-	{
-		// Move the format down the base of the morph
-		Ptr<ASTOpFixed> newOp = mu::Clone<ASTOpFixed>(at);
-		newOp->SetChild(newOp->op.args.MeshInterpolate.base, Visit(newOp->children[newOp->op.args.MeshInterpolate.base].child(), currentFormatOp));
-
-		// Reformat the morph targets to match the new format.
-		TSharedPtr<const FMesh> pTargetFormat = FindBaseMeshConstant(currentFormatOp->Format.child());
-		TSharedPtr<const FMesh> pTargetMorphFormat = MakeMorphTargetFormat(pTargetFormat);
-
-		mu::Ptr<ASTOpConstantResource> TargetMorphFormatOp = new ASTOpConstantResource();
-		TargetMorphFormatOp->Type = EOpType::ME_CONSTANT;
-		TargetMorphFormatOp->SetValue(pTargetMorphFormat, nullptr );
-		TargetMorphFormatOp->SourceDataDescriptor = at->GetSourceDataDescriptor();
-
-		for (int32 t = 0; t < MUTABLE_OP_MAX_INTERPOLATE_COUNT - 1; ++t)
-		{
-			if (newOp->children[newOp->op.args.MeshInterpolate.targets[t]])
-			{
-				mu::Ptr<ASTOpMeshFormat> newFormat = mu::Clone<ASTOpMeshFormat>(currentFormatOp);
-				newFormat->Flags = OP::MeshFormatArgs::Vertex | OP::MeshFormatArgs::IgnoreMissing;
-				newFormat->Format = TargetMorphFormatOp;
-
-				newOp->SetChild(newOp->op.args.MeshInterpolate.targets[t], Visit(newOp->children[newOp->op.args.MeshInterpolate.targets[t]].child(), newFormat.get()));
-			}
-		}
-
-		newAt = newOp;
-		break;
-	}
-
 	case EOpType::ME_REMOVEMASK:
 	{
 		auto newOp = mu::Clone<ASTOpMeshRemoveMask>(at);
diff --git a/Engine/Plugins/Mutable/Source/MutableTools/Private/MuT/ASTOpMeshGeometryOperation.cpp b/Engine/Plugins/Mutable/Source/MutableTools/Private/MuT/ASTOpMeshGeometryOperation.cpp
deleted file mode 100644
diff --git a/Engine/Plugins/Mutable/Source/MutableTools/Private/MuT/ASTOpMeshGeometryOperation.h b/Engine/Plugins/Mutable/Source/MutableTools/Private/MuT/ASTOpMeshGeometryOperation.h
deleted file mode 100644
diff --git a/Engine/Plugins/Mutable/Source/MutableTools/Private/MuT/CodeGenerator.cpp b/Engine/Plugins/Mutable/Source/MutableTools/Private/MuT/CodeGenerator.cpp
--- a/Engine/Plugins/Mutable/Source/MutableTools/Private/MuT/CodeGenerator.cpp
+++ b/Engine/Plugins/Mutable/Source/MutableTools/Private/MuT/CodeGenerator.cpp
@@ -66,8 +66,6 @@
 #include "MuT/NodeMeshConstantPrivate.h"
 #include "MuT/NodeMeshFormat.h"
 #include "MuT/NodeMeshFragment.h"
-#include "MuT/NodeMeshGeometryOperation.h"
-#include "MuT/NodeMeshInterpolate.h"
 #include "MuT/NodeMeshMorph.h"
 #include "MuT/NodeMeshReshape.h"
 #include "MuT/NodeModifier.h"
diff --git a/Engine/Plugins/Mutable/Source/MutableTools/Private/MuT/CodeGenerator.h b/Engine/Plugins/Mutable/Source/MutableTools/Private/MuT/CodeGenerator.h
--- a/Engine/Plugins/Mutable/Source/MutableTools/Private/MuT/CodeGenerator.h
+++ b/Engine/Plugins/Mutable/Source/MutableTools/Private/MuT/CodeGenerator.h
@@ -28,7 +28,6 @@
 #include "MuT/NodeImageProject.h"
 #include "MuT/NodeLOD.h"
 #include "MuT/NodeLayout.h"
-#include "MuT/NodeMeshGeometryOperation.h"
 #include "MuT/NodeMeshTable.h"
 #include "MuT/NodeObjectGroup.h"
 #include "MuT/NodeObjectNew.h"
@@ -82,7 +81,6 @@ namespace mu
 	class NodeMeshConstant;
 	class NodeMeshFormat;
 	class NodeMeshFragment;
-	class NodeMeshInterpolate;
 	class NodeMeshMakeMorph;
 	class NodeMeshMorph;
 	class NodeMeshReshape;
@@ -641,7 +639,6 @@ namespace mu
         void GenerateMesh_Morph(const FMeshGenerationOptions&, FMeshGenerationResult&, const NodeMeshMorph* );
         void GenerateMesh_MakeMorph(const FMeshGenerationOptions&, FMeshGenerationResult&, const NodeMeshMakeMorph* );
         void GenerateMesh_Fragment(const FMeshGenerationOptions&, FMeshGenerationResult&, const NodeMeshFragment* );
-        void GenerateMesh_Interpolate(const FMeshGenerationOptions&, FMeshGenerationResult&, const NodeMeshInterpolate* );
         void GenerateMesh_Switch(const FMeshGenerationOptions&, FMeshGenerationResult&, const NodeMeshSwitch* );
         void GenerateMesh_Transform(const FMeshGenerationOptions&, FMeshGenerationResult&, const NodeMeshTransform* );
         void GenerateMesh_ClipMorphPlane(const FMeshGenerationOptions&, FMeshGenerationResult&, const NodeMeshClipMorphPlane* );
@@ -649,7 +646,6 @@ namespace mu
         void GenerateMesh_ApplyPose(const FMeshGenerationOptions&, FMeshGenerationResult&, const NodeMeshApplyPose* );
         void GenerateMesh_Variation(const FMeshGenerationOptions&, FMeshGenerationResult&, const NodeMeshVariation* );
 		void GenerateMesh_Table(const FMeshGenerationOptions&, FMeshGenerationResult&, const NodeMeshTable*);
-		void GenerateMesh_GeometryOperation(const FMeshGenerationOptions&, FMeshGenerationResult&, const NodeMeshGeometryOperation*);
 		void GenerateMesh_Reshape(const FMeshGenerationOptions&, FMeshGenerationResult&, const NodeMeshReshape*);
 		void GenerateMesh_ClipDeform(const FMeshGenerationOptions&, FMeshGenerationResult&, const NodeMeshClipDeform*);
 
diff --git a/Engine/Plugins/Mutable/Source/MutableTools/Private/MuT/CodeGenerator_Color.cpp b/Engine/Plugins/Mutable/Source/MutableTools/Private/MuT/CodeGenerator_Color.cpp
--- a/Engine/Plugins/Mutable/Source/MutableTools/Private/MuT/CodeGenerator_Color.cpp
+++ b/Engine/Plugins/Mutable/Source/MutableTools/Private/MuT/CodeGenerator_Color.cpp
@@ -31,7 +31,6 @@
 #include "MuT/NodeColourVariation.h"
 #include "MuT/NodeImage.h"
 #include "MuT/NodeImageProject.h"
-#include "MuT/NodeMeshGeometryOperation.h"
 #include "MuT/NodeMeshTable.h"
 #include "MuT/NodeRange.h"
 #include "MuT/NodeScalar.h"
diff --git a/Engine/Plugins/Mutable/Source/MutableTools/Private/MuT/CodeGenerator_Image.cpp b/Engine/Plugins/Mutable/Source/MutableTools/Private/MuT/CodeGenerator_Image.cpp
--- a/Engine/Plugins/Mutable/Source/MutableTools/Private/MuT/CodeGenerator_Image.cpp
+++ b/Engine/Plugins/Mutable/Source/MutableTools/Private/MuT/CodeGenerator_Image.cpp
@@ -93,7 +93,6 @@
 #include "MuT/NodeImageTransformPrivate.h"
 #include "MuT/NodeImageVariation.h"
 #include "MuT/NodeImageVariationPrivate.h"
-#include "MuT/NodeMeshGeometryOperation.h"
 #include "MuT/NodeMeshTable.h"
 #include "MuT/NodeProjector.h"
 #include "MuT/NodeRange.h"
diff --git a/Engine/Plugins/Mutable/Source/MutableTools/Private/MuT/CodeGenerator_Mesh.cpp b/Engine/Plugins/Mutable/Source/MutableTools/Private/MuT/CodeGenerator_Mesh.cpp
--- a/Engine/Plugins/Mutable/Source/MutableTools/Private/MuT/CodeGenerator_Mesh.cpp
+++ b/Engine/Plugins/Mutable/Source/MutableTools/Private/MuT/CodeGenerator_Mesh.cpp
@@ -25,7 +25,6 @@
 #include "MuT/ASTOpMeshClipMorphPlane.h"
 #include "MuT/ASTOpMeshExtractLayoutBlocks.h"
 #include "MuT/ASTOpMeshFormat.h"
-#include "MuT/ASTOpMeshGeometryOperation.h"
 #include "MuT/ASTOpMeshMorphReshape.h"
 #include "MuT/ASTOpMeshTransform.h"
 #include "MuT/ASTOpMeshDifference.h"
@@ -50,10 +49,6 @@
 #include "MuT/NodeMeshFormat.h"
 #include "MuT/NodeMeshFormatPrivate.h"
 #include "MuT/NodeMeshFragment.h"
-#include "MuT/NodeMeshGeometryOperation.h"
-#include "MuT/NodeMeshGeometryOperationPrivate.h"
-#include "MuT/NodeMeshInterpolate.h"
-#include "MuT/NodeMeshInterpolatePrivate.h"
 #include "MuT/NodeMeshMakeMorph.h"
 #include "MuT/NodeMeshMakeMorphPrivate.h"
 #include "MuT/NodeMeshMorph.h"
@@ -604,7 +599,6 @@ namespace mu
 		case Node::EType::MeshMorph: GenerateMesh_Morph(InOptions, OutResult, static_cast<const NodeMeshMorph*>(Node)); break;
 		case Node::EType::MeshMakeMorph: GenerateMesh_MakeMorph(InOptions, OutResult, static_cast<const NodeMeshMakeMorph*>(Node)); break;
 		case Node::EType::MeshFragment: GenerateMesh_Fragment(InOptions, OutResult, static_cast<const NodeMeshFragment*>(Node)); break;
-		case Node::EType::MeshInterpolate: GenerateMesh_Interpolate(InOptions, OutResult, static_cast<const NodeMeshInterpolate*>(Node)); break;
 		case Node::EType::MeshSwitch: GenerateMesh_Switch(InOptions, OutResult, static_cast<const NodeMeshSwitch*>(Node)); break;
 		case Node::EType::MeshTransform: GenerateMesh_Transform(InOptions, OutResult, static_cast<const NodeMeshTransform*>(Node)); break;
 		case Node::EType::MeshClipMorphPlane: GenerateMesh_ClipMorphPlane(InOptions, OutResult, static_cast<const NodeMeshClipMorphPlane*>(Node)); break;
@@ -612,7 +606,6 @@ namespace mu
 		case Node::EType::MeshApplyPose: GenerateMesh_ApplyPose(InOptions, OutResult, static_cast<const NodeMeshApplyPose*>(Node)); break;
 		case Node::EType::MeshVariation: GenerateMesh_Variation(InOptions, OutResult, static_cast<const NodeMeshVariation*>(Node)); break;
 		case Node::EType::MeshTable: GenerateMesh_Table(InOptions, OutResult, static_cast<const NodeMeshTable*>(Node)); break;
-		case Node::EType::MeshGeometryOperation: GenerateMesh_GeometryOperation(InOptions, OutResult, static_cast<const NodeMeshGeometryOperation*>(Node)); break;
 		case Node::EType::MeshReshape: GenerateMesh_Reshape(InOptions, OutResult, static_cast<const NodeMeshReshape*>(Node)); break;
 		case Node::EType::MeshClipDeform: GenerateMesh_ClipDeform(InOptions, OutResult, static_cast<const NodeMeshClipDeform*>(Node)); break;
 		default: check(false);
@@ -817,90 +810,6 @@ namespace mu
 	}
 
 
-    //---------------------------------------------------------------------------------------------
-    void CodeGenerator::GenerateMesh_Interpolate(const FMeshGenerationOptions& InOptions, FMeshGenerationResult& OutResult,
-		const NodeMeshInterpolate* InterpolateNode )
-    {
-        NodeMeshInterpolate::Private& node = *InterpolateNode->GetPrivate();
-
-        // Generate the code
-        Ptr<ASTOpFixed> op = new ASTOpFixed();
-        op->op.type = EOpType::ME_INTERPOLATE;
-        OutResult.MeshOp = op;
-
-        // Factor
-        if ( Node* pFactor = node.m_pFactor.get() )
-        {
-            op->SetChild( op->op.args.MeshInterpolate.factor, Generate_Generic( pFactor, InOptions ) );
-        }
-        else
-        {
-            // This argument is required
-            op->SetChild( op->op.args.MeshInterpolate.factor,
-                          GenerateMissingScalarCode(TEXT("Interpolation factor"), 0.5f, InterpolateNode->GetMessageContext()) );
-        }
-
-        //
-        Ptr<ASTOp> base = 0;
-        int32 count = 0;
-        for ( int32 t=0
-            ; t<node.m_targets.Num() && t<MUTABLE_OP_MAX_INTERPOLATE_COUNT-1
-            ; ++t )
-        {
-            if ( NodeMesh* pA = node.m_targets[t].get() )
-            {
-				FMeshGenerationOptions TargetOptions = InOptions;
-				TargetOptions.OverrideLayouts.Empty();
-
-                FMeshGenerationResult TargetResult;
-                GenerateMesh( TargetOptions, TargetResult, pA );
-
-
-                // The first target is the base
-                if (count==0)
-                {
-                    base = TargetResult.MeshOp;
-                    op->SetChild( op->op.args.MeshInterpolate.base, TargetResult.MeshOp );
-
-                    OutResult.BaseMeshOp = TargetResult.BaseMeshOp;
-					OutResult.GeneratedLayouts = TargetResult.GeneratedLayouts;
-				}
-                else
-                {
-                    Ptr<ASTOpMeshDifference> dop = new ASTOpMeshDifference();
-                    dop->Base = base;
-                    dop->Target = TargetResult.MeshOp;
-
-                    // \todo Texcoords are broken?
-                    dop->bIgnoreTextureCoords = true;
-
-                    for ( size_t c=0; c<node.m_channels.Num(); ++c)
-                    {
-                        check( int32(node.m_channels[c].semantic) < 256 );
-						check(node.m_channels[c].semanticIndex < 256);
-						
-						ASTOpMeshDifference::FChannel Channel;
-						Channel.Semantic = uint8(node.m_channels[c].semantic);
-						Channel.SemanticIndex = uint8(node.m_channels[c].semanticIndex);
-						dop->Channels.Add(Channel);
-                    }
-
-                    op->SetChild( op->op.args.MeshInterpolate.targets[count-1], dop );
-                }
-                count++;
-            }
-        }
-
-        // At least one mesh is required
-        if (!count)
-        {
-            // TODO
-            //op.args.MeshInterpolate.target[0] = GenerateMissingImageCode( "First mesh", RGB_UByte );
-            ErrorLog->Add( "Mesh interpolation: at least the first mesh is required.", ELMT_ERROR, InterpolateNode->GetMessageContext());
-        }
-    }
-
-
     //---------------------------------------------------------------------------------------------
     void CodeGenerator::GenerateMesh_Switch(const FMeshGenerationOptions& InOptions, FMeshGenerationResult& OutResult, const NodeMeshSwitch* SwitchNode )
     {
@@ -1587,45 +1496,6 @@ namespace mu
     }
 
 
-	//---------------------------------------------------------------------------------------------
-	void CodeGenerator::GenerateMesh_GeometryOperation(const FMeshGenerationOptions& InOptions, FMeshGenerationResult& OutResult, const NodeMeshGeometryOperation* GeomNode)
-	{
-		const auto& node = *GeomNode->GetPrivate();
-
-		Ptr<ASTOpMeshGeometryOperation> op = new ASTOpMeshGeometryOperation();
-
-		// Mesh A
-		if (node.m_pMeshA)
-		{
-			GenerateMesh(InOptions, OutResult, node.m_pMeshA);
-			op->meshA = OutResult.MeshOp;
-		}
-		else
-		{
-			// This argument is required
-			ErrorLog->Add("Mesh geometric op mesh-a node is not set.", ELMT_ERROR, GeomNode->GetMessageContext());
-		}
-
-		// Mesh B
-		if (node.m_pMeshB)
-		{
-			FMeshGenerationOptions OtherOptions = InOptions;
-			OtherOptions.bLayouts = false;
-			OtherOptions.OverrideLayouts.Empty();
-			OtherOptions.ActiveTags.Empty();
-
-			FMeshGenerationResult bResult;
-			GenerateMesh(OtherOptions, bResult, node.m_pMeshB);
-			op->meshB = bResult.MeshOp;
-		}
-
-		op->scalarA = Generate_Generic(node.m_pScalarA, InOptions);
-		op->scalarB = Generate_Generic(node.m_pScalarB, InOptions);
-
-		OutResult.MeshOp = op;
-	}
-
-
 	//---------------------------------------------------------------------------------------------
 	void CodeGenerator::GenerateMesh_Reshape(const FMeshGenerationOptions& InOptions, FMeshGenerationResult& OutResult, const NodeMeshReshape* Reshape)
 	{
diff --git a/Engine/Plugins/Mutable/Source/MutableTools/Private/MuT/CodeGenerator_Scalar.cpp b/Engine/Plugins/Mutable/Source/MutableTools/Private/MuT/CodeGenerator_Scalar.cpp
--- a/Engine/Plugins/Mutable/Source/MutableTools/Private/MuT/CodeGenerator_Scalar.cpp
+++ b/Engine/Plugins/Mutable/Source/MutableTools/Private/MuT/CodeGenerator_Scalar.cpp
@@ -20,7 +20,6 @@
 #include "MuT/CodeGenerator_FirstPass.h"
 #include "MuT/ErrorLog.h"
 #include "MuT/NodeImage.h"
-#include "MuT/NodeMeshGeometryOperation.h"
 #include "MuT/NodeMeshTable.h"
 #include "MuT/NodeRange.h"
 #include "MuT/NodeScalar.h"
diff --git a/Engine/Plugins/Mutable/Source/MutableTools/Private/MuT/CodeOptimiser_Parameters.cpp b/Engine/Plugins/Mutable/Source/MutableTools/Private/MuT/CodeOptimiser_Parameters.cpp
--- a/Engine/Plugins/Mutable/Source/MutableTools/Private/MuT/CodeOptimiser_Parameters.cpp
+++ b/Engine/Plugins/Mutable/Source/MutableTools/Private/MuT/CodeOptimiser_Parameters.cpp
@@ -165,31 +165,6 @@ namespace mu
                         break;
                     }
 
-                    case EOpType::ME_INTERPOLATE:
-                    {
-						const ASTOpFixed* typed = static_cast<const ASTOpFixed*>(At.get());
-						FPendingItem ChildItem;
-                        ChildItem.ItemType = 0;
-                        ChildItem.OnlyLayoutsRelevant = item.OnlyLayoutsRelevant;
-
-                        if ( item.OnlyLayoutsRelevant==0 )
-                        {
-                            ChildItem.At = typed->children[typed->op.args.MeshInterpolate.factor].child();
-                            AddIfNeeded(ChildItem);
-                        }
-
-                        ChildItem.At = typed->children[typed->op.args.MeshInterpolate.base].child();
-                        AddIfNeeded(ChildItem);
-
-                        for (int32 t=0;t<MUTABLE_OP_MAX_INTERPOLATE_COUNT-1;++t)
-                        {
-                            ChildItem.At = typed->children[typed->op.args.MeshInterpolate.targets[t]].child();
-                            AddIfNeeded(ChildItem);
-                        }
-
-                        break;
-                    }
-
                     default:
                     {
                         At->ForEachChild([&](ASTChild& ref)
diff --git a/Engine/Plugins/Mutable/Source/MutableTools/Private/MuT/ErrorLog.cpp b/Engine/Plugins/Mutable/Source/MutableTools/Private/MuT/ErrorLog.cpp
--- a/Engine/Plugins/Mutable/Source/MutableTools/Private/MuT/ErrorLog.cpp
+++ b/Engine/Plugins/Mutable/Source/MutableTools/Private/MuT/ErrorLog.cpp
@@ -289,7 +289,6 @@ namespace mu
 		TEXT("ME_DIFFERENCE    "),
 		TEXT("ME_MORPH         "),
 		TEXT("ME_MERGE         "),
-		TEXT("ME_INTERPOLATE   "),
 		TEXT("ME_MASKCLIPMESH  "),
 		TEXT("ME_MASKCLIPUVMASK"),
 		TEXT("ME_MASKDIFF      "),
@@ -302,7 +301,6 @@ namespace mu
         TEXT("ME_SETSKELETON   "),
         TEXT("ME_PROJECT       "),
         TEXT("ME_APPLYPOSE     "),
-		TEXT("ME_GEOMETRYOP	   "),
 		TEXT("ME_BINDSHAPE	   "),
 		TEXT("ME_APPLYSHAPE	   "),
 		TEXT("ME_CLIPDEFORM	   "),
diff --git a/Engine/Plugins/Mutable/Source/MutableTools/Private/MuT/NodeMeshGeometryOperation.cpp b/Engine/Plugins/Mutable/Source/MutableTools/Private/MuT/NodeMeshGeometryOperation.cpp
deleted file mode 100644
diff --git a/Engine/Plugins/Mutable/Source/MutableTools/Private/MuT/NodeMeshGeometryOperationPrivate.h b/Engine/Plugins/Mutable/Source/MutableTools/Private/MuT/NodeMeshGeometryOperationPrivate.h
deleted file mode 100644
diff --git a/Engine/Plugins/Mutable/Source/MutableTools/Private/MuT/NodeMeshInterpolate.cpp b/Engine/Plugins/Mutable/Source/MutableTools/Private/MuT/NodeMeshInterpolate.cpp
deleted file mode 100644
diff --git a/Engine/Plugins/Mutable/Source/MutableTools/Private/MuT/NodeMeshInterpolatePrivate.h b/Engine/Plugins/Mutable/Source/MutableTools/Private/MuT/NodeMeshInterpolatePrivate.h
deleted file mode 100644
diff --git a/Engine/Plugins/Mutable/Source/MutableTools/Public/MuT/NodeMeshGeometryOperation.h b/Engine/Plugins/Mutable/Source/MutableTools/Public/MuT/NodeMeshGeometryOperation.h
deleted file mode 100644
diff --git a/Engine/Plugins/Mutable/Source/MutableTools/Public/MuT/NodeMeshInterpolate.h b/Engine/Plugins/Mutable/Source/MutableTools/Public/MuT/NodeMeshInterpolate.h
deleted file mode 100644
