[mutable] Protect access to som arrays that could be incorrect depending on data.
#rnx

diff --git a/Engine/Plugins/Mutable/Source/CustomizableObjectEditor/Private/MuCOE/GenerateMutableSource/GenerateMutableSourceMesh.cpp b/Engine/Plugins/Mutable/Source/CustomizableObjectEditor/Private/MuCOE/GenerateMutableSource/GenerateMutableSourceMesh.cpp
--- a/Engine/Plugins/Mutable/Source/CustomizableObjectEditor/Private/MuCOE/GenerateMutableSource/GenerateMutableSourceMesh.cpp
+++ b/Engine/Plugins/Mutable/Source/CustomizableObjectEditor/Private/MuCOE/GenerateMutableSource/GenerateMutableSourceMesh.cpp
@@ -3634,7 +3634,11 @@ mu::Ptr<mu::NodeMesh> GenerateMutableSourceMesh(const UEdGraphPin* Pin,
 					SectionIndexConnected :
 					SectionIndex;
 					
-				const int32 NumLayouts = ImportedModel->LODModels[LODIndexLayout].NumTexCoords;
+				int32 NumLayouts = 0;
+				if (ImportedModel && ImportedModel->LODModels.IsValidIndex(LODIndexLayout))
+				{
+					NumLayouts = ImportedModel->LODModels[LODIndexLayout].NumTexCoords;
+				}
 				MeshNode->SetLayoutCount(NumLayouts);
 				
 				const FLayoutGenerationFlags& LayoutFlags = GenerationContext.LayoutGenerationFlags.Last();
@@ -3648,7 +3652,7 @@ mu::Ptr<mu::NodeMesh> GenerateMutableSourceMesh(const UEdGraphPin* Pin,
 						MeshNode->SetLayout(LayoutIndex, CreateDefaultLayout());
 					
 						// Keep packing strategy if possible, Overlay can be valid with EPinMode == Passthrough	
-						if (Layouts[LayoutIndex])
+						if (Layouts.IsValidIndex(LayoutIndex) && Layouts[LayoutIndex])
 						{
 							MeshNode->GetLayout(LayoutIndex)->Strategy = ConvertLayoutStrategy(Layouts[LayoutIndex]->PackingStrategy);
 						}
@@ -3728,7 +3732,7 @@ mu::Ptr<mu::NodeMesh> GenerateMutableSourceMesh(const UEdGraphPin* Pin,
 				SectionIndex = -1;
 			}();
 			
-			check(SectionIndex < TypedNodeStatic->LODs[LODIndex].Materials.Num());
+			check(TypedNodeStatic->LODs.IsValidIndex(LODIndex) && SectionIndex < TypedNodeStatic->LODs[LODIndex].Materials.Num());
 
 			constexpr bool bIsReference = false;
 			TSharedPtr<mu::FMesh> MutableMesh = GenerateMutableMesh(TypedNodeStatic->StaticMesh, TSoftClassPtr<UAnimInstance>(), LODIndex, SectionIndex, LODIndex, SectionIndex, 
