[Mutable] Removed the code that didn't modify the "IsDirty" flag of a CO after a Synchronouly compilation.

#rnx

diff --git a/Engine/Plugins/Mutable/Source/CustomizableObjectEditor/Private/MuCOE/CustomizableObjectEditorFunctionLibrary.cpp b/Engine/Plugins/Mutable/Source/CustomizableObjectEditor/Private/MuCOE/CustomizableObjectEditorFunctionLibrary.cpp
--- a/Engine/Plugins/Mutable/Source/CustomizableObjectEditor/Private/MuCOE/CustomizableObjectEditorFunctionLibrary.cpp
+++ b/Engine/Plugins/Mutable/Source/CustomizableObjectEditor/Private/MuCOE/CustomizableObjectEditorFunctionLibrary.cpp
@@ -25,9 +25,6 @@ ECustomizableObjectCompilationState UCustomizableObjectEditorFunctionLibrary::Co
 		return ECustomizableObjectCompilationState::Failed;
 	}
 
-	// store package dirty state so that we can restore it - compile is not an edit:
-	const bool bPackageWasDirty = CustomizableObject->GetOutermost()->IsDirty();
-
 	const double StartTime = FPlatformTime::Seconds();
 
 	TSharedRef<FCompilationRequest> CompileRequest = MakeShared<FCompilationRequest>(*CustomizableObject, false);
@@ -57,8 +54,6 @@ PRAGMA_ENABLE_DEPRECATION_WARNINGS
 
 	check(CompileRequest->GetCompilationState() == ECompilationStatePrivate::Completed);
 
-	CustomizableObject->GetOutermost()->SetDirtyFlag(bPackageWasDirty);
-
 	const bool bCompilationSuccess = CompileRequest->GetCompilationResult() == ECompilationResultPrivate::Success ||
 		CompileRequest->GetCompilationResult() == ECompilationResultPrivate::Warnings;
 	
