[mutable] Fixed SkeletalMeshNode pins not updating when setting a SKM
- Added guard for logging compilation messages
- Removed unnecessary includes
- Forced the loading of the SKM part of the CustomziableObjectNodeSkeletalMesh to allow the correct pin setup of the node.

#rnx

diff --git a/Engine/Plugins/Mutable/Source/CustomizableObjectEditor/Private/MuCOE/GenerateMutableSource/GenerateMutableSource.cpp b/Engine/Plugins/Mutable/Source/CustomizableObjectEditor/Private/MuCOE/GenerateMutableSource/GenerateMutableSource.cpp
--- a/Engine/Plugins/Mutable/Source/CustomizableObjectEditor/Private/MuCOE/GenerateMutableSource/GenerateMutableSource.cpp
+++ b/Engine/Plugins/Mutable/Source/CustomizableObjectEditor/Private/MuCOE/GenerateMutableSource/GenerateMutableSource.cpp
@@ -15,14 +15,9 @@
 #include "MuCOE/EdGraphSchema_CustomizableObject.h"
 #include "MuCOE/GenerateMutableSource/GenerateMutableSourceGroupProjector.h"
 #include "MuCOE/GenerateMutableSource/GenerateMutableSourceModifier.h"
-#include "MuCOE/GenerateMutableSource/GenerateMutableSourceSurface.h"
 #include "MuCOE/GraphTraversal.h"
 #include "MuCOE/ICustomizableObjectExtensionNode.h"
-#include "MuCOE/Nodes/CustomizableObjectNodeModifierEditMeshSectionBase.h"
-#include "MuCOE/Nodes/CustomizableObjectNodeModifierExtendMeshSection.h"
 #include "MuCOE/Nodes/CustomizableObjectNodeMaterial.h"
-#include "MuCOE/Nodes/CustomizableObjectNodeMaterialVariation.h"
-#include "MuCOE/Nodes/CustomizableObjectNodeMaterialSwitch.h"
 #include "MuCOE/Nodes/CustomizableObjectNodeMesh.h"
 #include "MuCOE/Nodes/CustomizableObjectNodeModifierClipWithMesh.h"
 #include "MuCOE/Nodes/CustomizableObjectNodeObjectGroup.h"
@@ -45,7 +40,6 @@
 #include "MuCOE/Nodes/CustomizableObjectNodeTextureVariation.h"
 #include "MuCOE/UnrealEditorPortabilityHelpers.h"
 #include "MuT/NodeComponentEdit.h"
-#include "MuT/NodeLOD.h"
 #include "MuT/NodeMeshConstant.h"
 #include "MuT/NodeObjectGroup.h"
 #include "MuT/NodeObjectNew.h"
@@ -56,7 +50,6 @@
 #include "Math/NumericLimits.h"
 #include "Hash/CityHash.h"
 #include "MuCO/CustomizableObjectResourceDataTypes.h"
-#include "MuCOE/Nodes/CustomizableObjectNodeCopyMaterial.h"
 #include "MuCOE/Nodes/CustomizableObjectNodeGroupProjectorParameter.h"
 #include "UObject/ICookInfo.h"
 
@@ -151,7 +144,10 @@ void FMutableCompilationContext::Log(const FText& Message, const TArray<const UO
 
 void FMutableCompilationContext::Log(const FText& Message, const UObject* Context, const EMessageSeverity::Type MessageSeverity, const bool bAddBaseObjectInfo, const ELoggerSpamBin SpamBin) const
 {
-	Compiler->CompilerLog(Message, Context, MessageSeverity, bAddBaseObjectInfo, SpamBin);
+	if (Compiler)
+	{
+		Compiler->CompilerLog(Message, Context, MessageSeverity, bAddBaseObjectInfo, SpamBin);
+	}
 }
 
 
diff --git a/Engine/Plugins/Mutable/Source/CustomizableObjectEditor/Private/MuCOE/Nodes/CustomizableObjectNodeSkeletalMesh.cpp b/Engine/Plugins/Mutable/Source/CustomizableObjectEditor/Private/MuCOE/Nodes/CustomizableObjectNodeSkeletalMesh.cpp
--- a/Engine/Plugins/Mutable/Source/CustomizableObjectEditor/Private/MuCOE/Nodes/CustomizableObjectNodeSkeletalMesh.cpp
+++ b/Engine/Plugins/Mutable/Source/CustomizableObjectEditor/Private/MuCOE/Nodes/CustomizableObjectNodeSkeletalMesh.cpp
@@ -71,6 +71,9 @@ void UCustomizableObjectNodeSkeletalMesh::AllocateDefaultPins(UCustomizableObjec
 	
 	const UEdGraphSchema_CustomizableObject* Schema = GetDefault<UEdGraphSchema_CustomizableObject>();
 
+	// Force the loading of the skeletal mesh so we can later access the internal WeakObjectPtr (may be null if no SKM is set)
+	SkeletalMesh.LoadSynchronous();
+	
 	if (!SkeletalMesh)
 	{
 		UCustomizableObjectNodeSkeletalMeshPinDataMesh* PinData = NewObject<UCustomizableObjectNodeSkeletalMeshPinDataMesh>(this);
