[Mutable] Now NumExtraLODsToBuildPerPlatform values are deduced automatically from the CO's quality settings.

#rnx

diff --git a/Engine/Plugins/Mutable/Source/CustomizableObjectEditor/Private/MuCOE/GenerateMutableSource/GenerateMutableSource.cpp b/Engine/Plugins/Mutable/Source/CustomizableObjectEditor/Private/MuCOE/GenerateMutableSource/GenerateMutableSource.cpp
--- a/Engine/Plugins/Mutable/Source/CustomizableObjectEditor/Private/MuCOE/GenerateMutableSource/GenerateMutableSource.cpp
+++ b/Engine/Plugins/Mutable/Source/CustomizableObjectEditor/Private/MuCOE/GenerateMutableSource/GenerateMutableSource.cpp
@@ -920,14 +920,26 @@ mu::Ptr<mu::NodeObject> GenerateMutableSource(const UEdGraphPin * Pin, FMutableG
 				ObjectNode->AddStateParam(StateIndex, State.RuntimeParameters[ParamIndex]);
 			}
 
+			const FMutableLODSettings& LODSettings = GenerationContext.CompilationContext.Object->LODSettings;
 			const ITargetPlatform* TargetPlatform = GenerationContext.CompilationContext.Options.TargetPlatform;
+			const FSupportedQualityLevelArray SupportedQualityLevels = LODSettings.MinQualityLevelLOD.GetSupportedQualityLevels(*TargetPlatform->GetPlatformInfo().IniPlatformName.ToString());
 
-			int32 NumExtraLODsToBuildAfterFirstLOD = 0;
-			if (const int32* AuxNumExtraLODs = State.NumExtraLODsToBuildPerPlatform.Find(TargetPlatform->PlatformName()))
+			int32 MinQualityLevelLOD = INT32_MAX;
+			int32 MaxQualityLevelLOD = 0;
+
+			for (const int32& QualitySetting : SupportedQualityLevels)
 			{
-				NumExtraLODsToBuildAfterFirstLOD = *AuxNumExtraLODs;
+				if (LODSettings.MinQualityLevelLOD.IsQualityLevelValid(QualitySetting))
+				{
+					int32 QualityLod = LODSettings.MinQualityLevelLOD.GetValueForQualityLevel(QualitySetting);
+					MinQualityLevelLOD = FMath::Min(MinQualityLevelLOD, QualityLod);
+					MaxQualityLevelLOD = FMath::Max(MaxQualityLevelLOD, QualityLod);
+				}
 			}
 
+			// Generate the difference between the maximum and minimum LODs found.
+			const int32 NumExtraLODsToBuildAfterFirstLOD = MinQualityLevelLOD != INT32_MAX ? MaxQualityLevelLOD - MinQualityLevelLOD : 0;
+
 			ObjectNode->SetStateProperties(StateIndex, State.TextureCompressionStrategy, State.bBuildOnlyFirstLOD, NumExtraLODsToBuildAfterFirstLOD);
 
 			// UI Data
diff --git a/Engine/Plugins/Mutable/Source/CustomizableObjectEditor/Public/MuCOE/Nodes/CustomizableObjectNodeObject.h b/Engine/Plugins/Mutable/Source/CustomizableObjectEditor/Public/MuCOE/Nodes/CustomizableObjectNodeObject.h
--- a/Engine/Plugins/Mutable/Source/CustomizableObjectEditor/Public/MuCOE/Nodes/CustomizableObjectNodeObject.h
+++ b/Engine/Plugins/Mutable/Source/CustomizableObjectEditor/Public/MuCOE/Nodes/CustomizableObjectNodeObject.h
@@ -54,10 +54,6 @@ struct CUSTOMIZABLEOBJECTEDITOR_API FCustomizableObjectState
 	UPROPERTY(EditAnywhere, Category = CustomizableObject)
 	bool bBuildOnlyFirstLOD = false;
 
-	/** If there's an entry for a specific platform, when compiling for that platform Num LODs will be built after bBuildOnlyFirstLOD */
-	UPROPERTY(EditAnywhere, Category = CustomizableObject)
-	TMap<FString, int32> NumExtraLODsToBuildPerPlatform;
-
 	UPROPERTY(EditAnywhere, Category = CustomizableObject)
 	TMap<FString, FString> ForcedParameterValues;
 
