[Mutable] Remove Customizable Object Instance Usage Tick.

- Moved UCustomizableObjectInstanceUsage::Tick to UCusotmizableObjectSystem::TickInternal.
- Removed calls to SetSkeletalMesh that did not follow the Update flow. Callbacks must always be called.
- Removed UCustomizableObjectInstanceUsage SetSkeletalMesh. Added bAutomaticUpdateRequired which marks an Instance to start a normal Update flow.
- Removed EditorUpdateComponent SetSkeletalMesh. No longer necessary.
- UCustomizableObjectSystem::DiscardInstances and UCustomizableInstancePrivate::SetDefaultSkeletalMesh now sets the Reference Skeletal Mesh directly instead of using SetSkeletalMeshAndOverrideMaterials. SetSkeletalMeshAndOverrideMaterials had code specific to the normal Update flow.
- Inlined SetSkeletalMesh (+SetSkeletalMeshAndOverrideMaterials) and SetPhysicsAsset.
- Moved RequiresReinitPose to the UCustomizableObjectSystem.
- Simplified and renamed SetDefaultSkeletalMesh to SetReferenceSkeletalMesh.

diff --git a/Engine/Plugins/Mutable/Source/CustomizableObject/Internal/MuCO/CustomizableObjectInstancePrivate.h b/Engine/Plugins/Mutable/Source/CustomizableObject/Internal/MuCO/CustomizableObjectInstancePrivate.h
--- a/Engine/Plugins/Mutable/Source/CustomizableObject/Internal/MuCO/CustomizableObjectInstancePrivate.h
+++ b/Engine/Plugins/Mutable/Source/CustomizableObject/Internal/MuCO/CustomizableObjectInstancePrivate.h
@@ -287,7 +287,7 @@ public:
 	void ReleaseMutableResources(bool bCalledFromBeginDestroy, const UCustomizableObjectInstance& Instance);
 
 	/** Set the reference SkeletalMesh, or an empty mesh, to all actors using this instance. */
-	void SetDefaultSkeletalMesh(bool bSetEmptyMesh = false) const;
+	void SetReferenceSkeletalMesh() const;
 
 	const TArray<FAnimInstanceOverridePhysicsAsset>* GetGeneratedPhysicsAssetsForAnimInstance(TSubclassOf<UAnimInstance> AnimInstance) const;
 
@@ -414,6 +414,10 @@ public:
 	bool bShowOnlyRelevantParameters = true;
 	bool bShowUISections = false;
 	bool bShowUIThumbnails = false;
+
+	/** Automatic update required.
+	 * Set to true when a Customizable Object Instance Usage requires an automatic update (e.g., component reattached). */
+	bool bAutomaticUpdateRequired = false;
 	
 	// TEMP VARIABLE to check the Min desired LODs for this instance
 	TWeakObjectPtr<UCustomizableObjectInstanceUsage> NearestToActor;
diff --git a/Engine/Plugins/Mutable/Source/CustomizableObject/Internal/MuCO/CustomizableObjectInstanceUsagePrivate.h b/Engine/Plugins/Mutable/Source/CustomizableObject/Internal/MuCO/CustomizableObjectInstanceUsagePrivate.h
--- a/Engine/Plugins/Mutable/Source/CustomizableObject/Internal/MuCO/CustomizableObjectInstanceUsagePrivate.h
+++ b/Engine/Plugins/Mutable/Source/CustomizableObject/Internal/MuCO/CustomizableObjectInstanceUsagePrivate.h
@@ -14,47 +14,27 @@ class USkeletalMesh;
 class USkeletalMeshComponent;
 
 UCLASS()
-class CUSTOMIZABLEOBJECT_API UCustomizableObjectInstanceUsagePrivate : public UObject, public FTickableGameObject
+class CUSTOMIZABLEOBJECT_API UCustomizableObjectInstanceUsagePrivate : public UObject
 {
 	GENERATED_BODY()
 
 public:
-	// FTickableGameObject interface
-	virtual void Tick(float DeltaTime) override;
-	virtual ETickableTickType GetTickableTickType() const override;
-	virtual TStatId GetStatId() const override;
-	virtual bool IsTickableWhenPaused() const override;
-	virtual bool IsTickableInEditor() const override;
-	virtual bool IsTickable() const override;
-	
 	// Own interface
 	
 	/** Common end point of all updates. Even those which failed. */
 	void Callbacks();
 
-	void SetSkeletalMeshAndOverrideMaterials(USkeletalMeshComponent& Parent, USkeletalMesh* SkeletalMesh, const UCustomizableObjectInstance& CustomizableObjectInstance, bool* bOutSkeletalMeshUpdated, bool* bOutMaterialsUpdated);
-
 #if WITH_EDITOR
 	// Used to generate instances outside the CustomizableObject editor and PIE
 	void UpdateDistFromComponentToLevelEditorCamera(const FVector& CameraPosition);
-	void EditorUpdateComponent();
 #endif
 
 	USkeletalMesh* GetSkeletalMesh() const;
 
 	USkeletalMesh* GetAttachedSkeletalMesh() const;
-
-	void SetSkeletalMesh(USkeletalMesh* SkeletalMesh, bool* bOutSkeletalMeshUpdated = nullptr, bool* bOutMaterialsUpdated = nullptr);
-
-	void SetPhysicsAsset(UPhysicsAsset* PhysicsAsset, bool* bOutPhysicsAssetUpdated = nullptr);
-
+	
 	void UpdateDistFromComponentToPlayer(const AActor* const Pawn, bool bForceEvenIfNotBegunPlay = false);
-
-	/** Set to true to replace the SkeletalMesh of the parent component by the ReferenceSkeletalMesh or the generated SkeletalMesh. */
-	void SetPendingSetSkeletalMesh(bool bIsActive);
 	
-	bool GetPendingSetSkeletalMesh() const;
-
 	/** Return a valid pointer if this Usage is being used inside a UCustomizableSkeletalComponent. Nullptr if it is being used in Standalone mode. */
 	UCustomizableSkeletalComponent* GetCustomizableSkeletalComponent() const;
 	
@@ -64,7 +44,7 @@ public:
 
 	// Returns true if the NetMode of the associated UCustomizableSkeletalComponent (or the associated SkeletalMeshComponent if the former does not exist) is equal to InNetMode
 	bool IsNetMode(ENetMode InNetMode) const;
-	
-	// Used to replace the SkeletalMesh of the parent component by the ReferenceSkeletalMesh or the generated SkeletalMesh 
-	bool bUsedPendingSetSkeletalMesh = false;
+
+	/** Used to replace the Skeletal Mesh of the parent component by the Reference Skeletal Mesh or the generated Skeletal Mesh. */ 
+	bool bPendingSetSkeletalMesh = true;
 };
diff --git a/Engine/Plugins/Mutable/Source/CustomizableObject/Internal/MuCO/CustomizableSkeletalComponentPrivate.h b/Engine/Plugins/Mutable/Source/CustomizableObject/Internal/MuCO/CustomizableSkeletalComponentPrivate.h
--- a/Engine/Plugins/Mutable/Source/CustomizableObject/Internal/MuCO/CustomizableSkeletalComponentPrivate.h
+++ b/Engine/Plugins/Mutable/Source/CustomizableObject/Internal/MuCO/CustomizableSkeletalComponentPrivate.h
@@ -25,18 +25,6 @@ public:
 
 	USkeletalMesh* GetAttachedSkeletalMesh() const;
 	
-	void SetSkeletalMesh(USkeletalMesh* SkeletalMesh);
-	
-	void SetPhysicsAsset(UPhysicsAsset* PhysicsAsset);
-	
-	void SetPendingSetSkeletalMesh(bool bIsActive);
-
-#if WITH_EDITOR
-	void EditorUpdateComponent();
-#endif
-
-	bool& PendingSetSkeletalMesh();
-	
 	UCustomizableSkeletalComponent* GetPublic();
 	
 	const UCustomizableSkeletalComponent* GetPublic() const;
diff --git a/Engine/Plugins/Mutable/Source/CustomizableObject/Private/MuCO/CustomizableInstanceLODManagement.cpp b/Engine/Plugins/Mutable/Source/CustomizableObject/Private/MuCO/CustomizableInstanceLODManagement.cpp
--- a/Engine/Plugins/Mutable/Source/CustomizableObject/Private/MuCO/CustomizableInstanceLODManagement.cpp
+++ b/Engine/Plugins/Mutable/Source/CustomizableObject/Private/MuCO/CustomizableInstanceLODManagement.cpp
@@ -59,7 +59,6 @@ static TAutoConsoleVariable<bool> CVarOnlyUpdateCloseCustomizableObjects(
 #if WITH_EDITOR
 void UCustomizableInstanceLODManagementBase::EditorUpdateComponent(UCustomizableObjectInstanceUsage* InstanceUsage)
 {
-	InstanceUsage->GetPrivate()->EditorUpdateComponent();
 }
 #endif
 
@@ -176,12 +175,6 @@ void UCustomizableInstanceLODManagement::UpdateInstanceDistsAndLODs(FMutableInst
 								}
 							}
 						}
-
-						// Blueprint instances
-						else if (WorldType == EWorldType::EditorPreview)
-						{
-							CustomizableObjectInstanceUsage->GetPrivate()->EditorUpdateComponent();
-						}
 					}
 	#endif // WITH_EDITOR
 
@@ -384,7 +377,6 @@ void UCustomizableInstanceLODManagement::UpdateInstanceDistsAndLODs(FMutableInst
 
 				if (WorldType == EWorldType::EditorPreview || (!World && !bAttachParentActor))
 				{
-					CustomizableObjectInstanceUsage->GetPrivate()->EditorUpdateComponent();
 					continue;
 				}
 
diff --git a/Engine/Plugins/Mutable/Source/CustomizableObject/Private/MuCO/CustomizableObjectInstance.cpp b/Engine/Plugins/Mutable/Source/CustomizableObject/Private/MuCO/CustomizableObjectInstance.cpp
--- a/Engine/Plugins/Mutable/Source/CustomizableObject/Private/MuCO/CustomizableObjectInstance.cpp
+++ b/Engine/Plugins/Mutable/Source/CustomizableObject/Private/MuCO/CustomizableObjectInstance.cpp
@@ -3080,7 +3080,7 @@ void UCustomizableInstancePrivate::DiscardResources()
 }
 
 
-void UCustomizableInstancePrivate::SetDefaultSkeletalMesh(bool bSetEmptyMesh) const
+void UCustomizableInstancePrivate::SetReferenceSkeletalMesh() const
 {
 	UCustomizableObjectInstance* Instance = Cast<UCustomizableObjectInstance>(GetOuter());
 	if (!Instance)
@@ -3117,16 +3117,14 @@ void UCustomizableInstancePrivate::SetDefaultSkeletalMesh(bool bSetEmptyMesh) co
 		{
 			continue;
 		}
-
-		USkeletalMesh* SkeletalMesh = nullptr;
 		
-		if (!bSetEmptyMesh)
+		if (USkeletalMeshComponent* Parent = CustomizableObjectInstanceUsage->GetAttachParent())
 		{
-			// Force load the reference mesh if necessary. 
-			SkeletalMesh = ModelResources.ReferenceSkeletalMeshesData[ObjectComponentIndex].SoftSkeletalMesh.LoadSynchronous();
+			Parent->EmptyOverrideMaterials();
+			
+			USkeletalMesh* SkeletalMesh = ModelResources.ReferenceSkeletalMeshesData[ObjectComponentIndex].SoftSkeletalMesh.LoadSynchronous();
+			Parent->SetSkeletalMesh(SkeletalMesh);
 		}
-		
-		CustomizableObjectInstanceUsage->GetPrivate()->SetSkeletalMesh(SkeletalMesh);
 	}
 }
 
diff --git a/Engine/Plugins/Mutable/Source/CustomizableObject/Private/MuCO/CustomizableObjectInstanceUsage.cpp b/Engine/Plugins/Mutable/Source/CustomizableObject/Private/MuCO/CustomizableObjectInstanceUsage.cpp
--- a/Engine/Plugins/Mutable/Source/CustomizableObject/Private/MuCO/CustomizableObjectInstanceUsage.cpp
+++ b/Engine/Plugins/Mutable/Source/CustomizableObject/Private/MuCO/CustomizableObjectInstanceUsage.cpp
@@ -22,14 +22,6 @@
 #include "Stats/Stats.h"
 
 
-ETickableTickType UCustomizableObjectInstanceUsagePrivate::GetTickableTickType() const
-{ 
-	return (HasAnyFlags(RF_ClassDefaultObject))
-		? ETickableTickType::Never
-		: ETickableTickType::Conditional;
-}
-
-
 void UCustomizableObjectInstanceUsagePrivate::Callbacks()
 {
 	for (const UCustomizableObjectExtension* Extension : ICustomizableObjectModule::Get().GetRegisteredExtensions())
@@ -144,32 +136,6 @@ FName UCustomizableObjectInstanceUsage::GetComponentName() const
 }
 
 
-void UCustomizableObjectInstanceUsagePrivate::SetPendingSetSkeletalMesh(bool bIsActive)
-{
-	if (UCustomizableSkeletalComponent* CustomizableSkeletalComponent = GetCustomizableSkeletalComponent())
-	{
-		CustomizableSkeletalComponent->GetPrivate()->PendingSetSkeletalMesh() = bIsActive;
-	}
-	else
-	{
-		bUsedPendingSetSkeletalMesh = bIsActive;
-	}
-}
-
-
-bool UCustomizableObjectInstanceUsagePrivate::GetPendingSetSkeletalMesh() const
-{
-	if (UCustomizableSkeletalComponent* CustomizableSkeletalComponent = GetCustomizableSkeletalComponent())
-	{
-		return CustomizableSkeletalComponent->GetPrivate()->PendingSetSkeletalMesh();
-	}
-	else
-	{
-		return bUsedPendingSetSkeletalMesh;
-	}
-}
-
-
 UCustomizableSkeletalComponent* UCustomizableObjectInstanceUsagePrivate::GetCustomizableSkeletalComponent() const
 {
 	return Cast<UCustomizableSkeletalComponent>(GetPublic()->GetOuter()->GetOuter());
@@ -247,24 +213,18 @@ bool UCustomizableObjectInstanceUsage::GetSkipSetSkeletalMeshOnAttach() const
 
 void UCustomizableObjectInstanceUsage::AttachTo(USkeletalMeshComponent* SkeletalMeshComponent)
 {
-	if (GetPrivate()->GetCustomizableSkeletalComponent())
+	if (IsValid(SkeletalMeshComponent))
 	{
-		UE_LOG(LogMutable, Error, TEXT("Cannot change the attachment of a UCustomizableObjectInstanceUsage that has been automatically created by a CustomizableSkeletalComponent. Reattach the CustomizableSkeletalComponent instead."));
-		ensure(false);
+		UsedSkeletalMeshComponent = SkeletalMeshComponent;
 	}
 	else
 	{
-		if (IsValid(SkeletalMeshComponent))
-		{
-			UsedSkeletalMeshComponent = SkeletalMeshComponent;
-		}
-		else
-		{
-			UsedSkeletalMeshComponent = nullptr;
-		}
+		UsedSkeletalMeshComponent = nullptr;
+	}
 
-		// To mimic the behavior of UCustomizableSkeletalComponent::OnAttachmentChanged()
-		GetPrivate()->SetPendingSetSkeletalMesh(true);
+	if (!GetSkipSetSkeletalMeshOnAttach())
+	{
+		GetPrivate()->bPendingSetSkeletalMesh = true;
 	}
 }
 
@@ -292,128 +252,6 @@ USkeletalMesh* UCustomizableObjectInstanceUsagePrivate::GetSkeletalMesh() const
 }
 
 
-bool RequiresReinitPose(USkeletalMesh* CurrentSkeletalMesh, USkeletalMesh* SkeletalMesh)
-{
-	if (CurrentSkeletalMesh == SkeletalMesh)
-	{
-		return false;
-	}
-
-	if (!CurrentSkeletalMesh || !SkeletalMesh)
-	{
-		return SkeletalMesh != nullptr;
-	}
-
-	if (CurrentSkeletalMesh->GetLODNum() != SkeletalMesh->GetLODNum())
-	{
-		return true;
-	}
-
-	const FSkeletalMeshRenderData* CurrentRenderData = CurrentSkeletalMesh->GetResourceForRendering();
-	const FSkeletalMeshRenderData* NewRenderData = SkeletalMesh->GetResourceForRendering();
-	if (!CurrentRenderData || !NewRenderData)
-	{
-		return false;
-	}
-
-	const int32 NumLODs = SkeletalMesh->GetLODNum();
-	for (int32 LODIndex = 0; LODIndex < NumLODs; ++LODIndex)
-	{
-		if (CurrentRenderData->LODRenderData[LODIndex].RequiredBones != NewRenderData->LODRenderData[LODIndex].RequiredBones)
-		{
-			return true;
-		}
-	}
-
-	return false;
-}
-
-
-void UCustomizableObjectInstanceUsagePrivate::SetSkeletalMeshAndOverrideMaterials(USkeletalMeshComponent& Parent, USkeletalMesh* SkeletalMesh, 
-	const UCustomizableObjectInstance& CustomizableObjectInstance, bool* bOutSkeletalMeshUpdated, bool* bOutMaterialsUpdated)
-{
-	if (SkeletalMesh != Parent.GetSkeletalMeshAsset())
-	{
-		Parent.SetSkeletalMesh(SkeletalMesh, RequiresReinitPose(Parent.GetSkeletalMeshAsset(), SkeletalMesh));
-
-		if (bOutSkeletalMeshUpdated)
-		{
-			*bOutSkeletalMeshUpdated = true;
-		}
-	}
-
-	SetPendingSetSkeletalMesh(false);
-
-	TArray<TObjectPtr<UMaterialInterface>> OldOverridenMaterials = Parent.OverrideMaterials;
-
-	if (Parent.HasOverrideMaterials())
-	{
-#if !(UE_BUILD_SHIPPING || UE_BUILD_TEST)
-		if (Parent.GetClass()->GetFName() != FName(TEXT("SkeletalMeshComponentBudgeted"))) // Reduce unnecessary logging
-		{
-			UE_LOG(LogMutable, Warning, TEXT("Attaching Customizable Skeletal Component to Skeletal Mesh Component with overriden materials! Deleting overrides."));
-		}
-#endif
-
-		// For some reason the reference skeletal mesh materials are added as override materials, clear them if necessary
-		Parent.EmptyOverrideMaterials();
-	}
-
-	const UCustomizableObject* CustomizableObject = CustomizableObjectInstance.GetCustomizableObject();
-	if (!CustomizableObject)
-	{
-		return;
-	}
-	
-	const bool bIsTransientMesh = SkeletalMesh ? SkeletalMesh->HasAllFlags(EObjectFlags::RF_Transient) : false;
-	const bool bUseOverrideMaterials = !bIsTransientMesh || (CustomizableObject->bEnableMeshCache && UCustomizableObjectSystem::IsMeshCacheEnabled());
-
-	if (bUseOverrideMaterials)
-	{
-		if (FCustomizableInstanceComponentData* ComponentData = CustomizableObjectInstance.GetPrivate()->GetComponentData(GetPublic()->GetComponentName()))
-		{
-			for (int32 Index = 0; Index < ComponentData->OverrideMaterials.Num(); ++Index)
-			{
-				Parent.SetMaterial(Index, ComponentData->OverrideMaterials[Index]);
-			}
-		}
-	}
-
-	if (bOutMaterialsUpdated)
-	{
-		*bOutMaterialsUpdated = OldOverridenMaterials != Parent.OverrideMaterials;
-	}
-}
-
-
-void UCustomizableObjectInstanceUsagePrivate::SetSkeletalMesh(USkeletalMesh* SkeletalMesh, bool* bOutSkeletalMeshUpdated, bool* bOutMaterialsUpdated)
-{
-	USkeletalMeshComponent* Parent = Cast<USkeletalMeshComponent>(GetPublic()->GetAttachParent());
-	const UCustomizableObjectInstance* CustomizableObjectInstance = GetPublic()->GetCustomizableObjectInstance();
-
-	if (Parent && CustomizableObjectInstance)
-	{
-		SetSkeletalMeshAndOverrideMaterials(*Parent, SkeletalMesh, *CustomizableObjectInstance, bOutSkeletalMeshUpdated, bOutMaterialsUpdated);
-	}
-}
-
-void UCustomizableObjectInstanceUsagePrivate::SetPhysicsAsset(UPhysicsAsset* PhysicsAsset, bool* bOutPhysicsAssetUpdated)
-{
-	USkeletalMeshComponent* Parent = Cast<USkeletalMeshComponent>(GetPublic()->GetAttachParent());
-
-	if (Parent && Parent->GetWorld() &&
-		PhysicsAsset != Parent->GetPhysicsAsset())
-	{
-		Parent->SetPhysicsAsset(PhysicsAsset, true);
-		
-		if (bOutPhysicsAssetUpdated)
-		{
-			*bOutPhysicsAssetUpdated = true;
-		}
-	}
-}
-
-
 USkeletalMesh* UCustomizableObjectInstanceUsagePrivate::GetAttachedSkeletalMesh() const
 {
 	USkeletalMeshComponent* Parent = Cast<USkeletalMeshComponent>(GetPublic()->GetAttachParent());
@@ -473,50 +311,6 @@ void UCustomizableObjectInstanceUsagePrivate::UpdateDistFromComponentToLevelEdit
 			CustomizableObjectInstance->GetPrivate()->MinSquareDistFromComponentToPlayer = 
 				FMath::Min(SquareDist, CustomizableObjectInstance->GetPrivate()->MinSquareDistFromComponentToPlayer);
 		}
-
-		USkeletalMesh* AttachedSkeletalMesh = GetAttachedSkeletalMesh();
-		const FName& ComponentName = GetPublic()->GetComponentName();
-
-		const bool bInstanceGenerated = CustomizableObjectInstance->GetPrivate()->SkeletalMeshStatus != ESkeletalMeshStatus::NotGenerated;
-		USkeletalMesh* GeneratedSkeletalMesh = bInstanceGenerated ? CustomizableObjectInstance->GetComponentMeshSkeletalMesh(ComponentName) :
-			CustomizableObjectInstance->GetCustomizableObject()->GetComponentMeshReferenceSkeletalMesh(ComponentName);
-
-		if (AttachedSkeletalMesh != GeneratedSkeletalMesh)
-		{
-			SetSkeletalMesh(GeneratedSkeletalMesh);
-		}
-	}
-}
-
-
-void UCustomizableObjectInstanceUsagePrivate::EditorUpdateComponent()
-{
-	UCustomizableObjectInstance* CustomizableObjectInstance = GetPublic()->GetCustomizableObjectInstance();
-
-	if (CustomizableObjectInstance)
-	{
-		CustomizableObjectInstance->GetPrivate()->SetCOInstanceFlags(UsedByComponent);
-
-		USkeletalMeshComponent* SkeletalMeshComponent = GetPublic()->GetAttachParent();
-		AActor* ParentActor = SkeletalMeshComponent ? SkeletalMeshComponent->GetAttachmentRootActor() : nullptr;
-		
-		if (ParentActor)
-		{
-			USkeletalMesh* AttachedSkeletalMesh = GetAttachedSkeletalMesh();
-			const FName& ComponentName = GetPublic()->GetComponentName();
-
-			const bool bInstanceGenerated = CustomizableObjectInstance->GetPrivate()->SkeletalMeshStatus != ESkeletalMeshStatus::NotGenerated;
-			USkeletalMesh* GeneratedSkeletalMesh = bInstanceGenerated ?
-				CustomizableObjectInstance->GetComponentMeshSkeletalMesh(ComponentName) :
-				GetPublic()->GetSkipSetReferenceSkeletalMesh() ?
-					nullptr :
-					CustomizableObjectInstance->GetCustomizableObject()->GetComponentMeshReferenceSkeletalMesh(ComponentName);
-
-			if (AttachedSkeletalMesh != GeneratedSkeletalMesh)
-			{
-				SetSkeletalMesh(GeneratedSkeletalMesh);
-			}
-		}
 	}
 }
 #endif
@@ -587,102 +381,10 @@ void UCustomizableObjectInstanceUsagePrivate::UpdateDistFromComponentToPlayer(co
 				}
 			}
 		}
-
-		const FName& ComponentName = GetPublic()->GetComponentName();
-
-		if (ParentActor && GetAttachedSkeletalMesh() == nullptr && CustomizableObjectInstance->GetComponentMeshSkeletalMesh(ComponentName))
-		{
-			SetSkeletalMesh(CustomizableObjectInstance->GetComponentMeshSkeletalMesh(ComponentName));
-		}
 	}
 }
 
 
-void UCustomizableObjectInstanceUsagePrivate::Tick(float DeltaTime)
-{
-	if (!IsValid(this))
-	{
-		return;
-	}
-
-	UCustomizableObjectInstance* CustomizableObjectInstance = GetPublic()->GetCustomizableObjectInstance();
-
-	if (!GetPendingSetSkeletalMesh() || GetPublic()->GetSkipSetSkeletalMeshOnAttach() || !CustomizableObjectInstance || !IsValid(CustomizableObjectInstance))
-	{
-		return;
-	}
-	
-	UCustomizableObject* CustomizableObject = CustomizableObjectInstance->GetCustomizableObject();
-	if (!CustomizableObject || !IsValid(CustomizableObject))
-	{
-		return;
-	}	
-
-	if (USkeletalMeshComponent* Parent = Cast<USkeletalMeshComponent>(GetPublic()->GetAttachParent()))
-	{
-		const FName& ComponentName = GetPublic()->GetComponentName();
-
-		USkeletalMesh* SkeletalMesh = nullptr;
-
-		bool bInstanceGenerated = CustomizableObjectInstance->GetPrivate()->SkeletalMeshStatus == ESkeletalMeshStatus::Success;
-		
-#if WITH_EDITOR
-		bInstanceGenerated = bInstanceGenerated && !IsNetMode(NM_DedicatedServer);
-#endif
-		
-		if (bInstanceGenerated)
-		{
-			// Generated SkeletalMesh to set, can be null if the component is empty
-			SkeletalMesh = CustomizableObjectInstance->GetComponentMeshSkeletalMesh(ComponentName);
-		}
-		else
-		{
-			// If not generated yet, conditionally set the SkeletalMesh of reference
-			if (!bInstanceGenerated && !GetPublic()->GetSkipSetReferenceSkeletalMesh() && 
-				CustomizableObject->bEnableUseRefSkeletalMeshAsPlaceholder)
-			{
-				// Can be nullptr
-				SkeletalMesh = CustomizableObject->GetComponentMeshReferenceSkeletalMesh(ComponentName);
-			}
-		}
-
-		// Set SkeletalMesh
-		if (bInstanceGenerated || SkeletalMesh)
-		{
-			SetSkeletalMeshAndOverrideMaterials(*Parent, SkeletalMesh, *CustomizableObjectInstance, nullptr, nullptr);
-		}
-		else
-		{
-			SetPendingSetSkeletalMesh(false);
-		}
-	}
-}
-
-
-TStatId UCustomizableObjectInstanceUsagePrivate::GetStatId() const
-{
-	RETURN_QUICK_DECLARE_CYCLE_STAT(UCustomizableObjectInstanceUsage, STATGROUP_Tickables);
-}
-
-
-bool UCustomizableObjectInstanceUsagePrivate::IsTickableWhenPaused() const
-{
-	return true;
-}
-
-
-bool UCustomizableObjectInstanceUsagePrivate::IsTickableInEditor() const
-{
-	return true;
-}
-
-
-bool UCustomizableObjectInstanceUsagePrivate::IsTickable() const
-{
-	return !HasAnyFlags(RF_BeginDestroyed);
-}
-
-
 bool UCustomizableObjectInstanceUsagePrivate::IsNetMode(ENetMode InNetMode) const
 {
 	if (UCustomizableSkeletalComponent* CustomizableSkeletalComponent = GetCustomizableSkeletalComponent())
@@ -696,3 +398,4 @@ bool UCustomizableObjectInstanceUsagePrivate::IsNetMode(ENetMode InNetMode) cons
 
 	return false;
 }
+
diff --git a/Engine/Plugins/Mutable/Source/CustomizableObject/Private/MuCO/CustomizableObjectSystem.cpp b/Engine/Plugins/Mutable/Source/CustomizableObject/Private/MuCO/CustomizableObjectSystem.cpp
--- a/Engine/Plugins/Mutable/Source/CustomizableObject/Private/MuCO/CustomizableObjectSystem.cpp
+++ b/Engine/Plugins/Mutable/Source/CustomizableObject/Private/MuCO/CustomizableObjectSystem.cpp
@@ -853,6 +853,8 @@ void FinishUpdateGlobal(const TSharedRef<FUpdateContextPrivate>& Context)
 			PrivateInstance->CommittedDescriptor = Context->MoveCommittedDescriptor();
 			PrivateInstance->CommittedDescriptorHash = Context->GetCapturedDescriptorHash();
 
+			PrivateInstance->bAutomaticUpdateRequired = false;
+
 			if (SystemPrivate)
 			{
 				// Cache new Texture FParameters
@@ -982,6 +984,43 @@ void FinishUpdateGlobal(const TSharedRef<FUpdateContextPrivate>& Context)
 }
 
 
+bool RequiresReinitPose(USkeletalMesh* CurrentSkeletalMesh, USkeletalMesh* SkeletalMesh)
+{
+	if (CurrentSkeletalMesh == SkeletalMesh)
+	{
+		return false;
+	}
+
+	if (!CurrentSkeletalMesh || !SkeletalMesh)
+	{
+		return SkeletalMesh != nullptr;
+	}
+
+	if (CurrentSkeletalMesh->GetLODNum() != SkeletalMesh->GetLODNum())
+	{
+		return true;
+	}
+
+	const FSkeletalMeshRenderData* CurrentRenderData = CurrentSkeletalMesh->GetResourceForRendering();
+	const FSkeletalMeshRenderData* NewRenderData = SkeletalMesh->GetResourceForRendering();
+	if (!CurrentRenderData || !NewRenderData)
+	{
+		return false;
+	}
+
+	const int32 NumLODs = SkeletalMesh->GetLODNum();
+	for (int32 LODIndex = 0; LODIndex < NumLODs; ++LODIndex)
+	{
+		if (CurrentRenderData->LODRenderData[LODIndex].RequiredBones != NewRenderData->LODRenderData[LODIndex].RequiredBones)
+		{
+			return true;
+		}
+	}
+
+	return false;
+}
+
+
 /** Update the given Instance Skeletal Meshes */
 void UpdateSkeletalMesh(const TSharedRef<FUpdateContextPrivate>& Context)
 {
@@ -992,6 +1031,9 @@ void UpdateSkeletalMesh(const TSharedRef<FUpdateContextPrivate>& Context)
 	UCustomizableObjectInstance* CustomizableObjectInstance = Context->Instance.Get();
 	check(CustomizableObjectInstance);
 
+	UCustomizableObject* CustomizableObject = CustomizableObjectInstance->GetCustomizableObject();
+	check(CustomizableObject);
+
 	for (const TTuple<FName, TObjectPtr<USkeletalMesh>>& Pair : CustomizableObjectInstance->GetPrivate()->SkeletalMeshes)
 	{
 		FPreSetSkeletalMeshParams Params;
@@ -1006,37 +1048,85 @@ void UpdateSkeletalMesh(const TSharedRef<FUpdateContextPrivate>& Context)
 	check(CustomizableObjectInstancePrivateData != nullptr);
 	for (TObjectIterator<UCustomizableObjectInstanceUsage> It; It; ++It)
 	{
-		UCustomizableObjectInstanceUsage* CustomizableObjectInstanceUsage = *It;
+		UCustomizableObjectInstanceUsage* Usage = *It;
 
+		if (!IsValid(Usage))
+		{
+			continue;
+		}
+		
 #if WITH_EDITOR
-		if (IsValid(CustomizableObjectInstanceUsage) && CustomizableObjectInstanceUsage->GetPrivate()->IsNetMode(NM_DedicatedServer))
+		if (Usage->GetPrivate()->IsNetMode(NM_DedicatedServer))
 		{
 			continue;
 		}
 #endif
 
-		bool bSkeletalMeshUpdated = false;
-		bool bMaterialsUpdated = false;
-		bool bPhysicsAssetUpdated = false;	
-		
-		if (IsValid(CustomizableObjectInstanceUsage) &&
-			CustomizableObjectInstanceUsage->GetCustomizableObjectInstance() == CustomizableObjectInstance)
+		if (Usage->GetCustomizableObjectInstance() != CustomizableObjectInstance)
+		{
+			continue;
+		}
+
+		USkeletalMeshComponent* Parent = Cast<USkeletalMeshComponent>(Usage->GetAttachParent());
+		if (!Parent)
+		{
+			continue;
+		}
+
+		MUTABLE_CPUPROFILER_SCOPE(UpdateSkeletalMesh_SetSkeletalMesh);
+
+		bool bAttachedParentUpdated = false;
+
+		Usage->GetPrivate()->bPendingSetSkeletalMesh = false;
+
+		USkeletalMesh* SkeletalMesh = CustomizableObjectInstance->GetComponentMeshSkeletalMesh(Usage->GetComponentName());
+		if (SkeletalMesh != Parent->GetSkeletalMeshAsset())
 		{
-			MUTABLE_CPUPROFILER_SCOPE(UpdateSkeletalMesh_SetSkeletalMesh);
+			Parent->SetSkeletalMesh(SkeletalMesh, RequiresReinitPose(Parent->GetSkeletalMeshAsset(), SkeletalMesh));
+			bAttachedParentUpdated = true;
+		}
+
+		TArray<TObjectPtr<UMaterialInterface>> OldOverridenMaterials = Parent->OverrideMaterials;
+
+		if (Parent->HasOverrideMaterials())
+		{
+			UE_LOG(LogMutable, Warning, TEXT("Attaching Customizable Skeletal Component to Skeletal Mesh Component with overriden materials! Deleting overrides."));
+			Parent->EmptyOverrideMaterials();
+		}
+
+		const bool bIsTransientMesh = SkeletalMesh ? SkeletalMesh->HasAllFlags(EObjectFlags::RF_Transient) : false;
+		const bool bUseOverrideMaterials = !bIsTransientMesh || (CustomizableObject->bEnableMeshCache && UCustomizableObjectSystem::IsMeshCacheEnabled());
+		if (bUseOverrideMaterials)
+		{
+			if (FCustomizableInstanceComponentData* ComponentData = CustomizableObjectInstance->GetPrivate()->GetComponentData(Usage->GetComponentName()))
+			{
+				for (int32 Index = 0; Index < ComponentData->OverrideMaterials.Num(); ++Index)
+				{
+					Parent->SetMaterial(Index, ComponentData->OverrideMaterials[Index]);
+				}
+			}
+		}
 
-			USkeletalMesh* SkeletalMesh = CustomizableObjectInstance->GetComponentMeshSkeletalMesh(CustomizableObjectInstanceUsage->GetComponentName());
-			CustomizableObjectInstanceUsage->GetPrivate()->SetSkeletalMesh(SkeletalMesh, &bSkeletalMeshUpdated, &bMaterialsUpdated);
+		if (OldOverridenMaterials != Parent->OverrideMaterials)
+		{
+			bAttachedParentUpdated = true;
+		}
 
-			if (CustomizableObjectInstancePrivateData->HasCOInstanceFlags(ReplacePhysicsAssets) &&
-				SkeletalMesh)
+		if (CustomizableObjectInstancePrivateData->HasCOInstanceFlags(ReplacePhysicsAssets) &&
+			SkeletalMesh &&
+			Parent->GetWorld())
+		{
+			UPhysicsAsset* PhysicsAsset = SkeletalMesh->GetPhysicsAsset();
+			if (PhysicsAsset != Parent->GetPhysicsAsset())
 			{
-				CustomizableObjectInstanceUsage->GetPrivate()->SetPhysicsAsset(SkeletalMesh->GetPhysicsAsset(), &bPhysicsAssetUpdated);	
+				Parent->SetPhysicsAsset(PhysicsAsset, true);
+				bAttachedParentUpdated = true;
 			}
 		}
 
-		if (bSkeletalMeshUpdated || bMaterialsUpdated || bPhysicsAssetUpdated)
+		if (bAttachedParentUpdated)
 		{
-			Context->AttachedParentUpdated.Add(CustomizableObjectInstanceUsage);
+			Context->AttachedParentUpdated.Add(Usage);
 		}
 	}
 }
@@ -1190,7 +1280,8 @@ EUpdateRequired UCustomizableObjectSystemPrivate::IsUpdateRequired(const UCustom
 	const bool bNoInitialUpdate = bOnlyUpdateIfNotGenerated && bIsGenerated;
 
 	if (bNoUpdateLODs &&
-		bNoInitialUpdate)
+		bNoInitialUpdate &&
+		!Private->bAutomaticUpdateRequired)
 	{
 		return EUpdateRequired::NoUpdate;
 	}
@@ -1264,7 +1355,7 @@ void UCustomizableObjectSystemPrivate::EnqueueUpdateSkeletalMesh(const TSharedRe
 	{
 		// Mutable is disabled. Set the reference SkeletalMesh and finish the update with success to avoid breaking too many things.
 		Context->UpdateResult = EUpdateResult::Success;
-		Instance->GetPrivate()->SetDefaultSkeletalMesh();
+		Instance->GetPrivate()->SetReferenceSkeletalMesh();
 		FinishUpdateGlobal(Context);
 		return;
 	}
@@ -2924,6 +3015,15 @@ namespace impl
 			return;
 		}
 
+		UCustomizableObject* CustomizableObject = CustomizableObjectInstance->GetCustomizableObject();
+		if (!CustomizableObject)
+		{
+			System->ClearCurrentMutableOperation();
+			OperationData->UpdateResult = EUpdateResult::Error;
+			FinishUpdateGlobal(OperationData);
+			return;
+		}
+		
 		UCustomizableObjectSystemPrivate * CustomizableObjectSystemPrivateData = System->GetPrivate();
 
 		// Actual work
@@ -3702,6 +3802,52 @@ int32 UCustomizableObjectSystem::TickInternal(const bool bBlocking)
 
 	Private->UpdateStats();
 
+	for (TObjectIterator<UCustomizableObjectInstanceUsage> Usage; Usage; ++Usage)
+	{
+		UCustomizableObjectInstanceUsagePrivate* UsagePrivate = Usage->GetPrivate();
+		
+		if (!UsagePrivate->bPendingSetSkeletalMesh)
+		{
+			continue;
+		}
+
+		UsagePrivate->bPendingSetSkeletalMesh = false;
+
+		UCustomizableObjectInstance* Instance = Usage->GetCustomizableObjectInstance();
+		if (!Instance)
+		{
+			continue;
+		}
+		
+		Instance->GetPrivate()->SetCOInstanceFlags(UsedByComponent);
+
+		UCustomizableObject* Object = Instance->GetCustomizableObject();
+		if (!Object)
+		{
+			continue;
+		}
+
+		USkeletalMeshComponent* Parent = Cast<USkeletalMeshComponent>(Usage->GetAttachParent());
+		if (!Parent)
+		{
+			continue;
+		}
+		
+		if (!Usage->GetSkipSetReferenceSkeletalMesh() &&
+			Object->bEnableUseRefSkeletalMeshAsPlaceholder)
+		{
+			const FName& ComponentName = Usage->GetComponentName();
+			
+			if (USkeletalMesh* SkeletalMesh = Object->GetComponentMeshReferenceSkeletalMesh(ComponentName))
+			{
+				Parent->EmptyOverrideMaterials();
+				Parent->SetSkeletalMesh(SkeletalMesh);
+			}
+		}
+		
+		Instance->GetPrivate()->bAutomaticUpdateRequired = true;
+	}
+	
 	FMutableUpdateCandidate* LODUpdateCandidateFound = nullptr;
 	
 	bool bPendingCompilation = false;
@@ -3960,6 +4106,7 @@ void UCustomizableObjectSystem::DiscardInstances()
 					for (TObjectIterator<UCustomizableObjectInstanceUsage> It; It; ++It)
 					{
 						UCustomizableObjectInstanceUsage* InstanceUsage = *It;
+
 						if (InstanceUsage->GetCustomizableObjectInstance() != COI)
 						{
 							continue;
@@ -3971,20 +4118,24 @@ void UCustomizableObjectSystem::DiscardInstances()
 							continue;
 						}
 #endif
-
-						const FName& ComponentName = InstanceUsage->GetComponentName();
-						const int32 ObjectComponentIndex = COI->GetCustomizableObject()->GetPrivate()->GetModelResources().ComponentNamesPerObjectComponent.IndexOfByKey(ComponentName);
-						if (ModelResources.ReferenceSkeletalMeshesData.IsValidIndex(ObjectComponentIndex))
+						
+						if (USkeletalMeshComponent* Parent = Cast<USkeletalMeshComponent>(InstanceUsage->GetAttachParent()))
 						{
-							USkeletalMesh* SkeletalMesh = nullptr;
-					
-							if (IsReplaceDiscardedWithReferenceMeshEnabled())
+							const FName& ComponentName = InstanceUsage->GetComponentName();
+							const int32 ObjectComponentIndex = ModelResources.ComponentNamesPerObjectComponent.IndexOfByKey(ComponentName);
+							if (ModelResources.ReferenceSkeletalMeshesData.IsValidIndex(ObjectComponentIndex))
 							{
-								// Force load the reference mesh if necessary. 
-								SkeletalMesh = ModelResources.ReferenceSkeletalMeshesData[ObjectComponentIndex].SoftSkeletalMesh.LoadSynchronous();
+								USkeletalMesh* SkeletalMesh = nullptr;
+					
+								if (IsReplaceDiscardedWithReferenceMeshEnabled())
+								{
+									// Force load the reference mesh if necessary. 
+									SkeletalMesh = ModelResources.ReferenceSkeletalMeshesData[ObjectComponentIndex].SoftSkeletalMesh.LoadSynchronous();
+								}
+
+								Parent->EmptyOverrideMaterials();
+								Parent->SetSkeletalMesh(SkeletalMesh);
 							}
-		
-							InstanceUsage->GetPrivate()->SetSkeletalMesh(SkeletalMesh);
 						}
 					
 						for (const UCustomizableObjectExtension* Extension : ICustomizableObjectModule::Get().GetRegisteredExtensions())
diff --git a/Engine/Plugins/Mutable/Source/CustomizableObject/Private/MuCO/CustomizableSkeletalComponent.cpp b/Engine/Plugins/Mutable/Source/CustomizableObject/Private/MuCO/CustomizableSkeletalComponent.cpp
--- a/Engine/Plugins/Mutable/Source/CustomizableObject/Private/MuCO/CustomizableSkeletalComponent.cpp
+++ b/Engine/Plugins/Mutable/Source/CustomizableObject/Private/MuCO/CustomizableSkeletalComponent.cpp
@@ -43,33 +43,6 @@ USkeletalMesh* UCustomizableSkeletalComponentPrivate::GetSkeletalMesh() const
 }
 
 
-void UCustomizableSkeletalComponentPrivate::SetSkeletalMesh(USkeletalMesh* SkeletalMesh)
-{
-	if (InstanceUsage)
-	{
-		InstanceUsage->GetPrivate()->SetSkeletalMesh(SkeletalMesh);
-	}
-}
-
-
-void UCustomizableSkeletalComponentPrivate::SetPhysicsAsset(UPhysicsAsset* PhysicsAsset)
-{
-	if (InstanceUsage)
-	{
-		InstanceUsage->GetPrivate()->SetPhysicsAsset(PhysicsAsset);
-	}
-}
-
-
-void UCustomizableSkeletalComponentPrivate::SetPendingSetSkeletalMesh(bool bIsActive)
-{
-	if (InstanceUsage)
-	{
-		InstanceUsage->GetPrivate()->SetPendingSetSkeletalMesh(bIsActive);
-	}
-}
-
-
 USkeletalMesh* UCustomizableSkeletalComponentPrivate::GetAttachedSkeletalMesh() const
 {
 	if (InstanceUsage)
@@ -168,23 +141,6 @@ const UCustomizableSkeletalComponentPrivate* UCustomizableSkeletalComponent::Get
 }
 
 
-#if WITH_EDITOR
-void UCustomizableSkeletalComponentPrivate::EditorUpdateComponent()
-{
-	if (InstanceUsage)
-	{
-		InstanceUsage->GetPrivate()->EditorUpdateComponent();
-	}
-}
-#endif
-
-
-bool& UCustomizableSkeletalComponentPrivate::PendingSetSkeletalMesh()
-{
-	return GetPublic()->bPendingSetSkeletalMesh;	
-}
-
-
 UCustomizableSkeletalComponent* UCustomizableSkeletalComponentPrivate::GetPublic()
 {
 	UCustomizableSkeletalComponent* Public = StaticCast<UCustomizableSkeletalComponent*>(GetOuter());
@@ -207,13 +163,6 @@ void UCustomizableSkeletalComponent::OnAttachmentChanged()
 {
 	Super::OnAttachmentChanged();
 
-	if (Cast<USkeletalMeshComponent>(GetAttachParent()))
-	{
-		GetPrivate()->SetPendingSetSkeletalMesh(true);
-	}
-	else if(!GetAttachParent())
-	{
-		DestroyComponent();
-	}
+	GetPrivate()->InstanceUsage->AttachTo(Cast<USkeletalMeshComponent>(GetAttachParent()));
 }
 
diff --git a/Engine/Plugins/Mutable/Source/CustomizableObject/Public/MuCO/CustomizableSkeletalComponent.h b/Engine/Plugins/Mutable/Source/CustomizableObject/Public/MuCO/CustomizableSkeletalComponent.h
--- a/Engine/Plugins/Mutable/Source/CustomizableObject/Public/MuCO/CustomizableSkeletalComponent.h
+++ b/Engine/Plugins/Mutable/Source/CustomizableObject/Public/MuCO/CustomizableSkeletalComponent.h
@@ -44,9 +44,6 @@ public:
 	FCustomizableSkeletalComponentUpdatedDelegate UpdatedDelegate;
 
 private:
-	// Used to replace the SkeletalMesh of the parent component by the ReferenceSkeletalMesh or the generated SkeletalMesh 
-	bool bPendingSetSkeletalMesh = false;
-
 	UPROPERTY()
 	bool bSkipSkipSetSkeletalMeshOnAttach = false;
 	
