[Mutable] Moved COE headers to private.

These headers should not be public.

Kept COESchema as public but moved to private many of its functions.

diff --git a/Engine/Plugins/Mutable/Source/CustomizableObject/Internal/MuCO/ICustomizableObjectEditorModule.h b/Engine/Plugins/Mutable/Source/CustomizableObject/Internal/MuCO/ICustomizableObjectEditorModule.h
--- a/Engine/Plugins/Mutable/Source/CustomizableObject/Internal/MuCO/ICustomizableObjectEditorModule.h
+++ b/Engine/Plugins/Mutable/Source/CustomizableObject/Internal/MuCO/ICustomizableObjectEditorModule.h
@@ -6,7 +6,6 @@
 #include "Templates/SharedPointer.h"
 
 class FCustomizableObjectEditorLogger;
-class ICustomizableObjectDebugger;
 class ICustomizableObjectEditor;
 class ICustomizableObjectInstanceEditor;
 class IToolkitHost;
diff --git a/Engine/Plugins/Mutable/Source/CustomizableObjectEditor/Private/MuCOE/CustomizableObjectEditorModule.h b/Engine/Plugins/Mutable/Source/CustomizableObjectEditor/Private/MuCOE/CustomizableObjectEditorModule.h
--- a/Engine/Plugins/Mutable/Source/CustomizableObjectEditor/Private/MuCOE/CustomizableObjectEditorModule.h
+++ b/Engine/Plugins/Mutable/Source/CustomizableObjectEditor/Private/MuCOE/CustomizableObjectEditorModule.h
@@ -15,7 +15,6 @@ class FPropertyEditorModule;
 class FExtensibilityManager;
 class ICustomizableObjectEditor;
 class ICustomizableObjectInstanceEditor;
-class ICustomizableObjectDebugger;
 class FBakeOperationCompletedDelegate;
 class FName;
 
diff --git a/Engine/Plugins/Mutable/Source/CustomizableObjectEditor/Private/MuCOE/CustomizableObjectEditor_Deprecated.h b/Engine/Plugins/Mutable/Source/CustomizableObjectEditor/Private/MuCOE/CustomizableObjectEditor_Deprecated.h
new file mode 100644
--- /dev/null
+++ b/Engine/Plugins/Mutable/Source/CustomizableObjectEditor/Private/MuCOE/CustomizableObjectEditor_Deprecated.h
@@ -0,0 +1,193 @@
+﻿// Copyright Epic Games, Inc. All Rights Reserved.
+
+#pragma once
+
+#include "Containers/Array.h"
+#include "UObject/ObjectPtr.h"
+#include "EdGraph/EdGraphPin.h"
+
+#include "CustomizableObjectEditor_Deprecated.generated.h"
+
+class UTexture2D;
+struct FLODReductionSettings;
+
+
+// Place to hide all deprecated data structures. They are still needed for deserialization backwards compatibility.
+
+// UCustomizableObjectNodeMaterial
+// Deprecated, do not use!
+USTRUCT()
+struct CUSTOMIZABLEOBJECTEDITOR_API FCustomizableObjectNodeMaterialImage
+{
+	GENERATED_USTRUCT_BODY()
+
+	UPROPERTY()
+	FString Name;
+
+	UPROPERTY()
+	int32 UVLayout = 0;
+
+	UPROPERTY()
+	TObjectPtr<UTexture2D> ReferenceTexture = nullptr;
+	
+	UPROPERTY()
+	int32 LayerIndex = 0;
+
+	UPROPERTY()
+	FString PinName;
+};
+
+// Deprecated, do not use!
+USTRUCT()
+struct CUSTOMIZABLEOBJECTEDITOR_API FCustomizableObjectNodeMaterialVector
+{
+	GENERATED_USTRUCT_BODY()
+	
+	UPROPERTY()
+	FString Name;
+	
+	UPROPERTY()
+	int32 LayerIndex = 0;
+
+	UPROPERTY()
+	FString PinName;
+};
+
+// Deprecated, do not use!
+USTRUCT()
+struct CUSTOMIZABLEOBJECTEDITOR_API FCustomizableObjectNodeMaterialScalar
+{
+	GENERATED_USTRUCT_BODY()
+
+	UPROPERTY()
+	FString Name;
+	
+	UPROPERTY()
+	int32 LayerIndex = 0;
+
+	UPROPERTY()
+	FString PinName;
+};
+
+// Deprecated, do not use!
+constexpr int32 UV_LAYOUT_DEFAULT = -2;
+
+// UCustomizableObjectNodeModifierEditMeshSection
+// Deprecated, do not use!
+USTRUCT()
+struct CUSTOMIZABLEOBJECTEDITOR_API FCustomizableObjectNodeEditMaterialImage
+{
+	GENERATED_USTRUCT_BODY()
+
+	UPROPERTY()
+	FString Name;
+};
+
+// UCustomizableObjectNodeModifierExtendMeshSection
+// Deprecated, do not use!
+USTRUCT()
+struct CUSTOMIZABLEOBJECTEDITOR_API FCustomizableObjectNodeExtendMaterialImage
+{
+	GENERATED_USTRUCT_BODY()
+
+	UPROPERTY()
+	FString Name;
+};
+
+// UCustomizableObjectNodeSkeletalMesh
+// Deprecated, do not use!
+USTRUCT()
+struct CUSTOMIZABLEOBJECTEDITOR_API FCustomizableObjectNodeSkeletalMeshMaterial
+{
+	GENERATED_USTRUCT_BODY()
+
+	UPROPERTY()
+	FString Name;
+
+	UPROPERTY()
+	TObjectPtr<UEdGraphPin_Deprecated> MeshPin_DEPRECATED;
+
+	UPROPERTY()
+	TArray<TObjectPtr<UEdGraphPin_Deprecated>> LayoutPins_DEPRECATED;
+
+	UPROPERTY()
+	TArray<TObjectPtr<UEdGraphPin_Deprecated>> ImagePins_DEPRECATED;
+
+	UPROPERTY()
+	FEdGraphPinReference MeshPinRef;
+
+	UPROPERTY()
+	TArray<FEdGraphPinReference> LayoutPinsRef;
+
+	UPROPERTY()
+	TArray<FEdGraphPinReference> ImagePinsRef;
+};
+
+// Deprecated, do not use!
+USTRUCT()
+struct CUSTOMIZABLEOBJECTEDITOR_API FCustomizableObjectNodeSkeletalMeshLOD
+{
+	GENERATED_USTRUCT_BODY()
+
+	UPROPERTY()
+	TArray<FCustomizableObjectNodeSkeletalMeshMaterial> Materials;
+};
+
+// UCustomizableObjectNodeColorVariation
+// Deprecated, do not use!
+USTRUCT()
+struct CUSTOMIZABLEOBJECTEDITOR_API FCustomizableObjectColorVariation
+{
+	GENERATED_USTRUCT_BODY()
+
+	UPROPERTY(EditAnywhere, Category = CustomizableObject)
+	FString Tag;
+};
+
+// UCustomizableObjectFloatVariation
+// Deprecated, do not use!
+USTRUCT()
+struct CUSTOMIZABLEOBJECTEDITOR_API FCustomizableObjectFloatVariation
+{
+	GENERATED_USTRUCT_BODY()
+
+	UPROPERTY(EditAnywhere, Category = CustomizableObject)
+	FString Tag;
+};
+
+// UCustomizableObjectMaterialVariation
+// Deprecated, do not use!
+USTRUCT()
+struct CUSTOMIZABLEOBJECTEDITOR_API FCustomizableObjectMaterialVariation
+{
+	GENERATED_USTRUCT_BODY()
+
+	UPROPERTY(EditAnywhere, Category = CustomizableObject)
+	FString Tag;
+};
+
+// UCustomizableObjectMeshVariation
+// Deprecated, do not use!
+USTRUCT()
+struct CUSTOMIZABLEOBJECTEDITOR_API FCustomizableObjectMeshVariation
+{
+	GENERATED_USTRUCT_BODY()
+
+	UPROPERTY(EditAnywhere, Category = CustomizableObject)
+	FString Tag;
+};
+
+
+// UCustomizableObjectNodeObject
+// Deprecated, do not use!
+USTRUCT()
+struct FComponentSettings
+{
+	GENERATED_USTRUCT_BODY()
+	
+	UPROPERTY()
+	FString ComponentName;
+
+	UPROPERTY()
+	TArray<FLODReductionSettings> LODReductionSettings;
+};
diff --git a/Engine/Plugins/Mutable/Source/CustomizableObjectEditor/Private/MuCOE/CustomizableObjectGraph.h b/Engine/Plugins/Mutable/Source/CustomizableObjectEditor/Private/MuCOE/CustomizableObjectGraph.h
new file mode 100644
--- /dev/null
+++ b/Engine/Plugins/Mutable/Source/CustomizableObjectEditor/Private/MuCOE/CustomizableObjectGraph.h
@@ -0,0 +1,46 @@
+// Copyright Epic Games, Inc. All Rights Reserved.
+
+#pragma once
+
+#include "EdGraph/EdGraph.h"
+
+#include "CustomizableObjectGraph.generated.h"
+
+class UObject;
+
+
+UCLASS()
+class CUSTOMIZABLEOBJECTEDITOR_API UCustomizableObjectGraph : public UEdGraph
+{
+public:
+	GENERATED_BODY()
+
+	UCustomizableObjectGraph();
+
+	// UObject Interface
+	virtual void PostLoad() override;
+	void PostRename(UObject * OldOuter, const FName OldName) override;
+
+	// Own Interface
+	void NotifyNodeIdChanged(const FGuid& OldGuid, const FGuid& NewGuid);
+
+	FGuid RequestNotificationForNodeIdChange(const FGuid& OldGuid, const FGuid& NodeToNotifyGuid);
+
+	void PostDuplicate(bool bDuplicateForPIE) override;
+
+	/** Adds the necessary nodes for a CO to work */
+	void AddEssentialGraphNodes();
+
+	void BackwardsCompatibleFixup(int32 CustomizableObjectCustomVersion);
+
+	void PostBackwardsCompatibleFixup();
+	
+private:
+
+	// Request Node Id Update Map
+	TMap<FGuid, TSet<FGuid>> NodesToNotifyMap;
+
+	// Guid map with the key beeing the old Guid and the Value the new one, filled after duplicating COs
+	TMap<FGuid, FGuid> NotifiedNodeIdsMap;
+};
+
diff --git a/Engine/Plugins/Mutable/Source/CustomizableObjectEditor/Private/MuCOE/ICustomizableObjectEditor.h b/Engine/Plugins/Mutable/Source/CustomizableObjectEditor/Private/MuCOE/ICustomizableObjectEditor.h
new file mode 100644
--- /dev/null
+++ b/Engine/Plugins/Mutable/Source/CustomizableObjectEditor/Private/MuCOE/ICustomizableObjectEditor.h
@@ -0,0 +1,41 @@
+// Copyright Epic Games, Inc. All Rights Reserved.
+
+#pragma once
+
+#include "MuCOE/ICustomizableObjectInstanceEditor.h"
+
+class UEdGraphNode;
+
+class UCustomizableObject;
+class UCustomizableObjectInstance;
+class UClass;
+class UCustomizableObjectNode;
+
+
+// Public interface to Customizable Object Editor
+class ICustomizableObjectEditor : public ICustomizableObjectInstanceEditor
+{
+public:
+	// Retrieves the current Customizable Object displayed in the Editor.
+	virtual UCustomizableObject* GetCustomizableObject() = 0;
+	
+	// Checks whether nodes can currently be pasted */
+	virtual bool CanPasteNodes() const {return false;}
+
+	// Paste nodes at a specific location
+	virtual void PasteNodesHere(const FVector2D& Location) {}
+
+	virtual UEdGraphNode* CreateCommentBox(const FVector2D& NodePos) { return nullptr; }
+
+	virtual void UpdateGraphNodeProperties() {}
+
+	// Make sure the graph editor selects a specific node
+	virtual void SelectNode(const UEdGraphNode* Node) {}
+
+	virtual void ReconstructNode(UEdGraphNode* Node) {}
+
+	/** Reconstructs all child the nodes that match the given type.
+	 * @param StartNode Root node to start the graph traversal. This one also will be reconstructed.
+	 * @param NodeType Node types to reconstruct. */
+	virtual void ReconstructAllChildNodes(UCustomizableObjectNode& StartNode, const UClass& NodeType) {}
+};
diff --git a/Engine/Plugins/Mutable/Source/CustomizableObjectEditor/Private/MuCOE/Nodes/CustomizableObjectNodeAnimationPose.h b/Engine/Plugins/Mutable/Source/CustomizableObjectEditor/Private/MuCOE/Nodes/CustomizableObjectNodeAnimationPose.h
new file mode 100644
--- /dev/null
+++ b/Engine/Plugins/Mutable/Source/CustomizableObjectEditor/Private/MuCOE/Nodes/CustomizableObjectNodeAnimationPose.h
@@ -0,0 +1,46 @@
+// Copyright Epic Games, Inc. All Rights Reserved.
+
+#pragma once
+
+#include "MuCOE/Nodes/CustomizableObjectNode.h"
+
+#include "CustomizableObjectNodeAnimationPose.generated.h"
+
+namespace ENodeTitleType { enum Type : int; }
+
+class FString;
+class UCustomizableObjectNodeRemapPins;
+class UEdGraphPin;
+class UObject;
+class UPoseAsset;
+struct FPropertyChangedEvent;
+
+
+UCLASS()
+class CUSTOMIZABLEOBJECTEDITOR_API UCustomizableObjectNodeAnimationPose : public UCustomizableObjectNode
+{
+public:
+	GENERATED_BODY()
+
+	UCustomizableObjectNodeAnimationPose();
+
+	UPROPERTY(EditAnywhere, Category=NoCategory)
+	TObjectPtr<UPoseAsset> PoseAsset;
+
+	// UObject interface.
+	virtual void PostEditChangeProperty(FPropertyChangedEvent& PropertyChangedEvent) override;
+
+	// Begin EdGraphNode interface
+	UEdGraphPin* GetInputMeshPin() const;
+	UEdGraphPin* GetTablePosePin() const;
+	virtual FText GetNodeTitle(ENodeTitleType::Type TitleType) const override;
+	virtual FLinearColor GetNodeTitleColor() const override;
+	// End EdGraphNode interface
+
+	// UCustomizableObjectNode interface
+	void AllocateDefaultPins(UCustomizableObjectNodeRemapPins* RemapPins) override;
+
+	// Retrieve the pose information from the PoseAsset
+	static void StaticRetrievePoseInformation(UPoseAsset* PoseAsset, class USkeletalMesh* RefSkeletalMesh, TArray<FName>& OutArrayBoneName, TArray<FTransform>& OutArrayTransform);
+};
+
diff --git a/Engine/Plugins/Mutable/Source/CustomizableObjectEditor/Private/MuCOE/Nodes/CustomizableObjectNodeColorParameter.h b/Engine/Plugins/Mutable/Source/CustomizableObjectEditor/Private/MuCOE/Nodes/CustomizableObjectNodeColorParameter.h
new file mode 100644
--- /dev/null
+++ b/Engine/Plugins/Mutable/Source/CustomizableObjectEditor/Private/MuCOE/Nodes/CustomizableObjectNodeColorParameter.h
@@ -0,0 +1,41 @@
+// Copyright Epic Games, Inc. All Rights Reserved.
+
+#pragma once
+
+#include "MuCOE/Nodes/CustomizableObjectNode.h"
+
+#include "CustomizableObjectNodeColorParameter.generated.h"
+
+namespace ENodeTitleType { enum Type : int; }
+
+class UCustomizableObjectNodeRemapPins;
+class UObject;
+
+
+UCLASS()
+class CUSTOMIZABLEOBJECTEDITOR_API UCustomizableObjectNodeColorParameter : public UCustomizableObjectNode
+{
+public:
+	GENERATED_BODY()
+
+	UPROPERTY(EditAnywhere, Category=CustomizableObject, meta = (DontUpdateWhileEditing))
+	FLinearColor DefaultValue = FLinearColor(1, 1, 1, 1);
+
+	UPROPERTY(EditAnywhere, Category=CustomizableObject)
+	FString ParameterName = "Default Name";
+
+	UPROPERTY(EditAnywhere, Category = UI, meta = (DisplayName = "Parameter UI Metadata"))
+	FMutableParamUIMetadata ParamUIMetadata;
+
+	// Begin EdGraphNode interface
+	FText GetNodeTitle(ENodeTitleType::Type TitleType) const override;
+	FLinearColor GetNodeTitleColor() const override;
+	FText GetTooltipText() const override;
+	virtual void OnRenameNode(const FString& NewName) override;
+	virtual bool GetCanRenameNode() const override { return true; }
+
+	// UCustomizableObjectNode interface
+	void AllocateDefaultPins(UCustomizableObjectNodeRemapPins* RemapPins) override;
+	bool IsAffectedByLOD() const override { return false; }
+};
+
diff --git a/Engine/Plugins/Mutable/Source/CustomizableObjectEditor/Private/MuCOE/Nodes/CustomizableObjectNodeColorSwitch.h b/Engine/Plugins/Mutable/Source/CustomizableObjectEditor/Private/MuCOE/Nodes/CustomizableObjectNodeColorSwitch.h
new file mode 100644
--- /dev/null
+++ b/Engine/Plugins/Mutable/Source/CustomizableObjectEditor/Private/MuCOE/Nodes/CustomizableObjectNodeColorSwitch.h
@@ -0,0 +1,22 @@
+// Copyright Epic Games, Inc. All Rights Reserved.
+
+#pragma once
+
+#include "MuCOE/Nodes/CustomizableObjectNodeSwitchBase.h"
+
+#include "CustomizableObjectNodeColorSwitch.generated.h"
+
+
+UCLASS()
+class CUSTOMIZABLEOBJECTEDITOR_API UCustomizableObjectNodeColorSwitch : public UCustomizableObjectNodeSwitchBase
+{
+public:
+	GENERATED_BODY()
+
+	// UCustomizableObjectNode interface
+	virtual void BackwardsCompatibleFixup(int32 CustomizableObjectCustomVersion) override;
+
+	// UCustomizableObjectNodeSwitchBase interface
+	virtual FName GetCategory() const override;
+};
+
diff --git a/Engine/Plugins/Mutable/Source/CustomizableObjectEditor/Private/MuCOE/Nodes/CustomizableObjectNodeColorVariation.h b/Engine/Plugins/Mutable/Source/CustomizableObjectEditor/Private/MuCOE/Nodes/CustomizableObjectNodeColorVariation.h
new file mode 100644
--- /dev/null
+++ b/Engine/Plugins/Mutable/Source/CustomizableObjectEditor/Private/MuCOE/Nodes/CustomizableObjectNodeColorVariation.h
@@ -0,0 +1,28 @@
+// Copyright Epic Games, Inc. All Rights Reserved.
+
+#pragma once
+
+#include "MuCOE/Nodes/CustomizableObjectNodeVariation.h"
+
+#include "CustomizableObjectNodeColorVariation.generated.h"
+
+struct FCustomizableObjectColorVariation;
+
+
+UCLASS()
+class CUSTOMIZABLEOBJECTEDITOR_API UCustomizableObjectNodeColorVariation : public UCustomizableObjectNodeVariation
+{
+public:
+	GENERATED_BODY()
+
+	// Deprecated properties
+	UPROPERTY()
+	TArray<FCustomizableObjectColorVariation> Variations_DEPRECATED;
+
+	// UCustomizableObjectNode interface
+	virtual void BackwardsCompatibleFixup(int32 CustomizableObjectCustomVersion) override;
+	
+	// UCustomizableObjectNodeVariation interface
+	virtual FName GetCategory() const override;
+};
+
diff --git a/Engine/Plugins/Mutable/Source/CustomizableObjectEditor/Private/MuCOE/Nodes/CustomizableObjectNodeComponent.h b/Engine/Plugins/Mutable/Source/CustomizableObjectEditor/Private/MuCOE/Nodes/CustomizableObjectNodeComponent.h
new file mode 100644
--- /dev/null
+++ b/Engine/Plugins/Mutable/Source/CustomizableObjectEditor/Private/MuCOE/Nodes/CustomizableObjectNodeComponent.h
@@ -0,0 +1,21 @@
+// Copyright Epic Games, Inc. All Rights Reserved.
+
+#pragma once
+
+#include "MuCOE/Nodes/CustomizableObjectNode.h"
+
+#include "CustomizableObjectNodeComponent.generated.h"
+
+class UEdGraphPin;
+
+
+UCLASS(abstract)
+class CUSTOMIZABLEOBJECTEDITOR_API UCustomizableObjectNodeComponent : public UCustomizableObjectNode
+{
+public:
+	GENERATED_BODY()
+
+	// UCustomizableObjectNode interface
+	virtual bool IsAffectedByLOD() const override;
+};
+
diff --git a/Engine/Plugins/Mutable/Source/CustomizableObjectEditor/Private/MuCOE/Nodes/CustomizableObjectNodeEnumParameter.h b/Engine/Plugins/Mutable/Source/CustomizableObjectEditor/Private/MuCOE/Nodes/CustomizableObjectNodeEnumParameter.h
new file mode 100644
--- /dev/null
+++ b/Engine/Plugins/Mutable/Source/CustomizableObjectEditor/Private/MuCOE/Nodes/CustomizableObjectNodeEnumParameter.h
@@ -0,0 +1,62 @@
+// Copyright Epic Games, Inc. All Rights Reserved.
+
+#pragma once
+
+#include "MuCOE/Nodes/CustomizableObjectNode.h"
+
+#include "CustomizableObjectNodeEnumParameter.generated.h"
+
+namespace ENodeTitleType { enum Type : int; }
+
+class UCustomizableObjectNodeRemapPins;
+class UObject;
+struct FPropertyChangedEvent;
+
+
+USTRUCT()
+struct CUSTOMIZABLEOBJECTEDITOR_API FCustomizableObjectNodeEnumValue
+{
+	GENERATED_USTRUCT_BODY()
+
+	UPROPERTY(EditAnywhere, Category=CustomizableObject)
+	FString Name;
+
+	UPROPERTY(EditAnywhere, Category = UI, meta = (DisplayName = "Parameter UI Metadata"))
+	FMutableParamUIMetadata ParamUIMetadata;
+};
+
+
+UCLASS()
+class CUSTOMIZABLEOBJECTEDITOR_API UCustomizableObjectNodeEnumParameter : public UCustomizableObjectNode
+{
+public:
+	GENERATED_BODY()
+
+	UPROPERTY(EditAnywhere, Category=CustomizableObject)
+	int32 DefaultIndex = 0;
+
+	UPROPERTY(EditAnywhere, Category=CustomizableObject)
+	FString ParameterName = "Default Name";
+
+	UPROPERTY(EditAnywhere, Category = UI, meta = (DisplayName = "Parameter UI Metadata"))
+	FMutableParamUIMetadata ParamUIMetadata;
+
+	UPROPERTY(EditAnywhere, Category=CustomizableObject)
+	TArray<FCustomizableObjectNodeEnumValue> Values;
+
+	// UObject interface.
+	void PostEditChangeProperty(FPropertyChangedEvent& PropertyChangedEvent) override;
+
+	// Begin EdGraphNode interface
+	FText GetNodeTitle(ENodeTitleType::Type TitleType) const override;
+	FLinearColor GetNodeTitleColor() const override;
+	FText GetTooltipText() const override;
+	virtual void OnRenameNode(const FString& NewName) override;
+	virtual bool GetCanRenameNode() const override { return true; }
+
+	// UCustomizableObjectNode interface
+	void AllocateDefaultPins(UCustomizableObjectNodeRemapPins* RemapPins) override;
+	virtual bool IsAffectedByLOD() const override { return false; }
+
+};
+
diff --git a/Engine/Plugins/Mutable/Source/CustomizableObjectEditor/Private/MuCOE/Nodes/CustomizableObjectNodeFloatParameter.h b/Engine/Plugins/Mutable/Source/CustomizableObjectEditor/Private/MuCOE/Nodes/CustomizableObjectNodeFloatParameter.h
new file mode 100644
--- /dev/null
+++ b/Engine/Plugins/Mutable/Source/CustomizableObjectEditor/Private/MuCOE/Nodes/CustomizableObjectNodeFloatParameter.h
@@ -0,0 +1,60 @@
+// Copyright Epic Games, Inc. All Rights Reserved.
+
+#pragma once
+
+#include "MuCOE/Nodes/CustomizableObjectNode.h"
+
+#include "CustomizableObjectNodeFloatParameter.generated.h"
+
+namespace ENodeTitleType { enum Type : int; }
+
+class FArchive;
+class UCustomizableObjectNodeRemapPins;
+class UEdGraphPin;
+class UObject;
+struct FEdGraphPinReference;
+struct FPropertyChangedEvent;
+
+
+USTRUCT()
+struct CUSTOMIZABLEOBJECTEDITOR_API FCustomizableObjectNodeFloatDescription
+{
+	GENERATED_USTRUCT_BODY()
+
+	UPROPERTY(EditAnywhere, Category=CustomizableObject)
+	FString Name;
+};
+
+
+UCLASS()
+class CUSTOMIZABLEOBJECTEDITOR_API UCustomizableObjectNodeFloatParameter : public UCustomizableObjectNode
+{
+public:
+	GENERATED_BODY()
+
+	UPROPERTY(EditAnywhere, Category=CustomizableObject)
+	float DefaultValue = 1.0f;
+
+	UPROPERTY(EditAnywhere, Category=CustomizableObject)
+	FString ParameterName = "Default Name";
+
+	UPROPERTY(EditAnywhere, Category = UI, meta = (DisplayName = "Parameter UI Metadata"))
+	FMutableParamUIMetadata ParamUIMetadata;
+	
+	// UObject interface.
+	void PostEditChangeProperty(FPropertyChangedEvent& PropertyChangedEvent) override;
+	
+	// Begin EdGraphNode interface
+	FText GetNodeTitle(ENodeTitleType::Type TitleType) const override;
+	FLinearColor GetNodeTitleColor() const override;
+	FText GetTooltipText() const override;
+	virtual bool GetCanRenameNode() const override { return true; }
+
+	// UCustomizableObjectNode interface
+	void AllocateDefaultPins(UCustomizableObjectNodeRemapPins* RemapPins) override;
+	bool IsAffectedByLOD() const override;
+	void BackwardsCompatibleFixup(int32 CustomizableObjectCustomVersion) override;
+	virtual void OnRenameNode(const FString& NewName) override;
+
+};
+
diff --git a/Engine/Plugins/Mutable/Source/CustomizableObjectEditor/Private/MuCOE/Nodes/CustomizableObjectNodeFloatSwitch.h b/Engine/Plugins/Mutable/Source/CustomizableObjectEditor/Private/MuCOE/Nodes/CustomizableObjectNodeFloatSwitch.h
new file mode 100644
--- /dev/null
+++ b/Engine/Plugins/Mutable/Source/CustomizableObjectEditor/Private/MuCOE/Nodes/CustomizableObjectNodeFloatSwitch.h
@@ -0,0 +1,21 @@
+// Copyright Epic Games, Inc. All Rights Reserved.
+
+#pragma once
+
+#include "MuCOE/Nodes/CustomizableObjectNodeSwitchBase.h"
+
+#include "CustomizableObjectNodeFloatSwitch.generated.h"
+
+
+UCLASS()
+class CUSTOMIZABLEOBJECTEDITOR_API UCustomizableObjectNodeFloatSwitch : public UCustomizableObjectNodeSwitchBase
+{
+public:
+	GENERATED_BODY()
+
+	// UCustomizableObjectNode interface
+	virtual void BackwardsCompatibleFixup(int32 CustomizableObjectCustomVersion) override;
+
+	// UCustomizableObjectNodeSwitchBase interface
+	virtual FName GetCategory() const override;
+};
diff --git a/Engine/Plugins/Mutable/Source/CustomizableObjectEditor/Private/MuCOE/Nodes/CustomizableObjectNodeFloatVariation.h b/Engine/Plugins/Mutable/Source/CustomizableObjectEditor/Private/MuCOE/Nodes/CustomizableObjectNodeFloatVariation.h
new file mode 100644
--- /dev/null
+++ b/Engine/Plugins/Mutable/Source/CustomizableObjectEditor/Private/MuCOE/Nodes/CustomizableObjectNodeFloatVariation.h
@@ -0,0 +1,28 @@
+// Copyright Epic Games, Inc. All Rights Reserved.
+
+#pragma once
+
+#include "MuCOE/Nodes/CustomizableObjectNodeVariation.h"
+
+#include "CustomizableObjectNodeFloatVariation.generated.h"
+
+struct FCustomizableObjectFloatVariation;
+
+
+UCLASS()
+class CUSTOMIZABLEOBJECTEDITOR_API UCustomizableObjectNodeFloatVariation : public UCustomizableObjectNodeVariation
+{
+public:
+	GENERATED_BODY()
+	
+	// Deprecated properties
+	UPROPERTY()
+	TArray<FCustomizableObjectFloatVariation> Variations_DEPRECATED;
+
+	// UCustomizableObjectNode interface
+	virtual void BackwardsCompatibleFixup(int32 CustomizableObjectCustomVersion) override;
+	
+	// UCustomizableObjectNodeVariation interface
+	virtual FName GetCategory() const override;
+};
+
diff --git a/Engine/Plugins/Mutable/Source/CustomizableObjectEditor/Private/MuCOE/Nodes/CustomizableObjectNodeLayoutBlocks.h b/Engine/Plugins/Mutable/Source/CustomizableObjectEditor/Private/MuCOE/Nodes/CustomizableObjectNodeLayoutBlocks.h
new file mode 100644
--- /dev/null
+++ b/Engine/Plugins/Mutable/Source/CustomizableObjectEditor/Private/MuCOE/Nodes/CustomizableObjectNodeLayoutBlocks.h
@@ -0,0 +1,56 @@
+// Copyright Epic Games, Inc. All Rights Reserved.
+
+#pragma once
+
+#include "MuCOE/Nodes/CustomizableObjectNode.h"
+
+#include "CustomizableObjectNodeLayoutBlocks.generated.h"
+
+class UCustomizableObjectLayout;
+enum class ECustomizableObjectTextureLayoutPackingStrategy : uint8;
+namespace ENodeTitleType { enum Type : int; }
+struct FCustomizableObjectLayoutBlock;
+
+class FArchive;
+class UCustomizableObjectNodeRemapPins;
+class UEdGraphPin;
+class UObject;
+struct FPropertyChangedEvent;
+
+
+UCLASS()
+class CUSTOMIZABLEOBJECTEDITOR_API UCustomizableObjectNodeLayoutBlocks : public UCustomizableObjectNode
+{
+public:
+	GENERATED_BODY()
+
+	UCustomizableObjectNodeLayoutBlocks();
+
+	UPROPERTY()
+	FIntPoint GridSize_DEPRECATED;
+
+	/** Used with the fixed layout strategy. */
+	UPROPERTY()
+	FIntPoint MaxGridSize_DEPRECATED;
+
+	UPROPERTY()
+	TArray<FCustomizableObjectLayoutBlock> Blocks_DEPRECATED;
+
+	UPROPERTY()
+	ECustomizableObjectTextureLayoutPackingStrategy PackingStrategy_DEPRECATED;
+
+	UPROPERTY()
+	TObjectPtr<UCustomizableObjectLayout> Layout = nullptr;
+
+	// EdGraphNode interface
+	void Serialize(FArchive& Ar) override;
+
+	// UCustomizableObjectNode interface
+	virtual void BackwardsCompatibleFixup(int32 CustomizableObjectCustomVersion) override;
+	void AllocateDefaultPins(UCustomizableObjectNodeRemapPins* RemapPins) override;
+
+	UEdGraphPin* OutputPin() const
+	{
+		return FindPin(TEXT("Layout"));
+	}
+};
diff --git a/Engine/Plugins/Mutable/Source/CustomizableObjectEditor/Private/MuCOE/Nodes/CustomizableObjectNodeMaterial.h b/Engine/Plugins/Mutable/Source/CustomizableObjectEditor/Private/MuCOE/Nodes/CustomizableObjectNodeMaterial.h
new file mode 100644
--- /dev/null
+++ b/Engine/Plugins/Mutable/Source/CustomizableObjectEditor/Private/MuCOE/Nodes/CustomizableObjectNodeMaterial.h
@@ -0,0 +1,328 @@
+// Copyright Epic Games, Inc. All Rights Reserved.
+
+#pragma once
+
+#include "MuCOE/Nodes/CustomizableObjectNodeMaterialBase.h"
+#include "MuCOE/RemapPins/CustomizableObjectNodeRemapPinsByName.h"
+
+#include "CustomizableObjectNodeMaterial.generated.h"
+
+enum class EMaterialParameterType : uint8;
+namespace ENodeTitleType { enum Type : int; }
+
+class FArchive;
+class SGraphNode;
+class SWidget;
+class UCustomizableObjectNodeRemapPins;
+class UEdGraphPin;
+class UMaterial;
+class UMaterialInterface;
+class UObject;
+class UTexture2D;
+struct FCustomizableObjectNodeMaterialImage;
+struct FCustomizableObjectNodeMaterialScalar;
+struct FCustomizableObjectNodeMaterialVector;
+struct FEdGraphPinReference;
+struct FFrame;
+struct FPropertyChangedEvent;
+
+
+/** Custom remap pins by name action.
+ *
+ * Remap pins by Texture Parameter Id. */
+UCLASS()
+class UCustomizableObjectNodeMaterialRemapPinsByName : public UCustomizableObjectNodeRemapPinsByName
+{
+	GENERATED_BODY()
+public:
+	virtual bool Equal(const UCustomizableObjectNode& Node, const UEdGraphPin& OldPin, const UEdGraphPin& NewPin) const override;
+
+	virtual void RemapPins(const UCustomizableObjectNode& Node, const TArray<UEdGraphPin*>& OldPins, const TArray<UEdGraphPin*>& NewPins, TMap<UEdGraphPin*, UEdGraphPin*>& PinsToRemap, TArray<UEdGraphPin*>& PinsToOrphan) override;
+
+	bool HasSavedPinData(const UCustomizableObjectNode& Node, const UEdGraphPin &Pin) const;
+};
+
+
+/** Base class for all Material Parameters. */
+UCLASS()
+class CUSTOMIZABLEOBJECTEDITOR_API UCustomizableObjectNodeMaterialPinDataParameter : public UCustomizableObjectNodePinData
+{
+	GENERATED_BODY()
+
+public:
+
+	/** Parameter id + layer index */
+	UPROPERTY()
+	FNodeMaterialParameterId MaterialParameterId;
+
+	/** Texture Parameter Id. */
+	UPROPERTY()
+	FGuid ParameterId_DEPRECATED;
+
+	/** Returns true if all properties are in its default state. */
+	virtual bool IsDefault() const;
+};
+
+
+/** Node pin mode. All pins set to EPinMode::Default will use this this mode. */
+UENUM()
+enum class ENodePinMode
+{
+	Mutable UMETA(ToolTip = "All Material Texture FParameters go through Mutable."),
+	Passthrough UMETA(ToolTip = "All Material Texture FParameters are not modified by Mutable.")
+};
+
+
+/** Image pin, pin mode. */
+UENUM()
+enum class EPinMode
+{
+	Default UMETA(DisplayName = "Node Defined", ToolTip = "Use node's \"Default Texture Parameter Mode\"."),
+	Mutable UMETA(ToolTip = "The Material Texture FParameters goes through Mutable."),
+	Passthrough UMETA(ToolTip = "The Material Texture FParameters is not modified by Mutable.")
+};
+
+
+/** Image Pin, UV Layout Mode. */
+UENUM()
+enum class EUVLayoutMode
+{
+	/** Does not override the UV Index specified in the Material. */
+	FromMaterial,
+	/* Texture should not be transformed by any layout. These textures will not be reduced automatically for LODs. */
+	Ignore,
+	/** User specified UV Index. */
+	Index
+};
+
+
+/** Enum to FText. */
+FText EPinModeToText(EPinMode PinMode);
+
+UCLASS()
+class CUSTOMIZABLEOBJECTEDITOR_API UCustomizableObjectNodeMaterial : public UCustomizableObjectNodeMaterialBase
+{
+public:
+	GENERATED_BODY()
+	
+	// UObject interface.
+	virtual void Serialize(FArchive& Ar) override;
+	virtual void PostLoad() override;
+	virtual void PostEditChangeProperty(FPropertyChangedEvent& PropertyChangedEvent) override;
+
+	// EdGraphNode interface
+	virtual FText GetNodeTitle(ENodeTitleType::Type TitleType) const override;
+	virtual FText GetTooltipText() const override;
+	virtual TSharedPtr<SGraphNode> CreateVisualWidget() override;
+	virtual void PinConnectionListChanged(UEdGraphPin* Pin) override;
+	virtual void PostPasteNode() override;
+	virtual bool CanConnect(const UEdGraphPin* InOwnedInputPin, const UEdGraphPin* InOutputPin, bool& bOutIsOtherNodeBlocklisted, bool& bOutArePinsCompatible) const override;
+
+	// UCustomizableObjectNode interface
+	virtual void BackwardsCompatibleFixup(int32 CustomizableObjectCustomVersion) override;
+	virtual void AllocateDefaultPins(UCustomizableObjectNodeRemapPins* RemapPins) override;
+	virtual bool CanPinBeHidden(const UEdGraphPin& Pin) const override;
+	virtual bool HasPinViewer() const override;
+	virtual UCustomizableObjectNodeRemapPinsByName* CreateRemapPinsDefault() const override;
+	virtual bool ProvidesCustomPinRelevancyTest() const override { return true; }
+	virtual bool IsPinRelevant(const UEdGraphPin* Pin) const override;
+	virtual bool CustomRemovePin(UEdGraphPin& Pin) override;
+	virtual bool IsNodeOutDatedAndNeedsRefresh() override;
+	virtual FString GetRefreshMessage() const override;
+	virtual TSharedPtr<IDetailsView> CustomizePinDetails(const UEdGraphPin& Pin) const override;
+	virtual TArray<FString>* GetEnableTags() override;
+	virtual FString GetInternalTagDisplayName() override;
+	
+	// UCustomizableObjectNodeMaterialBase interface
+	virtual TArray<UCustomizableObjectLayout*> GetLayouts() const override;
+	virtual UEdGraphPin* OutputPin() const override;
+	virtual UMaterialInterface* GetMaterial() const override;
+	virtual bool IsReuseMaterialBetweenLODs() const override;
+	virtual int32 GetNumParameters(EMaterialParameterType Type) const override;
+	virtual FNodeMaterialParameterId GetParameterId(EMaterialParameterType Type, int32 ParameterIndex) const override;
+	virtual FName GetParameterName(EMaterialParameterType Type, int32 ParameterIndex) const override;
+	virtual int32 GetParameterLayerIndex(EMaterialParameterType Type, int32 ParameterIndex) const override;
+	virtual FText GetParameterLayerName(EMaterialParameterType Type, int32 ParameterIndex) const override;
+	virtual bool HasParameter(const FNodeMaterialParameterId& ParameterId) const override;
+	virtual UEdGraphPin* GetParameterPin(EMaterialParameterType Type, int32 ParameterIndex) const override;
+	virtual UEdGraphPin* GetParameterPin(const FNodeMaterialParameterId& ParameterId) const override;
+	virtual bool IsImageMutableMode(int32 ImageIndex) const override;
+	virtual bool IsImageMutableMode(const UEdGraphPin& Pin) const override;
+	virtual UTexture2D* GetImageReferenceTexture(int32 ImageIndex) const override;
+	virtual UTexture2D* GetImageValue(int32 ImageIndex) const override;
+	virtual int32 GetImageUVLayout(int32 ImageIndex) const override;
+	virtual UEdGraphPin* GetMeshPin() const override;
+	virtual UEdGraphPin* GetMaterialAssetPin() const override;
+	virtual const UCustomizableObjectNodeMaterial* GetMaterialNode() const override;
+	virtual bool RealMaterialDataHasChanged() const override;
+	virtual FPostImagePinModeChangedDelegate* GetPostImagePinModeChangedDelegate() override;
+	
+	// Own Interface
+
+	void SetMaterial(UMaterialInterface* InMaterial);
+
+	static bool HasParameter(const UMaterialInterface* InMaterial, const FNodeMaterialParameterId& ParameterId);
+	static int32 GetParameterLayerIndex(const UMaterialInterface* InMaterial, EMaterialParameterType Type, int32 ParameterIndex);
+
+	UPROPERTY(EditAnywhere, Category = Tags)
+	TArray<FString> Tags;
+
+private:
+	/** Set the Pin Mode of a Texture Parameter Pin. */
+	void SetImagePinMode(UEdGraphPin& Pin, EPinMode PinMode) const;
+	
+	/** Delegate called when a Texture Parameter Pin Mode changes. */
+	FPostImagePinModeChangedDelegate PostImagePinModeChangedDelegate;
+
+	UPROPERTY(EditAnywhere, Category=CustomizableObject)
+	TObjectPtr<UMaterialInterface> Material = nullptr;
+
+	UPROPERTY(EditAnywhere, Category=CustomizableObject, DisplayName = "Default Texture Parameter Mode", Meta = (ToolTip = "All Mateiral Texture FParameters set to \"Node Defined\" will use this mode."))
+	ENodePinMode TextureParametersMode = ENodePinMode::Passthrough;
+
+	UPROPERTY()
+	int32 MeshComponentIndex_DEPRECATED = 0;
+
+public:
+	/** Selects which Mesh component of the Instance this material belongs to */
+	UPROPERTY()
+	FName MeshComponentName_DEPRECATED;
+
+private:
+	/** Materials will be reused between LODs when possible unless explicitly disabled. */
+	UPROPERTY(EditAnywhere, Category = CustomizableObject)
+	bool bReuseMaterialBetweenLODs = true;
+
+	/** Last static or skeletal mesh connected. Used to remove the callback once disconnected. */
+	TWeakObjectPtr<UCustomizableObjectNode> LastMeshNodeConnected;
+
+	/** List of material parameter types that are actually relevant to mutable. */
+	static const TArray<EMaterialParameterType> ParameterTypes;
+
+	/** Relates a Parameter id (key) (and layer if is a layered material) to a Pin (value). Only used to improve performance.
+	  * If a deprecated pin and a non-deprecated pin have the same Parameter id, this the non-deprecated one prevails. */
+	UPROPERTY()
+	TMap<FNodeMaterialParameterId, FEdGraphPinReference> PinsParameterMap;
+	
+	/** Create the pin data of the given parameter type. */
+	UCustomizableObjectNodeMaterialPinDataParameter* CreatePinData(EMaterialParameterType Type, int32 ParameterIndex);
+
+	/** Allocate a pin for each parameter of the given type. */
+	void AllocateDefaultParameterPins(EMaterialParameterType Type);
+
+	/** Set the default Material from the connected static or skeletal mesh. */
+	void SetDefaultMaterial();
+
+	/** Connected NodeStaticMesh or NodeSkeletalMesh Mesh UPROPERTY changed callback function. Sets the default material. */
+	UFUNCTION()
+	void MeshPostEditChangeProperty(FPostEditChangePropertyDelegateParameters& Parameters);
+
+	/** Format pin name. */
+	FName GetPinName(EMaterialParameterType Type, int32 ParameterIndex) const;
+
+	/** Returns the texture coordinate of the given Material Expression. Returns -1 if not found. */
+	static int32 GetExpressionTextureCoordinate(UMaterial* Material, const FGuid &ImageId);
+
+	/** Return the Pin Category given the node NodePinMode. */
+	static FName NodePinModeToImagePinMode(ENodePinMode NodePinMode);
+
+	/** Return the Pin Category given a PinMode. */
+	FName GetImagePinMode(EPinMode PinMode) const;
+
+	/** Return the Pin Category given a Pin. */
+	FName GetImagePinMode(const UEdGraphPin& Pin) const;
+
+	/** Get the UV Layout Index defined in the Material. */
+	int32 GetImageUVLayoutFromMaterial(int32 ImageIndex) const;
+	
+	// Deprecated properties
+	/** Set all pins to Mutable mode. Even so, each pin can override its behaviour. */
+	UPROPERTY()
+	bool bDefaultPinModeMutable_DEPRECATED = false;
+	
+	UPROPERTY()
+	TArray<FCustomizableObjectNodeMaterialImage> Images_DEPRECATED;
+	
+	UPROPERTY()
+	TArray<FCustomizableObjectNodeMaterialVector> VectorParams_DEPRECATED;
+
+	UPROPERTY()
+	TArray<FCustomizableObjectNodeMaterialScalar> ScalarParams_DEPRECATED;
+
+	UPROPERTY()
+	TMap<FGuid, FEdGraphPinReference> PinsParameter_DEPRECATED;
+};
+
+
+/** Additional data for a Material Texture Parameter pin. */
+UCLASS()
+class CUSTOMIZABLEOBJECTEDITOR_API UCustomizableObjectNodeMaterialPinDataImage : public UCustomizableObjectNodeMaterialPinDataParameter
+{
+	GENERATED_BODY()
+
+	friend void UCustomizableObjectNodeMaterial::BackwardsCompatibleFixup(int32 CustomizableObjectCustomVersion);
+	
+public:
+	// UObject interface
+	virtual void PostEditChangeProperty(FPropertyChangedEvent& PropertyChangedEvent) override;
+	virtual bool CanEditChange(const FProperty* InProperty) const override;
+
+	// UCustomizableObjectNodePinData interface
+	virtual void BackwardsCompatibleFixup(int32 CustomizableObjectCustomVersion) override;
+
+	// NodePinDataParameter interface
+	/** Virtual function used to copy pin data when remapping pins. */
+	virtual void Copy(const UCustomizableObjectNodePinData& Other) override;
+	
+	// UCustomizableObjectNodeMaterialPinParameter interface
+	virtual bool IsDefault() const override;
+
+	// Own interface
+	/** Constructor parameters. Should always be called after a NewObject. */
+	void Init(UCustomizableObjectNodeMaterial& InNodeMaterial);
+
+	EPinMode GetPinMode() const;
+
+	void SetPinMode(EPinMode InPinMode);
+	
+private:
+	/** Image pin mode. If is not default, overrides the defined node behaviour. */
+	UPROPERTY(EditAnywhere, Category = NoCategory, DisplayName = "Texture Parameter Mode")
+	EPinMode PinMode = EPinMode::Default;
+
+public:
+	constexpr static int32 UV_LAYOUT_IGNORE = -1;
+
+	UPROPERTY(EditAnywhere, Category = NoCategory)
+	EUVLayoutMode UVLayoutMode = EUVLayoutMode::FromMaterial;
+	
+	/** Index of the UV channel that will be used with this image. It is necessary to apply the proper layout transformations to it. */
+	UPROPERTY(EditAnywhere, Category = NoCategory, meta = (EditCondition = "UVLayoutMode == EUVLayoutMode::Index", EditConditionHides))
+	int32 UVLayout = -2;
+
+	/** Reference Texture used to decide the texture properties of the mutable-generated textures
+	* connected to this material. If null, it will try to be guessed at compile time from
+	* the graph. */
+	UPROPERTY(EditAnywhere, Category = NoCategory) // Required to be EditAnywhere for the selector to work.
+	TObjectPtr<UTexture2D> ReferenceTexture = nullptr;
+
+private:
+	UPROPERTY()
+	TObjectPtr<UCustomizableObjectNodeMaterial> NodeMaterial = nullptr;
+};
+
+
+/** Additional data for a Material Vector Parameter pin. */
+UCLASS()
+class CUSTOMIZABLEOBJECTEDITOR_API UCustomizableObjectNodeMaterialPinDataVector : public UCustomizableObjectNodeMaterialPinDataParameter
+{
+	GENERATED_BODY()
+};
+
+
+/** Additional data for a Material Float Parameter pin. */
+UCLASS()
+class CUSTOMIZABLEOBJECTEDITOR_API UCustomizableObjectNodeMaterialPinDataScalar : public UCustomizableObjectNodeMaterialPinDataParameter
+{
+	GENERATED_BODY()
+};
diff --git a/Engine/Plugins/Mutable/Source/CustomizableObjectEditor/Private/MuCOE/Nodes/CustomizableObjectNodeMaterialVariation.h b/Engine/Plugins/Mutable/Source/CustomizableObjectEditor/Private/MuCOE/Nodes/CustomizableObjectNodeMaterialVariation.h
new file mode 100644
--- /dev/null
+++ b/Engine/Plugins/Mutable/Source/CustomizableObjectEditor/Private/MuCOE/Nodes/CustomizableObjectNodeMaterialVariation.h
@@ -0,0 +1,43 @@
+// Copyright Epic Games, Inc. All Rights Reserved.
+
+#pragma once
+
+#include "MuCOE/Nodes/CustomizableObjectNodeVariation.h"
+
+#include "CustomizableObjectNodeMaterialVariation.generated.h"
+
+struct FCustomizableObjectMaterialVariation;
+
+
+UENUM(BlueprintType)
+enum class ECustomizableObjectNodeMaterialVariationType : uint8
+{
+	Tag 		UMETA(DisplayName = "Tag"),
+	State 		UMETA(DisplayName = "State"),
+};
+
+
+UCLASS()
+class CUSTOMIZABLEOBJECTEDITOR_API UCustomizableObjectNodeMaterialVariation : public UCustomizableObjectNodeVariation
+{
+public:
+	GENERATED_BODY()
+
+	UPROPERTY(EditAnywhere, Category = CustomizableObject)
+	ECustomizableObjectNodeMaterialVariationType Type = ECustomizableObjectNodeMaterialVariationType::Tag;
+
+private:
+	// Deprecated properties
+	UPROPERTY()
+	TArray<FCustomizableObjectMaterialVariation> Variations_DEPRECATED;
+
+public:
+	// UCustomizableObjectNode interface
+	virtual void BackwardsCompatibleFixup(int32 CustomizableObjectCustomVersion) override;
+	virtual bool IsSingleOutputNode() const override;
+	
+	// UCustomizableObjectNodeVariation interface
+	virtual FName GetCategory() const override;
+	virtual bool IsInputPinArray() const override;
+};
+
diff --git a/Engine/Plugins/Mutable/Source/CustomizableObjectEditor/Private/MuCOE/Nodes/CustomizableObjectNodeMesh.h b/Engine/Plugins/Mutable/Source/CustomizableObjectEditor/Private/MuCOE/Nodes/CustomizableObjectNodeMesh.h
new file mode 100644
--- /dev/null
+++ b/Engine/Plugins/Mutable/Source/CustomizableObjectEditor/Private/MuCOE/Nodes/CustomizableObjectNodeMesh.h
@@ -0,0 +1,44 @@
+// Copyright Epic Games, Inc. All Rights Reserved.
+
+#pragma once
+
+#include "MuCOE/Nodes/CustomizableObjectNode.h"
+
+#include "CustomizableObjectNodeMesh.generated.h"
+
+class UCustomizableObjectLayout;
+class UEdGraphPin;
+class UObject;
+class UTexture2D;
+
+
+UCLASS(abstract)
+class CUSTOMIZABLEOBJECTEDITOR_API UCustomizableObjectNodeMesh : public UCustomizableObjectNode
+{
+public:
+	GENERATED_BODY()
+
+	// Own interface
+	/** Returns the multiple Layouts of a given Mesh pin. A pin can have multiple Layouts since it can have multiple UVs. Override. */
+	virtual TArray<UCustomizableObjectLayout*> GetLayouts(const UEdGraphPin& OutPin) const PURE_VIRTUAL(UCustomizableObjectNodeMesh::GetLayouts, return {};);
+
+	virtual UTexture2D* FindTextureForPin(const UEdGraphPin* Pin) const PURE_VIRTUAL(UCustomizableObjectNodeMesh::FindTextureForPin, return {};);
+
+	/** Returns the Unreal mesh (e.g., USkeletalMesh, UStaticMesh...). */
+	virtual TSoftObjectPtr<UObject> GetMesh() const PURE_VIRTUAL(UCustomizableObjectNodeMesh::GetMesh, return {};);
+
+	/** Returns the output Mesh pin associated to the given LODIndex and SectionIndex. Override. */
+	virtual UEdGraphPin* GetMeshPin(int32 LOD, int32 SectionIndex) const PURE_VIRTUAL(UCustomizableObjectNodeMesh::GetMeshPin, return {};);
+
+	/** Given a pin, return the Section and Layout index.
+	 *
+	 * Will always return a valid result. The result can be valid but out of sync with respect the Unreal mesh asset.
+	 * In other words, the pin still represents a LOD 3 but the asset no longer has this third LOD.
+	 * 
+	 * @param Pin In. Node's owning pin.
+	 * @param LODIndex Out.
+	 * @param SectionIndex Out.
+	 * @param LayoutIndex Out. */
+	virtual void GetPinSection(const UEdGraphPin& Pin, int32& OutLODIndex, int32& OutSectionIndex, int32& OutLayoutIndex) const PURE_VIRTUAL(UCustomizableObjectNodeMesh::GetPinSection, );
+};
+
diff --git a/Engine/Plugins/Mutable/Source/CustomizableObjectEditor/Private/MuCOE/Nodes/CustomizableObjectNodeMeshMorph.h b/Engine/Plugins/Mutable/Source/CustomizableObjectEditor/Private/MuCOE/Nodes/CustomizableObjectNodeMeshMorph.h
new file mode 100644
--- /dev/null
+++ b/Engine/Plugins/Mutable/Source/CustomizableObjectEditor/Private/MuCOE/Nodes/CustomizableObjectNodeMeshMorph.h
@@ -0,0 +1,84 @@
+// Copyright Epic Games, Inc. All Rights Reserved.
+
+#pragma once
+
+#include "MuCOE/Nodes/CustomizableObjectNode.h"
+#include "MuCOE/Nodes/CustomizableObjectNodeMeshReshapeCommon.h"
+
+#include "CustomizableObjectNodeMeshMorph.generated.h"
+
+class UCustomizableObjectNodeRemapPins;
+class UEdGraphPin;
+class UObject;
+struct FMeshReshapeBoneReference;
+enum class EBoneDeformSelectionMethod : uint8;
+
+UCLASS()
+class CUSTOMIZABLEOBJECTEDITOR_API UCustomizableObjectNodeMeshMorph : public UCustomizableObjectNode
+{
+public:
+	GENERATED_BODY()
+
+	UCustomizableObjectNodeMeshMorph();
+
+	// Begin EdGraphNode interface
+	FText GetNodeTitle(ENodeTitleType::Type TitleType) const override;
+	FLinearColor GetNodeTitleColor() const override;
+	FText GetTooltipText() const override;
+
+	// UCustomizableObjectNode interface
+	void AllocateDefaultPins(UCustomizableObjectNodeRemapPins* RemapPins) override;
+	bool IsNodeOutDatedAndNeedsRefresh() override;
+	FString GetRefreshMessage() const override;
+	void PinConnectionListChanged(UEdGraphPin* Pin) override;
+
+
+	// Own interface
+	class UCustomizableObjectNodeSkeletalMesh* GetSourceSkeletalMesh() const;
+
+	UEdGraphPin* MeshPin() const
+	{
+		return FindPin(TEXT("Mesh"), EGPD_Input);
+	}
+
+	UEdGraphPin* FactorPin() const
+	{
+		return FindPin(TEXT("Factor"));
+	}
+
+	void Serialize(FArchive& Ar) override;
+
+	UPROPERTY(Category=CustomizableObject, EditAnywhere)
+	FString MorphTargetName;
+
+	/** Experimental - Enable the deformation of the skeleton of the base mesh. */
+	UPROPERTY(EditAnywhere, Category = CustomizableObject)
+	bool bReshapeSkeleton = false;
+
+	/** Experimental - Enable the deformation of physics volumes of the base mesh */
+    UPROPERTY(EditAnywhere, Category = CustomizableObject)
+    bool bReshapePhysicsVolumes = false;
+
+	/** Bone Reshape Selection Method */
+	UPROPERTY(EditAnywhere, Category = ReshapeBones, meta = (EditCondition = "bReshapeSkeleton"))
+	EBoneDeformSelectionMethod SelectionMethod;
+	
+	/** Array with selected bones that will be deformed */
+	UPROPERTY(EditAnywhere, Category = ReshapeBones, meta = (EditCondition = "bReshapeSkeleton && (SelectionMethod == EBoneDeformSelectionMethod::ONLY_SELECTED || SelectionMethod == EBoneDeformSelectionMethod::ALL_BUT_SELECTED)"))
+	TArray<FMeshReshapeBoneReference> BonesToDeform;
+
+	UPROPERTY(EditAnywhere, Category = ReshapePhysics, meta = (DisplayName = "Selection Method", EditCondition = "bReshapePhysicsVolumes"))
+	EBoneDeformSelectionMethod PhysicsSelectionMethod;
+
+	/** Array with bones with physics bodies that will be deformed */
+	UPROPERTY(EditAnywhere, Category = ReshapePhysics, meta = (EditCondition = "bReshapePhysicsVolumes && (PhysicsSelectionMethod == EBoneDeformSelectionMethod::ONLY_SELECTED || PhysicsSelectionMethod == EBoneDeformSelectionMethod::ALL_BUT_SELECTED)"))
+	TArray<FMeshReshapeBoneReference> PhysicsBodiesToDeform;
+
+	UPROPERTY()
+	bool bDeformAllBones_DEPRECATED = false;
+
+	UPROPERTY()
+	bool bDeformAllPhysicsBodies_DEPRECATED = false;
+
+};
+
diff --git a/Engine/Plugins/Mutable/Source/CustomizableObjectEditor/Private/MuCOE/Nodes/CustomizableObjectNodeMeshReshape.h b/Engine/Plugins/Mutable/Source/CustomizableObjectEditor/Private/MuCOE/Nodes/CustomizableObjectNodeMeshReshape.h
new file mode 100644
--- /dev/null
+++ b/Engine/Plugins/Mutable/Source/CustomizableObjectEditor/Private/MuCOE/Nodes/CustomizableObjectNodeMeshReshape.h
@@ -0,0 +1,123 @@
+// Copyright Epic Games, Inc. All Rights Reserved.
+
+#pragma once
+
+#include "MuCOE/Nodes/CustomizableObjectNode.h"
+
+#include "CustomizableObjectNodeMeshReshape.generated.h"
+
+class UCustomizableObjectNodeRemapPins;
+class UEdGraphPin;
+class UObject;
+struct FMeshReshapeBoneReference;
+enum class EBoneDeformSelectionMethod : uint8;
+
+
+UENUM()
+enum class EMeshReshapeVertexColorChannelUsage
+{
+	None = 0,
+	RigidClusterId = 1,
+	MaskWeight = 2
+};
+
+USTRUCT()
+struct FMeshReshapeColorUsage
+{
+	GENERATED_BODY()
+
+	UPROPERTY(EditAnywhere, Category = CustomizableObject)
+	EMeshReshapeVertexColorChannelUsage R = EMeshReshapeVertexColorChannelUsage::None;
+
+	UPROPERTY(EditAnywhere, Category = CustomizableObject)
+	EMeshReshapeVertexColorChannelUsage G = EMeshReshapeVertexColorChannelUsage::None;
+
+	UPROPERTY(EditAnywhere, Category = CustomizableObject)
+	EMeshReshapeVertexColorChannelUsage B = EMeshReshapeVertexColorChannelUsage::None;
+
+	UPROPERTY(EditAnywhere, Category = CustomizableObject)
+	EMeshReshapeVertexColorChannelUsage A = EMeshReshapeVertexColorChannelUsage::None;
+};
+
+UCLASS()
+class CUSTOMIZABLEOBJECTEDITOR_API UCustomizableObjectNodeMeshReshape : public UCustomizableObjectNode
+{
+public:
+	GENERATED_BODY()
+
+	UCustomizableObjectNodeMeshReshape();
+
+	// Begin EdGraphNode interface
+	FText GetNodeTitle(ENodeTitleType::Type TitleType) const override;
+	FLinearColor GetNodeTitleColor() const override;
+	FText GetTooltipText() const override;
+
+	// UCustomizableObjectNode interface
+	void AllocateDefaultPins(UCustomizableObjectNodeRemapPins* RemapPins) override;
+	FString GetRefreshMessage() const override;
+	virtual bool IsExperimental() const override;
+
+	inline UEdGraphPin* BaseMeshPin() const
+	{
+		return FindPin(TEXT("Base Mesh"), EGPD_Input);
+	}
+
+	inline UEdGraphPin* BaseShapePin() const
+	{
+		return FindPin(TEXT("Base Shape"), EGPD_Input);
+	}
+
+	inline UEdGraphPin* TargetShapePin() const
+	{
+		return FindPin(TEXT("Target Shape"), EGPD_Input);
+	}
+
+	void Serialize(FArchive& Ar) override;
+
+	/** Enable the deformation of the vertices of the base mesh. */
+	UPROPERTY(EditAnywhere, Category = CustomizableObject)
+	bool bReshapeVertices = true;
+
+	UPROPERTY(EditAnywhere, Category = CustomizableObject)
+	bool bRecomputeNormals = false;
+
+	/** Enable laplacian smoothing to the result of the base mesh reshape. */
+	UPROPERTY(EditAnywhere, Category = CustomizableObject, meta = (EditCondition = "bReshapeVertices"))
+	bool bApplyLaplacianSmoothing = false;
+
+	/** Enable the deformation of the skeleton of the base mesh. */
+	UPROPERTY(EditAnywhere, Category = CustomizableObject)
+	bool bReshapeSkeleton = false;
+
+	/** Enable the deformation of physics volumes of the base mesh */
+    UPROPERTY(EditAnywhere, Category = CustomizableObject)
+    bool bReshapePhysicsVolumes = false;
+	
+	UPROPERTY()
+	bool bEnableRigidParts_DEPRECATED = false;
+
+	UPROPERTY(EditAnywhere, Category = CustomizableObject)
+	FMeshReshapeColorUsage VertexColorUsage;
+
+	/** Bone Reshape Selection Method */
+	UPROPERTY(EditAnywhere, Category = ReshapeBones, meta = (EditCondition = "bReshapeSkeleton"))
+	EBoneDeformSelectionMethod SelectionMethod;
+
+	/** Array with selected bones that will be deformed */
+	UPROPERTY(EditAnywhere, Category = ReshapeBones, meta = (EditCondition = "bReshapeSkeleton && (SelectionMethod == EBoneDeformSelectionMethod::ONLY_SELECTED || SelectionMethod == EBoneDeformSelectionMethod::ALL_BUT_SELECTED)"))
+	TArray<FMeshReshapeBoneReference> BonesToDeform;
+
+	UPROPERTY(EditAnywhere, Category = ReshapePhysics, meta = (DisplayName = "Selection Method", EditCondition = "bReshapePhysicsVolumes"))
+	EBoneDeformSelectionMethod PhysicsSelectionMethod;
+
+	/** Array with bones with physics bodies that will be deformed */
+	UPROPERTY(EditAnywhere, Category = ReshapePhysics, meta = (EditCondition = "bReshapePhysicsVolumes && (PhysicsSelectionMethod == EBoneDeformSelectionMethod::ONLY_SELECTED || PhysicsSelectionMethod == EBoneDeformSelectionMethod::ALL_BUT_SELECTED)"))
+	TArray<FMeshReshapeBoneReference> PhysicsBodiesToDeform;
+
+	UPROPERTY()
+	bool bDeformAllBones_DEPRECATED = false;
+
+	UPROPERTY()
+	bool bDeformAllPhysicsBodies_DEPRECATED = false;
+};
+
diff --git a/Engine/Plugins/Mutable/Source/CustomizableObjectEditor/Private/MuCOE/Nodes/CustomizableObjectNodeMeshSwitch.h b/Engine/Plugins/Mutable/Source/CustomizableObjectEditor/Private/MuCOE/Nodes/CustomizableObjectNodeMeshSwitch.h
new file mode 100644
--- /dev/null
+++ b/Engine/Plugins/Mutable/Source/CustomizableObjectEditor/Private/MuCOE/Nodes/CustomizableObjectNodeMeshSwitch.h
@@ -0,0 +1,22 @@
+// Copyright Epic Games, Inc. All Rights Reserved.
+
+#pragma once
+
+#include "MuCOE/Nodes/CustomizableObjectNodeSwitchBase.h"
+
+#include "CustomizableObjectNodeMeshSwitch.generated.h"
+
+
+UCLASS()
+class CUSTOMIZABLEOBJECTEDITOR_API UCustomizableObjectNodeMeshSwitch : public UCustomizableObjectNodeSwitchBase
+{
+public:
+	GENERATED_BODY()
+
+	// UCustomizableObjectNode interface
+	virtual void BackwardsCompatibleFixup(int32 CustomizableObjectCustomVersion) override;
+	
+	// UCustomizableObjectNodeSwitchBase interface
+	virtual FName GetCategory() const override;
+};
+
diff --git a/Engine/Plugins/Mutable/Source/CustomizableObjectEditor/Private/MuCOE/Nodes/CustomizableObjectNodeMeshVariation.h b/Engine/Plugins/Mutable/Source/CustomizableObjectEditor/Private/MuCOE/Nodes/CustomizableObjectNodeMeshVariation.h
new file mode 100644
--- /dev/null
+++ b/Engine/Plugins/Mutable/Source/CustomizableObjectEditor/Private/MuCOE/Nodes/CustomizableObjectNodeMeshVariation.h
@@ -0,0 +1,28 @@
+// Copyright Epic Games, Inc. All Rights Reserved.
+
+#pragma once
+
+#include "MuCOE/Nodes/CustomizableObjectNodeVariation.h"
+
+#include "CustomizableObjectNodeMeshVariation.generated.h"
+
+struct FCustomizableObjectMeshVariation;
+
+
+UCLASS()
+class CUSTOMIZABLEOBJECTEDITOR_API UCustomizableObjectNodeMeshVariation : public UCustomizableObjectNodeVariation
+{
+public:
+	GENERATED_BODY()
+
+	// Deprecated properties
+	UPROPERTY()
+	TArray<FCustomizableObjectMeshVariation> Variations_DEPRECATED;
+
+	// UCustomizableObjectNode interface
+	virtual void BackwardsCompatibleFixup(int32 CustomizableObjectCustomVersion) override;
+	
+	// UCustomizableObjectNodeVariation interface
+	virtual FName GetCategory() const override;
+};
+
diff --git a/Engine/Plugins/Mutable/Source/CustomizableObjectEditor/Private/MuCOE/Nodes/CustomizableObjectNodeModifierClipDeform.h b/Engine/Plugins/Mutable/Source/CustomizableObjectEditor/Private/MuCOE/Nodes/CustomizableObjectNodeModifierClipDeform.h
new file mode 100644
--- /dev/null
+++ b/Engine/Plugins/Mutable/Source/CustomizableObjectEditor/Private/MuCOE/Nodes/CustomizableObjectNodeModifierClipDeform.h
@@ -0,0 +1,55 @@
+// Copyright Epic Games, Inc. All Rights Reserved.
+
+#pragma once
+
+#include "MuCOE/Nodes/CustomizableObjectNodeModifierBase.h"
+#include "MuT/NodeModifier.h"
+
+#include "CustomizableObjectNodeModifierClipDeform.generated.h"
+
+class UCustomizableObjectNodeRemapPins;
+class UEdGraphPin;
+class UObject;
+
+UENUM()
+enum class EShapeBindingMethod : uint32
+{
+	ClosestProject = 0,
+	ClosestToSurface = 1,
+	NormalProject = 2
+};
+
+UCLASS()
+class CUSTOMIZABLEOBJECTEDITOR_API UCustomizableObjectNodeModifierClipDeform : public UCustomizableObjectNodeModifierBase
+{
+
+	GENERATED_BODY()
+
+public:
+
+	UPROPERTY()
+	TArray<FString> Tags_DEPRECATED;
+
+	UPROPERTY(EditAnywhere, Category = MeshClipDeform)
+	EShapeBindingMethod BindingMethod;
+
+	UPROPERTY(EditAnywhere, Category = RemoveOptions)
+	EFaceCullStrategy FaceCullStrategy = EFaceCullStrategy::AllVerticesCulled;
+
+public:
+
+	// EdGraphNode interface
+	virtual FText GetNodeTitle(ENodeTitleType::Type TitleType) const override;
+	virtual FText GetTooltipText() const override;
+
+	// UCustomizableObjectNode interface
+	virtual void AllocateDefaultPins(UCustomizableObjectNodeRemapPins* RemapPins) override;
+	virtual void BackwardsCompatibleFixup(int32 CustomizableObjectCustomVersion) override;
+	virtual bool IsExperimental() const override;
+
+	// Own interface
+
+	/** Input pins. */
+	UEdGraphPin* ClipShapePin() const;
+};
+
diff --git a/Engine/Plugins/Mutable/Source/CustomizableObjectEditor/Private/MuCOE/Nodes/CustomizableObjectNodeModifierClipWithMesh.h b/Engine/Plugins/Mutable/Source/CustomizableObjectEditor/Private/MuCOE/Nodes/CustomizableObjectNodeModifierClipWithMesh.h
new file mode 100644
--- /dev/null
+++ b/Engine/Plugins/Mutable/Source/CustomizableObjectEditor/Private/MuCOE/Nodes/CustomizableObjectNodeModifierClipWithMesh.h
@@ -0,0 +1,61 @@
+// Copyright Epic Games, Inc. All Rights Reserved.
+
+#pragma once
+
+#include "MuCOE/Nodes/CustomizableObjectNodeModifierBase.h"
+#include "MuT/NodeModifier.h"
+
+#include "CustomizableObjectNodeModifierClipWithMesh.generated.h"
+
+namespace ENodeTitleType { enum Type : int; }
+
+class UCustomizableObject;
+class UCustomizableObjectNodeRemapPins;
+class UEdGraphPin;
+class UObject;
+struct FGuid;
+
+UCLASS()
+class CUSTOMIZABLEOBJECTEDITOR_API UCustomizableObjectNodeModifierClipWithMesh : public UCustomizableObjectNodeModifierBase
+{
+	GENERATED_BODY()
+
+public:
+
+	UPROPERTY()
+	TArray<FString> Tags_DEPRECATED;
+
+	//!< If assigned, then a material inside this CO will be clipped by this node.
+    //!< If several materials with the same name, all are considered (to cover all LOD levels)
+    UPROPERTY()
+	TObjectPtr<UCustomizableObject> CustomizableObjectToClipWith_DEPRECATED = nullptr;
+
+    //!< Array with the Guids of the nodes with the same material inside the CustomizableObjectToClipWith CO (if any is assigned)
+    UPROPERTY()
+    TArray<FGuid> ArrayMaterialNodeToClipWithID_DEPRECATED;
+
+	/** Transform to apply to the clip mesh before clipping. */
+	UPROPERTY(EditAnywhere, Category = ClipMesh)
+	FTransform Transform;
+
+	UPROPERTY(EditAnywhere, Category = ClipMesh)
+	EFaceCullStrategy FaceCullStrategy = EFaceCullStrategy::AllVerticesCulled;
+
+public:
+
+	UCustomizableObjectNodeModifierClipWithMesh();
+
+	// EdGraphNode interface
+	virtual FText GetNodeTitle(ENodeTitleType::Type TitleType) const override;
+	virtual FText GetTooltipText() const override;
+
+	// UCustomizableObjectNode interface
+	virtual void AllocateDefaultPins(UCustomizableObjectNodeRemapPins* RemapPins) override;
+	virtual void BackwardsCompatibleFixup(int32 CustomizableObjectCustomVersion) override;
+
+	// Own interface
+	UEdGraphPin* OutputPin() const;
+
+	UEdGraphPin* ClipMeshPin() const;
+};
+
diff --git a/Engine/Plugins/Mutable/Source/CustomizableObjectEditor/Private/MuCOE/Nodes/CustomizableObjectNodeModifierExtendMeshSection.h b/Engine/Plugins/Mutable/Source/CustomizableObjectEditor/Private/MuCOE/Nodes/CustomizableObjectNodeModifierExtendMeshSection.h
new file mode 100644
--- /dev/null
+++ b/Engine/Plugins/Mutable/Source/CustomizableObjectEditor/Private/MuCOE/Nodes/CustomizableObjectNodeModifierExtendMeshSection.h
@@ -0,0 +1,77 @@
+// Copyright Epic Games, Inc. All Rights Reserved.
+
+#pragma once
+
+#include "MuCOE/Nodes/CustomizableObjectNodeModifierWithMaterial.h"
+#include "MuCOE/CustomizableObjectLayout.h"
+
+#include "CustomizableObjectNodeModifierExtendMeshSection.generated.h"
+
+namespace ENodeTitleType { enum Type : int; }
+
+class FArchive;
+class UCustomizableObject;
+class UCustomizableObjectNode;
+class UCustomizableObjectNodeRemapPins;
+class UEdGraphPin;
+class UObject;
+struct FCustomizableObjectNodeExtendMaterialImage;
+struct FEdGraphPinReference;
+
+
+UCLASS()
+class CUSTOMIZABLEOBJECTEDITOR_API UCustomizableObjectNodeModifierExtendMeshSection : public UCustomizableObjectNodeModifierWithMaterial
+{
+	GENERATED_BODY()
+
+public:
+
+	/** Tags enabled when the modifier is applied. Also used to decide what modifiers get applied to the data added by this node.  */
+	UPROPERTY(EditAnywhere, BlueprintReadWrite, Category = EnableTags)
+	TArray<FString> Tags;
+
+	/** First LOD where the modifier is applied to. */
+	UPROPERTY(EditAnywhere, Category = LOD, meta = (ClampMin = 0, ClampMax = 255))
+	uint32 FirstLOD = 0;
+
+public:
+	
+	// Begin EdGraphNode interface
+	virtual FText GetNodeTitle(ENodeTitleType::Type TitleType) const override;
+	virtual FText GetTooltipText() const override;
+
+	// UCustomizableObjectNode interface
+	virtual void BackwardsCompatibleFixup(int32 CustomizableObjectCustomVersion) override;
+	virtual void AllocateDefaultPins(UCustomizableObjectNodeRemapPins* RemapPins) override;
+	virtual FString GetRefreshMessage() const override;
+	virtual bool IsSingleOutputNode() const override;
+	virtual bool CustomRemovePin(UEdGraphPin& Pin) override;
+	virtual TArray<FString>* GetEnableTags() override;
+	virtual FString GetInternalTagDisplayName() override;
+
+	// Own interface
+	UEdGraphPin* AddMeshPin() const;
+
+	TArray<UCustomizableObjectLayout*> GetLayouts() const;
+
+private:
+
+	// Deprecated properties
+
+	UPROPERTY()
+	TObjectPtr<UCustomizableObject> ParentMaterialObject_DEPRECATED = nullptr;
+
+	UPROPERTY()
+	FGuid ParentMaterialNodeId_DEPRECATED;
+	
+
+	UPROPERTY()
+	TMap<FNodeMaterialParameterId, FEdGraphPinReference> PinsParameterMap_DEPRECATED;
+
+	UPROPERTY()
+	TMap<FGuid, FEdGraphPinReference> PinsParameter_DEPRECATED;
+
+	UPROPERTY()
+	TArray<FCustomizableObjectNodeExtendMaterialImage> Images_DEPRECATED;
+};
+
diff --git a/Engine/Plugins/Mutable/Source/CustomizableObjectEditor/Private/MuCOE/Nodes/CustomizableObjectNodeModifierRemoveMesh.h b/Engine/Plugins/Mutable/Source/CustomizableObjectEditor/Private/MuCOE/Nodes/CustomizableObjectNodeModifierRemoveMesh.h
new file mode 100644
--- /dev/null
+++ b/Engine/Plugins/Mutable/Source/CustomizableObjectEditor/Private/MuCOE/Nodes/CustomizableObjectNodeModifierRemoveMesh.h
@@ -0,0 +1,42 @@
+// Copyright Epic Games, Inc. All Rights Reserved.
+
+#pragma once
+
+#include "MuCOE/Nodes/CustomizableObjectNodeModifierEditMeshSectionBase.h"
+
+#include "CustomizableObjectNodeModifierRemoveMesh.generated.h"
+
+namespace ENodeTitleType { enum Type : int; }
+
+class UCustomizableObjectNodeRemapPins;
+class UEdGraphPin;
+class UObject;
+
+UCLASS()
+class CUSTOMIZABLEOBJECTEDITOR_API UCustomizableObjectNodeModifierRemoveMesh : public UCustomizableObjectNodeModifierEditMeshSectionBase
+{
+	GENERATED_BODY()
+
+public:
+
+	UPROPERTY(EditAnywhere, Category = RemoveOptions)
+	EFaceCullStrategy FaceCullStrategy = EFaceCullStrategy::AllVerticesCulled;
+
+public:
+
+	// Begin EdGraphNode interface
+	FText GetNodeTitle(ENodeTitleType::Type TitleType) const override;
+	FText GetTooltipText() const override;
+
+	// UCustomizableObjectNode interface
+	virtual void AllocateDefaultPins(UCustomizableObjectNodeRemapPins* RemapPins) override;
+	virtual void PinConnectionListChanged(UEdGraphPin * Pin);
+
+	UEdGraphPin* RemoveMeshPin() const
+	{
+		return FindPin(TEXT("Remove Mesh"));
+	}
+
+	bool IsSingleOutputNode() const override;
+};
+
diff --git a/Engine/Plugins/Mutable/Source/CustomizableObjectEditor/Private/MuCOE/Nodes/CustomizableObjectNodeModifierRemoveMeshBlocks.h b/Engine/Plugins/Mutable/Source/CustomizableObjectEditor/Private/MuCOE/Nodes/CustomizableObjectNodeModifierRemoveMeshBlocks.h
new file mode 100644
--- /dev/null
+++ b/Engine/Plugins/Mutable/Source/CustomizableObjectEditor/Private/MuCOE/Nodes/CustomizableObjectNodeModifierRemoveMeshBlocks.h
@@ -0,0 +1,41 @@
+// Copyright Epic Games, Inc. All Rights Reserved.
+
+#pragma once
+
+#include "MuCOE/Nodes/CustomizableObjectNodeModifierEditLayoutBlocks.h"
+
+#include "CustomizableObjectNodeModifierRemoveMeshBlocks.generated.h"
+
+namespace ENodeTitleType { enum Type : int; }
+
+class FArchive;
+class UCustomizableObjectNodeRemapPins;
+class UEdGraphPin;
+class UObject;
+struct FPropertyChangedEvent;
+
+UCLASS()
+class CUSTOMIZABLEOBJECTEDITOR_API UCustomizableObjectNodeModifierRemoveMeshBlocks : public UCustomizableObjectNodeModifierEditLayoutBlocks
+{
+	GENERATED_BODY()
+
+public:
+
+	UPROPERTY()
+	TArray<int32> Blocks_DEPRECATED;
+
+	UPROPERTY(EditAnywhere, Category = RemoveOptions)
+	EFaceCullStrategy FaceCullStrategy = EFaceCullStrategy::AllVerticesCulled;
+
+public:
+
+	// Begin EdGraphNode interface
+	virtual FText GetNodeTitle(ENodeTitleType::Type TitleType) const override;
+	virtual FText GetTooltipText() const override;
+
+	// UCustomizableObjectNode interface
+	virtual void BackwardsCompatibleFixup(int32 CustomizableObjectCustomVersion) override;
+	virtual void AllocateDefaultPins(UCustomizableObjectNodeRemapPins* RemapPins) override;
+	virtual bool IsSingleOutputNode() const override;
+};
+
diff --git a/Engine/Plugins/Mutable/Source/CustomizableObjectEditor/Private/MuCOE/Nodes/CustomizableObjectNodeModifierWithMaterial.h b/Engine/Plugins/Mutable/Source/CustomizableObjectEditor/Private/MuCOE/Nodes/CustomizableObjectNodeModifierWithMaterial.h
new file mode 100644
--- /dev/null
+++ b/Engine/Plugins/Mutable/Source/CustomizableObjectEditor/Private/MuCOE/Nodes/CustomizableObjectNodeModifierWithMaterial.h
@@ -0,0 +1,56 @@
+// Copyright Epic Games, Inc. All Rights Reserved.
+
+#pragma once
+
+#include "MuCOE/Nodes/CustomizableObjectNodeModifierBase.h"
+#include "MuCOE/Nodes/CustomizableObjectNodeMaterialBase.h"
+
+#include "CustomizableObjectNodeModifierWithMaterial.generated.h"
+
+class UObject;
+
+
+UCLASS(abstract)
+class CUSTOMIZABLEOBJECTEDITOR_API UCustomizableObjectNodeModifierWithMaterial : public UCustomizableObjectNodeModifierBase
+{
+
+	GENERATED_BODY()
+
+public:
+
+	/** Reference material that defines the structure of the data to be extended.
+	* The Sections modified with this modifier are supposed to have the same texture parameters, but don't need to have exactly
+	* the ReferenceMaterial set.
+	*/
+	UPROPERTY(EditAnywhere, Category = Parameters)
+	TObjectPtr<UMaterialInterface> ReferenceMaterial = nullptr;
+
+	/** Relates a Parameter id (and layer if is a layered material) to a Pin. Only used to improve performance. */
+	UPROPERTY()
+	TMap<FNodeMaterialParameterId, FEdGraphPinReference> PinsParameterMap;
+
+public:
+
+	//
+	virtual void PostEditChangeProperty(FPropertyChangedEvent& PropertyChangedEvent) override;
+
+	// UCustomizableObjectNode interface
+	virtual bool IsNodeOutDatedAndNeedsRefresh() override;
+
+	// Own interface
+
+	/** */
+	bool UsesImage(const FNodeMaterialParameterId& ImageId) const;
+
+	/** */
+	const UEdGraphPin* GetUsedImagePin(const FNodeMaterialParameterId& ImageId) const;
+
+	/** */
+	int32 GetNumParameters(const EMaterialParameterType Type) const;
+	FNodeMaterialParameterId GetParameterId(const EMaterialParameterType Type, const int32 ParameterIndex) const;
+	FName GetParameterName(const EMaterialParameterType Type, const int32 ParameterIndex) const;
+	static int32 GetParameterLayerIndex(const UMaterialInterface* InMaterial, const EMaterialParameterType Type, const int32 ParameterIndex);
+	int32 GetParameterLayerIndex(const EMaterialParameterType Type, const int32 ParameterIndex) const;
+	FText GetParameterLayerName(const EMaterialParameterType Type, const int32 ParameterIndex) const;
+};
+
diff --git a/Engine/Plugins/Mutable/Source/CustomizableObjectEditor/Private/MuCOE/Nodes/CustomizableObjectNodeObject.h b/Engine/Plugins/Mutable/Source/CustomizableObjectEditor/Private/MuCOE/Nodes/CustomizableObjectNodeObject.h
new file mode 100644
--- /dev/null
+++ b/Engine/Plugins/Mutable/Source/CustomizableObjectEditor/Private/MuCOE/Nodes/CustomizableObjectNodeObject.h
@@ -0,0 +1,207 @@
+// Copyright Epic Games, Inc. All Rights Reserved.
+
+#pragma once
+
+#include "CustomizableObjectNodeComponentMesh.h"
+#include "CustomizableObjectNodeComponentMeshBase.h"
+#include "MuCOE/Nodes/CustomizableObjectNode.h"
+#include "MuR/System.h"
+
+#include "CustomizableObjectNodeObject.generated.h"
+
+namespace ENodeTitleType { enum Type : int; }
+
+class UCustomizableObject;
+class UCustomizableObjectNodeMaterialBase;
+class UCustomizableObjectNodeRemapPins;
+class UObject;
+struct FPropertyChangedEvent;
+struct FSoftObjectPath;
+struct FComponentSettings;
+
+
+USTRUCT()
+struct CUSTOMIZABLEOBJECTEDITOR_API FCustomizableObjectState
+{
+	GENERATED_USTRUCT_BODY()
+
+	friend UCustomizableObjectNodeObject;
+	
+	UPROPERTY(EditAnywhere, Category=CustomizableObject)
+	FString Name;
+
+	UPROPERTY(EditAnywhere, Category = CustomizableObject, meta = (ShowParameterOptions))
+	TArray<FString> RuntimeParameters;
+
+	/** Special treatment of texture compression for this state. */
+	UPROPERTY(EditAnywhere, Category = CustomizableObject)
+	ETextureCompressionStrategy TextureCompressionStrategy = ETextureCompressionStrategy::None;
+
+	/** If this is enabled, texture streaming won't be used for this state, and full images will be generated when an instance is first updated. */
+	UPROPERTY(EditAnywhere, Category = CustomizableObject)
+	bool bDisableTextureStreaming = false;
+
+	/** LiveUpdateMode will reuse instance temp. data between updates and speed up update times, but spend much more memory. Good for customization screens, not for actual gameplay modes. */
+	UPROPERTY(EditAnywhere, Category = CustomizableObject)
+	bool bLiveUpdateMode = false;
+
+	// Enables the reuse of all possible textures when the instance is updated without any changes in geometry or state (the first update after creation doesn't reuse any)
+	// It will only work if the textures aren't compressed, so set the instance to a Mutable state with texture compression disabled
+	// WARNING! If texture reuse is enabled, do NOT keep external references to the textures of the instance. The instance owns the textures.
+	UPROPERTY(EditAnywhere, Category = CustomizableObject)
+	bool bReuseInstanceTextures = false;
+
+	UPROPERTY(EditAnywhere, Category = CustomizableObject)
+	bool bBuildOnlyFirstLOD = false;
+
+	UPROPERTY(EditAnywhere, Category = CustomizableObject)
+	TMap<FString, FString> ForcedParameterValues;
+
+	UPROPERTY(EditAnywhere, Category = UI, meta = (DisplayName = "State UI Metadata"))
+	FMutableStateUIMetadata UIMetadata;
+
+private:
+	// Deprecated
+	UPROPERTY()
+	FMutableParamUIMetadata StateUIMetadata_DEPRECATED;
+	
+	/** This is now TextureCompressionStrategy.  */
+	UPROPERTY()
+	bool bDontCompressRuntimeTextures_DEPRECATED = false;
+};
+
+USTRUCT()
+struct FSkeletalMeshMorphTargetOverride
+{
+	GENERATED_BODY()
+	
+	UPROPERTY(EditAnywhere, Category = RealTimeMorphTargets)
+	FName SkeletalMeshName;
+
+	UPROPERTY(EditAnywhere, Category = RealTimeMorphTargets)
+	ECustomizableObjectSelectionOverride SelectionOverride = ECustomizableObjectSelectionOverride::NoOverride;
+};
+
+USTRUCT()
+struct FRealTimeMorphSelectionOverride
+{
+	GENERATED_BODY()
+
+	UPROPERTY(EditAnywhere, Category = RealTimeMorphTargets)
+	FName MorphName;
+
+	UPROPERTY(EditAnywhere, Category = RealTimeMorphTargets)
+	ECustomizableObjectSelectionOverride SelectionOverride = ECustomizableObjectSelectionOverride::NoOverride;
+
+	UPROPERTY(EditAnywhere, Category = RealTimeMorphTargets, meta=(TitleProperty=SkeletalMeshName))
+	TArray<FSkeletalMeshMorphTargetOverride> SkeletalMeshes;
+
+    UPROPERTY()
+    TArray<FName> SkeletalMeshesNames_DEPRECATED;
+
+    UPROPERTY()
+    TArray<ECustomizableObjectSelectionOverride> Override_DEPRECATED;
+};
+
+UCLASS()
+class CUSTOMIZABLEOBJECTEDITOR_API UCustomizableObjectNodeObject : public UCustomizableObjectNode
+{
+public:
+	GENERATED_BODY()
+
+	UCustomizableObjectNodeObject();
+
+	UPROPERTY(EditAnywhere, Category=CustomizableObject)
+	FString ObjectName;
+
+	UPROPERTY(EditAnywhere, Category = UI, meta = (DisplayName = "Parameter UI Metadata"))
+	FMutableParamUIMetadata ParamUIMetadata;
+
+	UPROPERTY()
+	int32 NumLODs_DEPRECATED = 1;
+
+	UPROPERTY()
+	ECustomizableObjectAutomaticLODStrategy AutoLODStrategy_DEPRECATED = ECustomizableObjectAutomaticLODStrategy::AutomaticFromMesh;
+
+	UPROPERTY()
+	int32 NumMeshComponents_DEPRECATED = 1;
+
+	UPROPERTY(EditAnywhere, Category=CustomizableObject ,meta = (TitleProperty = "Name"))
+	TArray<FCustomizableObjectState> States;
+
+	UPROPERTY(EditAnywhere, Category = AttachedToExternalObject)
+	TObjectPtr<UCustomizableObject> ParentObject;
+
+	UPROPERTY(EditAnywhere, Category = AttachedToExternalObject)
+	FGuid ParentObjectGroupId;
+
+	UPROPERTY()
+	FGuid Identifier;
+
+	UPROPERTY(EditAnywhere, Category = RealTimeMorphTargets, meta=(TitleProperty=MorphName, EditCondition="ParentObject==nullptr"))
+	TArray<FRealTimeMorphSelectionOverride> RealTimeMorphSelectionOverrides;
+
+	UPROPERTY()
+	TArray<FComponentSettings> ComponentSettings_DEPRECATED;
+		
+	// UObject interface.
+	virtual void PostEditChangeProperty(FPropertyChangedEvent& PropertyChangedEvent) override;
+
+	// EdGraphNode interface
+	FText GetNodeTitle(ENodeTitleType::Type TitleType) const override;
+	FLinearColor GetNodeTitleColor() const override;
+	FText GetTooltipText() const override;
+	void PrepareForCopying() override;
+	void PostPasteNode() override;
+	void PostDuplicate(bool bDuplicateForPIE) override;
+	virtual bool GetCanRenameNode() const override { return true; }
+	virtual void OnRenameNode(const FString& NewName) override;
+
+	// UCustomizableObjectNode interface
+	virtual void BackwardsCompatibleFixup(int32 CustomizableObjectCustomVersion) override;
+	virtual void PostBackwardsCompatibleFixup() override;
+	void AllocateDefaultPins(UCustomizableObjectNodeRemapPins* RemapPins) override;
+
+	// Own Interface
+	void SetParentObject(UCustomizableObject* CustomizableParentObject);
+
+	// Own interface
+	UPROPERTY()
+	bool bIsBase;
+	
+	UEdGraphPin* ComponentsPin() const
+	{
+		return FindPin(ComponentsPinName);
+	}
+
+	UEdGraphPin* ModifiersPin() const
+	{
+		return FindPin(ModifiersPinName);
+	}
+
+	UEdGraphPin* ChildrenPin() const
+	{
+		return FindPin(ChildrenPinName);
+	}
+
+	UEdGraphPin* OutputPin() const
+	{
+		return FindPin(OutputPinName);
+	}
+
+	
+	virtual bool CanUserDeleteNode() const override;
+	virtual bool CanDuplicateNode() const override;
+	
+	bool IsSingleOutputNode() const override;
+	
+	// Array filled in the Details of the node to store all the parameter names of a CO graph (full tree)
+	TArray<FString> ParameterNames;
+
+private:
+	static const FName ChildrenPinName;
+	static const FName ComponentsPinName;
+	static const FName ModifiersPinName;
+	static const FName OutputPinName;
+};
+
diff --git a/Engine/Plugins/Mutable/Source/CustomizableObjectEditor/Private/MuCOE/Nodes/CustomizableObjectNodeObjectGroup.h b/Engine/Plugins/Mutable/Source/CustomizableObjectEditor/Private/MuCOE/Nodes/CustomizableObjectNodeObjectGroup.h
new file mode 100644
--- /dev/null
+++ b/Engine/Plugins/Mutable/Source/CustomizableObjectEditor/Private/MuCOE/Nodes/CustomizableObjectNodeObjectGroup.h
@@ -0,0 +1,69 @@
+// Copyright Epic Games, Inc. All Rights Reserved.
+
+#pragma once
+
+#include "MuCOE/Nodes/CustomizableObjectNode.h"
+
+#include "CustomizableObjectNodeObjectGroup.generated.h"
+
+enum class ECustomizableObjectGroupType : uint8;
+namespace ENodeTitleType { enum Type : int; }
+
+class FArchive;
+class UCustomizableObjectNodeRemapPins;
+class UEdGraphPin;
+class UObject;
+struct FPropertyChangedEvent;
+
+
+UCLASS()
+class CUSTOMIZABLEOBJECTEDITOR_API UCustomizableObjectNodeObjectGroup : public UCustomizableObjectNode
+{
+public:
+	GENERATED_BODY()
+
+	UCustomizableObjectNodeObjectGroup();
+
+	UPROPERTY(EditAnywhere, Category=GroupInfo, meta = (DisplayName= "Name"))
+	FString GroupName;
+
+	UPROPERTY(EditAnywhere, Category=GroupInfo, meta = (DisplayName = "Type"))
+	ECustomizableObjectGroupType GroupType;
+	
+	UPROPERTY(EditAnywhere, Category=GroupInfo)
+	FString DefaultValue;
+
+	UPROPERTY(EditAnywhere, Category = UI, meta = (DisplayName = "Parameter UI Metadata"))
+	FMutableParamUIMetadata ParamUIMetadata;
+
+	// The sockets defined in meshes deriving from this node will inherit this socket priority. When in the generated merged mesh there
+	// are clashes with socket names, the one with higher priority will be kept and the other discarded.
+	UPROPERTY(EditAnywhere, Category = MeshSockets)
+	int32 SocketPriority = 0;
+
+	// UObject interface.
+	void Serialize(FArchive& Ar) override;
+	void PostEditChangeProperty(FPropertyChangedEvent& PropertyChangedEvent) override;
+
+	// EdGraphNode interface
+	FText GetNodeTitle(ENodeTitleType::Type TitleType) const override;
+	FLinearColor GetNodeTitleColor() const override;
+	FText GetTooltipText() const override;
+	virtual bool GetCanRenameNode() const override { return true; }
+	virtual void OnRenameNode(const FString& NewName) override;
+
+	// UCustomizableObjectNode interface
+	void AllocateDefaultPins(UCustomizableObjectNodeRemapPins* RemapPins) override;
+	bool IsSingleOutputNode() const override;
+
+	// Own interface
+	UEdGraphPin* ObjectsPin() const;
+
+	UEdGraphPin* GroupProjectorsPin() const;
+
+	UEdGraphPin* GroupPin() const;
+
+private:
+	FString LastGroupName;
+};
+
diff --git a/Engine/Plugins/Mutable/Source/CustomizableObjectEditor/Private/MuCOE/Nodes/CustomizableObjectNodeProjectorConstant.h b/Engine/Plugins/Mutable/Source/CustomizableObjectEditor/Private/MuCOE/Nodes/CustomizableObjectNodeProjectorConstant.h
new file mode 100644
--- /dev/null
+++ b/Engine/Plugins/Mutable/Source/CustomizableObjectEditor/Private/MuCOE/Nodes/CustomizableObjectNodeProjectorConstant.h
@@ -0,0 +1,95 @@
+// Copyright Epic Games, Inc. All Rights Reserved.
+
+#pragma once
+
+#include "MuCOE/Nodes/CustomizableObjectNode.h"
+
+#include "CustomizableObjectNodeProjectorConstant.generated.h"
+
+namespace ENodeTitleType { enum Type : int; }
+
+class UCustomizableObjectNodeRemapPins;
+class UEdGraphPin;
+class UObject;
+struct FPropertyChangedEvent;
+
+
+UCLASS()
+class CUSTOMIZABLEOBJECTEDITOR_API UCustomizableObjectNodeProjectorConstant : public UCustomizableObjectNode
+{
+public:
+	GENERATED_BODY()
+
+	UCustomizableObjectNodeProjectorConstant();
+
+	/**  */
+	UPROPERTY(EditAnywhere, Category = CustomizableObject, meta = (ShowOnlyInnerProperties))
+	FCustomizableObjectProjector Value;
+
+	UPROPERTY(EditAnywhere, Category = UI, meta = (DisplayName = "Parameter UI Metadata"))
+	FMutableParamUIMetadata ParamUIMetadata;
+
+	UPROPERTY(EditAnywhere, Category = CustomizableObject, Meta = (DisplayName = "Projection Angle (degrees)"))
+	float ProjectionAngle;
+
+	UPROPERTY()
+	uint32 ReferenceSkeletonIndex_DEPRECATED = 0;
+
+	UPROPERTY(EditAnywhere, Category = ProjectorSnapToBone)
+	FName ReferenceSkeletonComponent;
+
+	UPROPERTY(EditAnywhere, Category = ProjectorSnapToBone)
+	FName ProjectorBone;
+
+	/** Temporary variable where to put the location information for bone combo box selection changes (in FCustomizableObjectNodeProjectorParameterDetails)*/
+	FVector BoneComboBoxLocation;
+
+	/** Temporary variable where to put the direction information for bone combo box selection changes (in FCustomizableObjectNodeProjectorParameterDetails)*/
+	FVector BoneComboBoxForwardDirection;
+
+	/** Temporary variable where to put the direction information for bone combo box selection changes (in FCustomizableObjectNodeProjectorParameterDetails)*/
+	FVector BoneComboBoxUpDirection;
+
+	// UObject interface.
+	void PostEditChangeProperty(FPropertyChangedEvent& PropertyChangedEvent) override;
+
+	// Begin EdGraphNode interface
+	FText GetNodeTitle(ENodeTitleType::Type TitleType) const override;
+	FLinearColor GetNodeTitleColor() const override;
+	FText GetTooltipText() const override;
+
+	// UCustomizableObjectNode interface
+	void AllocateDefaultPins(UCustomizableObjectNodeRemapPins* RemapPins) override;
+	virtual void BackwardsCompatibleFixup(int32 CustomizableObjectCustomVersion) override;
+
+	// Own interface
+	UEdGraphPin* ValuePin();
+
+	ECustomizableObjectProjectorType GetProjectorType() const;
+
+	FVector GetProjectorPosition() const;
+
+	void SetProjectorPosition(const FVector& Position);
+
+	FVector GetProjectorDirection() const;
+
+	void SetProjectorDirection(const FVector& Direction);
+
+	FVector GetProjectorUp() const;
+
+	void SetProjectorUp(const FVector& Up);
+
+	FVector GetProjectorScale() const;
+
+	void SetProjectorScale(const FVector& Scale);
+
+	float GetProjectorAngle();
+
+	void SetProjectorAngle(float Angle);
+
+private:
+
+	UPROPERTY()
+	ECustomizableObjectProjectorType ProjectionType_DEPRECATED;
+};
+
diff --git a/Engine/Plugins/Mutable/Source/CustomizableObjectEditor/Private/MuCOE/Nodes/CustomizableObjectNodeProjectorParameter.h b/Engine/Plugins/Mutable/Source/CustomizableObjectEditor/Private/MuCOE/Nodes/CustomizableObjectNodeProjectorParameter.h
new file mode 100644
--- /dev/null
+++ b/Engine/Plugins/Mutable/Source/CustomizableObjectEditor/Private/MuCOE/Nodes/CustomizableObjectNodeProjectorParameter.h
@@ -0,0 +1,101 @@
+// Copyright Epic Games, Inc. All Rights Reserved.
+
+#pragma once
+
+#include "MuCOE/Nodes/CustomizableObjectNode.h"
+
+#include "CustomizableObjectNodeProjectorParameter.generated.h"
+
+namespace ENodeTitleType { enum Type : int; }
+
+class UCustomizableObjectNodeRemapPins;
+class UEdGraphPin;
+class UObject;
+struct FPropertyChangedEvent;
+
+
+UCLASS()
+class CUSTOMIZABLEOBJECTEDITOR_API UCustomizableObjectNodeProjectorParameter : public UCustomizableObjectNode
+{
+public:
+	GENERATED_BODY()
+
+	UPROPERTY(EditAnywhere, Category=CustomizableObject, meta = (ShowOnlyInnerProperties))
+	FCustomizableObjectProjector DefaultValue;
+
+	UPROPERTY(EditAnywhere, Category=CustomizableObject)
+	FString ParameterName = "Default Name";
+
+	UPROPERTY(EditAnywhere, Category = UI, meta = (DisplayName = "Parameter UI Metadata"))
+	FMutableParamUIMetadata ParamUIMetadata;
+
+	UPROPERTY(EditAnywhere, Category = CustomizableObject, Meta = (DisplayName = "Projection Angle (degrees)"))
+	float ProjectionAngle = 360.0f;
+
+	UPROPERTY()
+	uint32 ReferenceSkeletonIndex_DEPRECATED = 0;
+
+	UPROPERTY(EditAnywhere, Category = ProjectorSnapToBone)
+	FName ReferenceSkeletonComponent;
+
+	UPROPERTY(EditAnywhere, Category = ProjectorSnapToBone)
+	FName ProjectorBone;
+
+	/** Temporary variable where to put the location information for bone combo box selection changes (in FCustomizableObjectNodeProjectorParameterDetails)*/
+	FVector BoneComboBoxLocation = FVector::ZeroVector;
+
+	/** Temporary variable where to put the direction information for bone combo box selection changes (in FCustomizableObjectNodeProjectorParameterDetails)*/
+	FVector BoneComboBoxForwardDirection = FVector::ZeroVector;
+
+	/** Temporary variable where to put the direction information for bone combo box selection changes (in FCustomizableObjectNodeProjectorParameterDetails)*/
+	FVector BoneComboBoxUpDirection = FVector::ZeroVector;
+
+	// UObject interface.
+	void PostEditChangeProperty(FPropertyChangedEvent& PropertyChangedEvent) override;
+
+	// Begin EdGraphNode interface
+	FText GetNodeTitle(ENodeTitleType::Type TitleType) const override;
+	FLinearColor GetNodeTitleColor() const override;
+	FText GetTooltipText() const override;
+	virtual bool GetCanRenameNode() const override { return true; }
+	virtual void OnRenameNode(const FString& NewName) override;
+
+	// UCustomizableObjectNode interface
+	void AllocateDefaultPins(UCustomizableObjectNodeRemapPins* RemapPins) override;
+	virtual bool IsAffectedByLOD() const override { return false; }
+	virtual void BackwardsCompatibleFixup(int32 CustomizableObjectCustomVersion) override;
+
+	// Own interface
+	UEdGraphPin* ProjectorPin() const
+	{
+		return FindPin(TEXT("Value"));
+	}
+
+	ECustomizableObjectProjectorType GetProjectorType() const;
+
+	FVector GetProjectorDefaultPosition() const;
+	
+	void SetProjectorDefaultPosition(const FVector& Position);
+
+	FVector GetProjectorDefaultDirection() const;
+
+	void SetProjectorDefaultDirection(const FVector& Direction);
+
+	FVector GetProjectorDefaultUp() const;
+
+	void SetProjectorDefaultUp(const FVector& Up);
+
+	FVector GetProjectorDefaultScale() const;
+
+	void SetProjectorDefaultScale(const FVector& Scale);
+
+	float GetProjectorDefaultAngle() const;
+
+	void SetProjectorDefaultAngle(float Angle);
+
+private:
+
+	UPROPERTY()
+	ECustomizableObjectProjectorType ProjectionType_DEPRECATED;
+};
+
diff --git a/Engine/Plugins/Mutable/Source/CustomizableObjectEditor/Private/MuCOE/Nodes/CustomizableObjectNodeStaticMesh.h b/Engine/Plugins/Mutable/Source/CustomizableObjectEditor/Private/MuCOE/Nodes/CustomizableObjectNodeStaticMesh.h
new file mode 100644
--- /dev/null
+++ b/Engine/Plugins/Mutable/Source/CustomizableObjectEditor/Private/MuCOE/Nodes/CustomizableObjectNodeStaticMesh.h
@@ -0,0 +1,171 @@
+// Copyright Epic Games, Inc. All Rights Reserved.
+
+#pragma once
+
+#include "MuCOE/Nodes/CustomizableObjectNodeMesh.h"
+#include "MuCOE/Nodes/SCustomizableObjectNode.h"
+
+#include "CustomizableObjectNodeStaticMesh.generated.h"
+
+namespace ENodeTitleType { enum Type : int; }
+
+class FArchive;
+class FAssetThumbnail;
+class FAssetThumbnailPool;
+class SOverlay;
+class SVerticalBox;
+class UCustomizableObjectLayout;
+class UCustomizableObjectNodeRemapPins;
+class UCustomizableObjectNodeRemapPinsByName;
+class UMaterialInterface;
+class UObject;
+class UStaticMesh;
+class UTexture2D;
+struct FPropertyChangedEvent;
+struct FSlateBrush;
+
+
+// Class to render the Static Mesh thumbnail of a CustomizableObjectNodeStaticMesh
+class SGraphNodeStaticMesh : public SCustomizableObjectNode
+{
+public:
+
+	SLATE_BEGIN_ARGS(SGraphNodeStaticMesh) {}
+	SLATE_END_ARGS();
+
+	SGraphNodeStaticMesh() : SCustomizableObjectNode() {};
+
+	// Builds the SGraphNodeStaticMesh when needed
+	void Construct(const FArguments& InArgs, UEdGraphNode* InGraphNode);
+
+	// Calls the needed functions to build the SGraphNode widgets
+	void UpdateGraphNode();
+
+	// Overriden functions to build the SGraphNode widgets
+	virtual void SetDefaultTitleAreaWidget(TSharedRef<SOverlay> DefaultTitleAreaWidget) override;
+	virtual void CreateBelowPinControls(TSharedPtr<SVerticalBox> MainBox) override;
+	virtual bool ShouldAllowCulling() const override { return false; }
+
+	// Callbacks for the widget
+	void OnExpressionPreviewChanged(const ECheckBoxState NewCheckedState);
+	ECheckBoxState IsExpressionPreviewChecked() const;
+	const FSlateBrush* GetExpressionPreviewArrow() const;
+	EVisibility ExpressionPreviewVisibility() const;
+
+public:
+
+	// Single property that only draws the combo box widget of the static mesh
+	TSharedPtr<class ISinglePropertyView> StaticMeshSelector;
+
+	// Pointer to the NodeStaticMesh that owns this SGraphNode
+	class UCustomizableObjectNodeStaticMesh* NodeStaticMesh = nullptr;
+
+private:
+
+	// Classes needed to get and render the thumbnail of the Static Mesh
+	TSharedPtr<FAssetThumbnailPool> AssetThumbnailPool;
+	TSharedPtr<FAssetThumbnail> AssetThumbnail;
+
+	//This parameter defines the size of the thumbnail widget inside the Node
+	float WidgetSize;
+
+	// This parameter defines the resolution of the thumbnail
+	uint32 ThumbnailSize;
+
+};
+
+
+USTRUCT()
+struct CUSTOMIZABLEOBJECTEDITOR_API FCustomizableObjectNodeStaticMeshMaterial
+{
+	GENERATED_USTRUCT_BODY()
+
+	UPROPERTY()
+	FString Name;
+
+	UPROPERTY()
+	TObjectPtr<UEdGraphPin_Deprecated> MeshPin_DEPRECATED;
+
+	UPROPERTY()
+	TObjectPtr<UEdGraphPin_Deprecated> LayoutPin_DEPRECATED;
+
+	UPROPERTY()
+	TArray< TObjectPtr<UEdGraphPin_Deprecated>> ImagePins_DEPRECATED;
+
+	UPROPERTY()
+	FEdGraphPinReference MeshPinRef;
+
+	UPROPERTY()
+	FEdGraphPinReference LayoutPinRef;
+
+	UPROPERTY()
+	TArray<FEdGraphPinReference> ImagePinsRef;
+};
+
+
+USTRUCT()
+struct CUSTOMIZABLEOBJECTEDITOR_API FCustomizableObjectNodeStaticMeshLOD
+{
+	GENERATED_USTRUCT_BODY()
+
+	UPROPERTY()
+	TArray<FCustomizableObjectNodeStaticMeshMaterial> Materials;
+};
+
+
+UCLASS()
+class CUSTOMIZABLEOBJECTEDITOR_API UCustomizableObjectNodeStaticMesh : public UCustomizableObjectNodeMesh
+{ 
+public:
+	GENERATED_BODY()
+
+	UPROPERTY(EditAnywhere, Category = CustomizableObject)
+	TSoftObjectPtr<UStaticMesh> StaticMesh;
+
+	/** Images */
+	UPROPERTY()
+	TArray<FCustomizableObjectNodeStaticMeshLOD> LODs;
+
+	/** Default pin when there is no mesh. */
+	UPROPERTY()
+	FEdGraphPinReference DefaultPin;
+	
+	// UObject interface.
+	virtual void PostLoad() override;
+	virtual void PostEditChangeProperty(FPropertyChangedEvent& PropertyChangedEvent) override;
+	virtual void Serialize(FArchive& Ar) override;
+
+	// UEdGraphNode interface
+	virtual FText GetNodeTitle(ENodeTitleType::Type TitleType) const override;
+	virtual FText GetTooltipText() const override;
+
+	// UCustomizableObjectNode interface
+	virtual void AllocateDefaultPins(UCustomizableObjectNodeRemapPins* RemapPins) override;
+	virtual bool ProvidesCustomPinRelevancyTest() const override { return true; }
+	virtual bool IsPinRelevant(const UEdGraphPin* Pin) const override;
+	virtual UCustomizableObjectNodeRemapPinsByName* CreateRemapPinsByName() const override;
+	virtual bool HasPinViewer() const override;
+	virtual bool IsNodeOutDatedAndNeedsRefresh() override;
+	virtual FString GetRefreshMessage() const override;
+
+	// UCustomizableObjectNodeMesh interface
+	virtual UTexture2D* FindTextureForPin(const UEdGraphPin* Pin) const override;
+	virtual TArray<UCustomizableObjectLayout*> GetLayouts(const UEdGraphPin& OutPin) const override;
+	virtual TSoftObjectPtr<UObject> GetMesh() const override;
+	virtual UEdGraphPin* GetMeshPin(int32 LOD, int32 SectionIndex) const override;
+	virtual void GetPinSection(const UEdGraphPin& Pin, int32& OutLODIndex, int32& OutSectionIndex, int32& OutLayoutIndex) const override;
+
+	/** Returns the material associated to the given output pin. */
+	UMaterialInterface* GetMaterialFor(const UEdGraphPin* Pin) const;
+
+	// Creates the SGraph Node widget for the thumbnail
+	virtual TSharedPtr<SGraphNode> CreateVisualWidget() override;
+
+	// Determines if the Node is collapsed or not
+	bool bCollapsed = true;
+
+	// Pointer to the SGraphNodeStaticMesh
+	TWeakPtr< SGraphNodeStaticMesh > GraphNodeStaticMesh;
+
+};
+
diff --git a/Engine/Plugins/Mutable/Source/CustomizableObjectEditor/Private/MuCOE/Nodes/CustomizableObjectNodeSwitchBase.h b/Engine/Plugins/Mutable/Source/CustomizableObjectEditor/Private/MuCOE/Nodes/CustomizableObjectNodeSwitchBase.h
new file mode 100644
--- /dev/null
+++ b/Engine/Plugins/Mutable/Source/CustomizableObjectEditor/Private/MuCOE/Nodes/CustomizableObjectNodeSwitchBase.h
@@ -0,0 +1,89 @@
+// Copyright Epic Games, Inc. All Rights Reserved.
+
+#pragma once
+
+#include "MuCOE/Nodes/CustomizableObjectNode.h"
+
+#include "CustomizableObjectNodeSwitchBase.generated.h"
+
+namespace ENodeTitleType { enum Type : int; }
+
+class FArchive;
+class UCustomizableObjectNodeRemapPins;
+class UObject;
+struct FFrame;
+
+
+UCLASS(Abstract)
+class CUSTOMIZABLEOBJECTEDITOR_API UCustomizableObjectNodeSwitchBase : public UCustomizableObjectNode
+{
+public:
+	GENERATED_BODY()
+
+	// UObject interface
+	virtual void Serialize(FArchive& Ar) override;
+	
+	// Begin EdGraphNode interface
+	virtual FText GetNodeTitle(ENodeTitleType::Type TitleType) const override;
+	virtual FLinearColor GetNodeTitleColor() const override;
+	virtual FText GetTooltipText() const override;
+	virtual void PinConnectionListChanged(UEdGraphPin* Pin) override;
+	virtual void PostPasteNode() override;
+
+	// UCustomizableObjectNode interface
+	virtual void PostBackwardsCompatibleFixup() override;
+	virtual void AllocateDefaultPins(UCustomizableObjectNodeRemapPins* RemapPins) override;
+	virtual void ReconstructNode(UCustomizableObjectNodeRemapPins* RemapPins) override;
+
+	/** Get the output pin category. Override. */
+	virtual FName GetCategory() const PURE_VIRTUAL(UCustomizableObjectNodeSwitchBase::GetCategory, return {}; );
+
+	UEdGraphPin* OutputPin() const;
+
+	UEdGraphPin* SwitchParameter() const;
+
+	UEdGraphPin* GetElementPin(int32 Index) const
+	{
+		return FindPin(GetPinPrefix(Index));
+	}
+
+	int32 GetNumElements() const;
+
+	/** Links the PostEditChangeProperty delegate */
+	void LinkPostEditChangePropertyDelegate(const UEdGraphPin& Pin);
+
+	/** Get the ouput pin name. Override. */
+	virtual FString GetOutputPinName() const;
+
+private:
+	/** Get the pin prefix. Used for retrocompatibility. Override. */
+	virtual FString GetPinPrefix() const;
+
+protected:
+	UPROPERTY()
+	FEdGraphPinReference OutputPinReference;
+
+	UPROPERTY()
+	TArray<FString> ReloadingElementsNames;
+
+private:
+	/** Get the pin prefix with index. Used for retrocompatibility.*/
+	FString GetPinPrefix(int32 Index) const;
+
+	void ReloadEnumParam();
+
+	/** Last NodeEnumParameter connected. Used to remove the callback once disconnected. */
+	TWeakObjectPtr<UCustomizableObjectNode> LastNodeEnumParameterConnected;
+
+	/** NodeEnumParameter property changed callback function. Reconstructs the node. */
+	UFUNCTION()
+	void EnumParameterPostEditChangeProperty(FPostEditChangePropertyDelegateParameters& Parameters);
+
+	/** The node has to be reconstructed. PinConnectionListChanged(...) can not reconstruct the node, flag used to reconstruct the node on NodeConnectionListChanged(). */
+	bool bMarkReconstruct = false;
+
+
+	UPROPERTY()
+	FEdGraphPinReference SwitchParameterPinReference;
+};
+
diff --git a/Engine/Plugins/Mutable/Source/CustomizableObjectEditor/Private/MuCOE/Nodes/CustomizableObjectNodeTextureBinarise.h b/Engine/Plugins/Mutable/Source/CustomizableObjectEditor/Private/MuCOE/Nodes/CustomizableObjectNodeTextureBinarise.h
new file mode 100644
--- /dev/null
+++ b/Engine/Plugins/Mutable/Source/CustomizableObjectEditor/Private/MuCOE/Nodes/CustomizableObjectNodeTextureBinarise.h
@@ -0,0 +1,40 @@
+// Copyright Epic Games, Inc. All Rights Reserved.
+#pragma once
+
+#include "MuCOE/Nodes/CustomizableObjectNode.h"
+
+#include "CustomizableObjectNodeTextureBinarise.generated.h"
+
+namespace ENodeTitleType { enum Type : int; }
+
+class FArchive;
+class UCustomizableObjectNodeRemapPins;
+class UObject;
+
+UCLASS()
+class CUSTOMIZABLEOBJECTEDITOR_API UCustomizableObjectNodeTextureBinarise : public UCustomizableObjectNode
+{
+public:
+	GENERATED_BODY()
+
+	// Begin EdGraphNode interface
+	FText GetNodeTitle(ENodeTitleType::Type TittleType) const override;
+	FLinearColor GetNodeTitleColor() const override;
+	FText GetTooltipText() const override;
+	virtual void Serialize(FArchive& Ar) override;
+	
+	// UCustomizableObjectNode interface
+	virtual void AllocateDefaultPins(UCustomizableObjectNodeRemapPins* RemapPins) override;
+	virtual void BackwardsCompatibleFixup(int32 CustomizableObjectCustomVersion) override;
+
+	UEdGraphPin* GetBaseImagePin() const;
+
+	UEdGraphPin* GetThresholdPin() const
+	{
+		return FindPin(TEXT("Threshold"));
+	}
+
+private:
+	UPROPERTY()
+	FEdGraphPinReference BaseImagePinReference;
+};
diff --git a/Engine/Plugins/Mutable/Source/CustomizableObjectEditor/Private/MuCOE/Nodes/CustomizableObjectNodeTextureColourMap.h b/Engine/Plugins/Mutable/Source/CustomizableObjectEditor/Private/MuCOE/Nodes/CustomizableObjectNodeTextureColourMap.h
new file mode 100644
--- /dev/null
+++ b/Engine/Plugins/Mutable/Source/CustomizableObjectEditor/Private/MuCOE/Nodes/CustomizableObjectNodeTextureColourMap.h
@@ -0,0 +1,46 @@
+// Copyright Epic Games, Inc. All Rights Reserved.
+
+#pragma once
+
+#include "MuCOE/Nodes/CustomizableObjectNode.h"
+
+#include "CustomizableObjectNodeTextureColourMap.generated.h"
+
+namespace ENodeTitleType { enum Type : int; }
+
+class UCustomizableObjectNodeRemapPins;
+class UEdGraphPin;
+class UObject;
+
+
+UCLASS()
+class CUSTOMIZABLEOBJECTEDITOR_API UCustomizableObjectNodeTextureColourMap : public UCustomizableObjectNode
+{
+public:
+	GENERATED_BODY()
+
+	// Begin EdGraphNode interface
+	FText GetNodeTitle(ENodeTitleType::Type TitleType) const override;
+	FLinearColor GetNodeTitleColor() const override;
+	FText GetTooltipText() const override;
+
+	// UCustomizableObjectNode interface
+	virtual void AllocateDefaultPins(UCustomizableObjectNodeRemapPins* RemapPins) override;
+	virtual void BackwardsCompatibleFixup(int32 CustomizableObjectCustomVersion) override;
+
+	UEdGraphPin* GetMapPin() const
+	{
+		return FindPin(TEXT("Map"));
+	}
+
+	UEdGraphPin* GetBasePin() const
+	{
+		return FindPin(TEXT("Base"));
+	}
+
+	UEdGraphPin* GetMaskPin() const
+	{
+		return FindPin(TEXT("Mask"));
+	}
+
+};
diff --git a/Engine/Plugins/Mutable/Source/CustomizableObjectEditor/Private/MuCOE/Nodes/CustomizableObjectNodeTextureFromChannels.h b/Engine/Plugins/Mutable/Source/CustomizableObjectEditor/Private/MuCOE/Nodes/CustomizableObjectNodeTextureFromChannels.h
new file mode 100644
--- /dev/null
+++ b/Engine/Plugins/Mutable/Source/CustomizableObjectEditor/Private/MuCOE/Nodes/CustomizableObjectNodeTextureFromChannels.h
@@ -0,0 +1,51 @@
+// Copyright Epic Games, Inc. All Rights Reserved.
+
+#pragma once
+
+#include "MuCOE/Nodes/CustomizableObjectNode.h"
+
+#include "CustomizableObjectNodeTextureFromChannels.generated.h"
+
+namespace ENodeTitleType { enum Type : int; }
+
+class UCustomizableObjectNodeRemapPins;
+class UEdGraphPin;
+class UObject;
+
+
+UCLASS()
+class CUSTOMIZABLEOBJECTEDITOR_API UCustomizableObjectNodeTextureFromChannels : public UCustomizableObjectNode
+{
+public:
+	GENERATED_BODY()
+
+	// Begin EdGraphNode interface
+	FText GetNodeTitle(ENodeTitleType::Type TitleType) const override;
+	FLinearColor GetNodeTitleColor() const override;
+	FText GetTooltipText() const override;
+
+	// UCustomizableObjectNode interface
+	virtual void AllocateDefaultPins(UCustomizableObjectNodeRemapPins* RemapPins) override;
+	virtual void BackwardsCompatibleFixup(int32 CustomizableObjectCustomVersion) override;
+
+	UEdGraphPin* RPin() const
+	{
+		return FindPin(TEXT("R"));
+	}
+
+	UEdGraphPin* GPin() const
+	{
+		return FindPin(TEXT("G"));
+	}
+
+	UEdGraphPin* BPin() const
+	{
+		return FindPin(TEXT("B"));
+	}
+
+	UEdGraphPin* APin() const
+	{
+		return FindPin(TEXT("A"));
+	}
+};
+
diff --git a/Engine/Plugins/Mutable/Source/CustomizableObjectEditor/Private/MuCOE/Nodes/CustomizableObjectNodeTextureFromColor.h b/Engine/Plugins/Mutable/Source/CustomizableObjectEditor/Private/MuCOE/Nodes/CustomizableObjectNodeTextureFromColor.h
new file mode 100644
--- /dev/null
+++ b/Engine/Plugins/Mutable/Source/CustomizableObjectEditor/Private/MuCOE/Nodes/CustomizableObjectNodeTextureFromColor.h
@@ -0,0 +1,36 @@
+// Copyright Epic Games, Inc. All Rights Reserved.
+
+#pragma once
+
+#include "MuCOE/Nodes/CustomizableObjectNode.h"
+
+#include "CustomizableObjectNodeTextureFromColor.generated.h"
+
+namespace ENodeTitleType { enum Type : int; }
+
+class UCustomizableObjectNodeRemapPins;
+class UEdGraphPin;
+class UObject;
+
+
+UCLASS()
+class CUSTOMIZABLEOBJECTEDITOR_API UCustomizableObjectNodeTextureFromColor : public UCustomizableObjectNode
+{
+public:
+	GENERATED_BODY()
+
+	// Begin EdGraphNode interface
+	FText GetNodeTitle(ENodeTitleType::Type TitleType) const override;
+	FLinearColor GetNodeTitleColor() const override;
+	FText GetTooltipText() const override;
+
+	// UCustomizableObjectNode interface
+	virtual void AllocateDefaultPins(UCustomizableObjectNodeRemapPins* RemapPins) override;
+	virtual void BackwardsCompatibleFixup(int32 CustomizableObjectCustomVersion) override;
+
+	UEdGraphPin* ColorPin() const
+	{
+		return FindPin(TEXT("Color"));
+	}
+};
+
diff --git a/Engine/Plugins/Mutable/Source/CustomizableObjectEditor/Private/MuCOE/Nodes/CustomizableObjectNodeTextureInterpolate.h b/Engine/Plugins/Mutable/Source/CustomizableObjectEditor/Private/MuCOE/Nodes/CustomizableObjectNodeTextureInterpolate.h
new file mode 100644
--- /dev/null
+++ b/Engine/Plugins/Mutable/Source/CustomizableObjectEditor/Private/MuCOE/Nodes/CustomizableObjectNodeTextureInterpolate.h
@@ -0,0 +1,63 @@
+// Copyright Epic Games, Inc. All Rights Reserved.
+
+#pragma once
+
+#include "MuCOE/Nodes/CustomizableObjectNode.h"
+
+#include "CustomizableObjectNodeTextureInterpolate.generated.h"
+
+namespace ENodeTitleType { enum Type : int; }
+
+class UCustomizableObjectNodeRemapPins;
+class UObject;
+struct FPropertyChangedEvent;
+
+
+UCLASS()
+class CUSTOMIZABLEOBJECTEDITOR_API UCustomizableObjectNodeTextureInterpolate : public UCustomizableObjectNode
+{
+public:
+	GENERATED_BODY()
+
+	UCustomizableObjectNodeTextureInterpolate();
+
+	// UObject interface.
+	void PostEditChangeProperty(FPropertyChangedEvent& PropertyChangedEvent) override;
+
+	// Begin EdGraphNode interface
+	FText GetNodeTitle(ENodeTitleType::Type TitleType) const override;
+	FLinearColor GetNodeTitleColor() const override;
+	FText GetTooltipText() const override;
+
+	// UCustomizableObjectNode interface
+	void AllocateDefaultPins(UCustomizableObjectNodeRemapPins* RemapPins) override;
+
+	UEdGraphPin* FactorPin() const
+	{
+		return FindPin(TEXT("Factor"));
+	}
+
+	UEdGraphPin* Targets(int32 Index) const
+	{
+		return FindPin(FString::Printf(TEXT("Target %d"), Index));
+	}
+
+	int32 GetNumTargets() const
+	{
+		int32 Count = 0;
+
+		for (UEdGraphPin* Pin : GetAllNonOrphanPins())
+		{
+			if (Pin->GetName().StartsWith(TEXT("Target ")))
+			{
+				Count++;
+			}
+		}
+
+		return Count;
+	}
+
+	UPROPERTY(EditAnywhere, Category=CustomizableObject, meta = (ClampMin = 2))
+	int32 NumTargets;
+};
+
