[Mutable] Moved COE headers to private.

diff --git a/Engine/Plugins/Mutable/Source/CustomizableObjectEditor/Private/MuCOE/Nodes/CustomizableObjectNodeTextureInvert.h b/Engine/Plugins/Mutable/Source/CustomizableObjectEditor/Private/MuCOE/Nodes/CustomizableObjectNodeTextureInvert.h
new file mode 100644
--- /dev/null
+++ b/Engine/Plugins/Mutable/Source/CustomizableObjectEditor/Private/MuCOE/Nodes/CustomizableObjectNodeTextureInvert.h
@@ -0,0 +1,35 @@
+// Copyright Epic Games, Inc. All Rights Reserved.
+#pragma once
+
+#include "MuCOE/Nodes/CustomizableObjectNode.h"
+
+#include "CustomizableObjectNodeTextureInvert.generated.h"
+
+namespace ENodeTitleType { enum Type : int; }
+
+class FArchive;
+class UCustomizableObjectNodeRemapPins;
+class UObject;
+
+UCLASS()
+class CUSTOMIZABLEOBJECTEDITOR_API UCustomizableObjectNodeTextureInvert : public UCustomizableObjectNode
+{
+public:
+	GENERATED_BODY()
+	
+	// EdGraphNode interface
+	virtual FText GetNodeTitle(ENodeTitleType::Type TittleType)const override;
+	virtual FLinearColor GetNodeTitleColor() const override;
+	virtual FText GetTooltipText() const override;
+	
+	// UCustomizableObjectNode interface
+	virtual void AllocateDefaultPins(UCustomizableObjectNodeRemapPins* RemapPins) override;
+	virtual void BackwardsCompatibleFixup(int32 CustomizableObjectCustomVersion) override;
+
+	// Own interface
+	UEdGraphPin* GetBaseImagePin() const;
+
+private:
+	UPROPERTY()
+	FEdGraphPinReference BaseImagePinReference;
+};
diff --git a/Engine/Plugins/Mutable/Source/CustomizableObjectEditor/Private/MuCOE/Nodes/CustomizableObjectNodeTextureLayer.h b/Engine/Plugins/Mutable/Source/CustomizableObjectEditor/Private/MuCOE/Nodes/CustomizableObjectNodeTextureLayer.h
new file mode 100644
--- /dev/null
+++ b/Engine/Plugins/Mutable/Source/CustomizableObjectEditor/Private/MuCOE/Nodes/CustomizableObjectNodeTextureLayer.h
@@ -0,0 +1,126 @@
+// Copyright Epic Games, Inc. All Rights Reserved.
+
+#pragma once
+
+#include "MuCOE/Nodes/CustomizableObjectNode.h"
+
+#include "CustomizableObjectNodeTextureLayer.generated.h"
+
+namespace ENodeTitleType { enum Type : int; }
+
+class FArchive;
+class UCustomizableObjectNodeRemapPins;
+class UEdGraph;
+class UObject;
+struct FPropertyChangedEvent;
+
+
+UENUM()
+enum ECustomizableObjectTextureLayerEffect : int
+{
+	COTLE_MODULATE			= 0 UMETA(DisplayName = "MODULATE"), 
+	COTLE_MULTIPLY				UMETA(DisplayName = "MULTIPLY"), 
+	COTLE_SOFTLIGHT				UMETA(DisplayName = "SOFTLIGHT"),
+	COTLE_HARDLIGHT				UMETA(DisplayName = "HARDLIGHT"),
+	COTLE_DODGE					UMETA(DisplayName = "DODGE"),
+	COTLE_BURN					UMETA(DisplayName = "BURN"),
+	COTLE_SCREEN				UMETA(DisplayName = "SCREEN"),
+	COTLE_OVERLAY				UMETA(DisplayName = "OVERLAY"),
+	COTLE_ALPHA_OVERLAY			UMETA(DisplayName = "LIGHTEN"),
+	COTLE_NORMAL_COMBINE		UMETA(DisplayName = "BLEND NORMALS")
+};
+
+
+USTRUCT()
+struct CUSTOMIZABLEOBJECTEDITOR_API FCustomizableObjectTextureLayer
+{
+	GENERATED_USTRUCT_BODY()
+
+	FCustomizableObjectTextureLayer()
+	{
+		Effect = COTLE_SOFTLIGHT;
+	}
+
+	UPROPERTY(EditAnywhere, Category=CustomizableObject)
+	TEnumAsByte<enum ECustomizableObjectTextureLayerEffect> Effect;
+};
+
+
+UCLASS()
+class CUSTOMIZABLEOBJECTEDITOR_API UCustomizableObjectNodeTextureLayer : public UCustomizableObjectNode
+{
+public:
+	GENERATED_BODY()
+
+	UCustomizableObjectNodeTextureLayer();
+
+	/**  */
+	UPROPERTY(EditAnywhere, Category=CustomizableObject)
+	TArray<FCustomizableObjectTextureLayer> Layers;
+
+	// UObject interface.
+	void PostEditChangeProperty(FPropertyChangedEvent& PropertyChangedEvent) override;
+	virtual void Serialize(FArchive& Ar) override;
+	
+	// Begin EdGraphNode interface
+	FText GetNodeTitle(ENodeTitleType::Type TitleType) const override;
+	FLinearColor GetNodeTitleColor() const override;
+	FText GetTooltipText() const override;
+	void PrepareForCopying() override;
+	void PostPasteNode() override;
+
+	// UCustomizableObjectNode interface
+	virtual void AllocateDefaultPins(UCustomizableObjectNodeRemapPins* RemapPins) override;
+	virtual void BackwardsCompatibleFixup(int32 CustomizableObjectCustomVersion) override;
+	virtual bool CanConnect(const UEdGraphPin* InOwnedInputPin, const UEdGraphPin* InOutputPin, bool& bOutIsOtherNodeBlocklisted, bool& bOutArePinsCompatible) const override;
+		
+	UEdGraphPin* BasePin() const
+	{
+		return FindPin(TEXT("Base"));
+	}
+
+	UEdGraphPin* OutputPin() const;
+
+	UEdGraphPin* LayerPin(int32 Index) const
+	{
+		FString PinName = FString::Printf(TEXT("Layer %d "), Index);
+		return FindPin(PinName);
+	}
+
+	UEdGraphPin* MaskPin(int32 Index) const
+	{
+		FString PinName = FString::Printf(TEXT("Mask %d "), Index);
+		return FindPin(PinName);
+	}
+
+	int32 GetNumLayers() const
+	{
+		int32 Count = 0;
+
+		for (UEdGraphPin* Pin : GetAllNonOrphanPins())
+		{
+			if (Pin->GetName().StartsWith(TEXT("Layer ")))
+			{
+				Count++;
+			}
+		}
+
+		return Count;
+	}
+protected:
+	// This struct will store the information of each textue layer node when its destroyed, to be able to recover the UCustomizableObjectNodeTextureLayer::Layers
+	// data in a post paste operation. Otherwise, this information is lost.
+	struct TextureLayerLayersData
+	{
+		FString Name; // Texture layer node name
+		TArray<FCustomizableObjectTextureLayer> Layers;
+		UEdGraph* Graph; // Graph that contained this node
+	};
+	static int32 GetIndexInArrayDestroyedNodes(UCustomizableObjectNodeTextureLayer* Node);
+	static TArray<TextureLayerLayersData> ArrayDestroyedNodes;
+
+private:
+	UPROPERTY()
+	FEdGraphPinReference OutputPinReference;
+};
+
diff --git a/Engine/Plugins/Mutable/Source/CustomizableObjectEditor/Private/MuCOE/Nodes/CustomizableObjectNodeTextureParameter.h b/Engine/Plugins/Mutable/Source/CustomizableObjectEditor/Private/MuCOE/Nodes/CustomizableObjectNodeTextureParameter.h
new file mode 100644
--- /dev/null
+++ b/Engine/Plugins/Mutable/Source/CustomizableObjectEditor/Private/MuCOE/Nodes/CustomizableObjectNodeTextureParameter.h
@@ -0,0 +1,57 @@
+// Copyright Epic Games, Inc. All Rights Reserved.
+
+#pragma once
+
+#include "MuCOE/Nodes/CustomizableObjectNode.h"
+
+#include "CustomizableObjectNodeTextureParameter.generated.h"
+
+namespace ENodeTitleType { enum Type : int; }
+
+class UCustomizableObjectNodeRemapPins;
+class UObject;
+class UTexture2D;
+
+
+UCLASS()
+class CUSTOMIZABLEOBJECTEDITOR_API UCustomizableObjectNodeTextureParameter : public UCustomizableObjectNode
+{
+public:
+	GENERATED_BODY()
+
+	/** Default value of the parameter. */
+	UPROPERTY(EditAnywhere, Category=CustomizableObject)
+	TObjectPtr<UTexture2D> DefaultValue;
+
+	/** Reference Texture where this parameter copies some properties from. */
+	UPROPERTY(EditAnywhere, Category=CustomizableObject)
+	TObjectPtr<UTexture2D> ReferenceValue;
+
+	UPROPERTY(EditAnywhere, Category=CustomizableObject)
+	FString ParameterName = "Default Name";
+
+	UPROPERTY(EditAnywhere, Category = UI, meta = (DisplayName = "Parameter UI Metadata"))
+	FMutableParamUIMetadata ParamUIMetadata;
+
+	/** Set the width of the Texture when there is no texture reference.*/
+	UPROPERTY(EditAnywhere, Category = CustomizableObject)
+	int32 TextureSizeX = 0;
+
+	/** Set the height of the Texture when there is no texture reference.*/
+	UPROPERTY(EditAnywhere, Category = CustomizableObject)
+	int32 TextureSizeY = 0;
+
+	// EdGraphNode interface
+	virtual FText GetNodeTitle(ENodeTitleType::Type TitleType) const override;
+	virtual FLinearColor GetNodeTitleColor() const override;
+	virtual FText GetTooltipText() const override;
+	virtual void OnRenameNode(const FString& NewName) override;
+	virtual bool GetCanRenameNode() const override { return true; }
+
+	// UCustomizableObjectNode interface
+	virtual void BackwardsCompatibleFixup(int32 CustomizableObjectCustomVersion) override;
+	virtual void AllocateDefaultPins(UCustomizableObjectNodeRemapPins* RemapPins) override;
+	virtual bool IsAffectedByLOD() const override { return false; }
+	virtual bool IsExperimental() const override;
+};
+
diff --git a/Engine/Plugins/Mutable/Source/CustomizableObjectEditor/Private/MuCOE/Nodes/CustomizableObjectNodeTextureSwitch.h b/Engine/Plugins/Mutable/Source/CustomizableObjectEditor/Private/MuCOE/Nodes/CustomizableObjectNodeTextureSwitch.h
new file mode 100644
--- /dev/null
+++ b/Engine/Plugins/Mutable/Source/CustomizableObjectEditor/Private/MuCOE/Nodes/CustomizableObjectNodeTextureSwitch.h
@@ -0,0 +1,20 @@
+// Copyright Epic Games, Inc. All Rights Reserved.
+
+#pragma once
+
+#include "MuCOE/Nodes/CustomizableObjectNodeSwitchBase.h"
+
+#include "CustomizableObjectNodeTextureSwitch.generated.h"
+
+
+UCLASS()
+class CUSTOMIZABLEOBJECTEDITOR_API UCustomizableObjectNodeTextureSwitch : public UCustomizableObjectNodeSwitchBase
+{
+public:
+	GENERATED_BODY()
+	// UCustomizableObjectNode interface
+	virtual void BackwardsCompatibleFixup(int32 CustomizableObjectCustomVersion) override;
+	
+	// UCustomizableObjectNodeSwitchBase interface
+	virtual FName GetCategory() const override;
+};
diff --git a/Engine/Plugins/Mutable/Source/CustomizableObjectEditor/Private/MuCOE/Nodes/CustomizableObjectNodeTextureToChannels.h b/Engine/Plugins/Mutable/Source/CustomizableObjectEditor/Private/MuCOE/Nodes/CustomizableObjectNodeTextureToChannels.h
new file mode 100644
--- /dev/null
+++ b/Engine/Plugins/Mutable/Source/CustomizableObjectEditor/Private/MuCOE/Nodes/CustomizableObjectNodeTextureToChannels.h
@@ -0,0 +1,62 @@
+// Copyright Epic Games, Inc. All Rights Reserved.
+
+#pragma once
+
+#include "MuCOE/Nodes/CustomizableObjectNode.h"
+
+#include "CustomizableObjectNodeTextureToChannels.generated.h"
+
+namespace ENodeTitleType { enum Type : int; }
+
+class FArchive;
+class UCustomizableObjectNodeRemapPins;
+class UObject;
+struct FPropertyChangedEvent;
+
+
+UCLASS()
+class CUSTOMIZABLEOBJECTEDITOR_API UCustomizableObjectNodeTextureToChannels : public UCustomizableObjectNode
+{
+public:
+	GENERATED_BODY()
+	
+	// UObject interface.
+	void PostEditChangeProperty(FPropertyChangedEvent& PropertyChangedEvent) override;
+	virtual void Serialize(FArchive& Ar) override;
+	
+	// Begin EdGraphNode interface
+	FText GetNodeTitle(ENodeTitleType::Type TitleType) const override;
+	FLinearColor GetNodeTitleColor() const override;
+	FText GetTooltipText() const override;
+
+	// UCustomizableObjectNode interface
+	virtual void AllocateDefaultPins(UCustomizableObjectNodeRemapPins* RemapPins) override;
+	virtual void BackwardsCompatibleFixup(int32 CustomizableObjectCustomVersion) override;
+
+	UEdGraphPin* InputPin() const;
+
+	UEdGraphPin* RPin() const
+	{
+		return FindPin(TEXT("R"));
+	}
+
+	UEdGraphPin* GPin() const
+	{
+		return FindPin(TEXT("G"));
+	}
+
+	UEdGraphPin* BPin() const
+	{
+		return FindPin(TEXT("B"));
+	}
+
+	UEdGraphPin* APin() const
+	{
+		return FindPin(TEXT("A"));
+	}
+
+private:
+	UPROPERTY()
+	FEdGraphPinReference InputPinReference;
+};
+
diff --git a/Engine/Plugins/Mutable/Source/CustomizableObjectEditor/Public/MuCOE/Nodes/CustomizableObjectNode.h b/Engine/Plugins/Mutable/Source/CustomizableObjectEditor/Public/MuCOE/Nodes/CustomizableObjectNode.h
--- a/Engine/Plugins/Mutable/Source/CustomizableObjectEditor/Public/MuCOE/Nodes/CustomizableObjectNode.h
+++ b/Engine/Plugins/Mutable/Source/CustomizableObjectEditor/Public/MuCOE/Nodes/CustomizableObjectNode.h
@@ -16,6 +16,7 @@ class UCustomizableObjectNodeRemapPins;
 class UCustomizableObjectNodeRemapPinsByName;
 class UCustomizableObjectNodeRemapPinsByPosition;
 class IDetailsView;
+class ICustomizableObjectEditor;
 struct FPropertyChangedEvent;
 
 /** Abstract base class for all pin data. */
@@ -150,7 +151,7 @@ public:
 	// Used to replace Ids of referenced nodes by their new ids after duplicating the CustomizableObject.
 	virtual void UpdateReferencedNodeId(const FGuid& NewGuid) {};
 
-	TSharedPtr<class ICustomizableObjectEditor> GetGraphEditor() const;
+	TSharedPtr<ICustomizableObjectEditor> GetGraphEditor() const;
 
 	virtual bool ProvidesCustomPinRelevancyTest() const { return false; } // Override and return true if IsPinRelevant is overridden.
 	virtual bool IsPinRelevant(const UEdGraphPin*) const { return true; } // Override if default pins allocated with AllocateDefaultPins do not provide enough information.
diff --git a/Engine/Plugins/Mutable/Source/MutableValidation/Private/MuV/AssetValidator_CustomizableObjects.cpp b/Engine/Plugins/Mutable/Source/MutableValidation/Private/MuV/AssetValidator_CustomizableObjects.cpp
--- a/Engine/Plugins/Mutable/Source/MutableValidation/Private/MuV/AssetValidator_CustomizableObjects.cpp
+++ b/Engine/Plugins/Mutable/Source/MutableValidation/Private/MuV/AssetValidator_CustomizableObjects.cpp
@@ -15,7 +15,6 @@
 #include "MuCO/CustomizableObjectCompilerTypes.h"
 #include "MuCO/CustomizableObjectPrivate.h"
 #include "MuCO/ICustomizableObjectEditorModule.h"
-#include "MuCOE/GraphTraversal.h"
 #include "UObject/NameTypes.h"
 #include "UObject/Object.h"
 
@@ -98,9 +97,9 @@ EDataValidationResult UAssetValidator_CustomizableObjects::IsCustomizableObjectV
 	{
 		OnPostCOValidationHandle = FEditorDelegates::OnPostAssetValidation.AddStatic(OnPostCOsValidation);	
 	}
-	
+
 	// Find out which is the root for this CO (it may be itself but that is OK)
-	UCustomizableObject* RootObject = GraphTraversal::GetRootObject(InCustomizableObject);
+	UCustomizableObject* RootObject = ICustomizableObjectEditorModule::GetChecked().GetRootObject(InCustomizableObject);
 	check (RootObject);
 	
 	// Check that the object to be compiled has not already been compiled
