[mutable] Added warning message when a parameter name is found more than once.
- Removed unnecessary includes
- Added missing context for the nodes of type Object and Object Group
- Added load Synchronous to ensure the SKM is loaded before checking for null

#rnx

diff --git a/Engine/Plugins/Mutable/Source/CustomizableObjectEditor/Private/MuCOE/GenerateMutableSource/GenerateMutableSource.cpp b/Engine/Plugins/Mutable/Source/CustomizableObjectEditor/Private/MuCOE/GenerateMutableSource/GenerateMutableSource.cpp
--- a/Engine/Plugins/Mutable/Source/CustomizableObjectEditor/Private/MuCOE/GenerateMutableSource/GenerateMutableSource.cpp
+++ b/Engine/Plugins/Mutable/Source/CustomizableObjectEditor/Private/MuCOE/GenerateMutableSource/GenerateMutableSource.cpp
@@ -888,6 +888,7 @@ mu::Ptr<mu::NodeObject> GenerateMutableSource(const UEdGraphPin * Pin, FMutableG
 		mu::Ptr<mu::NodeObjectNew> ObjectNode = new mu::NodeObjectNew();
 		Result = ObjectNode;
 
+		ObjectNode->SetMessageContext(Node);
 		ObjectNode->SetName(TypedNodeObj->ObjectName);
 		FGuid FinalGuid = GenerationContext.GetNodeIdUnique(TypedNodeObj);
 		if (FinalGuid != TypedNodeObj->NodeGuid)
@@ -1047,6 +1048,7 @@ mu::Ptr<mu::NodeObject> GenerateMutableSource(const UEdGraphPin * Pin, FMutableG
 		// All sockets from all mesh parts plugged into this group node will have the following priority when there's a socket name clash
 		GenerationContext.SocketPriorityStack.Push(TypedNodeGroup->SocketPriority);
 
+		GroupNode->SetMessageContext(Node);
 		GroupNode->SetName(TypedNodeGroup->GroupName);
 		GroupNode->SetUid(TypedNodeGroup->NodeGuid.ToString());
 		
diff --git a/Engine/Plugins/Mutable/Source/CustomizableObjectEditor/Private/MuCOE/GenerateMutableSource/GenerateMutableSourceMesh.cpp b/Engine/Plugins/Mutable/Source/CustomizableObjectEditor/Private/MuCOE/GenerateMutableSource/GenerateMutableSourceMesh.cpp
--- a/Engine/Plugins/Mutable/Source/CustomizableObjectEditor/Private/MuCOE/GenerateMutableSource/GenerateMutableSourceMesh.cpp
+++ b/Engine/Plugins/Mutable/Source/CustomizableObjectEditor/Private/MuCOE/GenerateMutableSource/GenerateMutableSourceMesh.cpp
@@ -3442,6 +3442,7 @@ mu::Ptr<mu::NodeMesh> GenerateMutableSourceMesh(const UEdGraphPin* Pin,
 		mu::NodeMeshConstantPtr MeshNode = new mu::NodeMeshConstant();
 		Result = MeshNode;
 
+		GenerationContext.LoadObject(TypedNodeSkel->SkeletalMesh, true);
 		if (TypedNodeSkel->SkeletalMesh)
 		{
 			int32 LODIndexConnected = -1; // LOD which the pin is connected to
@@ -3692,6 +3693,7 @@ mu::Ptr<mu::NodeMesh> GenerateMutableSourceMesh(const UEdGraphPin* Pin,
 
 	else if (const UCustomizableObjectNodeStaticMesh* TypedNodeStatic = Cast<UCustomizableObjectNodeStaticMesh>(Node))
 	{
+		GenerationContext.LoadObject(TypedNodeStatic->StaticMesh, true);
 		if (TypedNodeStatic->StaticMesh == nullptr)
 		{
 			FString Msg = FString::Printf(TEXT("The UCustomizableObjectNodeStaticMesh node %s has no static mesh assigned"), *Node->GetName());
@@ -3833,7 +3835,7 @@ mu::Ptr<mu::NodeMesh> GenerateMutableSourceMesh(const UEdGraphPin* Pin,
 						continue;
 					}
 
-					// Cheking if the morph exists in the application node
+					// Checking if the morph exists in the application node
 					FString MorphName = MorphPin->PinFriendlyName.ToString();
 					if (!MorphNames.Contains(MorphName))
 					{
diff --git a/Engine/Plugins/Mutable/Source/CustomizableObjectEditor/Private/MuCOE/Nodes/CustomizableObjectNodeSkeletalMesh.cpp b/Engine/Plugins/Mutable/Source/CustomizableObjectEditor/Private/MuCOE/Nodes/CustomizableObjectNodeSkeletalMesh.cpp
--- a/Engine/Plugins/Mutable/Source/CustomizableObjectEditor/Private/MuCOE/Nodes/CustomizableObjectNodeSkeletalMesh.cpp
+++ b/Engine/Plugins/Mutable/Source/CustomizableObjectEditor/Private/MuCOE/Nodes/CustomizableObjectNodeSkeletalMesh.cpp
@@ -16,6 +16,7 @@
 #include "MuCOE/Nodes/CustomizableObjectNodeLayoutBlocks.h"
 #include "MuCOE/RemapPins/CustomizableObjectNodeRemapPinsByNameDefaultPin.h"
 #include "MuCOE/UnrealEditorPortabilityHelpers.h"
+#include "UObject/ICookInfo.h"
 #include "Widgets/Input/SCheckBox.h"
 
 class UCustomizableObjectNodeRemapPinsByName;
@@ -33,7 +34,11 @@ static const TCHAR* SKELETAL_MESH_PIN_NAME = TEXT("Skeletal Mesh");
 void UCustomizableObjectNodeSkeletalMesh::PostLoad()
 {
 	Super::PostLoad();
+	
+	FCookLoadScope CookLoadScope(ECookLoadType::EditorOnly);
 
+	// Block only required by the FixUps : todo: Will be removed soon
+	SkeletalMesh.LoadSynchronous();
 	if (SkeletalMesh)
 	{
 		ConditionalPostLoadReference(*SkeletalMesh);
@@ -438,6 +443,8 @@ FText UCustomizableObjectNodeSkeletalMesh::GetTooltipText() const
 void UCustomizableObjectNodeSkeletalMesh::BackwardsCompatibleFixup(int32 CustomizableObjectCustomVersion)
 {
 	Super::BackwardsCompatibleFixup(CustomizableObjectCustomVersion);
+
+	FCookLoadScope CookLoadScope(ECookLoadType::EditorOnly);
 	
 	if (CustomizableObjectCustomVersion == FCustomizableObjectCustomVersion::PostLoadToCustomVersion)
 	{
@@ -519,6 +526,7 @@ void UCustomizableObjectNodeSkeletalMesh::BackwardsCompatibleFixup(int32 Customi
 					AddPinData(*Section.MeshPinRef.Get(), *PinData);
 				}
 
+				SkeletalMesh.LoadSynchronous();
 				if (SkeletalMesh)
 				{
 					if (const FSkeletalMaterial* SkeletalMaterial = GetSkeletalMaterialFor(LODIndex, SectionIndex))
@@ -595,6 +603,7 @@ void UCustomizableObjectNodeSkeletalMesh::BackwardsCompatibleFixup(int32 Customi
 
 	if (CustomizableObjectCustomVersion == FCustomizableObjectCustomVersion::MoveLayoutToNodeSkeletalMesh)
 	{
+		SkeletalMesh.LoadSynchronous();
 		if (const FSkeletalMeshModel* ImportedModel = SkeletalMesh ? SkeletalMesh->GetImportedModel() : nullptr)
 		{
 			TArray<UEdGraphPin*> PinsToDelete;
diff --git a/Engine/Plugins/Mutable/Source/MutableTools/Private/MuT/Compiler.cpp b/Engine/Plugins/Mutable/Source/MutableTools/Private/MuT/Compiler.cpp
--- a/Engine/Plugins/Mutable/Source/MutableTools/Private/MuT/Compiler.cpp
+++ b/Engine/Plugins/Mutable/Source/MutableTools/Private/MuT/Compiler.cpp
@@ -11,8 +11,6 @@
 #include "MuT/CompilerPrivate.h"
 #include "MuT/ErrorLog.h"
 #include "MuT/Node.h"
-#include "MuT/NodePrivate.h"
-#include "MuT/Table.h"
 #include "MuR/Image.h"
 #include "MuR/Mesh.h"
 #include "MuR/Model.h"
@@ -20,10 +18,8 @@
 #include "MuR/MutableTrace.h"
 #include "MuR/Operations.h"
 #include "MuR/ParametersPrivate.h"
-#include "MuR/Platform.h"
 #include "MuR/Serialisation.h"
 #include "MuR/MutableRuntimeModule.h"
-#include "MuR/System.h"
 
 #include "Trace/Detail/Channel.h"
 #include "Containers/Array.h"
@@ -39,7 +35,6 @@
 
 #include <inttypes.h>	// Required for 64-bit printf macros
 
-
 namespace mu
 {
 
@@ -301,53 +296,81 @@ namespace mu
     {
         MUTABLE_CPUPROFILER_SCOPE(Compile);
 
-        TArray< FStateCompilationData > states;
-        TSharedPtr<FErrorLog> genErrorLog;
+        TArray< FStateCompilationData > States;
+        TSharedPtr<FErrorLog> GenErrorLog;
 		TArray<FParameterDesc> Parameters;
         {
-            CodeGenerator gen( m_pD->Options->GetPrivate() );
+            CodeGenerator Generator( m_pD->Options->GetPrivate() );
 
-            gen.GenerateRoot( pNode );
+            Generator.GenerateRoot( pNode );
 
-            check( !gen.States.IsEmpty() );
+            check( !Generator.States.IsEmpty() );
 
-            for ( const TPair<FObjectState, Ptr<ASTOp>>& s: gen.States )
+            for ( const TPair<FObjectState, Ptr<ASTOp>>& s: Generator.States )
             {
-                FStateCompilationData data;
-                data.nodeState = s.Key;
-                data.root = s.Value;
-                data.state.Name = s.Key.Name;
-                states.Add( data );
+                FStateCompilationData Data;
+                Data.nodeState = s.Key;
+                Data.root = s.Value;
+                Data.state.Name = s.Key.Name;
+                States.Add( Data );
             }
 
-            genErrorLog = gen.ErrorLog;
+            GenErrorLog = Generator.ErrorLog;
 
-			// Set the parameter list from the non-optimized data, so that we have them all even if they are optimized out
-			int32 ParameterCount = gen.FirstPass.ParameterNodes.Num();
-			Parameters.SetNum(ParameterCount);
-			int32 ParameterIndex = 0;
-			for ( const TPair<Ptr<const Node>, Ptr<ASTOpParameter>>& Entry : gen.FirstPass.ParameterNodes )
+        	// Sort the parameters as deterministically as possible.
+        	Generator.FirstPass.ParameterNodes.ValueSort([](const Ptr<ASTOpParameter>& A, const Ptr<ASTOpParameter>& B)->bool
 			{
-				Parameters[ParameterIndex] = Entry.Value->Parameter;
-				++ParameterIndex;
-			}
-
-			// Sort the parameters as deterministically as possible.
-			struct FParameterSortPredicate
+				if (A->Parameter.Name < B->Parameter.Name)
+				{
+					return true;
+				}
+				if (A->Parameter.Name > B->Parameter.Name)
+				{
+					return false;
+				}
+				return A->Parameter.UID < B->Parameter.UID;
+			});
+        	
+			// Set the parameter list from the non-optimized data, so that we have them all even if they are optimized out
+            {
+            	const int32 ParameterCount = Generator.FirstPass.ParameterNodes.Num();
+            	Parameters.Reserve(ParameterCount);
+            	for ( const TPair<Ptr<const Node>, Ptr<ASTOpParameter>>& Entry : Generator.FirstPass.ParameterNodes )
+            	{
+            		Parameters.Add(Entry.Value->Parameter);
+            	}
+            }
+        	
+        	// Check for repeated parameter names and report them to the user
+			if (Parameters.Num() > 1)
 			{
-				bool operator()(const FParameterDesc& A, const FParameterDesc& B) const
+				TArray<Ptr<const Node>> ParameterNodes;
+				Generator.FirstPass.ParameterNodes.GenerateKeyArray(ParameterNodes);
+				
+				// Should be impossible as both arrays come from the same data.
+				check(ParameterNodes.Num() == Parameters.Num());
+
+				TArray<FString> RepeatedParameterNames;
+				
+				for (int32 Index = 1; Index < Parameters.Num(); ++Index)
 				{
-					if (A.Name < B.Name) return true;
-					if (A.Name > B.Name) return false;
-					return A.UID < B.UID;
+					const FString CurrentParameterName = Parameters[Index].Name;
+					const FString PreviousParameterName = Parameters[Index-1].Name;
+					
+					if (CurrentParameterName.Compare(PreviousParameterName) == 0 && !RepeatedParameterNames.Contains(CurrentParameterName))
+					{
+						// Report repeated parameter name
+						const FString WarningText = FString::Printf(TEXT("Repeated parameter found : \"%s\". Please use a different name."), *CurrentParameterName);
+						const void* MessageContext = ParameterNodes[Index]->GetMessageContext();
+						GenErrorLog->Add(WarningText, ELMT_WARNING, MessageContext);
+
+						RepeatedParameterNames.Add(CurrentParameterName);
+					}
 				}
-			};
-			
-			FParameterSortPredicate SortPredicate;
-			Parameters.Sort(SortPredicate);
+			}
 		}
 
-        mu::Private::RemoveDisabledAddInstanceOpsValues(m_pD->Options, states);
+        mu::Private::RemoveDisabledAddInstanceOpsValues(m_pD->Options, States);
 
         // Slow AST code verification for debugging.
         //TArray<Ptr<ASTOp>> roots;
@@ -359,13 +382,13 @@ namespace mu
 
         // Optimize the generated code
         {
-            CodeOptimiser optimiser( m_pD->Options, states );
-            optimiser.OptimiseAST( );
+            CodeOptimiser Optimiser( m_pD->Options, States );
+            Optimiser.OptimiseAST( );
         }
 
         // Link the Program and generate state data.
-		TSharedPtr<FModel> pResult = MakeShared<FModel>();
-        FProgram& Program = pResult->GetPrivate()->Program;
+		TSharedPtr<FModel> Result = MakeShared<FModel>();
+        FProgram& Program = Result->GetPrivate()->Program;
 
 		check(Program.Parameters.IsEmpty());
 		Program.Parameters = Parameters;
@@ -379,17 +402,17 @@ namespace mu
 		FImageOperator ImOp = FImageOperator::GetDefault(m_pD->Options->GetPrivate()->ImageFormatFunc);
 		FLinkerOptions LinkerOptions(ImOp);
 
-		for(FStateCompilationData& s: states )
+		for(FStateCompilationData& State : States )
         {
 			LinkerOptions.MinTextureResidentMipCount = m_pD->Options->GetPrivate()->MinTextureResidentMipCount;
 
-            if (s.root)
+            if (State.root)
             {
-				s.state.Root = ASTOp::FullLink(s.root, Program, &LinkerOptions);
+				State.state.Root = ASTOp::FullLink(State.root, Program, &LinkerOptions);
             }
             else
             {
-                s.state.Root = 0;
+                State.state.Root = 0;
             }
         }
 
@@ -397,68 +420,68 @@ namespace mu
 		Program.OpAddress.Shrink();
 
         // Set the runtime parameter indices.
-        for(FStateCompilationData& s: states )
+        for(FStateCompilationData& State : States )
         {
-            for ( int32 RuntimeParameterIndex=0; RuntimeParameterIndex <s.nodeState.RuntimeParams.Num(); ++RuntimeParameterIndex)
+            for ( int32 RuntimeParameterIndex=0; RuntimeParameterIndex <State.nodeState.RuntimeParams.Num(); ++RuntimeParameterIndex)
             {
-                int32 paramIndex = -1;
-                for ( int32 i=0; paramIndex<0 && i<Program.Parameters.Num(); ++i )
+                int32 ParameterIndex = -1;
+                for ( int32 i=0; ParameterIndex<0 && i<Program.Parameters.Num(); ++i )
                 {
                     if ( Program.Parameters[i].Name
                          ==
-                         s.nodeState.RuntimeParams[RuntimeParameterIndex] )
+                         State.nodeState.RuntimeParams[RuntimeParameterIndex] )
                     {
-                        paramIndex = (int)i;
+                        ParameterIndex = (int)i;
                     }
                 }
 
-                if (paramIndex>=0)
+                if (ParameterIndex>=0)
                 {
-                    s.state.m_runtimeParameters.Add( paramIndex );
+                    State.state.m_runtimeParameters.Add( ParameterIndex );
                 }
                 else
                 {
-					FString Temp = FString::Printf(TEXT(
+					FString WarningString = FString::Printf(TEXT(
 						"The state [%s] refers to a parameter [%s]  that has not been found in the model. This error can be "
 						"safely dismissed in case of partial compilation."), 
-						*s.nodeState.Name,
-						*s.nodeState.RuntimeParams[RuntimeParameterIndex]);
-                    m_pD->ErrorLog->Add(Temp, ELMT_WARNING, pNode->GetMessageContext() );
+						*State.nodeState.Name,
+						*State.nodeState.RuntimeParams[RuntimeParameterIndex]);
+                    m_pD->ErrorLog->Add(WarningString, ELMT_WARNING, pNode->GetMessageContext() );
                 }
             }
 
             // Generate the mask of update cache ops
-            for ( const auto& a: s.m_updateCache )
+            for (const Ptr<ASTOp>& Instruction : State.m_updateCache )
             {
-                s.state.m_updateCache.Emplace(a->linkedAddress);
+                State.state.m_updateCache.Emplace(Instruction->linkedAddress);
             }
 
             // Sort the update cache addresses for performance and determinism
-            s.state.m_updateCache.Sort();
+            State.state.m_updateCache.Sort();
 
             // Generate the mask of dynamic resources
-            for ( const auto& a: s.m_dynamicResources )
+            for (const TPair<Ptr<ASTOp>, TArray<FString>>& Instruction : State.m_dynamicResources )
             {
-                uint64 relevantMask = 0;
-                for ( const auto& b: a.Value )
+                uint64 RelevantMask = 0;
+                for (const FString& InstructionParameter : Instruction.Value )
                 {
                     // Find the index in the model parameter list
-                    int paramIndex = -1;
-                    for ( int32 i=0; paramIndex<0 && i<Program.Parameters.Num(); ++i )
+                    int ParameterIndex = -1;
+                    for ( int32 i=0; ParameterIndex<0 && i<Program.Parameters.Num(); ++i )
                     {
-                        if ( Program.Parameters[i].Name == b )
+                        if ( Program.Parameters[i].Name == InstructionParameter )
                         {
-                            paramIndex = (int)i;
+                            ParameterIndex = (int)i;
                         }
                     }
-                    check(paramIndex>=0);
+                    check(ParameterIndex>=0);
 
                     // Find the position in the state data vector.
-                    int32 IndexInRuntimeList = s.state.m_runtimeParameters.Find( paramIndex );
+                    const int32 IndexInRuntimeList = State.state.m_runtimeParameters.Find( ParameterIndex );
 
                     if (IndexInRuntimeList !=INDEX_NONE )
                     {
-                        relevantMask |= uint64(1) << IndexInRuntimeList;
+                        RelevantMask |= uint64(1) << IndexInRuntimeList;
                     }
                 }
 
@@ -466,26 +489,26 @@ namespace mu
                 // Maybe something with the difference of precision between the relevant parameters
                 // in operation subtrees.
                 //check(relevantMask!=0);
-                if (relevantMask!=0)
+                if (RelevantMask!=0)
                 {
-                    s.state.m_dynamicResources.Emplace( a.Key->linkedAddress, relevantMask );
+                    State.state.m_dynamicResources.Emplace( Instruction.Key->linkedAddress, RelevantMask );
                 }
             }            
 
             // Sort for performance and determinism
-            s.state.m_dynamicResources.Sort();
+            State.state.m_dynamicResources.Sort();
 
-            Program.States.Add(s.state);
+            Program.States.Add(State.state);
         }
 
 		UE_LOG(LogMutableCore, Verbose, TEXT("(int) %s : %ld"), TEXT("Program size"), int64(Program.OpAddress.Num()));
 
         // Merge the log in the right order
-        genErrorLog->Merge( m_pD->ErrorLog.Get() );
-        m_pD->ErrorLog = genErrorLog;
+        GenErrorLog->Merge( m_pD->ErrorLog.Get() );
+        m_pD->ErrorLog = GenErrorLog;
 
 		// Pack data
-		m_pD->GenerateRoms(pResult.Get(), LinkerOptions.AdditionalData );
+		m_pD->GenerateRoms(Result.Get(), LinkerOptions.AdditionalData );
 
 		UE_LOG(LogMutableCore, Verbose, TEXT("(int) %s : %ld"), TEXT("Program size"), int64(Program.OpAddress.Num()));
 
@@ -500,7 +523,7 @@ namespace mu
 		Program.ConstantPhysicsBodies.Shrink();
 		Program.Parameters.Shrink();
 
-        return pResult;
+        return Result;
     }
 
 
@@ -780,4 +803,3 @@ namespace mu
 
 
 }
-
