[Mutable] Moved COE headers to private.

These headers should not be public.

Kept COESchema as public but moved to private many of its functions.

diff --git a/Engine/Plugins/Experimental/MutablePopulation/Source/CustomizableObjectPopulationEditor/Private/MuCOPE/CustomizableObjectPopulationClassDetails.cpp b/Engine/Plugins/Experimental/MutablePopulation/Source/CustomizableObjectPopulationEditor/Private/MuCOPE/CustomizableObjectPopulationClassDetails.cpp
--- a/Engine/Plugins/Experimental/MutablePopulation/Source/CustomizableObjectPopulationEditor/Private/MuCOPE/CustomizableObjectPopulationClassDetails.cpp
+++ b/Engine/Plugins/Experimental/MutablePopulation/Source/CustomizableObjectPopulationEditor/Private/MuCOPE/CustomizableObjectPopulationClassDetails.cpp
@@ -2,7 +2,6 @@
 
 #include "MuCOPE/CustomizableObjectPopulationClassDetails.h"
 
-#include "MuCOE/UnrealEditorPortabilityHelpers.h"
 #include "Curves/CurveLinearColor.h"
 #include "DetailLayoutBuilder.h"
 #include "Framework/Views/TableViewMetadata.h"
@@ -252,7 +251,7 @@ void SPopulationClassTagManager::BuildAllowlistWidget()
 			.ToolTipText(LOCTEXT("AddAllowlistTag", "Add Allow List Tag"))
 			[
 				SNew(SImage)
-				.Image(UE_MUTABLE_GET_BRUSH(TEXT("Plus")))
+				.Image(FAppStyle::GetBrush(TEXT("Plus")))
 			]
 		]
 	];
@@ -309,7 +308,7 @@ void SPopulationClassTagManager::BuildBlocklistWidget()
 			.ToolTipText(LOCTEXT("AddBlocklistTag", "Add Block List Tag"))
 			[
 				SNew(SImage)
-				.Image(UE_MUTABLE_GET_BRUSH(TEXT("Plus")))
+				.Image(FAppStyle::GetBrush(TEXT("Plus")))
 			]
 		]
 	];
@@ -437,7 +436,7 @@ void SPopulationClassTagList::Construct(const FArguments& InArgs)
 		.ToolTipText(LOCTEXT("RemoveTag", "Remove Tag"))
 		[
 			SNew(SImage)
-			.Image(UE_MUTABLE_GET_BRUSH(TEXT("Cross")))
+			.Image(FAppStyle::GetBrush(TEXT("Cross")))
 		]
 	];
 
@@ -609,7 +608,7 @@ void SPopulationClassCharacteristic::Construct(const FArguments& InArgs)
 				.ToolTipText(LOCTEXT("AddConstraint", "Add a New Constraint"))
 				[
 					SNew(SImage)
-					.Image(UE_MUTABLE_GET_BRUSH(TEXT("Plus")))
+					.Image(FAppStyle::GetBrush(TEXT("Plus")))
 				]
 			];
 		}
@@ -691,7 +690,7 @@ void SPopulationClassCharacteristic::Construct(const FArguments& InArgs)
 				.OnClicked(this, &SPopulationClassCharacteristic::OnRemoveCharacteristicButtonPressed)
 				[
 					SNew(SImage)
-					.Image(UE_MUTABLE_GET_BRUSH(TEXT("Cross")))
+					.Image(FAppStyle::GetBrush(TEXT("Cross")))
 				]
 			]
 
@@ -863,7 +862,7 @@ void SPopulationClassConstraint::Construct(const FArguments& InArgs)
 							.Visibility(multiConstraint ? EVisibility::Visible : EVisibility::Collapsed)
 							[
 								SNew(SImage)
-								.Image(UE_MUTABLE_GET_BRUSH(TEXT("Cross")))
+								.Image(FAppStyle::GetBrush(TEXT("Cross")))
 							]
 						]
 				]
@@ -1661,7 +1660,7 @@ void SRangeWidget::Construct(const FArguments& InArgs)
 					.OnClicked(this, &SRangeWidget::OnRemoveRangeButtonPressed)
 					[
 						SNew(SImage)
-						.Image(UE_MUTABLE_GET_BRUSH(TEXT("Cross")))
+						.Image(FAppStyle::GetBrush(TEXT("Cross")))
 					]
 				]
 			]
diff --git a/Engine/Plugins/Experimental/MutablePopulation/Source/CustomizableObjectPopulationEditor/Private/MuCOPE/CustomizableObjectPopulationClassEditor.cpp b/Engine/Plugins/Experimental/MutablePopulation/Source/CustomizableObjectPopulationEditor/Private/MuCOPE/CustomizableObjectPopulationClassEditor.cpp
--- a/Engine/Plugins/Experimental/MutablePopulation/Source/CustomizableObjectPopulationEditor/Private/MuCOPE/CustomizableObjectPopulationClassEditor.cpp
+++ b/Engine/Plugins/Experimental/MutablePopulation/Source/CustomizableObjectPopulationEditor/Private/MuCOPE/CustomizableObjectPopulationClassEditor.cpp
@@ -29,7 +29,6 @@
 #include "MuCOPE/CustomizableObjectPopulationEditorModule.h"
 #include "MuCOPE/SCustomizableObjectPopulationClassTagsTool.h"
 #include "MuCOPE/SCustomizableObjectPopulationEditorViewport.h"
-#include "MuCOE/UnrealEditorPortabilityHelpers.h"
 #include "PropertyEditorModule.h"
 #include "SCurveEditor.h"
 #include "Widgets/Docking/SDockTab.h"
@@ -980,7 +979,7 @@ void FCustomizableObjectPopulationClassEditor::GeneratePopulationClassInstances(
 			{
 				//Search for asset with the same name
 				FString ObjectPath = FilePath + FString(".") + ObjectName;
-				AssetData = AssetRegistryModule.Get().GetAssetByObjectPath(UE_MUTABLE_OBJECTPATH(ObjectPath));
+				AssetData = AssetRegistryModule.Get().GetAssetByObjectPath(FSoftObjectPath(ObjectPath));
 
 				// Popup window
 				if (AssetData.IsValid() && AssetData.GetClass() == UCustomizableObjectInstance::StaticClass() && RetType != EAppReturnType::YesAll)
@@ -1098,7 +1097,7 @@ void SSelectPopulationClassFolderDlg::Construct(const FArguments& InArgs)
 		.Padding(2)
 		[
 			SNew(SBorder)
-			.BorderImage(UE_MUTABLE_GET_BRUSH("ToolPanel.GroupBorder"))
+			.BorderImage(FAppStyle::GetBrush("ToolPanel.GroupBorder"))
 			[
 				SNew(SVerticalBox)
 
@@ -1142,15 +1141,15 @@ void SSelectPopulationClassFolderDlg::Construct(const FArguments& InArgs)
 		.Padding(5)
 		[
 			SNew(SUniformGridPanel)
-			.SlotPadding(UE_MUTABLE_GET_MARGIN("StandardDialog.SlotPadding"))
-			.MinDesiredSlotWidth(UE_MUTABLE_GET_FLOAT("StandardDialog.MinDesiredSlotWidth"))
-			.MinDesiredSlotHeight(UE_MUTABLE_GET_FLOAT("StandardDialog.MinDesiredSlotHeight"))
+			.SlotPadding(FAppStyle::GetMargin("StandardDialog.SlotPadding"))
+			.MinDesiredSlotWidth(FAppStyle::GetFloat("StandardDialog.MinDesiredSlotWidth"))
+			.MinDesiredSlotHeight(FAppStyle::GetFloat("StandardDialog.MinDesiredSlotHeight"))
 
 			+ SUniformGridPanel::Slot(0, 0)
 			[
 				SNew(SButton)
 				.HAlign(HAlign_Center)
-				.ContentPadding(UE_MUTABLE_GET_MARGIN("StandardDialog.ContentPadding"))
+				.ContentPadding(FAppStyle::GetMargin("StandardDialog.ContentPadding"))
 				.Text(LOCTEXT("OK", "OK"))
 				.OnClicked(this, &SSelectPopulationClassFolderDlg::OnButtonClick, EAppReturnType::Ok)
 			]
@@ -1159,7 +1158,7 @@ void SSelectPopulationClassFolderDlg::Construct(const FArguments& InArgs)
 			[
 				SNew(SButton)
 				.HAlign(HAlign_Center)
-				.ContentPadding(UE_MUTABLE_GET_MARGIN("StandardDialog.ContentPadding"))
+				.ContentPadding(FAppStyle::GetMargin("StandardDialog.ContentPadding"))
 				.Text(LOCTEXT("Cancel", "Cancel"))
 				.OnClicked(this, &SSelectPopulationClassFolderDlg::OnButtonClick, EAppReturnType::Cancel)
 			]
diff --git a/Engine/Plugins/Experimental/MutablePopulation/Source/CustomizableObjectPopulationEditor/Private/MuCOPE/CustomizableObjectPopulationEditor.cpp b/Engine/Plugins/Experimental/MutablePopulation/Source/CustomizableObjectPopulationEditor/Private/MuCOPE/CustomizableObjectPopulationEditor.cpp
--- a/Engine/Plugins/Experimental/MutablePopulation/Source/CustomizableObjectPopulationEditor/Private/MuCOPE/CustomizableObjectPopulationEditor.cpp
+++ b/Engine/Plugins/Experimental/MutablePopulation/Source/CustomizableObjectPopulationEditor/Private/MuCOPE/CustomizableObjectPopulationEditor.cpp
@@ -26,7 +26,6 @@
 #include "MuCOPE/CustomizableObjectPopulationEditorActions.h"
 #include "MuCOPE/CustomizableObjectPopulationEditorModule.h"
 #include "MuCOPE/SCustomizableObjectPopulationEditorViewport.h"
-#include "MuCOE/UnrealEditorPortabilityHelpers.h"
 #include "PropertyEditorModule.h"
 #include "Widgets/Docking/SDockTab.h"
 #include "Widgets/Input/SButton.h"
@@ -615,7 +614,7 @@ void FCustomizableObjectPopulationEditor::GeneratePopulationInstances()
 			{
 				//Search for asset with the same name
 				FString ObjectPath = FilePath + FString(".") + ObjectName;
-				AssetData = AssetRegistryModule.Get().GetAssetByObjectPath(UE_MUTABLE_OBJECTPATH(ObjectPath));
+				AssetData = AssetRegistryModule.Get().GetAssetByObjectPath(FSoftObjectPath(ObjectPath));
 
 				// Popup window
 				if (AssetData.IsValid() && AssetData.GetClass() == UCustomizableObjectInstance::StaticClass() && RetType != EAppReturnType::YesAll)
@@ -735,7 +734,7 @@ void SSelectPopulationFolderDlg::Construct(const FArguments& InArgs)
 		.Padding(2)
 		[
 			SNew(SBorder)
-			.BorderImage(UE_MUTABLE_GET_BRUSH("ToolPanel.GroupBorder"))
+			.BorderImage(FAppStyle::GetBrush("ToolPanel.GroupBorder"))
 			[
 				SNew(SVerticalBox)
 
@@ -779,15 +778,15 @@ void SSelectPopulationFolderDlg::Construct(const FArguments& InArgs)
 		.Padding(5)
 		[
 			SNew(SUniformGridPanel)
-			.SlotPadding(UE_MUTABLE_GET_MARGIN("StandardDialog.SlotPadding"))
-			.MinDesiredSlotWidth(UE_MUTABLE_GET_FLOAT("StandardDialog.MinDesiredSlotWidth"))
-			.MinDesiredSlotHeight(UE_MUTABLE_GET_FLOAT("StandardDialog.MinDesiredSlotHeight"))
+			.SlotPadding(FAppStyle::GetMargin("StandardDialog.SlotPadding"))
+			.MinDesiredSlotWidth(FAppStyle::GetFloat("StandardDialog.MinDesiredSlotWidth"))
+			.MinDesiredSlotHeight(FAppStyle::GetFloat("StandardDialog.MinDesiredSlotHeight"))
 			
 			+ SUniformGridPanel::Slot(0, 0)
 			[
 				SNew(SButton)
 				.HAlign(HAlign_Center)
-				.ContentPadding(UE_MUTABLE_GET_MARGIN("StandardDialog.ContentPadding"))
+				.ContentPadding(FAppStyle::GetMargin("StandardDialog.ContentPadding"))
 				.Text(LOCTEXT("OK", "OK"))
 				.OnClicked(this, &SSelectPopulationFolderDlg::OnButtonClick, EAppReturnType::Ok)
 			]
@@ -796,7 +795,7 @@ void SSelectPopulationFolderDlg::Construct(const FArguments& InArgs)
 			[
 				SNew(SButton)
 				.HAlign(HAlign_Center)
-				.ContentPadding(UE_MUTABLE_GET_MARGIN("StandardDialog.ContentPadding"))
+				.ContentPadding(FAppStyle::GetMargin("StandardDialog.ContentPadding"))
 				.Text(LOCTEXT("Cancel", "Cancel"))
 				.OnClicked(this, &SSelectPopulationFolderDlg::OnButtonClick, EAppReturnType::Cancel)
 			]
diff --git a/Engine/Plugins/Experimental/MutablePopulation/Source/CustomizableObjectPopulationEditor/Private/MuCOPE/SCustomizableObjectPopulationClassTagManager.cpp b/Engine/Plugins/Experimental/MutablePopulation/Source/CustomizableObjectPopulationEditor/Private/MuCOPE/SCustomizableObjectPopulationClassTagManager.cpp
--- a/Engine/Plugins/Experimental/MutablePopulation/Source/CustomizableObjectPopulationEditor/Private/MuCOPE/SCustomizableObjectPopulationClassTagManager.cpp
+++ b/Engine/Plugins/Experimental/MutablePopulation/Source/CustomizableObjectPopulationEditor/Private/MuCOPE/SCustomizableObjectPopulationClassTagManager.cpp
@@ -2,7 +2,6 @@
 
 #include "MuCOPE/SCustomizableObjectPopulationClassTagManager.h"
 
-#include "MuCOE/UnrealEditorPortabilityHelpers.h"
 #include "MuCO/CustomizableObject.h"
 #include "MuCO/CustomizableObjectPrivate.h"
 #include "Widgets/Input/SButton.h"
@@ -103,7 +102,7 @@ void SCustomizableObjectPopulationClassTagsManager::BuildTagSelector()
 			.ToolTipText(LOCTEXT("AddTagSelector", "Add a Tag to the Customizable Object"))
 			[
 				SNew(SImage)
-				.Image(UE_MUTABLE_GET_BRUSH(TEXT("Plus")))
+				.Image(FAppStyle::GetBrush(TEXT("Plus")))
 			]
 		]
 	];
@@ -152,7 +151,7 @@ void SCustomizableObjectPopulationClassTagsManager::BuildTagSelector()
 				.ToolTipText(LOCTEXT("RemoveTag", "Remove Tag"))
 				[
 					SNew(SImage)
-					.Image(UE_MUTABLE_GET_BRUSH(TEXT("Cross")))
+					.Image(FAppStyle::GetBrush(TEXT("Cross")))
 				]
 			]
 		];
@@ -225,7 +224,7 @@ void SCustomizableObjectPopulationClassTagsManager::BuildTagManager()
 			.ToolTipText(LOCTEXT("AddTag", "Add a Tag to the Root Object"))
 			[
 				SNew(SImage)
-				.Image(UE_MUTABLE_GET_BRUSH(TEXT("Plus")))
+				.Image(FAppStyle::GetBrush(TEXT("Plus")))
 			]
 		]
 	];
@@ -263,7 +262,7 @@ void SCustomizableObjectPopulationClassTagsManager::BuildTagManager()
 				.ToolTipText(LOCTEXT("RemoveTag", "Remove Tag"))
 				[
 					SNew(SImage)
-					.Image(UE_MUTABLE_GET_BRUSH(TEXT("Cross")))
+					.Image(FAppStyle::GetBrush(TEXT("Cross")))
 				]
 			]
 		];
diff --git a/Engine/Plugins/Experimental/MutablePopulation/Source/CustomizableObjectPopulationEditor/Private/MuCOPE/SCustomizableObjectPopulationClassTagsTool.cpp b/Engine/Plugins/Experimental/MutablePopulation/Source/CustomizableObjectPopulationEditor/Private/MuCOPE/SCustomizableObjectPopulationClassTagsTool.cpp
--- a/Engine/Plugins/Experimental/MutablePopulation/Source/CustomizableObjectPopulationEditor/Private/MuCOPE/SCustomizableObjectPopulationClassTagsTool.cpp
+++ b/Engine/Plugins/Experimental/MutablePopulation/Source/CustomizableObjectPopulationEditor/Private/MuCOPE/SCustomizableObjectPopulationClassTagsTool.cpp
@@ -6,7 +6,6 @@
 #include "MuCO/CustomizableObjectPrivate.h"
 #include "MuCOP/CustomizableObjectPopulationClass.h"
 #include "MuCOPE/CustomizableObjectPopulationClassDetails.h"
-#include "MuCOE/UnrealEditorPortabilityHelpers.h"
 #include "Widgets/Input/SButton.h"
 #include "Widgets/Input/SSearchBox.h"
 #include "Widgets/Input/STextComboBox.h"
@@ -94,14 +93,14 @@ void SCustomizableObjectPopulationClassTagsTool::GenerateChildSlot()
 					.InitiallyCollapsed(false)
 					.BorderBackgroundColor(FLinearColor(0.6f, 0.6f, 0.6f, 1.0f))
 					.BorderImage_Lambda([this]() { return GetExpandableAreaBorderImage(*TagsManagerExpandableArea); })
-					.BodyBorderImage(UE_MUTABLE_GET_BRUSH("ToolPanel.GroupBorder"))
+					.BodyBorderImage(FAppStyle::GetBrush("ToolPanel.GroupBorder"))
 					.BodyBorderBackgroundColor(FLinearColor::White)
 					.HeaderContent()
 					[
 						SNew(STextBlock)
 						.Margin(FMargin(0.0f,5.0f))
 						.Text(LOCTEXT("ClassTagManagementText", "Population Tag Management for the current Customizable Object"))
-						.Font(UE_MUTABLE_GET_FONTSTYLE("DetailsView.CategoryFontStyle"))
+						.Font(FAppStyle::GetFontStyle("DetailsView.CategoryFontStyle"))
 						.ShadowOffset(FVector2D(1.0f, 1.0f))
 					]
 					.BodyContent()
@@ -152,14 +151,14 @@ void SCustomizableObjectPopulationClassTagsTool::GenerateChildSlot()
 					.InitiallyCollapsed(false)
 					.BorderBackgroundColor(FLinearColor(0.6f, 0.6f, 0.6f, 1.0f))
 					.BorderImage_Lambda([this]() { return GetExpandableAreaBorderImage(*ParameterOptionsManagerExpandableArea); })
-					.BodyBorderImage(UE_MUTABLE_GET_BRUSH("ToolPanel.GroupBorder"))
+					.BodyBorderImage(FAppStyle::GetBrush("ToolPanel.GroupBorder"))
 					.BodyBorderBackgroundColor(FLinearColor::White)
 					.HeaderContent()
 					[
 						SNew(STextBlock)
 						.Margin(FMargin(0.0f, 5.0f))
 						.Text(LOCTEXT("ParameterOptionsManagerText", "Parameter Option Population Tags"))
-						.Font(UE_MUTABLE_GET_FONTSTYLE("DetailsView.CategoryFontStyle"))
+						.Font(FAppStyle::GetFontStyle("DetailsView.CategoryFontStyle"))
 						.ShadowOffset(FVector2D(1.0f, 1.0f))
 					]
 					.BodyContent()
@@ -276,7 +275,7 @@ void SCustomizableObjectPopulationClassTagsTool::GenerateTagsViewer()
 					.ButtonColorAndOpacity(FSlateColor(FLinearColor(1.0f,0.2f,0.2f)))
 					[
 						SNew(SImage)
-						.Image(UE_MUTABLE_GET_BRUSH(TEXT("Cross")))
+						.Image(FAppStyle::GetBrush(TEXT("Cross")))
 					]
 				]
 			];
@@ -306,7 +305,7 @@ void SCustomizableObjectPopulationClassTagsTool::GenerateTagsViewer()
 						.OnClicked(this, &SCustomizableObjectPopulationClassTagsTool::OnRemoveTagButtonPressed,i)
 						[
 							SNew(SImage)
-							.Image(UE_MUTABLE_GET_BRUSH(TEXT("Cross")))
+							.Image(FAppStyle::GetBrush(TEXT("Cross")))
 						]
 					]
 				];
@@ -663,9 +662,9 @@ const FSlateBrush* SCustomizableObjectPopulationClassTagsTool::GetExpandableArea
 {
 	if (Area.IsTitleHovered())
 	{
-		return Area.IsExpanded() ? UE_MUTABLE_GET_BRUSH("DetailsView.CategoryTop_Hovered") : UE_MUTABLE_GET_BRUSH("DetailsView.CollapsedCategory_Hovered");
+		return Area.IsExpanded() ? FAppStyle::GetBrush("DetailsView.CategoryTop_Hovered") : FAppStyle::GetBrush("DetailsView.CollapsedCategory_Hovered");
 	}
-	return Area.IsExpanded() ? UE_MUTABLE_GET_BRUSH("DetailsView.CategoryTop") : UE_MUTABLE_GET_BRUSH("DetailsView.CollapsedCategory");
+	return Area.IsExpanded() ? FAppStyle::GetBrush("DetailsView.CategoryTop") : FAppStyle::GetBrush("DetailsView.CollapsedCategory");
 }
 
 
@@ -751,7 +750,7 @@ void SSingleTagWidget::Construct(const FArguments& InArgs)
 				.OnClicked(this, &SSingleTagWidget::OnRemoveTagButtonPressed)
 				[
 					SNew(SImage)
-					.Image(UE_MUTABLE_GET_BRUSH(TEXT("Cross")))
+					.Image(FAppStyle::GetBrush(TEXT("Cross")))
 				]
 			]
 		]
diff --git a/Engine/Plugins/Experimental/MutablePopulation/Source/CustomizableObjectPopulationEditor/Private/MuCOPE/SCustomizableObjectPopulationEditorViewportToolBar.cpp b/Engine/Plugins/Experimental/MutablePopulation/Source/CustomizableObjectPopulationEditor/Private/MuCOPE/SCustomizableObjectPopulationEditorViewportToolBar.cpp
--- a/Engine/Plugins/Experimental/MutablePopulation/Source/CustomizableObjectPopulationEditor/Private/MuCOPE/SCustomizableObjectPopulationEditorViewportToolBar.cpp
+++ b/Engine/Plugins/Experimental/MutablePopulation/Source/CustomizableObjectPopulationEditor/Private/MuCOPE/SCustomizableObjectPopulationEditorViewportToolBar.cpp
@@ -7,8 +7,6 @@
 #include "Framework/MultiBox/MultiBoxBuilder.h"
 #include "LevelEditor.h"
 #include "Modules/ModuleManager.h"
-#include "MuCOE/CustomizableObjectEditorViewportMenuCommands.h"
-#include "MuCOE/UnrealEditorPortabilityHelpers.h"
 #include "MuCOPE/SCustomizableObjectPopulationEditorViewport.h"
 #include "SEditorViewportViewMenu.h"
 #include "Widgets/Input/SMenuAnchor.h"
@@ -57,7 +55,7 @@ void SCustomizableObjectPopulationEditorViewportToolBar::Construct(const FArgume
 	FLinearColor ButtonColor2 = FLinearColor(0.2f, 0.2f, 0.2f, 0.75f);
 	FLinearColor TextColor1 = FLinearColor(1.0f, 1.0f, 1.0f, 1.0f);
 	FLinearColor TextColor2 = FLinearColor(0.8f, 0.8f, 0.8f, 0.8f);
-	FSlateFontInfo Info = UE_MUTABLE_GET_FONTSTYLE("BoldFont");
+	FSlateFontInfo Info = FAppStyle::GetFontStyle("BoldFont");
 	Info.Size += 26;
 
 	ChildSlot
@@ -67,10 +65,10 @@ void SCustomizableObjectPopulationEditorViewportToolBar::Construct(const FArgume
 		.VAlign(VAlign_Top)
 		[
 			SNew(SBorder)
-			.BorderImage(UE_MUTABLE_GET_BRUSH("NoBorder"))
+			.BorderImage(FAppStyle::GetBrush("NoBorder"))
 		//// Color and opacity is changed based on whether or not the mouse cursor is hovering over the toolbar area
 		//.ColorAndOpacity(this, &SViewportToolBar::OnGetColorAndOpacity)
-		.ForegroundColor(UE_MUTABLE_GET_SLATECOLOR(DefaultForegroundName))
+		.ForegroundColor(FAppStyle::GetSlateColor(DefaultForegroundName))
 		[
 			SNew(SVerticalBox)
 			+ SVerticalBox::Slot()
@@ -94,8 +92,6 @@ void SCustomizableObjectPopulationEditorViewportToolBar::Construct(const FArgume
 
 TSharedRef<SWidget> SCustomizableObjectPopulationEditorViewportToolBar::GenerateViewMenu() const
 {
-	const FCustomizableObjectEditorViewportMenuCommands& Actions = FCustomizableObjectEditorViewportMenuCommands::Get();
-
 	const bool bInShouldCloseWindowAfterMenuSelection = true;
 	FMenuBuilder ViewMenuBuilder(bInShouldCloseWindowAfterMenuSelection, Viewport.Pin()->GetCommandList());
 
@@ -197,7 +193,7 @@ TSharedRef<SWidget> SCustomizableObjectPopulationEditorViewportToolBar::Generate
 		.WidthOverride(100.0f)
 		[
 			SNew(SSpinBox<float>)
-			.Font(UE_MUTABLE_GET_FONTSTYLE(TEXT("MenuItem.Font")))
+			.Font(FAppStyle::GetFontStyle(TEXT("MenuItem.Font")))
 		.MinValue(FOVMin)
 		.MaxValue(FOVMax)
 		.Value(this, &SCustomizableObjectPopulationEditorViewportToolBar::OnGetFOVValue)
diff --git a/Engine/Plugins/Mutable/Source/CustomizableObjectEditor/Private/MuCOE/SCustomizableObjectEditorAdvancedPreviewSettings.cpp b/Engine/Plugins/Mutable/Source/CustomizableObjectEditor/Private/MuCOE/SCustomizableObjectEditorAdvancedPreviewSettings.cpp
--- a/Engine/Plugins/Mutable/Source/CustomizableObjectEditor/Private/MuCOE/SCustomizableObjectEditorAdvancedPreviewSettings.cpp
+++ b/Engine/Plugins/Mutable/Source/CustomizableObjectEditor/Private/MuCOE/SCustomizableObjectEditorAdvancedPreviewSettings.cpp
@@ -7,7 +7,6 @@
 #include "IContentBrowserSingleton.h"
 #include "Misc/MessageDialog.h"
 #include "Misc/Paths.h"
-#include "MuCOE/UnrealEditorPortabilityHelpers.h"
 #include "MuCOE/CustomizableObjectEditor.h"
 #include "MuCOE/CustomizableObjectEditorLevelProfile.h"
 #include "Widgets/Input/SButton.h"
@@ -29,7 +28,7 @@ void SLevelSelectWindow::Construct(const FArguments& InArgs)
 	DefaultSettings = InArgs._DefaultSettings.Get();
 
 	FAssetPickerConfig AssetPickerConfig;
-	UE_MUTABLE_GET_CLASSPATHS(AssetPickerConfig.Filter).Add(UE_MUTABLE_TOPLEVELASSETPATH(TEXT("/Script/Engine"), TEXT("World")));
+	AssetPickerConfig.Filter.ClassPaths.Add(FTopLevelAssetPath(TEXT("/Script/Engine"), TEXT("World")));
 	AssetPickerConfig.OnAssetSelected = FOnAssetSelected::CreateRaw(this, &SLevelSelectWindow::OnLevelAssetSelected);
 	AssetPickerConfig.ThumbnailScale = 0.4f;
 	FContentBrowserModule& ContentBrowserModule = FModuleManager::LoadModuleChecked<FContentBrowserModule>("ContentBrowser");
@@ -45,14 +44,14 @@ void SLevelSelectWindow::Construct(const FArguments& InArgs)
 			.Padding(2)
 			[
 				SNew(SBorder)
-				.BorderImage(UE_MUTABLE_GET_BRUSH("ToolPanel.GroupBorder"))
+				.BorderImage(FAppStyle::GetBrush("ToolPanel.GroupBorder"))
 				[
 					SNew(SVerticalBox)
 					+ SVerticalBox::Slot()
 					.FillHeight(1)
 					[
 						SNew(SBorder)
-						.BorderImage(UE_MUTABLE_GET_BRUSH("ToolPanel.GroupBorder"))
+						.BorderImage(FAppStyle::GetBrush("ToolPanel.GroupBorder"))
 						[
 							ContentBrowserModule.Get().CreateAssetPicker(AssetPickerConfig)
 						]
@@ -65,14 +64,14 @@ void SLevelSelectWindow::Construct(const FArguments& InArgs)
 			.Padding(5)
 			[
 				SNew(SUniformGridPanel)
-				.SlotPadding(UE_MUTABLE_GET_MARGIN("StandardDialog.SlotPadding"))
-				.MinDesiredSlotWidth(UE_MUTABLE_GET_FLOAT("StandardDialog.MinDesiredSlotWidth"))
-				.MinDesiredSlotHeight(UE_MUTABLE_GET_FLOAT("StandardDialog.MinDesiredSlotHeight"))
+				.SlotPadding(FAppStyle::GetMargin("StandardDialog.SlotPadding"))
+				.MinDesiredSlotWidth(FAppStyle::GetFloat("StandardDialog.MinDesiredSlotWidth"))
+				.MinDesiredSlotHeight(FAppStyle::GetFloat("StandardDialog.MinDesiredSlotHeight"))
 				+ SUniformGridPanel::Slot(0, 0)
 				[
 					SNew(SButton)
 					.HAlign(HAlign_Center)
-					.ContentPadding(UE_MUTABLE_GET_MARGIN("StandardDialog.ContentPadding"))
+					.ContentPadding(FAppStyle::GetMargin("StandardDialog.ContentPadding"))
 					.Text(LOCTEXT("OK", "OK"))
 					.OnClicked(this, &SLevelSelectWindow::OnButtonClick, EAppReturnType::Ok)
 				]
@@ -80,7 +79,7 @@ void SLevelSelectWindow::Construct(const FArguments& InArgs)
 				[
 					SNew(SButton)
 					.HAlign(HAlign_Center)
-					.ContentPadding(UE_MUTABLE_GET_MARGIN("StandardDialog.ContentPadding"))
+					.ContentPadding(FAppStyle::GetMargin("StandardDialog.ContentPadding"))
 					.Text(LOCTEXT("Cancel", "Cancel"))
 					.OnClicked(this, &SLevelSelectWindow::OnButtonClick, EAppReturnType::Cancel)
 				]
@@ -98,13 +97,13 @@ FReply SLevelSelectWindow::OnButtonClick(EAppReturnType::Type ButtonID)
 
 	if (UserResponse == EAppReturnType::Ok)
 	{
-		if (UE_MUTABLE_ASSETCLASS(LevelAssetData) != UE_MUTABLE_TOPLEVELASSETPATH(TEXT("/Script/Engine"), TEXT("World")))
+		if (LevelAssetData.AssetClassPath != FTopLevelAssetPath(TEXT("/Script/Engine"), TEXT("World")))
 		{
 			FMessageDialog::Open(EAppMsgType::Ok, LOCTEXT("PostProcessNameValueMissing", "Please select a Level asset"));
 			return FReply::Handled();
 		}
 
-		LevelAssetPath = UE_MUTABLE_GETOBJECTPATH(LevelAssetData);
+		LevelAssetPath = LevelAssetData.GetObjectPathString();
 
 		Result = LevelProfileManage::LoadProfileUObjectNames(LevelAssetPath, ArrayDirectionalLightName, ArrayPostProcessName, ArraySkyLightName);
 
@@ -215,19 +214,19 @@ void SUObjectSelecWindow::Construct(const FArguments& InArgs)
 			.Padding(2)
 			[
 				SNew(SBorder)
-				.BorderImage(UE_MUTABLE_GET_BRUSH("ToolPanel.GroupBorder"))
+				.BorderImage(FAppStyle::GetBrush("ToolPanel.GroupBorder"))
 				[
 					SNew(SVerticalBox)
 					+ SVerticalBox::Slot()
 					.AutoHeight()
 					[
 						SNew(SBorder)
-						.BorderImage(UE_MUTABLE_GET_BRUSH("ToolPanel.GroupBorder"))
+						.BorderImage(FAppStyle::GetBrush("ToolPanel.GroupBorder"))
 						[
 							SNew(SUniformGridPanel)
-							.SlotPadding(UE_MUTABLE_GET_MARGIN("StandardDialog.SlotPadding"))
-							.MinDesiredSlotWidth(UE_MUTABLE_GET_FLOAT("StandardDialog.MinDesiredSlotWidth"))
-							.MinDesiredSlotHeight(UE_MUTABLE_GET_FLOAT("StandardDialog.MinDesiredSlotHeight"))
+							.SlotPadding(FAppStyle::GetMargin("StandardDialog.SlotPadding"))
+							.MinDesiredSlotWidth(FAppStyle::GetFloat("StandardDialog.MinDesiredSlotWidth"))
+							.MinDesiredSlotHeight(FAppStyle::GetFloat("StandardDialog.MinDesiredSlotHeight"))
 							+ SUniformGridPanel::Slot(0, 0)
 							[
 								SNew(STextBlock)
@@ -246,12 +245,12 @@ void SUObjectSelecWindow::Construct(const FArguments& InArgs)
 					.AutoHeight()
 					[
 						SNew(SBorder)
-						.BorderImage(UE_MUTABLE_GET_BRUSH("ToolPanel.GroupBorder"))
+						.BorderImage(FAppStyle::GetBrush("ToolPanel.GroupBorder"))
 						[
 							SNew(SUniformGridPanel)
-							.SlotPadding(UE_MUTABLE_GET_MARGIN("StandardDialog.SlotPadding"))
-							.MinDesiredSlotWidth(UE_MUTABLE_GET_FLOAT("StandardDialog.MinDesiredSlotWidth"))
-							.MinDesiredSlotHeight(UE_MUTABLE_GET_FLOAT("StandardDialog.MinDesiredSlotHeight"))
+							.SlotPadding(FAppStyle::GetMargin("StandardDialog.SlotPadding"))
+							.MinDesiredSlotWidth(FAppStyle::GetFloat("StandardDialog.MinDesiredSlotWidth"))
+							.MinDesiredSlotHeight(FAppStyle::GetFloat("StandardDialog.MinDesiredSlotHeight"))
 							+ SUniformGridPanel::Slot(0, 0)
 							[
 								SNew(STextBlock)
@@ -269,12 +268,12 @@ void SUObjectSelecWindow::Construct(const FArguments& InArgs)
 					.AutoHeight()
 					[
 						SNew(SBorder)
-						.BorderImage(UE_MUTABLE_GET_BRUSH("ToolPanel.GroupBorder"))
+						.BorderImage(FAppStyle::GetBrush("ToolPanel.GroupBorder"))
 						[
 							SNew(SUniformGridPanel)
-							.SlotPadding(UE_MUTABLE_GET_MARGIN("StandardDialog.SlotPadding"))
-							.MinDesiredSlotWidth(UE_MUTABLE_GET_FLOAT("StandardDialog.MinDesiredSlotWidth"))
-							.MinDesiredSlotHeight(UE_MUTABLE_GET_FLOAT("StandardDialog.MinDesiredSlotHeight"))
+							.SlotPadding(FAppStyle::GetMargin("StandardDialog.SlotPadding"))
+							.MinDesiredSlotWidth(FAppStyle::GetFloat("StandardDialog.MinDesiredSlotWidth"))
+							.MinDesiredSlotHeight(FAppStyle::GetFloat("StandardDialog.MinDesiredSlotHeight"))
 							+ SUniformGridPanel::Slot(0, 0)
 							[
 								SNew(STextBlock)
@@ -292,12 +291,12 @@ void SUObjectSelecWindow::Construct(const FArguments& InArgs)
 					.AutoHeight()
 					[
 						SNew(SBorder)
-						.BorderImage(UE_MUTABLE_GET_BRUSH("ToolPanel.GroupBorder"))
+						.BorderImage(FAppStyle::GetBrush("ToolPanel.GroupBorder"))
 						[
 							SNew(SUniformGridPanel)
-							.SlotPadding(UE_MUTABLE_GET_MARGIN("StandardDialog.SlotPadding"))
-							.MinDesiredSlotWidth(UE_MUTABLE_GET_FLOAT("StandardDialog.MinDesiredSlotWidth"))
-							.MinDesiredSlotHeight(UE_MUTABLE_GET_FLOAT("StandardDialog.MinDesiredSlotHeight"))
+							.SlotPadding(FAppStyle::GetMargin("StandardDialog.SlotPadding"))
+							.MinDesiredSlotWidth(FAppStyle::GetFloat("StandardDialog.MinDesiredSlotWidth"))
+							.MinDesiredSlotHeight(FAppStyle::GetFloat("StandardDialog.MinDesiredSlotHeight"))
 							+ SUniformGridPanel::Slot(0, 0)
 							[
 								SNew(STextBlock)
@@ -319,14 +318,14 @@ void SUObjectSelecWindow::Construct(const FArguments& InArgs)
 			.Padding(5)
 			[
 				SNew(SUniformGridPanel)
-				.SlotPadding(UE_MUTABLE_GET_MARGIN("StandardDialog.SlotPadding"))
-				.MinDesiredSlotWidth(UE_MUTABLE_GET_FLOAT("StandardDialog.MinDesiredSlotWidth"))
-				.MinDesiredSlotHeight(UE_MUTABLE_GET_FLOAT("StandardDialog.MinDesiredSlotHeight"))
+				.SlotPadding(FAppStyle::GetMargin("StandardDialog.SlotPadding"))
+				.MinDesiredSlotWidth(FAppStyle::GetFloat("StandardDialog.MinDesiredSlotWidth"))
+				.MinDesiredSlotHeight(FAppStyle::GetFloat("StandardDialog.MinDesiredSlotHeight"))
 				+ SUniformGridPanel::Slot(0, 0)
 				[
 					SNew(SButton)
 					.HAlign(HAlign_Center)
-					.ContentPadding(UE_MUTABLE_GET_MARGIN("StandardDialog.ContentPadding"))
+					.ContentPadding(FAppStyle::GetMargin("StandardDialog.ContentPadding"))
 					.Text(LOCTEXT("OK", "OK"))
 					.OnClicked(this, &SUObjectSelecWindow::OnButtonClick, EAppReturnType::Ok)
 				]
@@ -334,7 +333,7 @@ void SUObjectSelecWindow::Construct(const FArguments& InArgs)
 				[
 					SNew(SButton)
 					.HAlign(HAlign_Center)
-					.ContentPadding(UE_MUTABLE_GET_MARGIN("StandardDialog.ContentPadding"))
+					.ContentPadding(FAppStyle::GetMargin("StandardDialog.ContentPadding"))
 					.Text(LOCTEXT("Cancel", "Cancel"))
 					.OnClicked(this, &SUObjectSelecWindow::OnButtonClick, EAppReturnType::Cancel)
 				]
@@ -438,9 +437,9 @@ void SCustomizableObjectEditorAdvancedPreviewSettings::Construct(const FArgument
 				.HAlign(HAlign_Fill)
 				[
 					SNew(SUniformGridPanel)
-					.SlotPadding(UE_MUTABLE_GET_MARGIN("StandardDialog.SlotPadding"))
-					.MinDesiredSlotWidth(UE_MUTABLE_GET_FLOAT("StandardDialog.MinDesiredSlotWidth"))
-					.MinDesiredSlotHeight(UE_MUTABLE_GET_FLOAT("StandardDialog.MinDesiredSlotHeight"))
+					.SlotPadding(FAppStyle::GetMargin("StandardDialog.SlotPadding"))
+					.MinDesiredSlotWidth(FAppStyle::GetFloat("StandardDialog.MinDesiredSlotWidth"))
+					.MinDesiredSlotHeight(FAppStyle::GetFloat("StandardDialog.MinDesiredSlotHeight"))
 					+ SUniformGridPanel::Slot(0, 0)
 					[
 						SNew(STextBlock)
diff --git a/Engine/Plugins/Mutable/Source/CustomizableObjectEditor/Public/MuCOE/Nodes/CustomizableObjectNode.h b/Engine/Plugins/Mutable/Source/CustomizableObjectEditor/Public/MuCOE/Nodes/CustomizableObjectNode.h
--- a/Engine/Plugins/Mutable/Source/CustomizableObjectEditor/Public/MuCOE/Nodes/CustomizableObjectNode.h
+++ b/Engine/Plugins/Mutable/Source/CustomizableObjectEditor/Public/MuCOE/Nodes/CustomizableObjectNode.h
@@ -16,6 +16,7 @@ class UCustomizableObjectNodeRemapPins;
 class UCustomizableObjectNodeRemapPinsByName;
 class UCustomizableObjectNodeRemapPinsByPosition;
 class IDetailsView;
+class ICustomizableObjectEditor;
 struct FPropertyChangedEvent;
 
 /** Abstract base class for all pin data. */
@@ -150,7 +151,7 @@ public:
 	// Used to replace Ids of referenced nodes by their new ids after duplicating the CustomizableObject.
 	virtual void UpdateReferencedNodeId(const FGuid& NewGuid) {};
 
-	TSharedPtr<class ICustomizableObjectEditor> GetGraphEditor() const;
+	TSharedPtr<ICustomizableObjectEditor> GetGraphEditor() const;
 
 	virtual bool ProvidesCustomPinRelevancyTest() const { return false; } // Override and return true if IsPinRelevant is overridden.
 	virtual bool IsPinRelevant(const UEdGraphPin*) const { return true; } // Override if default pins allocated with AllocateDefaultPins do not provide enough information.
diff --git a/Engine/Plugins/Mutable/Source/MutableValidation/Private/MuV/AssetValidator_CustomizableObjects.cpp b/Engine/Plugins/Mutable/Source/MutableValidation/Private/MuV/AssetValidator_CustomizableObjects.cpp
--- a/Engine/Plugins/Mutable/Source/MutableValidation/Private/MuV/AssetValidator_CustomizableObjects.cpp
+++ b/Engine/Plugins/Mutable/Source/MutableValidation/Private/MuV/AssetValidator_CustomizableObjects.cpp
@@ -15,7 +15,6 @@
 #include "MuCO/CustomizableObjectCompilerTypes.h"
 #include "MuCO/CustomizableObjectPrivate.h"
 #include "MuCO/ICustomizableObjectEditorModule.h"
-#include "MuCOE/GraphTraversal.h"
 #include "UObject/NameTypes.h"
 #include "UObject/Object.h"
 
@@ -98,9 +97,9 @@ EDataValidationResult UAssetValidator_CustomizableObjects::IsCustomizableObjectV
 	{
 		OnPostCOValidationHandle = FEditorDelegates::OnPostAssetValidation.AddStatic(OnPostCOsValidation);	
 	}
-	
+
 	// Find out which is the root for this CO (it may be itself but that is OK)
-	UCustomizableObject* RootObject = GraphTraversal::GetRootObject(InCustomizableObject);
+	UCustomizableObject* RootObject = ICustomizableObjectEditorModule::GetChecked().GetRootObject(InCustomizableObject);
 	check (RootObject);
 	
 	// Check that the object to be compiled has not already been compiled
