[Mutable] Make CO and COI factories public.

diff --git a/Engine/Plugins/Mutable/Source/CustomizableObjectEditor/Private/MuCOE/CustomizableObjectFactory.h b/Engine/Plugins/Mutable/Source/CustomizableObjectEditor/Private/MuCOE/CustomizableObjectFactory.h
deleted file mode 100644
diff --git a/Engine/Plugins/Mutable/Source/CustomizableObjectEditor/Private/MuCOE/CustomizableObjectInstanceFactory.h b/Engine/Plugins/Mutable/Source/CustomizableObjectEditor/Private/MuCOE/CustomizableObjectInstanceFactory.h
deleted file mode 100644
diff --git a/Engine/Plugins/Mutable/Source/CustomizableObjectEditor/Public/MuCOE/CustomizableObjectFactory.h b/Engine/Plugins/Mutable/Source/CustomizableObjectEditor/Public/MuCOE/CustomizableObjectFactory.h
new file mode 100644
--- /dev/null
+++ b/Engine/Plugins/Mutable/Source/CustomizableObjectEditor/Public/MuCOE/CustomizableObjectFactory.h
@@ -0,0 +1,38 @@
+// Copyright Epic Games, Inc. All Rights Reserved.
+
+#pragma once
+
+#include "Factories/Factory.h"
+
+#include "CustomizableObjectFactory.generated.h"
+
+enum class ECheckBoxState : uint8;
+
+namespace ESelectInfo { enum Type : int; }
+namespace ETextCommit { enum Type : int; }
+
+class FFeedbackContext;
+class FReply;
+class STextComboBox;
+class SWindow;
+class UClass;
+class UCustomizableObject;
+class UObject;
+class USkeletalMesh;
+
+struct EVisibility;
+
+
+UCLASS(MinimalAPI)
+class UCustomizableObjectFactory : public UFactory
+{
+	GENERATED_BODY()
+
+	UCustomizableObjectFactory();
+
+public:
+	// Begin UFactory Interface
+	virtual bool DoesSupportClass(UClass * Class) override;
+	virtual UClass* ResolveSupportedClass() override;
+	virtual UObject* FactoryCreateNew(UClass* Class, UObject* InParent,FName Name,EObjectFlags Flags,UObject* Context,FFeedbackContext* Warn) override;
+};
diff --git a/Engine/Plugins/Mutable/Source/CustomizableObjectEditor/Public/MuCOE/CustomizableObjectInstanceFactory.h b/Engine/Plugins/Mutable/Source/CustomizableObjectEditor/Public/MuCOE/CustomizableObjectInstanceFactory.h
new file mode 100644
--- /dev/null
+++ b/Engine/Plugins/Mutable/Source/CustomizableObjectEditor/Public/MuCOE/CustomizableObjectInstanceFactory.h
@@ -0,0 +1,25 @@
+// Copyright Epic Games, Inc. All Rights Reserved.
+
+#pragma once
+
+#include "ActorFactories/ActorFactory.h"
+#include "CustomizableObjectInstanceFactory.generated.h"
+
+class AActor;
+class FText;
+class UObject;
+class USkeletalMesh;
+struct FAssetData;
+
+UCLASS(MinimalAPI, config=Editor)
+class UCustomizableObjectInstanceFactory : public UActorFactory
+{
+	GENERATED_UCLASS_BODY()
+
+protected:
+	// UActorFactory interface
+	virtual void PostSpawnActor(UObject* Asset, AActor* NewActor) override;
+	virtual UObject* GetAssetFromActorInstance(AActor* ActorInstance) override;
+	virtual bool CanCreateActorFrom(const FAssetData& AssetData, FText& OutErrorMsg) override;
+	virtual FQuat AlignObjectToSurfaceNormal(const FVector& InSurfaceNormal, const FQuat& ActorRotation) const override;
+};
