[Mutable] Add functionaliy to only generate image mips if the data needed to generate is available.

- Image descriptor calculation now checks which is the first LOD it can generatate with the available data.
- Added CVar mutable.IgnoreFirstAvalilableLODCalculation to ignore the first lod calculated.
- Added DoesBlockExist() to the ModelReader interface.
- Removed MipsToSkip from GetExternalImageDesc as it was not used.

#rnx

diff --git a/Engine/Plugins/Mutable/Source/CustomizableObject/Internal/MuCO/UnrealMutableModelDiskStreamer.h b/Engine/Plugins/Mutable/Source/CustomizableObject/Internal/MuCO/UnrealMutableModelDiskStreamer.h
--- a/Engine/Plugins/Mutable/Source/CustomizableObject/Internal/MuCO/UnrealMutableModelDiskStreamer.h
+++ b/Engine/Plugins/Mutable/Source/CustomizableObject/Internal/MuCO/UnrealMutableModelDiskStreamer.h
@@ -60,7 +60,8 @@ public:
 	void EndStreaming();
 
 	// mu::ModelReader interface
-	FOperationID BeginReadBlock(const mu::FModel*, uint32 key0, void* pBuffer, uint64 size, TFunction<void(bool bSuccess)>* CompletionCallback) override;
+	bool DoesBlockExist(const mu::FModel*, uint32 BlockKey) override;
+	FOperationID BeginReadBlock(const mu::FModel*, uint32 BlockKey, void* pBuffer, uint64 size, TFunction<void(bool bSuccess)>* CompletionCallback) override;
 	bool IsReadCompleted(FOperationID) override;
 	bool EndRead(FOperationID) override;
 
diff --git a/Engine/Plugins/Mutable/Source/CustomizableObject/Private/MuCO/CustomizableObjectSystem.cpp b/Engine/Plugins/Mutable/Source/CustomizableObject/Private/MuCO/CustomizableObjectSystem.cpp
--- a/Engine/Plugins/Mutable/Source/CustomizableObject/Private/MuCO/CustomizableObjectSystem.cpp
+++ b/Engine/Plugins/Mutable/Source/CustomizableObject/Private/MuCO/CustomizableObjectSystem.cpp
@@ -119,6 +119,12 @@ TAutoConsoleVariable<bool> CVarEnableRealTimeMorphTargets(
 	true,
 	TEXT("Enable or disable generation of realtime morph targets."));
 
+
+TAutoConsoleVariable<bool> CVarIgnoreFirstAvailableLODCalculation(
+	TEXT("mutable.IgnoreFirstAvalilableLODCalculation"),
+	false,
+	TEXT("If set to true, ignores the first available LOD calculation to set the generated tetxure size."));
+
 #if WITH_EDITOR
 bool bEnableLODManagmentInEditor = false;
 
@@ -1756,7 +1762,7 @@ namespace impl
 		double StartTime,
 		const TSharedPtr<TArray<mu::FResourceID>>& ImagesInThisInstance,
 		int32 ImageIndex,
-		UE::Tasks::TTask<mu::FImageDesc> GetImageDescTask);
+		UE::Tasks::TTask<mu::FExtendedImageDesc> GetImageDescTask);
 
 	
 	/** Process the next Image. If there are no more Images, go to the end of the task. */
@@ -2724,7 +2730,7 @@ namespace impl
 		FInstanceUpdateData::FImage& Image = OperationData->InstanceUpdateData.Images[ImageIndex];
 		
 		// This should only be done when using progressive images, since GetImageDesc does some actual processing.
-		UE::Tasks::TTask<mu::FImageDesc> GetImageDescTask = OperationData->MutableSystem->GetImageDesc(OperationData->InstanceID, Image.ImageID);
+		UE::Tasks::TTask<mu::FExtendedImageDesc> GetImageDescTask = OperationData->MutableSystem->GetImageDesc(OperationData->InstanceID, Image.ImageID);
 		
 		UE::Tasks::AddNested(UE::Tasks::Launch(TEXT("Task_Mutable_GetImages_GetImage"), [=]()
 		{
@@ -2800,18 +2806,17 @@ namespace impl
 		Task_Mutable_GetImages_Loop(OperationData, StartTime, ImagesInThisInstance, ++ImageIndex);
 	}
 
-
 	/** See declaration. */
 	void Task_Mutable_GetImages_GetImage(
 		const TSharedRef<FUpdateContextPrivate>& OperationData,
 		double StartTime,
 		const TSharedPtr<TArray<mu::FResourceID>>& ImagesInThisInstance,
 		int32 ImageIndex,
-		UE::Tasks::TTask<mu::FImageDesc> GetImageDescTask)
+		UE::Tasks::TTask<mu::FExtendedImageDesc> GetImageDescTask)
 	{
 		MUTABLE_CPUPROFILER_SCOPE(Task_Mutable_GetImages_GetImage)
 
-		const mu::FImageDesc& ImageDesc = GetImageDescTask.GetResult();
+		const mu::FExtendedImageDesc& ImageDesc = GetImageDescTask.GetResult();
 
 		FInstanceUpdateData::FImage& Image = OperationData->InstanceUpdateData.Images[ImageIndex];
 
@@ -2820,18 +2825,23 @@ namespace impl
 		{
 			const uint16 MaxTextureSizeToGenerate = static_cast<uint16>(CustomizableObjectSystemPrivateData->MaxTextureSizeToGenerate);
 			const uint16 MaxSize = FMath::Max(ImageDesc.m_size[0], ImageDesc.m_size[1]);
-			uint16 Reduction = 1;
 
+			Image.BaseMip = 0;
 			if (MaxTextureSizeToGenerate > 0 && MaxSize > MaxTextureSizeToGenerate)
 			{
 				// Find the reduction factor, and the BaseMip of the texture.
 				const uint32 NextPowerOfTwo = FMath::RoundUpToPowerOfTwo(FMath::DivideAndRoundUp(MaxSize, MaxTextureSizeToGenerate));
-				Reduction = FMath::Max(NextPowerOfTwo, 2U); // At least divide the texture by a factor of two
+				uint16 Reduction = FMath::Max(NextPowerOfTwo, 2U); // At least divide the texture by a factor of two
 				Image.BaseMip = FMath::FloorLog2(Reduction);
 			}
 
-			Image.FullImageSizeX = ImageDesc.m_size[0] / Reduction;
-			Image.FullImageSizeY = ImageDesc.m_size[1] / Reduction;
+			if (!CVarIgnoreFirstAvailableLODCalculation.GetValueOnAnyThread())
+			{
+				Image.BaseMip = FMath::Max<uint8>(Image.BaseMip, ImageDesc.FirstLODAvailable);
+			}
+
+			Image.FullImageSizeX = ImageDesc.m_size[0] >> Image.BaseMip;
+			Image.FullImageSizeY = ImageDesc.m_size[1] >> Image.BaseMip;
 		}
 
 		const bool bCached = ImagesInThisInstance->Contains(Image.ImageID) || // See if it is cached from this same instance (can happen with LODs)
@@ -2879,7 +2889,6 @@ namespace impl
 		}
 	}
 
-
 	/** See declaration. */
 	void Task_Mutable_GetImages_Loop(
 		const TSharedRef<FUpdateContextPrivate>& OperationData,
diff --git a/Engine/Plugins/Mutable/Source/CustomizableObject/Private/MuCO/UnrealMutableImageProvider.cpp b/Engine/Plugins/Mutable/Source/CustomizableObject/Private/MuCO/UnrealMutableImageProvider.cpp
--- a/Engine/Plugins/Mutable/Source/CustomizableObject/Private/MuCO/UnrealMutableImageProvider.cpp
+++ b/Engine/Plugins/Mutable/Source/CustomizableObject/Private/MuCO/UnrealMutableImageProvider.cpp
@@ -391,7 +391,7 @@ TTuple<UE::Tasks::FTask, TFunction<void()>> FUnrealMutableResourceProvider::GetR
 
 
 // This should mantain parity with the descriptor of the images generated by GetImageAsync 
-mu::FImageDesc FUnrealMutableResourceProvider::GetImageDesc(FName Id, uint8 MipmapsToSkip)
+mu::FExtendedImageDesc FUnrealMutableResourceProvider::GetImageDesc(FName Id)
 {
 	MUTABLE_CPUPROFILER_SCOPE(FUnrealMutableImageProvider::GetImageDesc);
 
@@ -409,7 +409,10 @@ mu::FImageDesc FUnrealMutableResourceProvider::GetImageDesc(FName Id, uint8 Mipm
 	if (ImageInfo.Image)
 	{
 		// Easy case where the image was directly provided
-		return mu::FImageDesc(ImageInfo.Image->GetSize(), ImageInfo.Image->GetFormat(), ImageInfo.Image->GetLODCount());
+		
+		const mu::FImageDesc ImageDesc = mu::FImageDesc 
+				{ ImageInfo.Image->GetSize(), ImageInfo.Image->GetFormat(), (uint8)ImageInfo.Image->GetLODCount() };
+		return mu::FExtendedImageDesc { ImageDesc, 0 }; 
 	}
 
 #if WITH_EDITOR
@@ -417,36 +420,31 @@ mu::FImageDesc FUnrealMutableResourceProvider::GetImageDesc(FName Id, uint8 Mipm
 	{
 		const FTextureSource& Source = SourceTextureData->GetSource();
 
-		const int32 MipIndex = FMath::Min(static_cast<int32>(MipmapsToSkip),  Source.GetNumMips() - 1);
-
-		const mu::FImageSize ImageSize = mu::FImageSize(
-				Source.GetSizeX() >> MipIndex,
-				Source.GetSizeY() >> MipIndex);
+		const mu::FImageSize ImageSize = mu::FImageSize(Source.GetSizeX(), Source.GetSizeY());
+		const uint8 LODs = 1;
 
-		const int32 Lods = 1;
-		
-		return mu::FImageDesc(ImageSize, mu::EImageFormat::None, Lods);
+		return mu::FExtendedImageDesc { mu::FImageDesc { ImageSize, mu::EImageFormat::None, LODs }, 0 };
 	}
 #else
 	if (UTexture2D* TextureToLoad = ImageInfo.TextureToLoad)
 	{
-		mu::FImageDesc Result;	
+		mu::FExtendedImageDesc Result;	
 
 		// It's safe to access TextureToLoad because ExternalImagesLock guarantees that the data in GlobalExternalImages is valid,
 		// not being modified by the game thread at the moment and the texture cannot be GCed because of the AddReferencedObjects
 		// in the FUnrealMutableImageProvider
 
-		int32 MipIndex = MipmapsToSkip < TextureToLoad->GetPlatformData()->Mips.Num() ? MipmapsToSkip : TextureToLoad->GetPlatformData()->Mips.Num() - 1;
+		const int32 TextureToLoadNumMips = TextureToLoad->GetPlatformData()->Mips.Num();
 
-		// Mips in the mip tail are inlined and can't be streamed, find the smallest mip available.
-		for (; MipIndex > 0; --MipIndex)
+		int32 FirstLODAvailable = 0;
+		for (; FirstLODAvailable < TextureToLoadNumMips; ++FirstLODAvailable)
 		{
-			if (TextureToLoad->GetPlatformData()->Mips[MipIndex].BulkData.CanLoadFromDisk())
+			if (TextureToLoad->GetPlatformData()->Mips[FirstLODAvailable].BulkData.DoesExist())
 			{
 				break;
 			}
 		}
-
+	
 		// Texture format and the equivalent mutable format
 		const EPixelFormat Format = TextureToLoad->GetPlatformData()->PixelFormat;
 		const mu::EImageFormat MutableFormat = GetMutablePixelFormat(Format);
@@ -458,13 +456,12 @@ mu::FImageDesc FUnrealMutableResourceProvider::GetImageDesc(FName Id, uint8 Mipm
 			return CreateDummyDesc();
 		}
 
-		const mu::FImageSize ImageSize = mu::FImageSize(
-				TextureToLoad->GetSizeX() >> MipIndex,
-				TextureToLoad->GetSizeY() >> MipIndex);
+		const mu::FImageDesc ImageDesc = mu::FImageDesc 
+			{ mu::FImageSize(TextureToLoad->GetSizeX(), TextureToLoad->GetSizeY()), MutableFormat, 1 }; 
 
-		const int32 Lods = 1;
+		Result = mu::FExtendedImageDesc { ImageDesc, (uint8)FirstLODAvailable }; 
 
-		return mu::FImageDesc(ImageSize, MutableFormat, Lods);
+		return Result;
 	}
 #endif
 	
@@ -801,9 +798,9 @@ TSharedPtr<mu::FImage> FUnrealMutableResourceProvider::CreateDummy()
 }
 
 
-mu::FImageDesc FUnrealMutableResourceProvider::CreateDummyDesc()
+mu::FExtendedImageDesc FUnrealMutableResourceProvider::CreateDummyDesc()
 {
-	return DUMMY_IMAGE_DESC;
+	return mu::FExtendedImageDesc{ {DUMMY_IMAGE_DESC}, 0 };
 }
 
 
diff --git a/Engine/Plugins/Mutable/Source/CustomizableObject/Private/MuCO/UnrealMutableImageProvider.h b/Engine/Plugins/Mutable/Source/CustomizableObject/Private/MuCO/UnrealMutableImageProvider.h
--- a/Engine/Plugins/Mutable/Source/CustomizableObject/Private/MuCO/UnrealMutableImageProvider.h
+++ b/Engine/Plugins/Mutable/Source/CustomizableObject/Private/MuCO/UnrealMutableImageProvider.h
@@ -26,7 +26,7 @@ public:
 	// Thread: worker
 	virtual TTuple<UE::Tasks::FTask, TFunction<void()>> GetImageAsync(FName Id, uint8 MipmapsToSkip, TFunction<void(TSharedPtr<mu::FImage>)>& ResultCallback) override;
 	virtual TTuple<UE::Tasks::FTask, TFunction<void()>> GetReferencedImageAsync(const void* ModelPtr, int32 Id, uint8 MipmapsToSkip, TFunction<void(TSharedPtr<mu::FImage>)>& ResultCallback) override;
-	virtual mu::FImageDesc GetImageDesc(FName Id, uint8 MipmapsToSkip) override;
+	virtual mu::FExtendedImageDesc GetImageDesc(FName Id) override;
 	virtual TTuple<UE::Tasks::FTask, TFunction<void()>> GetMeshAsync(FName Id, TFunction<void(TSharedPtr<mu::FMesh>)>& ResultCallback) override;
 	virtual TTuple<UE::Tasks::FTask, TFunction<void()>> GetReferencedMeshAsync(const void* ModelPtr, int32 Id, TFunction<void(TSharedPtr<mu::FMesh>)>& ResultCallback) override;
 
@@ -88,11 +88,11 @@ private:
 	};
 
 	static inline const mu::FImageDesc DUMMY_IMAGE_DESC = 
-			mu::FImageDesc(mu::FImageSize(32, 32), mu::EImageFormat::RGBA_UByte, 1);
+			mu::FImageDesc {mu::FImageSize(32, 32), mu::EImageFormat::RGBA_UByte, 1};
 
 	/** This will be called if an image Id has been requested by Mutable core but it has not been provided by any provider. */
 	static TSharedPtr<mu::FImage> CreateDummy();
-	static mu::FImageDesc CreateDummyDesc();
+	static mu::FExtendedImageDesc CreateDummyDesc();
 
 	/** Map of Ids to external textures that may be required for any instance or Mutable texture mip under construction.
 	* This is only safely written from the game thread protected by the following critical section, and it
diff --git a/Engine/Plugins/Mutable/Source/CustomizableObject/Private/MuCO/UnrealMutableModelDiskStreamer.cpp b/Engine/Plugins/Mutable/Source/CustomizableObject/Private/MuCO/UnrealMutableModelDiskStreamer.cpp
--- a/Engine/Plugins/Mutable/Source/CustomizableObject/Private/MuCO/UnrealMutableModelDiskStreamer.cpp
+++ b/Engine/Plugins/Mutable/Source/CustomizableObject/Private/MuCO/UnrealMutableModelDiskStreamer.cpp
@@ -212,6 +212,64 @@ FAutoConsoleVariableRef CVarStreamPriority(
 	StreamPriority,
 	TEXT(""));
 
+namespace Private
+{
+FORCEINLINE FString MakeBlockFilePath(const FString& Prefix, const FMutableStreamableBlock& Block)
+{
+	FString Result;
+
+#if WITH_EDITOR
+	Result = Prefix;
+#else
+	Result = FString::Printf(TEXT("%s-%08x.mut"), *Prefix, Block.FileId);
+	if (Block.Flags == uint16(EMutableFileFlags::HighRes))
+	{
+		Result += TEXT(".high");
+	}
+#endif
+
+	return Result;
+}
+};
+
+bool FUnrealMutableModelBulkReader::DoesBlockExist(const mu::FModel* Model, uint32 Key)
+{
+	FObjectData* ObjectData = Objects.FindByPredicate(
+			[Model](const FObjectData& ObjectData) { return ObjectData.Model.Pin().Get() == Model; });
+
+	if (!ObjectData)
+	{
+		return false;
+	}
+
+	const FMutableStreamableBlock* Block = ObjectData->ModelStreamableBulkData->ModelStreamables.Find(Key);
+
+	if (!Block)
+	{
+		return false;
+	}
+
+	if (ObjectData->ModelStreamableBulkData->StreamableBulkData.IsValidIndex(Block->FileId))
+	{	
+		FByteBulkData& BulkData = ObjectData->ModelStreamableBulkData->StreamableBulkData[Block->FileId];
+		return BulkData.DoesExist();
+	}
+	else
+	{
+#if WITH_EDITOR
+		if (ObjectData->bIsStoredInDDC)
+		{
+			return true;
+		}
+		else
+#endif
+		{
+			const FString FilePath = Private::MakeBlockFilePath(ObjectData->BulkFilePrefix, *Block);
+			return FPlatformFileManager::Get().GetPlatformFile().FileExists(*FilePath);
+		}
+	}	
+}
+
 
 mu::FModelReader::FOperationID FUnrealMutableModelBulkReader::BeginReadBlock(const mu::FModel* Model, uint32 Key, void* Buffer, uint64 Size, TFunction<void(bool bSuccess)>* CompletionCallback)
 {
@@ -221,7 +279,7 @@ mu::FModelReader::FOperationID FUnrealMutableModelBulkReader::BeginReadBlock(con
 
 	// Find the object we are streaming for
 	FObjectData* ObjectData = Objects.FindByPredicate(
-		[Model](const FObjectData& d) {return d.Model.Pin().Get() == Model; });
+			[Model](const FObjectData& d) { return d.Model.Pin().Get() == Model; });
 
 	if (!ObjectData)
 	{
@@ -359,16 +417,7 @@ mu::FModelReader::FOperationID FUnrealMutableModelBulkReader::BeginReadBlock(con
 				TSharedPtr<IAsyncReadFileHandle>& Found = ObjectData->ReadFileHandles.FindOrAdd(Block->FileId);
 				if (!Found)
 				{
-#if WITH_EDITOR
-					FString FilePath = ObjectData->BulkFilePrefix;
-#else
-					FString FilePath = FString::Printf(TEXT("%s-%08x.mut"), *ObjectData->BulkFilePrefix, Block->FileId);
-					if (Block->Flags == uint16(EMutableFileFlags::HighRes))
-					{
-						FilePath += TEXT(".high");
-					}
-#endif
-
+					FString FilePath = Private::MakeBlockFilePath(ObjectData->BulkFilePrefix, *Block);
 					Found = MakeShareable(FPlatformFileManager::Get().GetPlatformFile().OpenAsyncRead(*FilePath));
 
 					if (!Found)
diff --git a/Engine/Plugins/Mutable/Source/CustomizableObjectEditor/Private/MuCOE/SMutableCodeViewer.cpp b/Engine/Plugins/Mutable/Source/CustomizableObjectEditor/Private/MuCOE/SMutableCodeViewer.cpp
--- a/Engine/Plugins/Mutable/Source/CustomizableObjectEditor/Private/MuCOE/SMutableCodeViewer.cpp
+++ b/Engine/Plugins/Mutable/Source/CustomizableObjectEditor/Private/MuCOE/SMutableCodeViewer.cpp
@@ -2529,9 +2529,9 @@ namespace
 			return MakeTuple(UE::Tasks::MakeCompletedTask<void>(), []() -> void {});
 		}
 
-		virtual mu::FImageDesc GetImageDesc(FName Id, uint8 MipmapsToSkip) override
+		virtual mu::FExtendedImageDesc GetImageDesc(FName Id) override
 		{
-			return IMAGE_DESC;
+			return mu::FExtendedImageDesc{ IMAGE_DESC };
 		}
 
 
diff --git a/Engine/Plugins/Mutable/Source/MutableRuntime/Private/MuR/CodeRunner.cpp b/Engine/Plugins/Mutable/Source/MutableRuntime/Private/MuR/CodeRunner.cpp
--- a/Engine/Plugins/Mutable/Source/MutableRuntime/Private/MuR/CodeRunner.cpp
+++ b/Engine/Plugins/Mutable/Source/MutableRuntime/Private/MuR/CodeRunner.cpp
@@ -225,7 +225,7 @@ namespace mu
 
 	
     //---------------------------------------------------------------------------------------------
-	FImageDesc CodeRunner::GetExternalImageDesc(FName Id, uint8 MipmapsToSkip)
+	FExtendedImageDesc CodeRunner::GetExternalImageDesc(FName Id)
 	{
 		MUTABLE_CPUPROFILER_SCOPE(GetExternalImageDesc);
 
@@ -233,7 +233,7 @@ namespace mu
 
 		if (System->ExternalResourceProvider)
 		{
-			return System->ExternalResourceProvider->GetImageDesc(Id, MipmapsToSkip);
+			return System->ExternalResourceProvider->GetImageDesc(Id);
 		}
 		else
 		{
@@ -241,7 +241,7 @@ namespace mu
 			check(false);
 		}
 
-		return FImageDesc();
+		return FExtendedImageDesc();
 	}
 
 	
@@ -6000,84 +6000,121 @@ namespace mu
     }
 
 	//---------------------------------------------------------------------------------------------
-	void CodeRunner::RunCodeImageDesc(const FScheduledOp& item, const FParameters* pParams, const FModel* pModel,  uint32 lodMask )
+	void CodeRunner::RunCodeImageDesc(const FScheduledOp& Item, const FParameters* pParams, const FModel* pModel,  uint32 lodMask )
 	{
 		MUTABLE_CPUPROFILER_SCOPE(RunCodeImageDesc);
 
-		check(item.Type == FScheduledOp::EType::ImageDesc);
+		check(Item.Type == FScheduledOp::EType::ImageDesc);
 
-		// Ensure there is room for the result in the heap.
-		if (item.CustomState >= uint32(HeapData.Num()))
+		// Add the result heap address if not present.
+		if (HeapImageDesc.Num() == 0)
 		{
-			HeapImageDesc.SetNum(item.CustomState+1);
+			HeapImageDesc.AddDefaulted();
 		}
 
+		// Item.CustomState can point to HeapData or HeapImageDesc depending on the operation.
+		// Can not check here. Assume access to the arrays will do the check. 
+		// We could set the Item.CustomState highest bit to distinguish if needed.
+		//check((int32)Item.CustomState < HeapImageDesc.Num());
 
 		const FProgram& Program = Model->GetPrivate()->Program;
 
-		EOpType type = Program.GetOpType(item.At);
-		switch (type)
+		EOpType Type = Program.GetOpType(Item.At);
+		switch (Type)
 		{
 
 		case EOpType::IM_CONSTANT:
 		{
-			check(item.Stage == 0);
-			OP::ResourceConstantArgs args = Program.GetOpArgs<OP::ResourceConstantArgs>(item.At);
-			int32 ImageIndex = args.value;
+			check(Item.Stage == 0);
+			OP::ResourceConstantArgs Args = Program.GetOpArgs<OP::ResourceConstantArgs>(Item.At);
+			int32 ImageIndex = Args.value;
+
+			FExtendedImageDesc& Result = HeapImageDesc[Item.CustomState];
 
-			FImageDesc& Result = HeapImageDesc[item.CustomState];
 			Result.m_format = Program.ConstantImages[ImageIndex].ImageFormat;
 			Result.m_size[0] = Program.ConstantImages[ImageIndex].ImageSizeX;
 			Result.m_size[1] = Program.ConstantImages[ImageIndex].ImageSizeY;
 			Result.m_lods = Program.ConstantImages[ImageIndex].LODCount;
-			StoreValidDesc(item);
+
+			int32 LODIndexIndex = Program.ConstantImages[ImageIndex].FirstIndex;
+			int32 LODIndex = 0;
+			for (; LODIndex < Result.m_lods; ++LODIndex)
+			{
+				const int32 CurrentIndexIndex = LODIndexIndex + LODIndex;
+				const FConstantResourceIndex CurrentIndex = Program.ConstantImageLODIndices[CurrentIndexIndex];
+
+				bool bIsLODAvailable = false;
+				if (!CurrentIndex.Streamable)
+				{
+					bIsLODAvailable = true;
+				}
+				else
+				{
+					uint32 RomId = CurrentIndex.Index;
+					bIsLODAvailable = System->StreamInterface->DoesBlockExist(pModel, RomId);
+				}
+
+				if (bIsLODAvailable)
+				{
+					break;
+				}
+			}
+
+			Result.FirstLODAvailable = LODIndex;
+
+			StoreValidDesc(Item);
 			break;
 		}
 
 		case EOpType::IM_PARAMETER:
 		{
-			check(item.Stage == 0);
-			OP::ParameterArgs args = Program.GetOpArgs<OP::ParameterArgs>(item.At);
-			FName Id = pParams->GetImageValue(args.variable);
-			uint8 MipsToSkip = item.ExecutionOptions;
-			HeapImageDesc[item.CustomState] = GetExternalImageDesc(Id, MipsToSkip);
-			StoreValidDesc(item);
+			check(Item.Stage == 0);
+
+			OP::ParameterArgs Args = Program.GetOpArgs<OP::ParameterArgs>(Item.At);
+			FName Id = pParams->GetImageValue(Args.variable);
+			HeapImageDesc[Item.CustomState] = GetExternalImageDesc(Id);
+			StoreValidDesc(Item);
 			break;
 		}
 
 		case EOpType::IM_REFERENCE:
 		{
-			check(item.Stage == 0);
-			OP::ResourceReferenceArgs Args = Program.GetOpArgs<OP::ResourceReferenceArgs>(item.At);
-			FImageDesc& Result = HeapImageDesc[item.CustomState];
-			Result = Args.ImageDesc;
-			StoreValidDesc(item);
+			check(Item.Stage == 0);
+
+			OP::ResourceReferenceArgs Args = Program.GetOpArgs<OP::ResourceReferenceArgs>(Item.At);
+			FExtendedImageDesc& Result = HeapImageDesc[Item.CustomState];
+			Result = FExtendedImageDesc{ Args.ImageDesc, 0 };
+
+			StoreValidDesc(Item);
 			break;
 		}
 
 		case EOpType::IM_CONDITIONAL:
 		{
-			OP::ConditionalArgs args = Program.GetOpArgs<OP::ConditionalArgs>(item.At);
-			switch (item.Stage)
+			OP::ConditionalArgs Args = Program.GetOpArgs<OP::ConditionalArgs>(Item.At);
+			switch (Item.Stage)
 			{
 			case 0:
 			{
 				// We need to run the full condition result
-				FScheduledOp FullConditionOp(args.condition, item);
+				FScheduledOp FullConditionOp(Args.condition, Item);
 				FullConditionOp.Type = FScheduledOp::EType::Full;
-				AddOp(FScheduledOp(item.At, item, 1), FullConditionOp);
+
+				AddOp(FScheduledOp(Item.At, Item, 1, Item.CustomState), FullConditionOp);
 				break;
 			}
 
 			case 1:
 			{
-				bool value = LoadBool(FCacheAddress(args.condition, item.ExecutionIndex, item.ExecutionOptions));
-				OP::ADDRESS resultAt = value ? args.yes : args.no;
-				AddOp(FScheduledOp(item.At, item, 2), FScheduledOp(resultAt, item));
+				bool bValue = LoadBool(FCacheAddress(Args.condition, Item.ExecutionIndex, Item.ExecutionOptions));
+				OP::ADDRESS ResultAt = bValue ? Args.yes : Args.no;
+
+				AddOp(FScheduledOp(Item.At, Item, 2, Item.CustomState), 
+						FScheduledOp(ResultAt, Item, 0, Item.CustomState));
 				break;
 			}
 
-			case 2: StoreValidDesc(item); break;
+			case 2: StoreValidDesc(Item); break;
 			default: check(false);
 			}
 			break;
@@ -6085,34 +6122,34 @@ namespace mu
 
 		case EOpType::IM_SWITCH:
 		{
-			const uint8* data = Program.GetOpArgsPointer(item.At);
+			const uint8* Data = Program.GetOpArgsPointer(Item.At);
 		
 			OP::ADDRESS VarAddress;
-			FMemory::Memcpy( &VarAddress, data, sizeof(OP::ADDRESS));
-			data += sizeof(OP::ADDRESS);
+			FMemory::Memcpy(&VarAddress, Data, sizeof(OP::ADDRESS));
+			Data += sizeof(OP::ADDRESS);
 
 			OP::ADDRESS DefAddress;
-			FMemory::Memcpy( &DefAddress, data, sizeof(OP::ADDRESS));
-			data += sizeof(OP::ADDRESS);
+			FMemory::Memcpy(&DefAddress, Data, sizeof(OP::ADDRESS));
+			Data += sizeof(OP::ADDRESS);
 
 			uint32 CaseCount;
-			FMemory::Memcpy( &CaseCount, data, sizeof(uint32));
-			data += sizeof(uint32);
+			FMemory::Memcpy(&CaseCount, Data, sizeof(uint32));
+			Data += sizeof(uint32);
 	
-			switch (item.Stage)
+			switch (Item.Stage)
 			{
 			case 0:
 			{
 				if (VarAddress)
 				{
 					// We need to run the full condition result
-					FScheduledOp FullVariableOp(VarAddress, item);
+					FScheduledOp FullVariableOp(VarAddress, Item);
 					FullVariableOp.Type = FScheduledOp::EType::Full;
-					AddOp(FScheduledOp(item.At, item, 1), FullVariableOp);
+					AddOp(FScheduledOp(Item.At, Item, 1, Item.CustomState), FullVariableOp);
 				}
 				else
 				{
-					StoreValidDesc(item);
+					StoreValidDesc(Item);
 				}
 				break;
 			}
@@ -6120,33 +6157,33 @@ namespace mu
 			case 1:
 			{
 				// Get the variable result
-				int32 var = LoadInt(FCacheAddress(VarAddress, item));
+				int32 Var = LoadInt(FCacheAddress(VarAddress, Item));
 
-				OP::ADDRESS valueAt = DefAddress;
+				OP::ADDRESS ValueAt = DefAddress;
 				for (uint32 C = 0; C < CaseCount; ++C)
 				{
 					int32 Condition;
-					FMemory::Memcpy( &Condition, data, sizeof(int32) );
-					data += sizeof(int32);
+					FMemory::Memcpy(&Condition, Data, sizeof(int32));
+					Data += sizeof(int32);
 					
 					OP::ADDRESS At;
-					FMemory::Memcpy( &At, data, sizeof(OP::ADDRESS) );
-					data += sizeof(OP::ADDRESS);
+					FMemory::Memcpy(&At, Data, sizeof(OP::ADDRESS));
+					Data += sizeof(OP::ADDRESS);
 
-					if (At && var == (int32)Condition)
+					if (At && Var == (int32)Condition)
 					{
-						valueAt = At;
+						ValueAt = At;
 						break;
 					}
 				}
 
-				AddOp(FScheduledOp(item.At, item, 2, valueAt),
-					  FScheduledOp(valueAt, item));
+				AddOp(FScheduledOp(Item.At, Item, 2, Item.CustomState),
+						FScheduledOp(ValueAt, Item, 0, Item.CustomState));
 
 				break;
 			}
 
-			case 2: StoreValidDesc(item); break;
+			case 2: StoreValidDesc(Item); break;
 			default: check(false); break;
 			}
 			break;
@@ -6154,11 +6191,48 @@ namespace mu
 
 		case EOpType::IM_LAYERCOLOUR:
 		{
-			OP::ImageLayerColourArgs args = Program.GetOpArgs<OP::ImageLayerColourArgs>(item.At);
-			switch (item.Stage)
+			OP::ImageLayerColourArgs Args = Program.GetOpArgs<OP::ImageLayerColourArgs>(Item.At);
+			switch (Item.Stage)
+			{
+			case 0:
+			{
+				FScheduledOpData ResultArrayData;
+				ResultArrayData.ReturnArray.ReturnAddress = Item.CustomState;
+				ResultArrayData.ReturnArray.Begin = HeapImageDesc.Num();
+				ResultArrayData.ReturnArray.Num = 2;
+
+				HeapImageDesc.AddDefaulted(ResultArrayData.ReturnArray.Num);
+
+				int32 ResultArrayAddress = HeapData.Add(ResultArrayData);
+
+				AddOp(FScheduledOp(Item.At, Item, 1, ResultArrayAddress), 
+						FScheduledOp(Args.base, Item, 0, ResultArrayData.ReturnArray.Begin + 0),
+						FScheduledOp(Args.mask, Item, 0, ResultArrayData.ReturnArray.Begin + 1));
+
+				break;
+			}
+			case 1:
 			{
-			case 0: AddOp(FScheduledOp(item.At, item, 1), FScheduledOp(args.base, item)); break;
-			case 1: StoreValidDesc(item); break;
+				FScheduledOpData ResultArrayData = HeapData[Item.CustomState];
+				TArrayView<FExtendedImageDesc> DepResults(
+						HeapImageDesc.GetData() + ResultArrayData.ReturnArray.Begin, 
+						ResultArrayData.ReturnArray.Num);
+				
+				check(DepResults.Num() == 2);
+	
+				FExtendedImageDesc& BaseImageDesc = DepResults[0];
+				FExtendedImageDesc& MaskImageDesc = DepResults[1];
+
+				FExtendedImageDesc& Result = HeapImageDesc[ResultArrayData.ReturnArray.ReturnAddress]; 
+				Result = BaseImageDesc;
+
+				Result.FirstLODAvailable = FMath::Max(
+							Result.FirstLODAvailable, MaskImageDesc.FirstLODAvailable);
+
+				StoreValidDesc(Item); 
+
+				break;
+			}
 			default: check(false);
 			}
 			break;
@@ -6166,11 +6240,50 @@ namespace mu
 
 		case EOpType::IM_LAYER:
 		{
-			OP::ImageLayerArgs args = Program.GetOpArgs<OP::ImageLayerArgs>(item.At);
-			switch (item.Stage)
+			OP::ImageLayerArgs Args = Program.GetOpArgs<OP::ImageLayerArgs>(Item.At);
+			switch (Item.Stage)
+			{
+			case 0:
+			{
+				FScheduledOpData ResultArrayData;
+				ResultArrayData.ReturnArray.ReturnAddress = Item.CustomState;
+				ResultArrayData.ReturnArray.Begin = HeapImageDesc.Num();
+				ResultArrayData.ReturnArray.Num = 3;
+
+				HeapImageDesc.AddDefaulted(ResultArrayData.ReturnArray.Num);
+
+				int32 ResultArrayAddress = HeapData.Add(ResultArrayData);
+
+				AddOp(FScheduledOp(Item.At, Item, 1, ResultArrayAddress), 
+						FScheduledOp(Args.base, Item, 0, ResultArrayData.ReturnArray.Begin + 0),
+						FScheduledOp(Args.mask, Item, 0, ResultArrayData.ReturnArray.Begin + 1),
+						FScheduledOp(Args.blended, Item, 0, ResultArrayData.ReturnArray.Begin + 2));
+				break;
+			}
+			case 1:
 			{
-			case 0: AddOp(FScheduledOp(item.At, item, 1), FScheduledOp(args.base, item)); break;
-			case 1: StoreValidDesc(item); break;
+				FScheduledOpData ResultArrayData = HeapData[Item.CustomState];
+				TArrayView<FExtendedImageDesc> DepResults(
+						HeapImageDesc.GetData() + ResultArrayData.ReturnArray.Begin, 
+						ResultArrayData.ReturnArray.Num);
+				
+				check(DepResults.Num() == 3);
+	
+				FExtendedImageDesc& BaseImageDesc = DepResults[0];
+				FExtendedImageDesc& MaskImageDesc = DepResults[1];
+				FExtendedImageDesc& BlendedImageDesc = DepResults[2];
+
+				FExtendedImageDesc& Result = HeapImageDesc[ResultArrayData.ReturnArray.ReturnAddress]; 
+				Result = BaseImageDesc;
+
+				Result.FirstLODAvailable = FMath::Max3(
+							Result.FirstLODAvailable, 
+							MaskImageDesc.FirstLODAvailable,
+							BlendedImageDesc.FirstLODAvailable);
+
+				StoreValidDesc(Item); 
+				break;
+			}
 			default: check(false);
 			}
 			break;
@@ -6178,11 +6291,18 @@ namespace mu
 
 		case EOpType::IM_MULTILAYER:
 		{
-			OP::ImageMultiLayerArgs args = Program.GetOpArgs<OP::ImageMultiLayerArgs>(item.At);
-			switch (item.Stage)
+			OP::ImageMultiLayerArgs Args = Program.GetOpArgs<OP::ImageMultiLayerArgs>(Item.At);
+			switch (Item.Stage)
 			{
-			case 0: AddOp(FScheduledOp(item.At, item, 1), FScheduledOp(args.base, item)); break;
-			case 1: StoreValidDesc(item); break;
+			case 0:
+			{
+				//TODO: For now multilayer operations will only check the base to get the descriptor.
+				// but all iterations should be checked for available mips. 
+				AddOp(FScheduledOp(Item.At, Item, 1, Item.CustomState), 
+						FScheduledOp(Args.base, Item, 0, Item.CustomState));
+				break;
+			}
+			case 1: StoreValidDesc(Item); break;
 			default: check(false);
 			}
 			break;
@@ -6190,11 +6310,16 @@ namespace mu
 
 		case EOpType::IM_NORMALCOMPOSITE:
 		{
-			OP::ImageNormalCompositeArgs args = Program.GetOpArgs<OP::ImageNormalCompositeArgs>(item.At);
-			switch (item.Stage)
+			OP::ImageNormalCompositeArgs Args = Program.GetOpArgs<OP::ImageNormalCompositeArgs>(Item.At);
+			switch (Item.Stage)
 			{
-			case 0: AddOp(FScheduledOp(item.At, item, 1), FScheduledOp(args.base, item)); break;
-			case 1: StoreValidDesc(item); break;
+			case 0: 
+			{
+				AddOp(FScheduledOp(Item.At, Item, 1, Item.CustomState), 
+						FScheduledOp(Args.base, Item, 0, Item.CustomState)); 
+				break;
+			}
+			case 1: StoreValidDesc(Item); break;
 			default: check(false);
 			}
 			break;
@@ -6202,26 +6327,27 @@ namespace mu
 
 		case EOpType::IM_PIXELFORMAT:
 		{
-			OP::ImagePixelFormatArgs args = Program.GetOpArgs<OP::ImagePixelFormatArgs>(item.At);
-			switch (item.Stage)
+			OP::ImagePixelFormatArgs Args = Program.GetOpArgs<OP::ImagePixelFormatArgs>(Item.At);
+			switch (Item.Stage)
 			{
 			case 0:
-				AddOp(FScheduledOp(item.At, item, 1), FScheduledOp(args.source, item));
+				AddOp(FScheduledOp(Item.At, Item, 1, Item.CustomState), 
+						FScheduledOp(Args.source, Item, 0, Item.CustomState));
 				break;
 
 			case 1:
 			{
 				// Update directly in the heap
-				EImageFormat OldFormat = HeapImageDesc[item.CustomState].m_format;
-				EImageFormat NewFormat = args.format;
-				if (args.formatIfAlpha != EImageFormat::None
+				EImageFormat OldFormat = HeapImageDesc[Item.CustomState].m_format;
+				EImageFormat NewFormat = Args.format;
+				if (Args.formatIfAlpha != EImageFormat::None
 					&&
 					GetImageFormatData(OldFormat).Channels > 3)
 				{
-					NewFormat = args.formatIfAlpha;
+					NewFormat = Args.formatIfAlpha;
 				}
-				HeapImageDesc[item.CustomState].m_format = NewFormat;				
-				StoreValidDesc(item);
+				HeapImageDesc[Item.CustomState].m_format = NewFormat;				
+				StoreValidDesc(Item);
 				break;
 			}
 
@@ -6234,18 +6360,19 @@ namespace mu
 
 		case EOpType::IM_MIPMAP:
 		{
-			OP::ImageMipmapArgs args = Program.GetOpArgs<OP::ImageMipmapArgs>(item.At);
-			switch (item.Stage)
+			OP::ImageMipmapArgs Args = Program.GetOpArgs<OP::ImageMipmapArgs>(Item.At);
+			switch (Item.Stage)
 			{
 			case 0:
-				AddOp(FScheduledOp(item.At, item, 1), FScheduledOp(args.source, item));
+				AddOp(FScheduledOp(Item.At, Item, 1, Item.CustomState), 
+						FScheduledOp(Args.source, Item, 0, Item.CustomState));
 				break;
 
 			case 1:
 			{
 				// Somewhat synched with Full op execution code.
-				FImageDesc BaseDesc = HeapImageDesc[item.CustomState];
-				int32 LevelCount = args.levels;
+				FExtendedImageDesc BaseDesc = HeapImageDesc[Item.CustomState];
+				int32 LevelCount = Args.levels;
 				int32 MaxLevelCount = FImage::GetMipmapCount(BaseDesc.m_size[0], BaseDesc.m_size[1]);
 				if (LevelCount == 0)
 				{
@@ -6263,8 +6390,8 @@ namespace mu
 				LevelCount = FMath::Max(StartLevel, LevelCount);
 
 				// Update result.
-				HeapImageDesc[item.CustomState].m_lods = LevelCount;
-				StoreValidDesc(item);
+				HeapImageDesc[Item.CustomState].m_lods = LevelCount;
+				StoreValidDesc(Item);
 				break;
 			}
 
@@ -6276,17 +6403,18 @@ namespace mu
 
 		case EOpType::IM_RESIZE:
 		{
-			OP::ImageResizeArgs args = Program.GetOpArgs<OP::ImageResizeArgs>(item.At);
-			switch (item.Stage)
+			OP::ImageResizeArgs Args = Program.GetOpArgs<OP::ImageResizeArgs>(Item.At);
+			switch (Item.Stage)
 			{
 			case 0:
-				AddOp(FScheduledOp(item.At, item, 1), FScheduledOp(args.source, item));
+				AddOp(FScheduledOp(Item.At, Item, 1, Item.CustomState), 
+						FScheduledOp(Args.source, Item, 0, Item.CustomState));
 				break;
 
 			case 1:
-				HeapImageDesc[item.CustomState].m_size[0] = args.size[0];
-				HeapImageDesc[item.CustomState].m_size[1] = args.size[1];
-				StoreValidDesc(item);
+				HeapImageDesc[Item.CustomState].m_size[0] = Args.size[0];
+				HeapImageDesc[Item.CustomState].m_size[1] = Args.size[1];
+				StoreValidDesc(Item);
 				break;
 
 			default:
@@ -6297,30 +6425,44 @@ namespace mu
 
 		case EOpType::IM_RESIZELIKE:
 		{
-			OP::ImageResizeLikeArgs args = Program.GetOpArgs<OP::ImageResizeLikeArgs>(item.At);
-			switch (item.Stage)
+			OP::ImageResizeLikeArgs Args = Program.GetOpArgs<OP::ImageResizeLikeArgs>(Item.At);
+			switch (Item.Stage)
 			{
 			case 0:
 			{
-				int32 ResultAndBaseDesc = item.CustomState;
-				int32 SourceDescAddress = HeapImageDesc.Add({});
-				FScheduledOpData Data;
-				Data.ResizeLike.ResultDescAt = ResultAndBaseDesc;
-				Data.ResizeLike.SourceDescAt = SourceDescAddress;
-				int32 SecondStageData = HeapData.Add(Data);
-				AddOp(FScheduledOp(item.At, item, 1, SecondStageData),
-					FScheduledOp(args.source, item, 0, ResultAndBaseDesc),
-					FScheduledOp(args.sizeSource, item, 0, SourceDescAddress));
+				FScheduledOpData ResultArrayData;
+				ResultArrayData.ReturnArray.ReturnAddress = Item.CustomState;
+				ResultArrayData.ReturnArray.Begin = HeapImageDesc.Num();
+				ResultArrayData.ReturnArray.Num = 2;
+
+				HeapImageDesc.AddDefaulted(ResultArrayData.ReturnArray.Num);
+
+				int32 ResultArrayAddress = HeapData.Add(ResultArrayData);
+
+				AddOp(FScheduledOp(Item.At, Item, 1, ResultArrayAddress), 
+						FScheduledOp(Args.source, Item, 0, ResultArrayData.ReturnArray.Begin + 0),
+						FScheduledOp(Args.sizeSource, Item, 0, ResultArrayData.ReturnArray.Begin + 1));
+
 				break;
 			}
 
 			case 1:
 			{
-				const FScheduledOpData& SecondStageData = HeapData[ item.CustomState ];
-				FImageDesc& ResultAndBaseDesc = HeapImageDesc[SecondStageData.ResizeLike.ResultDescAt];
-				const FImageDesc& SourceDesc = HeapImageDesc[SecondStageData.ResizeLike.SourceDescAt];
-				ResultAndBaseDesc.m_size = SourceDesc.m_size;
-				StoreValidDesc(item);
+				FScheduledOpData ResultArrayData = HeapData[Item.CustomState];
+				TArrayView<FExtendedImageDesc> DepResults(
+						HeapImageDesc.GetData() + ResultArrayData.ReturnArray.Begin, 
+						ResultArrayData.ReturnArray.Num);
+				
+				check(DepResults.Num() == 2);
+	
+				FExtendedImageDesc& SourceImageDesc = DepResults[0];
+				FExtendedImageDesc& SizeImageDesc = DepResults[1];
+
+				FExtendedImageDesc& Result = HeapImageDesc[ResultArrayData.ReturnArray.ReturnAddress]; 
+				Result = SourceImageDesc;
+				Result.m_size = SizeImageDesc.m_size;
+
+				StoreValidDesc(Item);
 				break;
 			}
 
@@ -6333,21 +6475,22 @@ namespace mu
 
 		case EOpType::IM_RESIZEREL:
 		{
-			OP::ImageResizeRelArgs args = Program.GetOpArgs<OP::ImageResizeRelArgs>(item.At);
-			switch (item.Stage)
+			OP::ImageResizeRelArgs Args = Program.GetOpArgs<OP::ImageResizeRelArgs>(Item.At);
+			switch (Item.Stage)
 			{
 			case 0:
-				AddOp(FScheduledOp(item.At, item, 1), FScheduledOp(args.source, item));
+				AddOp(FScheduledOp(Item.At, Item, 1, Item.CustomState), 
+						FScheduledOp(Args.source, Item, 0, Item.CustomState));
 				break;
 
 			case 1:
 			{
-				FImageDesc& ResultAndBaseDesc = HeapImageDesc[item.CustomState];
-				FImageSize destSize(
-					uint16(ResultAndBaseDesc.m_size[0] * args.factor[0] + 0.5f),
-					uint16(ResultAndBaseDesc.m_size[1] * args.factor[1] + 0.5f));
-				ResultAndBaseDesc.m_size = destSize;
-				StoreValidDesc(item);
+				FExtendedImageDesc& ResultAndBaseDesc = HeapImageDesc[Item.CustomState];
+				FImageSize DestSize(
+					uint16(ResultAndBaseDesc.m_size[0] * Args.factor[0] + 0.5f),
+					uint16(ResultAndBaseDesc.m_size[1] * Args.factor[1] + 0.5f));
+				ResultAndBaseDesc.m_size = DestSize;
+				StoreValidDesc(Item);
 				break;
 			}
 
@@ -6360,43 +6503,43 @@ namespace mu
 
 		case EOpType::IM_BLANKLAYOUT:
 		{
-			OP::ImageBlankLayoutArgs args = Program.GetOpArgs<OP::ImageBlankLayoutArgs>(item.At);
-			switch (item.Stage)
+			OP::ImageBlankLayoutArgs Args = Program.GetOpArgs<OP::ImageBlankLayoutArgs>(Item.At);
+			switch (Item.Stage)
 			{
 			case 0:
 			{
 				// We need to run the full layout
-				FScheduledOp FullLayoutOp(args.layout, item);
+				FScheduledOp FullLayoutOp(Args.layout, Item);
 				FullLayoutOp.Type = FScheduledOp::EType::Full;
-				AddOp(FScheduledOp(item.At, item, 1), FullLayoutOp);
+				AddOp(FScheduledOp(Item.At, Item, 1, Item.CustomState), FullLayoutOp);
 				break;
 			}
 
 			case 1:
 			{
-				TSharedPtr<const FLayout> pLayout = LoadLayout(FCacheAddress(args.layout, item));
+				TSharedPtr<const FLayout> pLayout = LoadLayout(FCacheAddress(Args.layout, Item));
 
 				FIntPoint SizeInBlocks = pLayout->GetGridSize();
-				FIntPoint BlockSizeInPixels(args.blockSize[0], args.blockSize[1]);
+				FIntPoint BlockSizeInPixels(Args.blockSize[0], Args.blockSize[1]);
 				FIntPoint ImageSizeInPixels = SizeInBlocks * BlockSizeInPixels;
 
-				FImageDesc& ResultAndBaseDesc = HeapImageDesc[item.CustomState];
-				FImageSize destSize(uint16(ImageSizeInPixels.X), uint16(ImageSizeInPixels.Y));
-				ResultAndBaseDesc.m_size = destSize;
-				ResultAndBaseDesc.m_format = args.format;
+				FExtendedImageDesc& ResultAndBaseDesc = HeapImageDesc[Item.CustomState];
+				FImageSize DestSize(uint16(ImageSizeInPixels.X), uint16(ImageSizeInPixels.Y));
+				ResultAndBaseDesc.m_size = DestSize;
+				ResultAndBaseDesc.m_format = Args.format;
 				
-				if (args.generateMipmaps)
+				if (Args.generateMipmaps)
 				{
-					if (args.mipmapCount == 0)
+					if (Args.mipmapCount == 0)
 					{
 						ResultAndBaseDesc.m_lods = FImage::GetMipmapCount(ImageSizeInPixels.X, ImageSizeInPixels.Y);
 					}
 					else
 					{
-						ResultAndBaseDesc.m_lods = args.mipmapCount;
+						ResultAndBaseDesc.m_lods = Args.mipmapCount;
 					}
 				}
-				StoreValidDesc(item);
+				StoreValidDesc(Item);
 				break;
 			}
 
@@ -6409,11 +6552,46 @@ namespace mu
 
 		case EOpType::IM_COMPOSE:
 		{
-			OP::ImageComposeArgs args = Program.GetOpArgs<OP::ImageComposeArgs>(item.At);
-			switch (item.Stage)
+			OP::ImageComposeArgs Args = Program.GetOpArgs<OP::ImageComposeArgs>(Item.At);
+			switch (Item.Stage)
+			{
+			case 0: 
 			{
-			case 0: AddOp(FScheduledOp(item.At, item, 1), FScheduledOp(args.base, item)); break;
-			case 1: StoreValidDesc(item); break;
+				FScheduledOpData ResultArrayData;
+				ResultArrayData.ReturnArray.ReturnAddress = Item.CustomState;
+				ResultArrayData.ReturnArray.Begin = HeapImageDesc.Num();
+				ResultArrayData.ReturnArray.Num = 2;
+
+				HeapImageDesc.AddDefaulted(ResultArrayData.ReturnArray.Num);
+
+				int32 ResultArrayAddress = HeapData.Add(ResultArrayData);
+
+				AddOp(FScheduledOp(Item.At, Item, 1, ResultArrayAddress), 
+						FScheduledOp(Args.base, Item, 0, ResultArrayData.ReturnArray.Begin + 0),
+						FScheduledOp(Args.blockImage, Item, 0, ResultArrayData.ReturnArray.Begin + 1));
+
+				break;
+			} 
+			case 1:
+			{
+				FScheduledOpData ResultArrayData = HeapData[Item.CustomState];
+				TArrayView<FExtendedImageDesc> DepResults(
+						HeapImageDesc.GetData() + ResultArrayData.ReturnArray.Begin, 
+						ResultArrayData.ReturnArray.Num);
+				
+				check(DepResults.Num() == 2);
+	
+				FExtendedImageDesc& BaseImageDesc = DepResults[0];
+				FExtendedImageDesc& BlockImageDesc = DepResults[1];
+
+				FExtendedImageDesc& Result = HeapImageDesc[ResultArrayData.ReturnArray.ReturnAddress]; 
+				Result = BaseImageDesc;
+				Result.FirstLODAvailable = 
+						FMath::Max(BaseImageDesc.FirstLODAvailable, BlockImageDesc.FirstLODAvailable);
+
+				StoreValidDesc(Item); 
+				break;
+			}
 			default: check(false);
 			}
 			break;
@@ -6421,11 +6599,58 @@ namespace mu
 
 		case EOpType::IM_INTERPOLATE:
 		{
-			OP::ImageInterpolateArgs args = Program.GetOpArgs<OP::ImageInterpolateArgs>(item.At);
-			switch (item.Stage)
+			OP::ImageInterpolateArgs Args = Program.GetOpArgs<OP::ImageInterpolateArgs>(Item.At);
+
+			int32 NumImages = 0;
+			for (; NumImages < MUTABLE_OP_MAX_INTERPOLATE_COUNT; ++NumImages)
+			{
+				if (!Args.targets[NumImages])
+				{
+					break;
+				}
+			}
+
+			switch (Item.Stage)
+			{
+			case 0: 
+			{
+				FScheduledOpData ResultArrayData;
+				ResultArrayData.ReturnArray.ReturnAddress = Item.CustomState;
+				ResultArrayData.ReturnArray.Begin = HeapImageDesc.Num();
+				ResultArrayData.ReturnArray.Num = NumImages;
+
+				HeapImageDesc.AddDefaulted(ResultArrayData.ReturnArray.Num);
+				int32 ResultArrayAddress = HeapData.Add(ResultArrayData);
+
+				TArray<FScheduledOp, TInlineAllocator<MUTABLE_OP_MAX_INTERPOLATE_COUNT>> Deps;
+
+				for (int32 Index = 0; Index < NumImages; ++Index)
+				{
+					Deps.Add(FScheduledOp(Args.targets[Index], Item, 0, ResultArrayData.ReturnArray.Begin + Index));
+				}
+
+				AddOp(FScheduledOp(Item.At, Item, 1, ResultArrayAddress), Deps); 
+				break;
+			}
+			case 1: 
 			{
-			case 0: AddOp(FScheduledOp(item.At, item, 1), FScheduledOp(args.targets[0], item)); break;
-			case 1: StoreValidDesc(item); break;
+				FScheduledOpData ResultArrayData = HeapData[Item.CustomState];
+				TArrayView<FExtendedImageDesc> DepResults(
+						HeapImageDesc.GetData() + ResultArrayData.ReturnArray.Begin, 
+						ResultArrayData.ReturnArray.Num);
+				
+				check(DepResults.Num() == NumImages && NumImages > 0);
+
+				FExtendedImageDesc& Result = HeapImageDesc[ResultArrayData.ReturnArray.ReturnAddress]; 
+				Result = DepResults[0];
+
+				for (const FExtendedImageDesc& ResultImageDesc : DepResults)
+				{
+					Result.FirstLODAvailable = FMath::Max(Result.FirstLODAvailable, ResultImageDesc.FirstLODAvailable);
+				}
+
+				StoreValidDesc(Item); break;
+			}
 			default: check(false);
 			}
 			break;
@@ -6433,11 +6658,16 @@ namespace mu
 
 		case EOpType::IM_SATURATE:
 		{
-			OP::ImageSaturateArgs args = Program.GetOpArgs<OP::ImageSaturateArgs>(item.At);
-			switch (item.Stage)
+			OP::ImageSaturateArgs Args = Program.GetOpArgs<OP::ImageSaturateArgs>(Item.At);
+			switch (Item.Stage)
+			{
+			case 0: 
 			{
-			case 0: AddOp(FScheduledOp(item.At, item, 1), FScheduledOp(args.base, item)); break;
-			case 1: StoreValidDesc(item); break;
+				AddOp(FScheduledOp(Item.At, Item, 1, Item.CustomState), 
+						FScheduledOp(Args.base, Item, 0, Item.CustomState)); 
+				break;
+			}
+			case 1: StoreValidDesc(Item); break;
 			default: check(false);
 			}
 			break;
@@ -6445,16 +6675,18 @@ namespace mu
 
 		case EOpType::IM_LUMINANCE:
 		{
-			OP::ImageLuminanceArgs args = Program.GetOpArgs<OP::ImageLuminanceArgs>(item.At);
-			switch (item.Stage)
+			OP::ImageLuminanceArgs Args = Program.GetOpArgs<OP::ImageLuminanceArgs>(Item.At);
+			switch (Item.Stage)
 			{
 			case 0:
-				AddOp(FScheduledOp(item.At, item, 1), FScheduledOp(args.base, item));
+			{
+				AddOp(FScheduledOp(Item.At, Item, 1, Item.CustomState), 
+						FScheduledOp(Args.base, Item, 0, Item.CustomState));
 				break;
-
+			}
 			case 1:
-				HeapImageDesc[item.CustomState].m_format = EImageFormat::L_UByte;
-				StoreValidDesc(item);
+				HeapImageDesc[Item.CustomState].m_format = EImageFormat::L_UByte;
+				StoreValidDesc(Item);
 				break;
 
 			default:
@@ -6466,18 +6698,60 @@ namespace mu
 
 		case EOpType::IM_SWIZZLE:
 		{
-			OP::ImageSwizzleArgs args = Program.GetOpArgs<OP::ImageSwizzleArgs>(item.At);
-			switch (item.Stage)
+			OP::ImageSwizzleArgs Args = Program.GetOpArgs<OP::ImageSwizzleArgs>(Item.At);
+			switch (Item.Stage)
 			{
 			case 0:
-				AddOp(FScheduledOp(item.At, item, 1), FScheduledOp(args.sources[0], item));
-				break;
+			{
+				FScheduledOpData ResultArrayData;
+				ResultArrayData.ReturnArray.ReturnAddress = Item.CustomState;
+				ResultArrayData.ReturnArray.Begin = HeapImageDesc.Num();
+				ResultArrayData.ReturnArray.Num = 4;
 
-			case 1:
-				HeapImageDesc[item.CustomState].m_format = args.format;
-				StoreValidDesc(item);
+				HeapImageDesc.AddDefaulted(ResultArrayData.ReturnArray.Num);
+
+				int32 ResultArrayAddress = HeapData.Add(ResultArrayData);
+
+				AddOp(FScheduledOp(Item.At, Item, 1, ResultArrayAddress), 
+						FScheduledOp(Args.sources[0], Item, 0, ResultArrayData.ReturnArray.Begin + 0),
+						FScheduledOp(Args.sources[1], Item, 0, ResultArrayData.ReturnArray.Begin + 1),
+						FScheduledOp(Args.sources[2], Item, 0, ResultArrayData.ReturnArray.Begin + 2),
+						FScheduledOp(Args.sources[3], Item, 0, ResultArrayData.ReturnArray.Begin + 3));
 				break;
+			}
+			case 1:
+			{
+				FScheduledOpData ResultArrayData = HeapData[Item.CustomState];
+				TArrayView<FExtendedImageDesc> DepResults(
+						HeapImageDesc.GetData() + ResultArrayData.ReturnArray.Begin, 
+						ResultArrayData.ReturnArray.Num);
+				
+				check(DepResults.Num() == 4);
+	
+				FExtendedImageDesc& Result = HeapImageDesc[ResultArrayData.ReturnArray.ReturnAddress];
+
+				int32 FirstValidSourceIndex = 0;
+				for (; FirstValidSourceIndex < 4; ++FirstValidSourceIndex)
+				{
+					if (Args.sources[FirstValidSourceIndex])
+					{
+						Result = DepResults[FirstValidSourceIndex];
+						break;
+					}
+				}
+				Result.m_format = Args.format;
+				
+				for (int32 Index = FirstValidSourceIndex + 1; Index < 4; ++Index)
+				{
+					if (Args.sources[Index])
+					{
+						Result.FirstLODAvailable = FMath::Max(Result.FirstLODAvailable, DepResults[Index].FirstLODAvailable);
+					}
+				}
 
+				StoreValidDesc(Item);
+				break;
+			}
 			default:
 				check(false);
 			}
@@ -6487,11 +6761,16 @@ namespace mu
 
 		case EOpType::IM_COLOURMAP:
 		{
-			OP::ImageColourMapArgs args = Program.GetOpArgs<OP::ImageColourMapArgs>(item.At);
-			switch (item.Stage)
+			OP::ImageColourMapArgs Args = Program.GetOpArgs<OP::ImageColourMapArgs>(Item.At);
+			switch (Item.Stage)
 			{
-			case 0: AddOp(FScheduledOp(item.At, item, 1), FScheduledOp(args.base, item)); break;
-			case 1: StoreValidDesc(item); break;
+			case 0: 
+			{	
+				AddOp(FScheduledOp(Item.At, Item, 1, Item.CustomState), 
+						FScheduledOp(Args.base, Item, 0, Item.CustomState));
+				break;
+			}
+			case 1: StoreValidDesc(Item); break;
 			default: check(false);
 			}
 			break;
@@ -6499,29 +6778,32 @@ namespace mu
 
 		case EOpType::IM_GRADIENT:
 		{
-			OP::ImageGradientArgs args = Program.GetOpArgs<OP::ImageGradientArgs>(item.At);
-			HeapImageDesc[item.CustomState].m_size[0] = args.size[0];
-			HeapImageDesc[item.CustomState].m_size[1] = args.size[1];
-			HeapImageDesc[item.CustomState].m_lods = 1;
-			HeapImageDesc[item.CustomState].m_format = EImageFormat::RGB_UByte;
-			StoreValidDesc(item);
+			OP::ImageGradientArgs Args = Program.GetOpArgs<OP::ImageGradientArgs>(Item.At);
+			HeapImageDesc[Item.CustomState].m_size[0] = Args.size[0];
+			HeapImageDesc[Item.CustomState].m_size[1] = Args.size[1];
+			HeapImageDesc[Item.CustomState].m_lods = 1;
+			HeapImageDesc[Item.CustomState].m_format = EImageFormat::RGB_UByte;
+			StoreValidDesc(Item);
 			break;
 		}
 
 		case EOpType::IM_BINARISE:
 		{
-			OP::ImageBinariseArgs args = Program.GetOpArgs<OP::ImageBinariseArgs>(item.At);
-			switch (item.Stage)
+			OP::ImageBinariseArgs Args = Program.GetOpArgs<OP::ImageBinariseArgs>(Item.At);
+			switch (Item.Stage)
 			{
 			case 0:
-				AddOp(FScheduledOp(item.At, item, 1), FScheduledOp(args.base, item));
+			{
+				AddOp(FScheduledOp(Item.At, Item, 1, Item.CustomState), 
+						FScheduledOp(Args.base, Item, 0, Item.CustomState));
 				break;
-
+			}
 			case 1:
-				HeapImageDesc[item.CustomState].m_format = EImageFormat::L_UByte;
-				StoreValidDesc(item);
+			{
+				HeapImageDesc[Item.CustomState].m_format = EImageFormat::L_UByte;
+				StoreValidDesc(Item);
 				break;
-
+			}
 			default:
 				check(false);
 			}
@@ -6530,11 +6812,16 @@ namespace mu
 
 		case EOpType::IM_INVERT:
 		{
-			OP::ImageInvertArgs args = Program.GetOpArgs<OP::ImageInvertArgs>(item.At);
-			switch (item.Stage)
+			OP::ImageInvertArgs Args = Program.GetOpArgs<OP::ImageInvertArgs>(Item.At);
+			switch (Item.Stage)
+			{
+			case 0: 
 			{
-			case 0: AddOp(FScheduledOp(item.At, item, 1), FScheduledOp(args.base, item)); break;
-			case 1: StoreValidDesc(item); break;
+				AddOp(FScheduledOp(Item.At, Item, 1, Item.CustomState), 
+						FScheduledOp(Args.base, Item, 0, Item.CustomState)); 
+				break;
+			}
+			case 1: StoreValidDesc(Item); break;
 			default: check(false);
 			}
 			break;
@@ -6542,31 +6829,34 @@ namespace mu
 
 		case EOpType::IM_PLAINCOLOUR:
 		{
-			OP::ImagePlainColourArgs args = Program.GetOpArgs<OP::ImagePlainColourArgs>(item.At);
-			HeapImageDesc[item.CustomState].m_size[0] = args.size[0];
-			HeapImageDesc[item.CustomState].m_size[1] = args.size[1];
-			HeapImageDesc[item.CustomState].m_lods = args.LODs;
-			HeapImageDesc[item.CustomState].m_format = args.format;
-			StoreValidDesc(item);
+			OP::ImagePlainColourArgs Args = Program.GetOpArgs<OP::ImagePlainColourArgs>(Item.At);
+			HeapImageDesc[Item.CustomState].m_size[0] = Args.size[0];
+			HeapImageDesc[Item.CustomState].m_size[1] = Args.size[1];
+			HeapImageDesc[Item.CustomState].m_lods = Args.LODs;
+			HeapImageDesc[Item.CustomState].m_format = Args.format;
+			StoreValidDesc(Item);
 			break;
 		}
 
 		case EOpType::IM_CROP:
 		{
-			OP::ImageCropArgs args = Program.GetOpArgs<OP::ImageCropArgs>(item.At);
-			switch (item.Stage)
+			OP::ImageCropArgs Args = Program.GetOpArgs<OP::ImageCropArgs>(Item.At);
+			switch (Item.Stage)
 			{
 			case 0:
-				AddOp(FScheduledOp(item.At, item, 1), FScheduledOp(args.source, item));
+			{
+				AddOp(FScheduledOp(Item.At, Item, 1, Item.CustomState), 
+						FScheduledOp(Args.source, Item, 0, Item.CustomState));
 				break;
-
+			}
 			case 1:
-				HeapImageDesc[item.CustomState].m_size[0] = args.sizeX;
-				HeapImageDesc[item.CustomState].m_size[1] = args.sizeY;
-				HeapImageDesc[item.CustomState].m_lods = 1;
-				StoreValidDesc(item);
+			{
+				HeapImageDesc[Item.CustomState].m_size[0] = Args.sizeX;
+				HeapImageDesc[Item.CustomState].m_size[1] = Args.sizeY;
+				HeapImageDesc[Item.CustomState].m_lods = 1;
+				StoreValidDesc(Item);
 				break;
-
+			}
 			default:
 				check(false);
 			}
@@ -6575,11 +6865,43 @@ namespace mu
 
 		case EOpType::IM_PATCH:
 		{
-			OP::ImagePatchArgs args = Program.GetOpArgs<OP::ImagePatchArgs>(item.At);
-			switch (item.Stage)
+			OP::ImagePatchArgs Args = Program.GetOpArgs<OP::ImagePatchArgs>(Item.At);
+			switch (Item.Stage)
 			{
-			case 0: AddOp(FScheduledOp(item.At, item, 1), FScheduledOp(args.base, item)); break;
-			case 1: StoreValidDesc(item); break;
+			case 0: 
+			{
+				FScheduledOpData ResultArrayData;
+				ResultArrayData.ReturnArray.ReturnAddress = Item.CustomState;
+				ResultArrayData.ReturnArray.Begin = HeapImageDesc.Num();
+				ResultArrayData.ReturnArray.Num = 2;
+
+				HeapImageDesc.AddDefaulted(ResultArrayData.ReturnArray.Num);
+
+				int32 ResultArrayAddress = HeapData.Add(ResultArrayData);
+
+				AddOp(FScheduledOp(Item.At, Item, 1, ResultArrayAddress), 
+						FScheduledOp(Args.base, Item, 0, ResultArrayData.ReturnArray.Begin),
+						FScheduledOp(Args.patch, Item, 0, ResultArrayData.ReturnArray.Begin + 1));
+				break;
+			}
+			case 1: 
+			{
+				FScheduledOpData ResultArrayData = HeapData[Item.CustomState];
+				TArrayView<FExtendedImageDesc> DepResults(
+						HeapImageDesc.GetData() + ResultArrayData.ReturnArray.Begin, 
+						ResultArrayData.ReturnArray.Num);
+				
+				check(DepResults.Num() == 2);
+	
+				FExtendedImageDesc& BaseImageDesc = DepResults[0];
+				FExtendedImageDesc& PatchImageDesc = DepResults[1];
+
+				FExtendedImageDesc& Result = HeapImageDesc[ResultArrayData.ReturnArray.ReturnAddress]; 
+				Result = BaseImageDesc;
+				Result.FirstLODAvailable = FMath::Max(Result.FirstLODAvailable, PatchImageDesc.FirstLODAvailable);
+
+				StoreValidDesc(Item); break;
+			}
 			default: check(false);
 			}
 			break;
@@ -6587,30 +6909,33 @@ namespace mu
 
 		case EOpType::IM_RASTERMESH:
 		{
-			OP::ImageRasterMeshArgs args = Program.GetOpArgs<OP::ImageRasterMeshArgs>(item.At);
-			HeapImageDesc[item.CustomState].m_size[0] = args.sizeX;
-			HeapImageDesc[item.CustomState].m_size[1] = args.sizeY;
-			HeapImageDesc[item.CustomState].m_lods = 1;
-			HeapImageDesc[item.CustomState].m_format = EImageFormat::L_UByte;
-			StoreValidDesc(item);
+			OP::ImageRasterMeshArgs Args = Program.GetOpArgs<OP::ImageRasterMeshArgs>(Item.At);
+			HeapImageDesc[Item.CustomState].m_size[0] = Args.sizeX;
+			HeapImageDesc[Item.CustomState].m_size[1] = Args.sizeY;
+			HeapImageDesc[Item.CustomState].m_lods = 1;
+			HeapImageDesc[Item.CustomState].m_format = EImageFormat::L_UByte;
+			StoreValidDesc(Item);
 			break;
 		}
 
 		case EOpType::IM_MAKEGROWMAP:
 		{
-			OP::ImageMakeGrowMapArgs args = Program.GetOpArgs<OP::ImageMakeGrowMapArgs>(item.At);
-			switch (item.Stage)
+			OP::ImageMakeGrowMapArgs Args = Program.GetOpArgs<OP::ImageMakeGrowMapArgs>(Item.At);
+			switch (Item.Stage)
 			{
 			case 0:
-				AddOp(FScheduledOp(item.At, item, 1), FScheduledOp(args.mask, item));
+			{
+				AddOp(FScheduledOp(Item.At, Item, 1, Item.CustomState), 
+						FScheduledOp(Args.mask, Item, 0, Item.CustomState));
 				break;
-
+			}
 			case 1:
-				HeapImageDesc[item.CustomState].m_format = EImageFormat::L_UByte;
-				HeapImageDesc[item.CustomState].m_lods = 1;
-				StoreValidDesc(item);
+			{
+				HeapImageDesc[Item.CustomState].m_format = EImageFormat::L_UByte;
+				HeapImageDesc[Item.CustomState].m_lods = 1;
+				StoreValidDesc(Item);
 				break;
-
+			}
 			default:
 				check(false);
 			}
@@ -6620,11 +6945,16 @@ namespace mu
 
 		case EOpType::IM_DISPLACE:
 		{
-			OP::ImageDisplaceArgs args = Program.GetOpArgs<OP::ImageDisplaceArgs>(item.At);
-			switch (item.Stage)
+			OP::ImageDisplaceArgs Args = Program.GetOpArgs<OP::ImageDisplaceArgs>(Item.At);
+			switch (Item.Stage)
+			{
+			case 0: 
 			{
-			case 0: AddOp(FScheduledOp(item.At, item, 1), FScheduledOp(args.source, item)); break;
-			case 1: StoreValidDesc(item); break;
+				AddOp(FScheduledOp(Item.At, Item, 1, Item.CustomState),
+						FScheduledOp(Args.source, Item, 0, Item.CustomState));
+				break;
+			}
+			case 1: StoreValidDesc(Item); break;
 			default: check(false);
 			}
 			break;
@@ -6633,27 +6963,28 @@ namespace mu
         case EOpType::IM_TRANSFORM:
         {
 
-			OP::ImageTransformArgs Args = Program.GetOpArgs<OP::ImageTransformArgs>(item.At);
+			OP::ImageTransformArgs Args = Program.GetOpArgs<OP::ImageTransformArgs>(Item.At);
 
-            switch (item.Stage)
+            switch (Item.Stage)
             {
             case 0:
 			{
-				AddOp(FScheduledOp(item.At, item, 1), FScheduledOp(Args.Base, item));	
+				AddOp(FScheduledOp(Item.At, Item, 1, Item.CustomState), 
+						FScheduledOp(Args.Base, Item, 0, Item.CustomState));	
                 break;
 			}
             case 1:
             {
-				HeapImageDesc[item.CustomState].m_lods = 1;
-				HeapImageDesc[item.CustomState].m_format = GetUncompressedFormat(HeapImageDesc[item.CustomState].m_format);
+				HeapImageDesc[Item.CustomState].m_lods = 1;
+				HeapImageDesc[Item.CustomState].m_format = GetUncompressedFormat(HeapImageDesc[Item.CustomState].m_format);
 				
 				if (!(Args.SizeX == 0 && Args.SizeY == 0))
 				{
-					HeapImageDesc[item.CustomState].m_size[0] = Args.SizeX;
-					HeapImageDesc[item.CustomState].m_size[1] = Args.SizeY;
+					HeapImageDesc[Item.CustomState].m_size[0] = Args.SizeX;
+					HeapImageDesc[Item.CustomState].m_size[1] = Args.SizeY;
 				}
 
-				StoreValidDesc(item);
+				StoreValidDesc(Item);
                 break;
             }
 
@@ -6665,11 +6996,11 @@ namespace mu
 		}
 
 		default:
-			if (type != EOpType::NONE)
+			if (Type != EOpType::NONE)
 			{
 				// Operation not implemented
 				check(false);
-				HeapImageDesc[item.CustomState] = FImageDesc();
+				HeapImageDesc[Item.CustomState] = FExtendedImageDesc();
 			}
 			break;
 		}
diff --git a/Engine/Plugins/Mutable/Source/MutableRuntime/Private/MuR/CodeRunner.h b/Engine/Plugins/Mutable/Source/MutableRuntime/Private/MuR/CodeRunner.h
--- a/Engine/Plugins/Mutable/Source/MutableRuntime/Private/MuR/CodeRunner.h
+++ b/Engine/Plugins/Mutable/Source/MutableRuntime/Private/MuR/CodeRunner.h
@@ -72,6 +72,13 @@ namespace  mu
         {
 			union
 			{
+				struct
+				{
+					int32 Begin;
+					int32 Num;
+					int32 ReturnAddress;
+				} ReturnArray;
+			
 				struct
 				{
 					float Bifactor;
@@ -85,12 +92,6 @@ namespace  mu
 					bool bBlendOnlyOneMip;
 				} MultiLayer;
 
-				struct
-				{
-					int32 ResultDescAt;
-					int32 SourceDescAt;
-				} ResizeLike;
-
 				struct
 				{
 					uint8 Mip;
@@ -152,7 +153,7 @@ namespace  mu
 
 		/** Load an external image asynchronously, returns an event to wait for complition and a cleanup function that must be called once the event has completed. */
 		TTuple<UE::Tasks::FTask, TFunction<void()>> LoadExternalImageAsync(FExternalResourceId Id, uint8 MipmapsToSkip, TFunction<void(TSharedPtr<FImage>)>& ResultCallback);
- 	    mu::FImageDesc GetExternalImageDesc(FName Id, uint8 MipmapsToSkip);
+ 	    mu::FExtendedImageDesc GetExternalImageDesc(FName Id);
 
 		/** Load an external mesh asynchronously, returns an event to wait for complition and a cleanup function that must be called once the event has completed. */
 		TTuple<UE::Tasks::FTask, TFunction<void()>> LoadExternalMeshAsync(FExternalResourceId Id, TFunction<void(TSharedPtr<FMesh>)>& ResultCallback);
@@ -161,10 +162,12 @@ namespace  mu
 		FSettings Settings;
 
     protected:
-        //! Heap of intermediate data pushed by some instructions and referred by others.
-        //! It is not released until no operations are pending.
+        /** 
+		 * Heap of intermediate data pushed by some instructions and referred by others.
+         * It is not released until no operations are pending.
+		 */
 		TArray<FScheduledOpData> HeapData;
-		TArray<FImageDesc> HeapImageDesc;
+		TArray<FExtendedImageDesc> HeapImageDesc;
 
 		/** Only used for correct mip skipping with external images. It is the LOD for which the image is build. */
 		int32 ImageLOD;
@@ -178,7 +181,7 @@ namespace  mu
 		UE::Tasks::FTask StartRun(bool bForceInlineExecution);
 
 		//!
-		void GetImageDescResult(FImageDesc& OutDesc);
+		void GetImageDescResult(FExtendedImageDesc& OutDesc);
 
 		//!
 		FProgramCache& GetMemory();
diff --git a/Engine/Plugins/Mutable/Source/MutableRuntime/Private/MuR/CodeRunnerTasked.cpp b/Engine/Plugins/Mutable/Source/MutableRuntime/Private/MuR/CodeRunnerTasked.cpp
--- a/Engine/Plugins/Mutable/Source/MutableRuntime/Private/MuR/CodeRunnerTasked.cpp
+++ b/Engine/Plugins/Mutable/Source/MutableRuntime/Private/MuR/CodeRunnerTasked.cpp
@@ -637,9 +637,9 @@ namespace mu
 
 	
 	//---------------------------------------------------------------------------------------------
-	void CodeRunner::GetImageDescResult(FImageDesc& OutDesc)
+	void CodeRunner::GetImageDescResult(FExtendedImageDesc& OutDesc)
 	{
-		check( HeapImageDesc.Num()>0 );
+		check(HeapImageDesc.Num() > 0);
 		OutDesc = HeapImageDesc[0];
 	}
 
diff --git a/Engine/Plugins/Mutable/Source/MutableRuntime/Private/MuR/System.cpp b/Engine/Plugins/Mutable/Source/MutableRuntime/Private/MuR/System.cpp
--- a/Engine/Plugins/Mutable/Source/MutableRuntime/Private/MuR/System.cpp
+++ b/Engine/Plugins/Mutable/Source/MutableRuntime/Private/MuR/System.cpp
@@ -416,12 +416,12 @@ namespace mu
 		TEXT("If different than 0, clear the image desc cache at every image."),
 		ECVF_Scalability);
 
-	void FSystem::GetImageDescInline(FInstance::FID instanceID, FResourceID ImageId, FImageDesc& OutDesc)
+	void FSystem::GetImageDescInline(FInstance::FID instanceID, FResourceID ImageId, FExtendedImageDesc& OutDesc)
 	{
 		LLM_SCOPE_BYNAME(TEXT("MutableRuntime"));
 		MUTABLE_CPUPROFILER_SCOPE(SystemGetImageDesc);
 
-		OutDesc = FImageDesc();
+		OutDesc = FExtendedImageDesc();
 
 		// Find the live instance
 		FLiveInstance* pLiveInstance = m_pD->FindLiveInstance(instanceID);
@@ -463,7 +463,7 @@ namespace mu
 	}
 
 
-	UE::Tasks::TTask<FImageDesc> FSystem::GetImageDesc(FInstance::FID instanceID, FResourceID ImageId)
+	UE::Tasks::TTask<FExtendedImageDesc> FSystem::GetImageDesc(FInstance::FID instanceID, FResourceID ImageId)
 	{
 		LLM_SCOPE_BYNAME(TEXT("MutableRuntime"));
 		MUTABLE_CPUPROFILER_SCOPE(SystemGetImageDesc);
@@ -492,7 +492,7 @@ namespace mu
 			m_pD->WorkingMemoryManager.EndRunnerThread();
 			m_pD->WorkingMemoryManager.CurrentInstanceCache = nullptr;
 
-			return UE::Tasks::MakeCompletedTask<FImageDesc>(); 
+			return UE::Tasks::MakeCompletedTask<FExtendedImageDesc>(); 
 		}
 
 		// GetImageDesc may call normal execution paths where meshes are computed.
@@ -504,9 +504,9 @@ namespace mu
 		UE::Tasks::FTask RunnerCompletionEvent = Runner->StartRun(bForceInlineExecution);
 		
 		return UE::Tasks::Launch(TEXT("FSystem::GetImageDescResultTask"),
-				[SystemPrivate = m_pD, Runner]() -> FImageDesc
+				[SystemPrivate = m_pD, Runner]() -> FExtendedImageDesc
 				{
-					FImageDesc Result;
+					FExtendedImageDesc Result;
 					Runner->GetImageDescResult(Result);
 
 					SystemPrivate->WorkingMemoryManager.EndRunnerThread();
diff --git a/Engine/Plugins/Mutable/Source/MutableRuntime/Public/MuR/ImageTypes.h b/Engine/Plugins/Mutable/Source/MutableRuntime/Public/MuR/ImageTypes.h
--- a/Engine/Plugins/Mutable/Source/MutableRuntime/Public/MuR/ImageTypes.h
+++ b/Engine/Plugins/Mutable/Source/MutableRuntime/Public/MuR/ImageTypes.h
@@ -81,26 +81,36 @@ namespace mu
 		{
 		}
 
-		//!
+		/** */
 		FImageSize m_size = FImageSize(0, 0);
 
-		//!
+		/** */
 		EImageFormat m_format = EImageFormat::None;
 
-		//! Levels of detail (mipmaps)
+		/** Levels of detail (mipmaps) */
 		uint8 m_lods = 0;
 
-		//!
+		/** */
 		FORCEINLINE bool operator==(const FImageDesc& Other) const
 		{
-			return m_format == Other.m_format &&
-				   m_lods == Other.m_lods     &&
-				   m_size == Other.m_size;
+			return
+				m_size == Other.m_size && 
+				m_format == Other.m_format && 
+				m_lods == Other.m_lods;
 		}
 	};
 
+	/** An image descriptor with information about missing data LODs */
+	struct FExtendedImageDesc : public FImageDesc
+	{
+		/** 
+		 * Number of LODs that cannot be generated because the data is not available.
+		 * It is assumed that non available LODs are consecutive starting from the mip with higher resolution.
+		 **/
+		uint8 FirstLODAvailable = 0;
+	};
 
-	//! List of supported modes in generic image layering operations.
+	/** List of supported modes in generic image layering operations. */
 	enum class EBlendType
 	{
 		BT_NONE = 0,
diff --git a/Engine/Plugins/Mutable/Source/MutableRuntime/Public/MuR/Serialisation.h b/Engine/Plugins/Mutable/Source/MutableRuntime/Public/MuR/Serialisation.h
--- a/Engine/Plugins/Mutable/Source/MutableRuntime/Public/MuR/Serialisation.h
+++ b/Engine/Plugins/Mutable/Source/MutableRuntime/Public/MuR/Serialisation.h
@@ -90,9 +90,11 @@ namespace mu
 		//! Negative values indicate an error.
         typedef int32 FOperationID;
 
+		virtual bool DoesBlockExist(const FModel*, uint32 BlockKey) = 0;
+
 		//! \brief Start a data request operation.
 		//! \param Model.
-        //! \param key0 key identifying the model data fragment that is requested.
+        //! \param BlockKey key identifying the model data fragment that is requested.
         //!         This key interpretation depends on the implementation of the ModelStreamer,
 		//! \param Buffer is an already-allocated buffer big enough to receive the expected data.
 		//! \param size is the size of the pBuffer buffer, which must match the size of the data
@@ -116,7 +118,6 @@ namespace mu
 		* \return true if the data was loaded successfully.
 		*/
         virtual bool EndRead(FOperationID) = 0;
-
     };
 
 
diff --git a/Engine/Plugins/Mutable/Source/MutableRuntime/Public/MuR/System.h b/Engine/Plugins/Mutable/Source/MutableRuntime/Public/MuR/System.h
--- a/Engine/Plugins/Mutable/Source/MutableRuntime/Public/MuR/System.h
+++ b/Engine/Plugins/Mutable/Source/MutableRuntime/Public/MuR/System.h
@@ -79,7 +79,7 @@ namespace mu
 		virtual TTuple<UE::Tasks::FTask, TFunction<void()>> GetImageAsync(FName Id, uint8 MipmapsToSkip, TFunction<void(TSharedPtr<FImage>)>& ResultCallback) = 0;
 		virtual TTuple<UE::Tasks::FTask, TFunction<void()>> GetReferencedImageAsync(const void* ModelPtr, int32 Id, uint8 MipmapsToSkip, TFunction<void(TSharedPtr<FImage>)>& ResultCallback) { check(false); return {}; }
 
-		virtual mu::FImageDesc GetImageDesc(FName Id, uint8 MipmapsToSkip) = 0;
+		virtual mu::FExtendedImageDesc GetImageDesc(FName Id) = 0;
 
 		/** Returns the completion event and a cleanup function that must be called once event is completed. */
 		virtual TTuple<UE::Tasks::FTask, TFunction<void()>> GetMeshAsync(FName Id, TFunction<void(TSharedPtr<FMesh>)>& ResultCallback) = 0;
@@ -152,7 +152,7 @@ namespace mu
 
 		//! Only valid between BeginUpdate and EndUpdate
 		//! Calculate the description of an image, without generating it.
-		UE::Tasks::TTask<FImageDesc> GetImageDesc(FInstance::FID InstanceID, FResourceID ImageId);
+		UE::Tasks::TTask<FExtendedImageDesc> GetImageDesc(FInstance::FID InstanceID, FResourceID ImageId);
 
 		//! Only valid between BeginUpdate and EndUpdate
 		//! \param MipsToSkip Number of mips to skip compared from the full image.
@@ -165,7 +165,7 @@ namespace mu
 
 		//! Only valid between BeginUpdate and EndUpdate
 		//! Calculate the description of an image, without generating it.
-		void GetImageDescInline(FInstance::FID InstanceID, FResourceID ImageId, FImageDesc& OutDesc);
+		void GetImageDescInline(FInstance::FID InstanceID, FResourceID ImageId, FExtendedImageDesc& OutDesc);
 
 		//! Only valid between BeginUpdate and EndUpdate
 		//! \param MipsToSkip Number of mips to skip compared from the full image.
