[mutable] Improved sorting of parameters when compiling CO.
- Current approach does not handle well cases where the Name and UID are equal.
- New approach will be able to deterministically handle those cases.
- UID is not guaranteed to be unique at the moment.

#rnx

diff --git a/Engine/Plugins/Mutable/Source/MutableTools/Private/MuT/Compiler.cpp b/Engine/Plugins/Mutable/Source/MutableTools/Private/MuT/Compiler.cpp
--- a/Engine/Plugins/Mutable/Source/MutableTools/Private/MuT/Compiler.cpp
+++ b/Engine/Plugins/Mutable/Source/MutableTools/Private/MuT/Compiler.cpp
@@ -318,7 +318,7 @@ namespace mu
             GenErrorLog = Generator.ErrorLog;
 
         	// Sort the parameters as deterministically as possible.
-        	Generator.FirstPass.ParameterNodes.ValueSort([](const Ptr<ASTOpParameter>& A, const Ptr<ASTOpParameter>& B)->bool
+        	Generator.FirstPass.ParameterNodes.ValueStableSort([](const Ptr<ASTOpParameter>& A, const Ptr<ASTOpParameter>& B)->bool
 			{
 				if (A->Parameter.Name < B->Parameter.Name)
 				{
