[Mutable] Add a tooltip to UCustomizableObjectNodeComponentMesh properties.

#rnx

diff --git a/Engine/Plugins/Mutable/Source/CustomizableObjectEditor/Private/MuCOE/Nodes/CustomizableObjectNodeComponentMesh.h b/Engine/Plugins/Mutable/Source/CustomizableObjectEditor/Private/MuCOE/Nodes/CustomizableObjectNodeComponentMesh.h
--- a/Engine/Plugins/Mutable/Source/CustomizableObjectEditor/Private/MuCOE/Nodes/CustomizableObjectNodeComponentMesh.h
+++ b/Engine/Plugins/Mutable/Source/CustomizableObjectEditor/Private/MuCOE/Nodes/CustomizableObjectNodeComponentMesh.h
@@ -61,9 +61,16 @@ public:
 	UPROPERTY(EditAnywhere, Category = ComponentMesh)
 	FName ComponentName;
 	
+	/** All the Skeletal Meshes generated for this component will use the Reference Skeletal Mesh properties
+	 * for everything Mutable doesn't create or modify. This includes data like LOD distances, Physics
+	 * properties, Bounding Volumes, Base Skeleton, and more.
+	 *
+	 * The Reference Skeletal Mesh can be used as a placeholder mesh when there are too many actors or in 
+	 * situations of stress where the generation of the Skeletal Mesh might take a few seconds to complete. */
 	UPROPERTY(EditAnywhere, Category = ComponentMesh)
 	TObjectPtr<USkeletalMesh> ReferenceSkeletalMesh;
 
+	/** LOD reduction settings to apply on top of the LOD settings of the Reference Skeletal Mesh. */
 	UPROPERTY(EditAnywhere, Category = ComponentMesh)
 	TArray<FLODReductionSettings> LODReductionSettings;
 
diff --git a/Engine/Plugins/Mutable/Source/CustomizableObjectEditor/Private/MuCOE/Nodes/CustomizableObjectNodeComponentMeshBase.h b/Engine/Plugins/Mutable/Source/CustomizableObjectEditor/Private/MuCOE/Nodes/CustomizableObjectNodeComponentMeshBase.h
--- a/Engine/Plugins/Mutable/Source/CustomizableObjectEditor/Private/MuCOE/Nodes/CustomizableObjectNodeComponentMeshBase.h
+++ b/Engine/Plugins/Mutable/Source/CustomizableObjectEditor/Private/MuCOE/Nodes/CustomizableObjectNodeComponentMeshBase.h
@@ -14,7 +14,14 @@ enum class ECustomizableObjectAutomaticLODStrategy : uint8
 	Inherited = 0 UMETA(DisplayName = "Inherit from parent component"),
 	// Don't try to generate LODs automatically for the child nodes. Only the ones tha explicitely define them will be used.
 	Manual = 1 UMETA(DisplayName = "Only manually created LODs"),
-	// Try to generate the same material structure than LOd 0 if the source meshes have LODs.
+	/** Automatic LODs from Mesh will try to replicate the explicit graph connections to higher LODs.
+	 *
+	 * For example, given a Mesh with n LODs connected to a Component m-th LOD, LODs m1, m2... will be automatically connected
+	 * using Mesh LODs n1, n2.
+	 *
+	 * Note: If the mesh section is disabled at LOD 'n' or doesn't exist, it will not be added from that LOD onwards.
+	 * Note: If the mesh section uses a different MaterialSlot at LOD 'n', it will not be added from that LOD onwards.
+	 */
 	AutomaticFromMesh = 2 UMETA(DisplayName = "Automatic from mesh")
 };
 
