[Mutable] Made UEdGraphSchema_CustomizableObject funcitons private.

diff --git a/Engine/Plugins/Mutable/Source/CustomizableObjectEditor/Private/MuCOE/CustomizableObjectGraph.cpp b/Engine/Plugins/Mutable/Source/CustomizableObjectEditor/Private/MuCOE/CustomizableObjectGraph.cpp
--- a/Engine/Plugins/Mutable/Source/CustomizableObjectEditor/Private/MuCOE/CustomizableObjectGraph.cpp
+++ b/Engine/Plugins/Mutable/Source/CustomizableObjectEditor/Private/MuCOE/CustomizableObjectGraph.cpp
@@ -2,6 +2,7 @@
 
 #include "MuCOE/CustomizableObjectGraph.h"
 
+#include "CustomizableObjectSchemaActions.h"
 #include "MuCO/CustomizableObjectCustomVersion.h"
 #include "MuCOE/EdGraphSchema_CustomizableObject.h"
 #include "MuCOE/Nodes/CustomizableObjectNodeObject.h"
diff --git a/Engine/Plugins/Mutable/Source/CustomizableObjectEditor/Private/MuCOE/CustomizableObjectSchemaActions.h b/Engine/Plugins/Mutable/Source/CustomizableObjectEditor/Private/MuCOE/CustomizableObjectSchemaActions.h
new file mode 100644
--- /dev/null
+++ b/Engine/Plugins/Mutable/Source/CustomizableObjectEditor/Private/MuCOE/CustomizableObjectSchemaActions.h
@@ -0,0 +1,72 @@
+﻿// Copyright Epic Games, Inc. All Rights Reserved.
+
+#pragma once
+
+#include "EdGraphSchema_K2_Actions.h"
+
+#include "CustomizableObjectSchemaActions.generated.h"
+
+
+/** Action to add a node to the graph */
+USTRUCT()
+struct CUSTOMIZABLEOBJECTEDITOR_API FCustomizableObjectSchemaAction_NewNode : public FEdGraphSchemaAction
+{
+	GENERATED_USTRUCT_BODY();
+
+	/** Template of node we want to create */
+	UPROPERTY()
+	TObjectPtr<class UEdGraphNode> NodeTemplate;
+
+
+	FCustomizableObjectSchemaAction_NewNode() 
+		: FEdGraphSchemaAction()
+		, NodeTemplate(NULL)
+	{}
+
+	FCustomizableObjectSchemaAction_NewNode(const FString& InNodeCategory, const FText& InMenuDesc, const FText& InToolTip, const int32 InGrouping, const FText& InKeywords = FText(), int32 InSectionID = 0)
+		: FEdGraphSchemaAction(FText::FromString(InNodeCategory), InMenuDesc, InToolTip, InGrouping, InKeywords, InSectionID)
+		, NodeTemplate(NULL)
+	{}
+
+	// FEdGraphSchemaAction interface
+	virtual UEdGraphNode* PerformAction(class UEdGraph* ParentGraph, UEdGraphPin* FromPin, const FVector2D Location, bool bSelectNewNode = true) override;
+	virtual UEdGraphNode* PerformAction(class UEdGraph* ParentGraph, TArray<UEdGraphPin*>& FromPins, const FVector2D Location, bool bSelectNewNode = true) override;
+	virtual void AddReferencedObjects( FReferenceCollector& Collector ) override;
+
+	/** Reimplementaiton of EdGraphSchema::CreateNode(...). Performs the overlap calculation before calling AutowireNewNode(...). AutowireNewNode can induce a call to a ReconstrucNode() which removes pins required for the calculation. */
+	static UEdGraphNode* CreateNode(class UEdGraph* ParentGraph, UEdGraphPin* FromPin, const FVector2D Location, class UEdGraphNode* InNodeTemplate);
+	// End of FEdGraphSchemaAction interface
+
+	template <typename NodeType>
+	static NodeType* InstantSpawn(class UEdGraph* ParentGraph, NodeType* InTemplateNode, const FVector2D Location)
+	{
+		FEdGraphSchemaAction_K2NewNode Action;
+		Action.NodeTemplate = InTemplateNode;
+
+		return Cast<NodeType>(Action.PerformAction(ParentGraph, NULL, Location));
+	}
+};
+
+
+/** Action to paste clipboard contents into the graph */
+USTRUCT()
+struct CUSTOMIZABLEOBJECTEDITOR_API FCustomizableObjectSchemaAction_Paste : public FEdGraphSchemaAction
+{
+	GENERATED_USTRUCT_BODY();
+
+	// Simple type info
+	static FName StaticGetTypeId() { static FName Type("FCustomizableObjectSchemaAction_Paste"); return Type; }
+	virtual FName GetTypeId() const override { return StaticGetTypeId(); }
+
+	FCustomizableObjectSchemaAction_Paste()
+		: FEdGraphSchemaAction()
+	{}
+
+	FCustomizableObjectSchemaAction_Paste(const FText& InNodeCategory, const FText& InMenuDesc, const FString& InToolTip, const int32 InGrouping)
+		: FEdGraphSchemaAction(InNodeCategory, InMenuDesc, FText::FromString(InToolTip), InGrouping)
+	{}
+
+	// FEdGraphSchemaAction interface
+	virtual UEdGraphNode* PerformAction(class UEdGraph* ParentGraph, UEdGraphPin* FromPin, const FVector2D Location, bool bSelectNewNode = true) override;
+	// End of FEdGraphSchemaAction interface
+};
diff --git a/Engine/Plugins/Mutable/Source/CustomizableObjectEditor/Private/MuCOE/EdGraphSchema_CustomizableObject.cpp b/Engine/Plugins/Mutable/Source/CustomizableObjectEditor/Private/MuCOE/EdGraphSchema_CustomizableObject.cpp
--- a/Engine/Plugins/Mutable/Source/CustomizableObjectEditor/Private/MuCOE/EdGraphSchema_CustomizableObject.cpp
+++ b/Engine/Plugins/Mutable/Source/CustomizableObjectEditor/Private/MuCOE/EdGraphSchema_CustomizableObject.cpp
@@ -1,6 +1,7 @@
 // Copyright Epic Games, Inc. All Rights Reserved.
 
 #include "MuCOE/EdGraphSchema_CustomizableObject.h"
+#include "MuCOE/CustomizableObjectSchemaActions.h"
 
 #include "CustomizableObjectConnectionDrawingPolicy.h"
 #include "EdGraphNode_Comment.h"
@@ -105,6 +106,25 @@ namespace MutablePrivate
 }
 
 
+TSharedPtr<ICustomizableObjectEditor> GetCustomizableObjectEditor(const UEdGraph* ParentGraph)
+{
+	// Find the associated Editor
+	UCustomizableObject* CustomizableObject = Cast<UCustomizableObject>(ParentGraph->GetOuter());
+
+	TSharedPtr<ICustomizableObjectEditor> CustomizableObjectEditor;
+	if (CustomizableObject)
+	{
+		TSharedPtr< IToolkit > FoundAssetEditor = FToolkitManager::Get().FindEditorForAsset(CustomizableObject);
+		if (FoundAssetEditor.IsValid())
+		{
+			return StaticCastSharedPtr<ICustomizableObjectEditor>(FoundAssetEditor);
+		}
+	}
+
+	return nullptr;
+}
+
+
 UEdGraphNode* FCustomizableObjectSchemaAction_NewNode::PerformAction(class UEdGraph* ParentGraph, UEdGraphPin* FromPin, const FVector2D Location, bool bSelectNewNode)
 {
 	UEdGraphNode* ResultNode = nullptr;
@@ -210,7 +230,7 @@ void FCustomizableObjectSchemaAction_NewNode::AddReferencedObjects( FReferenceCo
 
 UEdGraphNode* FCustomizableObjectSchemaAction_Paste::PerformAction(class UEdGraph* ParentGraph, UEdGraphPin* FromPin, const FVector2D Location, bool bSelectNewNode/* = true*/)
 {
-	TSharedPtr<ICustomizableObjectEditor> CustomizableObjectEditor = UEdGraphSchema_CustomizableObject::GetCustomizableObjectEditor(ParentGraph);
+	TSharedPtr<ICustomizableObjectEditor> CustomizableObjectEditor = GetCustomizableObjectEditor(ParentGraph);
 
 	if (CustomizableObjectEditor.IsValid() && CustomizableObjectEditor->CanPasteNodes())
 	{
@@ -249,9 +269,9 @@ const FName UEdGraphSchema_CustomizableObject::PC_Transform("transform");
 const FText UEdGraphSchema_CustomizableObject::NC_Experimental(FText::FromString("Experimental"));
 
 
-TSharedPtr<FCustomizableObjectSchemaAction_NewNode> UEdGraphSchema_CustomizableObject::AddNewNodeAction(FGraphActionListBuilderBase& ContextMenuBuilder, const FString& Category, const FText& MenuDesc, const FText& Tooltip, const int32 Grouping, const FString& Keywords)
+TSharedPtr<FCustomizableObjectSchemaAction_NewNode> AddNewNodeAction(FGraphActionListBuilderBase& ContextMenuBuilder, const FString& Category, const FText& MenuDesc, const FText& Tooltip, const int32 Grouping = 0, const FString& Keywords = FString())
 {
-	TSharedPtr<FCustomizableObjectSchemaAction_NewNode> NewActionNode = TSharedPtr<FCustomizableObjectSchemaAction_NewNode>(new FCustomizableObjectSchemaAction_NewNode(Category, MenuDesc, Tooltip, Grouping, FText::FromString(Keywords)));
+	TSharedPtr<FCustomizableObjectSchemaAction_NewNode> NewActionNode = MakeShared<FCustomizableObjectSchemaAction_NewNode>(Category, MenuDesc, Tooltip, Grouping, FText::FromString(Keywords));
 
 	ContextMenuBuilder.AddAction( NewActionNode );
 
@@ -313,7 +333,7 @@ namespace
 			return;
 		}
 
-		TSharedPtr<FCustomizableObjectSchemaAction_NewNode> Action = UEdGraphSchema_CustomizableObject::AddNewNodeAction(ContextMenuBuilder, Category, MenuDesc, FText(), Grouping);
+		TSharedPtr<FCustomizableObjectSchemaAction_NewNode> Action = AddNewNodeAction(ContextMenuBuilder, Category, MenuDesc, FText(), Grouping);
 		Action->NodeTemplate = TemplateNode;
 	}
 
@@ -809,7 +829,7 @@ bool UEdGraphSchema_CustomizableObject::ShouldHidePinDefaultValue(UEdGraphPin* P
 }
 
 
-void UEdGraphSchema_CustomizableObject::GetContextMenuActionsReconstructAllChildNodes(UToolMenu* Menu, TWeakObjectPtr<UGraphNodeContextMenuContext> WeakContext) const
+void GetContextMenuActionsReconstructAllChildNodes(UToolMenu* Menu, TWeakObjectPtr<UGraphNodeContextMenuContext> WeakContext)
 {
 	FToolMenuSection& SubSection = Menu->AddSection("Section");
 
@@ -922,7 +942,7 @@ void UEdGraphSchema_CustomizableObject::GetContextMenuActions(UToolMenu* Menu, U
 				"ReconstructAllNodes",
 				LOCTEXT("ReconstructChildAllNodes", "Refresh All Child Nodes"),
 				LOCTEXT("ReconstructAllChildNodes_Tooltip", "Refresh all child nodes from the selected ones (inclusive)."),
-				FNewToolMenuDelegate::CreateUObject(this, &UEdGraphSchema_CustomizableObject::GetContextMenuActionsReconstructAllChildNodes, MakeWeakObjectPtr(Context)));
+				FNewToolMenuDelegate::CreateStatic(&GetContextMenuActionsReconstructAllChildNodes, MakeWeakObjectPtr(Context)));
 		}
 
 		struct SCommentUtility
@@ -931,7 +951,10 @@ void UEdGraphSchema_CustomizableObject::GetContextMenuActions(UToolMenu* Menu, U
 			{
 				if (Schema && Graph)
 				{
-					Schema->AddComment(Graph, NULL, FVector2D::ZeroVector, true);
+					if (TSharedPtr<ICustomizableObjectEditor> CustomizableObjectEditor = GetCustomizableObjectEditor(Graph))
+					{
+						CustomizableObjectEditor->CreateCommentBox(FVector2D::ZeroVector);
+					}
 				}
 			}
 		};
@@ -1032,64 +1055,55 @@ void UEdGraphSchema_CustomizableObject::BreakSinglePinLink(UEdGraphPin* SourcePi
 }
 
 
-TSharedPtr<ICustomizableObjectEditor> UEdGraphSchema_CustomizableObject::GetCustomizableObjectEditor(const class UEdGraph* ParentGraph)
+/** Enum containing all the object types that we are able to convert onto a node when dragging
+	 * and dropping and asset of that type onto the CO graph.
+	 * Each value of this enumeration will, in practice, have a CO node to be represented by.
+	 */
+// TODO: Replace the usage of this enum class with something similar to typeID (not casting)
+enum class ESpawnableObjectType : int32
 {
-	// Find the associated Editor
-	UCustomizableObject* CustomizableObject = Cast<UCustomizableObject>(ParentGraph->GetOuter());
-
-	TSharedPtr<ICustomizableObjectEditor> CustomizableObjectEditor;
-	if (CustomizableObject)
-	{
-		TSharedPtr< IToolkit > FoundAssetEditor = FToolkitManager::Get().FindEditorForAsset(CustomizableObject);
-		if (FoundAssetEditor.IsValid())
-		{
-			return StaticCastSharedPtr<ICustomizableObjectEditor>(FoundAssetEditor);
-		}
-	}
-
-	return nullptr;
-}
+	None = -1,		// Invalid value
+		
+	UTexture2D,
+	USkeletalMesh,
+	UStaticMesh,
+	UMaterialInterface,
+};
 
 
-UEdGraphNode* UEdGraphSchema_CustomizableObject::AddComment(class UEdGraph* ParentGraph, UEdGraphPin* FromPin, const FVector2D Location, bool bSelectNewNode/* = true*/) const
+bool IsSpawnableAsset(const FAssetData& InAsset, ESpawnableObjectType& OutObjectType)
 {
-	TSharedPtr<ICustomizableObjectEditor> CustomizableObjectEditor = GetCustomizableObjectEditor(ParentGraph);
+	UObject* Object = InAsset.GetAsset();
 
-	if(CustomizableObjectEditor.IsValid())
+	// Type used to know what kind of UObject this asset is
+	OutObjectType = ESpawnableObjectType::None;
+
+	// Check if the provided object can be casted to any of the UObject types we can spawn as CO Nodes
+	if (Cast<UTexture2D>(Object))
 	{
-		return CustomizableObjectEditor->CreateCommentBox(Location);
+		OutObjectType = ESpawnableObjectType::UTexture2D;
+		return true;
 	}
-
-	return nullptr;
-}
-
-
-void UEdGraphSchema_CustomizableObject::GetBreakLinkToSubMenuActions( class FMenuBuilder& MenuBuilder, UEdGraphPin* InGraphPin ) const
-{
-	// Make sure we have a unique name for every entry in the list
-	TMap< FString, uint32 > LinkTitleCount;
-
-	// Add all the links we could break from
-	for(TArray<class UEdGraphPin*>::TConstIterator Links(InGraphPin->LinkedTo); Links; ++Links)
+	else if (Cast<USkeletalMesh>(Object))
 	{
-		FString Title = FString::Printf(TEXT("%s (%s)"), *(*Links)->GetOwningNode()->GetNodeTitle(ENodeTitleType::ListView).ToString(), *(*Links)->PinFriendlyName.ToString());
-		uint32 &Count = LinkTitleCount.FindOrAdd(Title);
-		FString Description;
-		if(Count == 0)
-		{
-			Description = FString::Printf(TEXT("Break link to %s"), *Title);
-		}
-		else
-		{
-			Description = FString::Printf(TEXT("Break link to %s (%d)"), *Title, Count);
-		}
-		++Count;
-		MenuBuilder.AddMenuEntry( 
-			FText::FromString(Description), 
-			FText::FromString(Description), 
-			FSlateIcon(),
-			FUIAction( FExecuteAction::CreateUObject((UEdGraphSchema_CustomizableObject*const)this, &UEdGraphSchema_CustomizableObject::BreakSinglePinLink, const_cast< UEdGraphPin* >(InGraphPin), *Links) ) 
-			);
+		OutObjectType = ESpawnableObjectType::USkeletalMesh;
+		return true;
+	}
+	else if (Cast<UStaticMesh>(Object))
+	{
+		OutObjectType = ESpawnableObjectType::UStaticMesh;
+		return true;
+	}
+	else if (Cast<UMaterialInterface>(Object))
+	{
+		OutObjectType = ESpawnableObjectType::UMaterialInterface;
+		return true;
+	}
+	// Add more compatible types here, sync it up with ESpawnableObjectType
+	else
+	{
+		// Non spawnable object
+		return false;
 	}
 }
 
@@ -1208,43 +1222,6 @@ FConnectionDrawingPolicy* UEdGraphSchema_CustomizableObject::CreateConnectionDra
 }
 
 
-bool UEdGraphSchema_CustomizableObject::IsSpawnableAsset(const FAssetData& InAsset, ESpawnableObjectType& OutObjectType) const
-{
-	UObject* Object = InAsset.GetAsset();
-
-	// Type used to know what kind of UObject this asset is
-	OutObjectType = ESpawnableObjectType::None;
-
-	// Check if the provided object can be casted to any of the UObject types we can spawn as CO Nodes
-	if (Cast<UTexture2D>(Object))
-	{
-		OutObjectType = ESpawnableObjectType::UTexture2D;
-		return true;
-	}
-	else if (Cast<USkeletalMesh>(Object))
-	{
-		OutObjectType = ESpawnableObjectType::USkeletalMesh;
-		return true;
-	}
-	else if (Cast<UStaticMesh>(Object))
-	{
-		OutObjectType = ESpawnableObjectType::UStaticMesh;
-		return true;
-	}
-	else if (Cast<UMaterialInterface>(Object))
-	{
-		OutObjectType = ESpawnableObjectType::UMaterialInterface;
-		return true;
-	}
-	// Add more compatible types here, sync it up with ESpawnableObjectType
-	else
-	{
-		// Non spawnable object
-		return false;
-	}
-}
-
-
 void UEdGraphSchema_CustomizableObject::GetAssetsGraphHoverMessage(const TArray<FAssetData>& Assets,
 	const UEdGraph* HoverGraph, FString& OutTooltipText, bool& OutOkIcon) const
 {
diff --git a/Engine/Plugins/Mutable/Source/CustomizableObjectEditor/Public/MuCOE/EdGraphSchema_CustomizableObject.h b/Engine/Plugins/Mutable/Source/CustomizableObjectEditor/Public/MuCOE/EdGraphSchema_CustomizableObject.h
--- a/Engine/Plugins/Mutable/Source/CustomizableObjectEditor/Public/MuCOE/EdGraphSchema_CustomizableObject.h
+++ b/Engine/Plugins/Mutable/Source/CustomizableObjectEditor/Public/MuCOE/EdGraphSchema_CustomizableObject.h
@@ -2,7 +2,7 @@
 
 #pragma once
 
-#include "EdGraphSchema_K2_Actions.h"
+#include "EdGraph/EdGraphSchema.h"
 
 #include "EdGraphSchema_CustomizableObject.generated.h"
 
@@ -16,70 +16,6 @@ struct FAssetData;
 struct FEdGraphPinType;
 class ICustomizableObjectEditor;
 
-/** Action to add a node to the graph */
-USTRUCT()
-struct CUSTOMIZABLEOBJECTEDITOR_API FCustomizableObjectSchemaAction_NewNode : public FEdGraphSchemaAction
-{
-	GENERATED_USTRUCT_BODY();
-
-	/** Template of node we want to create */
-	UPROPERTY()
-	TObjectPtr<class UEdGraphNode> NodeTemplate;
-
-
-	FCustomizableObjectSchemaAction_NewNode() 
-		: FEdGraphSchemaAction()
-		, NodeTemplate(NULL)
-	{}
-
-	FCustomizableObjectSchemaAction_NewNode(const FString& InNodeCategory, const FText& InMenuDesc, const FText& InToolTip, const int32 InGrouping, const FText& InKeywords = FText(), int32 InSectionID = 0)
-		: FEdGraphSchemaAction(FText::FromString(InNodeCategory), InMenuDesc, InToolTip, InGrouping, InKeywords, InSectionID)
-		, NodeTemplate(NULL)
-	{}
-
-	// FEdGraphSchemaAction interface
-	virtual UEdGraphNode* PerformAction(class UEdGraph* ParentGraph, UEdGraphPin* FromPin, const FVector2D Location, bool bSelectNewNode = true) override;
-	virtual UEdGraphNode* PerformAction(class UEdGraph* ParentGraph, TArray<UEdGraphPin*>& FromPins, const FVector2D Location, bool bSelectNewNode = true) override;
-	virtual void AddReferencedObjects( FReferenceCollector& Collector ) override;
-
-	/** Reimplementaiton of EdGraphSchema::CreateNode(...). Performs the overlap calculation before calling AutowireNewNode(...). AutowireNewNode can induce a call to a ReconstrucNode() which removes pins required for the calculation. */
-	static UEdGraphNode* CreateNode(class UEdGraph* ParentGraph, UEdGraphPin* FromPin, const FVector2D Location, class UEdGraphNode* InNodeTemplate);
-	// End of FEdGraphSchemaAction interface
-
-	template <typename NodeType>
-	static NodeType* InstantSpawn(class UEdGraph* ParentGraph, NodeType* InTemplateNode, const FVector2D Location)
-	{
-		FEdGraphSchemaAction_K2NewNode Action;
-		Action.NodeTemplate = InTemplateNode;
-
-		return Cast<NodeType>(Action.PerformAction(ParentGraph, NULL, Location));
-	}
-};
-
-
-/** Action to paste clipboard contents into the graph */
-USTRUCT()
-struct CUSTOMIZABLEOBJECTEDITOR_API FCustomizableObjectSchemaAction_Paste : public FEdGraphSchemaAction
-{
-	GENERATED_USTRUCT_BODY();
-
-	// Simple type info
-	static FName StaticGetTypeId() { static FName Type("FCustomizableObjectSchemaAction_Paste"); return Type; }
-	virtual FName GetTypeId() const override { return StaticGetTypeId(); }
-
-	FCustomizableObjectSchemaAction_Paste()
-		: FEdGraphSchemaAction()
-	{}
-
-	FCustomizableObjectSchemaAction_Paste(const FText& InNodeCategory, const FText& InMenuDesc, const FString& InToolTip, const int32 InGrouping)
-		: FEdGraphSchemaAction(InNodeCategory, InMenuDesc, FText::FromString(InToolTip), InGrouping)
-	{}
-
-	// FEdGraphSchemaAction interface
-	virtual UEdGraphNode* PerformAction(class UEdGraph* ParentGraph, UEdGraphPin* FromPin, const FVector2D Location, bool bSelectNewNode = true) override;
-	// End of FEdGraphSchemaAction interface
-};
-
 
 UCLASS()
 class CUSTOMIZABLEOBJECTEDITOR_API UEdGraphSchema_CustomizableObject : public UEdGraphSchema
@@ -130,47 +66,8 @@ public:
 
 	//  Own interface
 	static FLinearColor GetPinTypeColor(const FName& PinType);
-
-	void GetContextMenuActionsReconstructAllChildNodes(UToolMenu* Menu, TWeakObjectPtr<UGraphNodeContextMenuContext> Context) const;
-
-	/** Utility to create a node creating action */
-	static TSharedPtr<FCustomizableObjectSchemaAction_NewNode> AddNewNodeAction(FGraphActionListBuilderBase& ContextMenuBuilder, const FString& Category, const FText& MenuDesc, const FText& Tooltip, const int32 Grouping = 0, const FString& Keywords = FString());
-
-	/** */
-	void GetBreakLinkToSubMenuActions( class FMenuBuilder& MenuBuilder, UEdGraphPin* InGraphPin ) const;
-	
-private: 
-/** Enum containing all the object types that we are able to convert onto a node when dragging
- * and dropping and asset of that type onto the CO graph.
- * Each value of this enumeration will, in practice, have a CO node to be represented by.
- */
-// TODO: Replace the usage of this enum class with something similar to typeID (not casting)
-enum class ESpawnableObjectType : int32
-{
-	None = -1,		// Invalid value
-	
-	UTexture2D,
-	USkeletalMesh,
-	UStaticMesh,
-	UMaterialInterface,
-};
-
-public:
-
-	/**
-	 * Provided an asset object check if mutable does have a relative node to be able to host this asset's data.
-	 *  @param InAsset - The asset object we want to check if can be used to generate a new mutable node containing it.
-	 *  @param OutObjectType - The type of object of the asset.
-	 *  @return True if the asset can be used to spawn a mutable node, false if not.
-	 *  @note All possible object types that we can convert onto mutable nodes must be defined here.
-	 */
-	bool IsSpawnableAsset(const FAssetData& InAsset, ESpawnableObjectType& OutObjectType) const;
-
-	static TSharedPtr<class ICustomizableObjectEditor> GetCustomizableObjectEditor(const class UEdGraph* ParentGraph);
-	UEdGraphNode* AddComment(class UEdGraph* ParentGraph, UEdGraphPin* FromPin, const FVector2D Location, bool bSelectNewNode = true) const;
 	
 	/** Given a pin category name, get its friendly name (user readable). */
 	static FName GetPinCategoryName(const FName& PinCategory);
 	static FText GetPinCategoryFriendlyName(const FName& PinCategory);
-
 };
