[Backout] - CL39340396
#rnx
Original CL Desc
-----------------------------------------------------------------
[mutable] Added warning message when a parameter name is found more than once.
- Removed unnecessary includes
- Added missing context for the nodes of type Object and Object Group
- Added load Synchronous to ensure the SKM is loaded before checking for null

#rnx

diff --git a/Engine/Plugins/Mutable/Source/CustomizableObjectEditor/Private/MuCOE/GenerateMutableSource/GenerateMutableSource.cpp b/Engine/Plugins/Mutable/Source/CustomizableObjectEditor/Private/MuCOE/GenerateMutableSource/GenerateMutableSource.cpp
--- a/Engine/Plugins/Mutable/Source/CustomizableObjectEditor/Private/MuCOE/GenerateMutableSource/GenerateMutableSource.cpp
+++ b/Engine/Plugins/Mutable/Source/CustomizableObjectEditor/Private/MuCOE/GenerateMutableSource/GenerateMutableSource.cpp
@@ -888,7 +888,6 @@ mu::Ptr<mu::NodeObject> GenerateMutableSource(const UEdGraphPin * Pin, FMutableG
 		mu::Ptr<mu::NodeObjectNew> ObjectNode = new mu::NodeObjectNew();
 		Result = ObjectNode;
 
-		ObjectNode->SetMessageContext(Node);
 		ObjectNode->SetName(TypedNodeObj->ObjectName);
 		FGuid FinalGuid = GenerationContext.GetNodeIdUnique(TypedNodeObj);
 		if (FinalGuid != TypedNodeObj->NodeGuid)
@@ -1048,7 +1047,6 @@ mu::Ptr<mu::NodeObject> GenerateMutableSource(const UEdGraphPin * Pin, FMutableG
 		// All sockets from all mesh parts plugged into this group node will have the following priority when there's a socket name clash
 		GenerationContext.SocketPriorityStack.Push(TypedNodeGroup->SocketPriority);
 
-		GroupNode->SetMessageContext(Node);
 		GroupNode->SetName(TypedNodeGroup->GroupName);
 		GroupNode->SetUid(TypedNodeGroup->NodeGuid.ToString());
 		
diff --git a/Engine/Plugins/Mutable/Source/CustomizableObjectEditor/Private/MuCOE/GenerateMutableSource/GenerateMutableSourceMesh.cpp b/Engine/Plugins/Mutable/Source/CustomizableObjectEditor/Private/MuCOE/GenerateMutableSource/GenerateMutableSourceMesh.cpp
--- a/Engine/Plugins/Mutable/Source/CustomizableObjectEditor/Private/MuCOE/GenerateMutableSource/GenerateMutableSourceMesh.cpp
+++ b/Engine/Plugins/Mutable/Source/CustomizableObjectEditor/Private/MuCOE/GenerateMutableSource/GenerateMutableSourceMesh.cpp
@@ -3442,7 +3442,6 @@ mu::Ptr<mu::NodeMesh> GenerateMutableSourceMesh(const UEdGraphPin* Pin,
 		mu::NodeMeshConstantPtr MeshNode = new mu::NodeMeshConstant();
 		Result = MeshNode;
 
-		GenerationContext.LoadObject(TypedNodeSkel->SkeletalMesh, true);
 		if (TypedNodeSkel->SkeletalMesh)
 		{
 			int32 LODIndexConnected = -1; // LOD which the pin is connected to
@@ -3693,7 +3692,6 @@ mu::Ptr<mu::NodeMesh> GenerateMutableSourceMesh(const UEdGraphPin* Pin,
 
 	else if (const UCustomizableObjectNodeStaticMesh* TypedNodeStatic = Cast<UCustomizableObjectNodeStaticMesh>(Node))
 	{
-		GenerationContext.LoadObject(TypedNodeStatic->StaticMesh, true);
 		if (TypedNodeStatic->StaticMesh == nullptr)
 		{
 			FString Msg = FString::Printf(TEXT("The UCustomizableObjectNodeStaticMesh node %s has no static mesh assigned"), *Node->GetName());
@@ -3835,7 +3833,7 @@ mu::Ptr<mu::NodeMesh> GenerateMutableSourceMesh(const UEdGraphPin* Pin,
 						continue;
 					}
 
-					// Checking if the morph exists in the application node
+					// Cheking if the morph exists in the application node
 					FString MorphName = MorphPin->PinFriendlyName.ToString();
 					if (!MorphNames.Contains(MorphName))
 					{
diff --git a/Engine/Plugins/Mutable/Source/CustomizableObjectEditor/Private/MuCOE/Nodes/CustomizableObjectNodeSkeletalMesh.cpp b/Engine/Plugins/Mutable/Source/CustomizableObjectEditor/Private/MuCOE/Nodes/CustomizableObjectNodeSkeletalMesh.cpp
--- a/Engine/Plugins/Mutable/Source/CustomizableObjectEditor/Private/MuCOE/Nodes/CustomizableObjectNodeSkeletalMesh.cpp
+++ b/Engine/Plugins/Mutable/Source/CustomizableObjectEditor/Private/MuCOE/Nodes/CustomizableObjectNodeSkeletalMesh.cpp
@@ -16,7 +16,6 @@
 #include "MuCOE/Nodes/CustomizableObjectNodeLayoutBlocks.h"
 #include "MuCOE/RemapPins/CustomizableObjectNodeRemapPinsByNameDefaultPin.h"
 #include "MuCOE/UnrealEditorPortabilityHelpers.h"
-#include "UObject/ICookInfo.h"
 #include "Widgets/Input/SCheckBox.h"
 
 class UCustomizableObjectNodeRemapPinsByName;
@@ -34,11 +33,7 @@ static const TCHAR* SKELETAL_MESH_PIN_NAME = TEXT("Skeletal Mesh");
 void UCustomizableObjectNodeSkeletalMesh::PostLoad()
 {
 	Super::PostLoad();
-	
-	FCookLoadScope CookLoadScope(ECookLoadType::EditorOnly);
 
-	// Block only required by the FixUps : todo: Will be removed soon
-	SkeletalMesh.LoadSynchronous();
 	if (SkeletalMesh)
 	{
 		ConditionalPostLoadReference(*SkeletalMesh);
@@ -443,8 +438,6 @@ FText UCustomizableObjectNodeSkeletalMesh::GetTooltipText() const
 void UCustomizableObjectNodeSkeletalMesh::BackwardsCompatibleFixup(int32 CustomizableObjectCustomVersion)
 {
 	Super::BackwardsCompatibleFixup(CustomizableObjectCustomVersion);
-
-	FCookLoadScope CookLoadScope(ECookLoadType::EditorOnly);
 	
 	if (CustomizableObjectCustomVersion == FCustomizableObjectCustomVersion::PostLoadToCustomVersion)
 	{
@@ -526,7 +519,6 @@ void UCustomizableObjectNodeSkeletalMesh::BackwardsCompatibleFixup(int32 Customi
 					AddPinData(*Section.MeshPinRef.Get(), *PinData);
 				}
 
-				SkeletalMesh.LoadSynchronous();
 				if (SkeletalMesh)
 				{
 					if (const FSkeletalMaterial* SkeletalMaterial = GetSkeletalMaterialFor(LODIndex, SectionIndex))
@@ -603,7 +595,6 @@ void UCustomizableObjectNodeSkeletalMesh::BackwardsCompatibleFixup(int32 Customi
 
 	if (CustomizableObjectCustomVersion == FCustomizableObjectCustomVersion::MoveLayoutToNodeSkeletalMesh)
 	{
-		SkeletalMesh.LoadSynchronous();
 		if (const FSkeletalMeshModel* ImportedModel = SkeletalMesh ? SkeletalMesh->GetImportedModel() : nullptr)
 		{
 			TArray<UEdGraphPin*> PinsToDelete;
diff --git a/Engine/Plugins/Mutable/Source/MutableTools/Private/MuT/Compiler.cpp b/Engine/Plugins/Mutable/Source/MutableTools/Private/MuT/Compiler.cpp
--- a/Engine/Plugins/Mutable/Source/MutableTools/Private/MuT/Compiler.cpp
+++ b/Engine/Plugins/Mutable/Source/MutableTools/Private/MuT/Compiler.cpp
@@ -11,6 +11,8 @@
 #include "MuT/CompilerPrivate.h"
 #include "MuT/ErrorLog.h"
 #include "MuT/Node.h"
+#include "MuT/NodePrivate.h"
+#include "MuT/Table.h"
 #include "MuR/Image.h"
 #include "MuR/Mesh.h"
 #include "MuR/Model.h"
@@ -18,8 +20,10 @@
 #include "MuR/MutableTrace.h"
 #include "MuR/Operations.h"
 #include "MuR/ParametersPrivate.h"
+#include "MuR/Platform.h"
 #include "MuR/Serialisation.h"
 #include "MuR/MutableRuntimeModule.h"
+#include "MuR/System.h"
 
 #include "Trace/Detail/Channel.h"
 #include "Containers/Array.h"
@@ -35,6 +39,7 @@
 
 #include <inttypes.h>	// Required for 64-bit printf macros
 
+
 namespace mu
 {
 
@@ -296,81 +301,53 @@ namespace mu
     {
         MUTABLE_CPUPROFILER_SCOPE(Compile);
 
-        TArray< FStateCompilationData > States;
-        TSharedPtr<FErrorLog> GenErrorLog;
+        TArray< FStateCompilationData > states;
+        TSharedPtr<FErrorLog> genErrorLog;
 		TArray<FParameterDesc> Parameters;
         {
-            CodeGenerator Generator( m_pD->Options->GetPrivate() );
+            CodeGenerator gen( m_pD->Options->GetPrivate() );
 
-            Generator.GenerateRoot( pNode );
+            gen.GenerateRoot( pNode );
 
-            check( !Generator.States.IsEmpty() );
+            check( !gen.States.IsEmpty() );
 
-            for ( const TPair<FObjectState, Ptr<ASTOp>>& s: Generator.States )
+            for ( const TPair<FObjectState, Ptr<ASTOp>>& s: gen.States )
             {
-                FStateCompilationData Data;
-                Data.nodeState = s.Key;
-                Data.root = s.Value;
-                Data.state.Name = s.Key.Name;
-                States.Add( Data );
+                FStateCompilationData data;
+                data.nodeState = s.Key;
+                data.root = s.Value;
+                data.state.Name = s.Key.Name;
+                states.Add( data );
             }
 
-            GenErrorLog = Generator.ErrorLog;
+            genErrorLog = gen.ErrorLog;
 
-        	// Sort the parameters as deterministically as possible.
-        	Generator.FirstPass.ParameterNodes.ValueStableSort([](const Ptr<ASTOpParameter>& A, const Ptr<ASTOpParameter>& B)->bool
-			{
-				if (A->Parameter.Name < B->Parameter.Name)
-				{
-					return true;
-				}
-				if (A->Parameter.Name > B->Parameter.Name)
-				{
-					return false;
-				}
-				return A->Parameter.UID < B->Parameter.UID;
-			});
-        	
 			// Set the parameter list from the non-optimized data, so that we have them all even if they are optimized out
-            {
-            	const int32 ParameterCount = Generator.FirstPass.ParameterNodes.Num();
-            	Parameters.Reserve(ParameterCount);
-            	for ( const TPair<Ptr<const Node>, Ptr<ASTOpParameter>>& Entry : Generator.FirstPass.ParameterNodes )
-            	{
-            		Parameters.Add(Entry.Value->Parameter);
-            	}
-            }
-        	
-        	// Check for repeated parameter names and report them to the user
-			if (Parameters.Num() > 1)
+			int32 ParameterCount = gen.FirstPass.ParameterNodes.Num();
+			Parameters.SetNum(ParameterCount);
+			int32 ParameterIndex = 0;
+			for ( const TPair<Ptr<const Node>, Ptr<ASTOpParameter>>& Entry : gen.FirstPass.ParameterNodes )
 			{
-				TArray<Ptr<const Node>> ParameterNodes;
-				Generator.FirstPass.ParameterNodes.GenerateKeyArray(ParameterNodes);
-				
-				// Should be impossible as both arrays come from the same data.
-				check(ParameterNodes.Num() == Parameters.Num());
-
-				TArray<FString> RepeatedParameterNames;
-				
-				for (int32 Index = 1; Index < Parameters.Num(); ++Index)
-				{
-					const FString CurrentParameterName = Parameters[Index].Name;
-					const FString PreviousParameterName = Parameters[Index-1].Name;
-					
-					if (CurrentParameterName.Compare(PreviousParameterName) == 0 && !RepeatedParameterNames.Contains(CurrentParameterName))
-					{
-						// Report repeated parameter name
-						const FString WarningText = FString::Printf(TEXT("Repeated parameter found : \"%s\". Please use a different name."), *CurrentParameterName);
-						const void* MessageContext = ParameterNodes[Index]->GetMessageContext();
-						GenErrorLog->Add(WarningText, ELMT_WARNING, MessageContext);
+				Parameters[ParameterIndex] = Entry.Value->Parameter;
+				++ParameterIndex;
+			}
 
-						RepeatedParameterNames.Add(CurrentParameterName);
-					}
+			// Sort the parameters as deterministically as possible.
+			struct FParameterSortPredicate
+			{
+				bool operator()(const FParameterDesc& A, const FParameterDesc& B) const
+				{
+					if (A.Name < B.Name) return true;
+					if (A.Name > B.Name) return false;
+					return A.UID < B.UID;
 				}
-			}
+			};
+			
+			FParameterSortPredicate SortPredicate;
+			Parameters.Sort(SortPredicate);
 		}
 
-        mu::Private::RemoveDisabledAddInstanceOpsValues(m_pD->Options, States);
+        mu::Private::RemoveDisabledAddInstanceOpsValues(m_pD->Options, states);
 
         // Slow AST code verification for debugging.
         //TArray<Ptr<ASTOp>> roots;
@@ -382,13 +359,13 @@ namespace mu
 
         // Optimize the generated code
         {
-            CodeOptimiser Optimiser( m_pD->Options, States );
-            Optimiser.OptimiseAST( );
+            CodeOptimiser optimiser( m_pD->Options, states );
+            optimiser.OptimiseAST( );
         }
 
         // Link the Program and generate state data.
-		TSharedPtr<FModel> Result = MakeShared<FModel>();
-        FProgram& Program = Result->GetPrivate()->Program;
+		TSharedPtr<FModel> pResult = MakeShared<FModel>();
+        FProgram& Program = pResult->GetPrivate()->Program;
 
 		check(Program.Parameters.IsEmpty());
 		Program.Parameters = Parameters;
@@ -402,17 +379,17 @@ namespace mu
 		FImageOperator ImOp = FImageOperator::GetDefault(m_pD->Options->GetPrivate()->ImageFormatFunc);
 		FLinkerOptions LinkerOptions(ImOp);
 
-		for(FStateCompilationData& State : States )
+		for(FStateCompilationData& s: states )
         {
 			LinkerOptions.MinTextureResidentMipCount = m_pD->Options->GetPrivate()->MinTextureResidentMipCount;
 
-            if (State.root)
+            if (s.root)
             {
-				State.state.Root = ASTOp::FullLink(State.root, Program, &LinkerOptions);
+				s.state.Root = ASTOp::FullLink(s.root, Program, &LinkerOptions);
             }
             else
             {
-                State.state.Root = 0;
+                s.state.Root = 0;
             }
         }
 
@@ -420,68 +397,68 @@ namespace mu
 		Program.OpAddress.Shrink();
 
         // Set the runtime parameter indices.
-        for(FStateCompilationData& State : States )
+        for(FStateCompilationData& s: states )
         {
-            for ( int32 RuntimeParameterIndex=0; RuntimeParameterIndex <State.nodeState.RuntimeParams.Num(); ++RuntimeParameterIndex)
+            for ( int32 RuntimeParameterIndex=0; RuntimeParameterIndex <s.nodeState.RuntimeParams.Num(); ++RuntimeParameterIndex)
             {
-                int32 ParameterIndex = -1;
-                for ( int32 i=0; ParameterIndex<0 && i<Program.Parameters.Num(); ++i )
+                int32 paramIndex = -1;
+                for ( int32 i=0; paramIndex<0 && i<Program.Parameters.Num(); ++i )
                 {
                     if ( Program.Parameters[i].Name
                          ==
-                         State.nodeState.RuntimeParams[RuntimeParameterIndex] )
+                         s.nodeState.RuntimeParams[RuntimeParameterIndex] )
                     {
-                        ParameterIndex = (int)i;
+                        paramIndex = (int)i;
                     }
                 }
 
-                if (ParameterIndex>=0)
+                if (paramIndex>=0)
                 {
-                    State.state.m_runtimeParameters.Add( ParameterIndex );
+                    s.state.m_runtimeParameters.Add( paramIndex );
                 }
                 else
                 {
-					FString WarningString = FString::Printf(TEXT(
+					FString Temp = FString::Printf(TEXT(
 						"The state [%s] refers to a parameter [%s]  that has not been found in the model. This error can be "
 						"safely dismissed in case of partial compilation."), 
-						*State.nodeState.Name,
-						*State.nodeState.RuntimeParams[RuntimeParameterIndex]);
-                    m_pD->ErrorLog->Add(WarningString, ELMT_WARNING, pNode->GetMessageContext() );
+						*s.nodeState.Name,
+						*s.nodeState.RuntimeParams[RuntimeParameterIndex]);
+                    m_pD->ErrorLog->Add(Temp, ELMT_WARNING, pNode->GetMessageContext() );
                 }
             }
 
             // Generate the mask of update cache ops
-            for (const Ptr<ASTOp>& Instruction : State.m_updateCache )
+            for ( const auto& a: s.m_updateCache )
             {
-                State.state.m_updateCache.Emplace(Instruction->linkedAddress);
+                s.state.m_updateCache.Emplace(a->linkedAddress);
             }
 
             // Sort the update cache addresses for performance and determinism
-            State.state.m_updateCache.Sort();
+            s.state.m_updateCache.Sort();
 
             // Generate the mask of dynamic resources
-            for (const TPair<Ptr<ASTOp>, TArray<FString>>& Instruction : State.m_dynamicResources )
+            for ( const auto& a: s.m_dynamicResources )
             {
-                uint64 RelevantMask = 0;
-                for (const FString& InstructionParameter : Instruction.Value )
+                uint64 relevantMask = 0;
+                for ( const auto& b: a.Value )
                 {
                     // Find the index in the model parameter list
-                    int ParameterIndex = -1;
-                    for ( int32 i=0; ParameterIndex<0 && i<Program.Parameters.Num(); ++i )
+                    int paramIndex = -1;
+                    for ( int32 i=0; paramIndex<0 && i<Program.Parameters.Num(); ++i )
                     {
-                        if ( Program.Parameters[i].Name == InstructionParameter )
+                        if ( Program.Parameters[i].Name == b )
                         {
-                            ParameterIndex = (int)i;
+                            paramIndex = (int)i;
                         }
                     }
-                    check(ParameterIndex>=0);
+                    check(paramIndex>=0);
 
                     // Find the position in the state data vector.
-                    const int32 IndexInRuntimeList = State.state.m_runtimeParameters.Find( ParameterIndex );
+                    int32 IndexInRuntimeList = s.state.m_runtimeParameters.Find( paramIndex );
 
                     if (IndexInRuntimeList !=INDEX_NONE )
                     {
-                        RelevantMask |= uint64(1) << IndexInRuntimeList;
+                        relevantMask |= uint64(1) << IndexInRuntimeList;
                     }
                 }
 
@@ -489,26 +466,26 @@ namespace mu
                 // Maybe something with the difference of precision between the relevant parameters
                 // in operation subtrees.
                 //check(relevantMask!=0);
-                if (RelevantMask!=0)
+                if (relevantMask!=0)
                 {
-                    State.state.m_dynamicResources.Emplace( Instruction.Key->linkedAddress, RelevantMask );
+                    s.state.m_dynamicResources.Emplace( a.Key->linkedAddress, relevantMask );
                 }
             }            
 
             // Sort for performance and determinism
-            State.state.m_dynamicResources.Sort();
+            s.state.m_dynamicResources.Sort();
 
-            Program.States.Add(State.state);
+            Program.States.Add(s.state);
         }
 
 		UE_LOG(LogMutableCore, Verbose, TEXT("(int) %s : %ld"), TEXT("Program size"), int64(Program.OpAddress.Num()));
 
         // Merge the log in the right order
-        GenErrorLog->Merge( m_pD->ErrorLog.Get() );
-        m_pD->ErrorLog = GenErrorLog;
+        genErrorLog->Merge( m_pD->ErrorLog.Get() );
+        m_pD->ErrorLog = genErrorLog;
 
 		// Pack data
-		m_pD->GenerateRoms(Result.Get(), LinkerOptions.AdditionalData );
+		m_pD->GenerateRoms(pResult.Get(), LinkerOptions.AdditionalData );
 
 		UE_LOG(LogMutableCore, Verbose, TEXT("(int) %s : %ld"), TEXT("Program size"), int64(Program.OpAddress.Num()));
 
@@ -523,7 +500,7 @@ namespace mu
 		Program.ConstantPhysicsBodies.Shrink();
 		Program.Parameters.Shrink();
 
-        return Result;
+        return pResult;
     }
 
 
@@ -803,3 +780,4 @@ namespace mu
 
 
 }
+
