[Mutable] Add cook scope to remove unwanted assets from builds.

#rnx

diff --git a/Engine/Plugins/Mutable/Source/CustomizableObjectEditor/Private/MuCOE/CustomizableObjectEditorModule.cpp b/Engine/Plugins/Mutable/Source/CustomizableObjectEditor/Private/MuCOE/CustomizableObjectEditorModule.cpp
--- a/Engine/Plugins/Mutable/Source/CustomizableObjectEditor/Private/MuCOE/CustomizableObjectEditorModule.cpp
+++ b/Engine/Plugins/Mutable/Source/CustomizableObjectEditor/Private/MuCOE/CustomizableObjectEditorModule.cpp
@@ -532,6 +532,8 @@ void GetReferencingCustomizableObjects(FName CustomizableObjectName, TArray<FNam
 
 void GetReferencingPackages(const UCustomizableObject& Object, TArray<FAssetData>& ReferencingAssets)
 {
+	FCookLoadScope CookLoadScope(ECookLoadType::EditorOnly);
+
 	// Gather all child CustomizableObjects
 	TArray<FName> VisitedObjectNames;
 	GetReferencingCustomizableObjects(Object.GetPackage()->GetFName(), VisitedObjectNames, ReferencingAssets);
diff --git a/Engine/Plugins/Mutable/Source/CustomizableObjectEditor/Private/MuCOE/GenerateMutableSource/GenerateMutableSource.cpp b/Engine/Plugins/Mutable/Source/CustomizableObjectEditor/Private/MuCOE/GenerateMutableSource/GenerateMutableSource.cpp
--- a/Engine/Plugins/Mutable/Source/CustomizableObjectEditor/Private/MuCOE/GenerateMutableSource/GenerateMutableSource.cpp
+++ b/Engine/Plugins/Mutable/Source/CustomizableObjectEditor/Private/MuCOE/GenerateMutableSource/GenerateMutableSource.cpp
@@ -51,15 +51,9 @@
 #include "Hash/CityHash.h"
 #include "MuCO/CustomizableObjectResourceDataTypes.h"
 #include "MuCOE/Nodes/CustomizableObjectNodeGroupProjectorParameter.h"
-#include "UObject/ICookInfo.h"
 
 #define LOCTEXT_NAMESPACE "CustomizableObjectEditor"
 
-TAutoConsoleVariable<bool> CVarLoadObjectCookLoadTypeEditorOnly(
-	TEXT("Mutable.LoadObjectCookLoadTypeEditorOnly"),
-	true,
-	TEXT("Controls whether calls to LoadObject use a CookLoadScope of type EditorOnly")
-);
 
 uint32 GetTypeHash(const FGraphCycleKey& Key)
 {
@@ -498,19 +492,11 @@ FMutableComponentInfo* FMutableGraphGenerationContext::GetCurrentComponentInfo()
 
 UObject* FMutableGraphGenerationContext::LoadObject(const FSoftObjectPtr& SoftObject, bool bParticipatingObjectsPassLoad)
 {
-	if (CVarLoadObjectCookLoadTypeEditorOnly.GetValueOnAnyThread())
-	{ 
-		FCookLoadScope CookLoadScope(ECookLoadType::EditorOnly);
-		return bLoadObjects && (!bParticipatingObjectsPass || bParticipatingObjectsPassLoad) ?
-			SoftObject.LoadSynchronous() :
-			nullptr;
-	}
-	else
-	{
-		return bLoadObjects && (!bParticipatingObjectsPass || bParticipatingObjectsPassLoad) ?
-			SoftObject.LoadSynchronous() :
-			nullptr;
-	}
+	FCookLoadScope CookLoadScope(ECookLoadType::EditorOnly);
+
+	return bLoadObjects && (!bParticipatingObjectsPass || bParticipatingObjectsPassLoad) ?
+		SoftObject.LoadSynchronous() :
+		nullptr;
 }
 
 
diff --git a/Engine/Plugins/Mutable/Source/CustomizableObjectEditor/Private/MuCOE/GenerateMutableSource/GenerateMutableSource.h b/Engine/Plugins/Mutable/Source/CustomizableObjectEditor/Private/MuCOE/GenerateMutableSource/GenerateMutableSource.h
--- a/Engine/Plugins/Mutable/Source/CustomizableObjectEditor/Private/MuCOE/GenerateMutableSource/GenerateMutableSource.h
+++ b/Engine/Plugins/Mutable/Source/CustomizableObjectEditor/Private/MuCOE/GenerateMutableSource/GenerateMutableSource.h
@@ -29,6 +29,7 @@
 #include "MuR/Mesh.h"
 #include "MuR/Ptr.h"
 #include "UObject/Package.h"
+#include "UObject/ICookInfo.h"
 
 class UCustomizableObjectNodeGroupProjectorParameter;
 class UCustomizableObjectNodeMaterialBase;
@@ -772,6 +773,7 @@ public:
 	template<typename T>
 	T* LoadObject(const TSoftObjectPtr<T>& SoftObject, bool bParticipatingObjectsPassLoad = false)
 	{
+		FCookLoadScope CookLoadScope(ECookLoadType::EditorOnly);
 		return bLoadObjects && (!bParticipatingObjectsPass || bParticipatingObjectsPassLoad) ?
 			SoftObject.LoadSynchronous() :
 			nullptr;
@@ -780,6 +782,7 @@ public:
 	template<typename T>
 	UClass* LoadClass(const TSoftClassPtr<T>& SoftClass, bool bParticipatingObjectsPassLoad = false)
 	{
+		FCookLoadScope CookLoadScope(ECookLoadType::EditorOnly);
 		return bLoadObjects && (!bParticipatingObjectsPass || bParticipatingObjectsPassLoad) ?
 			SoftClass.LoadSynchronous() :
 			nullptr;
diff --git a/Engine/Plugins/Mutable/Source/CustomizableObjectEditor/Private/MuCOE/Nodes/CustomizableObjectNodeTable.cpp b/Engine/Plugins/Mutable/Source/CustomizableObjectEditor/Private/MuCOE/Nodes/CustomizableObjectNodeTable.cpp
--- a/Engine/Plugins/Mutable/Source/CustomizableObjectEditor/Private/MuCOE/Nodes/CustomizableObjectNodeTable.cpp
+++ b/Engine/Plugins/Mutable/Source/CustomizableObjectEditor/Private/MuCOE/Nodes/CustomizableObjectNodeTable.cpp
@@ -285,6 +285,7 @@ void UCustomizableObjectNodeTable::PostLoad()
 		DefaultDataArray.SetNumZeroed(TableStruct->GetStructureSize());
 		TableStruct->InitializeStruct(DefaultDataArray.GetData());
 
+		FCookLoadScope CookLoadScope(ECookLoadType::EditorOnly);
 		for (TFieldIterator<FProperty> It(TableStruct); It; ++It)
 		{
 			FProperty* ColumnProperty = *It;
@@ -449,6 +450,7 @@ void UCustomizableObjectNodeTable::AllocateDefaultPins(UCustomizableObjectNodeRe
 	TSet<FString> PinNames;
 	bool bHasRepeatedNames = false;
 	
+	FCookLoadScope CookLoadScope(ECookLoadType::EditorOnly);
 	for (TFieldIterator<FProperty> It(TableStruct); It; ++It)
 	{
 		FProperty* ColumnProperty = *It;
@@ -754,6 +756,7 @@ bool UCustomizableObjectNodeTable::IsNodeOutDatedAndNeedsRefresh()
 
 	bool bNeedsUpdate = false;
 
+	FCookLoadScope CookLoadScope(ECookLoadType::EditorOnly);
 	for (TFieldIterator<FProperty> It(TableStruct); It && !bNeedsUpdate; ++It)
 	{
 		FProperty* ColumnProperty = *It;
diff --git a/Engine/Plugins/Mutable/Source/CustomizableObjectEditor/Private/MuCOE/Nodes/CustomizableObjectNodeTable.h b/Engine/Plugins/Mutable/Source/CustomizableObjectEditor/Private/MuCOE/Nodes/CustomizableObjectNodeTable.h
--- a/Engine/Plugins/Mutable/Source/CustomizableObjectEditor/Private/MuCOE/Nodes/CustomizableObjectNodeTable.h
+++ b/Engine/Plugins/Mutable/Source/CustomizableObjectEditor/Private/MuCOE/Nodes/CustomizableObjectNodeTable.h
@@ -6,6 +6,7 @@
 #include "Engine/Texture2DArray.h"
 #include "MuCOE/Nodes/CustomizableObjectNode.h"
 #include "MuCOE/RemapPins/CustomizableObjectNodeRemapPinsByName.h"
+#include "UObject/ICookInfo.h"
 
 #include "CustomizableObjectNodeTable.generated.h"
 
@@ -323,6 +324,7 @@ public:
 
 					if (CellData)
 					{
+						FCookLoadScope CookLoadScope(ECookLoadType::EditorOnly);
 						Object = SoftObjectProperty->GetPropertyValue(CellData).LoadSynchronous();
 					}
 
