[Mutable] Fix UCusotmizableComponent and UCustomizableObjectInstanceUsage using CDO objects.

diff --git a/Engine/Plugins/Mutable/Source/CustomizableObject/Internal/MuCO/CustomizableSkeletalComponentPrivate.h b/Engine/Plugins/Mutable/Source/CustomizableObject/Internal/MuCO/CustomizableSkeletalComponentPrivate.h
--- a/Engine/Plugins/Mutable/Source/CustomizableObject/Internal/MuCO/CustomizableSkeletalComponentPrivate.h
+++ b/Engine/Plugins/Mutable/Source/CustomizableObject/Internal/MuCO/CustomizableSkeletalComponentPrivate.h
@@ -29,7 +29,7 @@ public:
 	
 	const UCustomizableSkeletalComponent* GetPublic() const;
 
-	UPROPERTY()
+	UPROPERTY(Instanced)
 	TObjectPtr<UCustomizableObjectInstanceUsage> InstanceUsage;
 };
 
diff --git a/Engine/Plugins/Mutable/Source/CustomizableObject/Public/MuCO/CustomizableObjectInstanceUsage.h b/Engine/Plugins/Mutable/Source/CustomizableObject/Public/MuCO/CustomizableObjectInstanceUsage.h
--- a/Engine/Plugins/Mutable/Source/CustomizableObject/Public/MuCO/CustomizableObjectInstanceUsage.h
+++ b/Engine/Plugins/Mutable/Source/CustomizableObject/Public/MuCO/CustomizableObjectInstanceUsage.h
@@ -114,7 +114,7 @@ private:
 	UPROPERTY()
 	bool bUsedSkipSetSkeletalMeshOnAttach = false;
 	
-	UPROPERTY( VisibleAnywhere, Category=CO)
+	UPROPERTY(Instanced)
 	TObjectPtr<UCustomizableObjectInstanceUsagePrivate> Private;
 	
 	friend UCustomizableObjectInstanceUsagePrivate;
diff --git a/Engine/Plugins/Mutable/Source/CustomizableObject/Public/MuCO/CustomizableSkeletalComponent.h b/Engine/Plugins/Mutable/Source/CustomizableObject/Public/MuCO/CustomizableSkeletalComponent.h
--- a/Engine/Plugins/Mutable/Source/CustomizableObject/Public/MuCO/CustomizableSkeletalComponent.h
+++ b/Engine/Plugins/Mutable/Source/CustomizableObject/Public/MuCO/CustomizableSkeletalComponent.h
@@ -47,7 +47,7 @@ private:
 	UPROPERTY()
 	bool bSkipSkipSetSkeletalMeshOnAttach = false;
 	
-	UPROPERTY()
+	UPROPERTY(Instanced)
 	TObjectPtr<UCustomizableSkeletalComponentPrivate> Private;
 
 public:
