[mutable] Remove unused static constant tables, and move some to the tools library.
Also cleanup naming convention in related code.
#rnx

diff --git a/Engine/Plugins/Mutable/Source/CustomizableObject/Internal/MuCO/CustomizableObjectPrivate.h b/Engine/Plugins/Mutable/Source/CustomizableObject/Internal/MuCO/CustomizableObjectPrivate.h
--- a/Engine/Plugins/Mutable/Source/CustomizableObject/Internal/MuCO/CustomizableObjectPrivate.h
+++ b/Engine/Plugins/Mutable/Source/CustomizableObject/Internal/MuCO/CustomizableObjectPrivate.h
@@ -870,7 +870,7 @@ enum class EMutableFileFlags : uint8
 #if WITH_EDITOR
 namespace MutablePrivate
 {
-	enum class EDataType : uint32 // uint32 for padding and DDC purposes
+	enum class EStreamableDataType : uint32 // uint32 for padding and DDC purposes
 	{
 		None = 0,
 		Model,
@@ -908,7 +908,7 @@ namespace MutablePrivate
 
 	struct CUSTOMIZABLEOBJECT_API FFile
 	{
-		EDataType DataType = EDataType::None;
+		EStreamableDataType  DataType = EStreamableDataType::None;
 
 		/** Rom ResourceType. */
 		uint16 ResourceType = 0;
@@ -943,12 +943,12 @@ namespace MutablePrivate
 
 	struct FFileCategoryID
 	{
-		FFileCategoryID(EDataType DataType, uint16 ResourceType, uint16 Flags);
+		FFileCategoryID(EStreamableDataType DataType, uint16 ResourceType, uint16 Flags);
 
 		FFileCategoryID() = default;
 
-		// DATATYPE
-		EDataType DataType = EDataType::None;
+		// EDataType
+		EStreamableDataType DataType = EStreamableDataType::None;
 
 		/** Rom ResourceType. */
 		uint16 ResourceType = 0;
@@ -1820,6 +1820,8 @@ public:
 
 		RemoveObsoletMeshInterpolateAndGeometryOp,
 
+		RemoveObsoleteDataTypesFromEnum,
+
 		// -----<new versions can be added above this line>--------
 		LastCustomizableObjectVersion
 	};
diff --git a/Engine/Plugins/Mutable/Source/CustomizableObject/Private/MuCO/CustomizableObject.cpp b/Engine/Plugins/Mutable/Source/CustomizableObject/Private/MuCO/CustomizableObject.cpp
--- a/Engine/Plugins/Mutable/Source/CustomizableObject/Private/MuCO/CustomizableObject.cpp
+++ b/Engine/Plugins/Mutable/Source/CustomizableObject/Private/MuCO/CustomizableObject.cpp
@@ -2682,7 +2682,7 @@ namespace MutablePrivate
 		const int32 NumBlocks = Blocks.Num();
 
 		const uint8* SourceData = nullptr;
-		if (DataType == EDataType::Model)
+		if (DataType == EStreamableDataType::Model)
 		{
 			for (int32 BlockIndex = 0; BlockIndex < NumBlocks; ++BlockIndex)
 			{
@@ -2692,7 +2692,7 @@ namespace MutablePrivate
 			}
 			return;
 		}
-		else if (DataType == EDataType::RealTimeMorph)
+		else if (DataType == EStreamableDataType::RealTimeMorph)
 		{
 			for (int32 BlockIndex = 0; BlockIndex < NumBlocks; ++BlockIndex)
 			{
@@ -2701,7 +2701,7 @@ namespace MutablePrivate
 				PlatformData->MorphStreamableData.Get(Block.Id, TArrayView64<uint8>(DestData.GetData() + Block.Offset, Block.Size), bDropData);
 			}
 		}
-		else if (DataType == EDataType::Clothing)
+		else if (DataType == EStreamableDataType::Clothing)
 		{
 
 			for (int32 BlockIndex = 0; BlockIndex < NumBlocks; ++BlockIndex)
@@ -2717,7 +2717,7 @@ namespace MutablePrivate
 		}
 	}
 
-	FFileCategoryID::FFileCategoryID(EDataType InDataType, uint16 InResourceType, uint16 InFlags)
+	FFileCategoryID::FFileCategoryID(EStreamableDataType InDataType, uint16 InResourceType, uint16 InFlags)
 	{
 		DataType = InDataType;
 		ResourceType = InResourceType;
@@ -2801,20 +2801,20 @@ namespace MutablePrivate
 		// Model Roms. Iterate all Model roms to distribute them in categories.
 		{
 			// Add meshes and low-res textures to the Default bucket 
-			DefaultBucket.Categories.Add({ FFileCategoryID(EDataType::Model, mu::DATATYPE::DT_MESH, 0), 0, 0, 0 });
-			DefaultBucket.Categories.Add({ FFileCategoryID(EDataType::Model, mu::DATATYPE::DT_IMAGE, 0), 0, 0, 0 });
+			DefaultBucket.Categories.Add({ FFileCategoryID(EStreamableDataType::Model, uint16(mu::EDataType::Mesh), 0), 0, 0, 0 });
+			DefaultBucket.Categories.Add({ FFileCategoryID(EStreamableDataType::Model, uint16(mu::EDataType::Image), 0), 0, 0, 0 });
 
 			// Add High-res textures to the Optional bucket
-			OptionalBucket.Categories.Add({ FFileCategoryID(EDataType::Model, mu::DATATYPE::DT_IMAGE, uint16(EMutableFileFlags::HighRes)), 0, 0, 0});
+			OptionalBucket.Categories.Add({ FFileCategoryID(EStreamableDataType::Model, uint16(mu::EDataType::Image), uint16(EMutableFileFlags::HighRes)), 0, 0, 0});
 
 			const int32 NumRoms = Model->GetRomCount();
 			for (int32 RomIndex = 0; RomIndex < NumRoms; ++RomIndex)
 			{
 				const uint32 BlockSize = Model->GetRomSize(RomIndex);
-				const uint16 BlockResourceType = Model->IsMeshData(RomIndex) ? mu::DATATYPE::DT_MESH : mu::DATATYPE::DT_IMAGE;
+				const uint16 BlockResourceType = Model->IsMeshData(RomIndex) ? uint16(mu::EDataType::Mesh) : uint16(mu::EDataType::Image);
 				const EMutableFileFlags BlockFlags = Model->IsRomHighRes(RomIndex) ? EMutableFileFlags::HighRes : EMutableFileFlags::None;
 
-				FFileCategoryID CategoryID = { EDataType::Model, BlockResourceType, uint16(BlockFlags) };
+				FFileCategoryID CategoryID = { EStreamableDataType::Model, BlockResourceType, uint16(BlockFlags) };
 				TPair<FFileBucket&, FFileCategory&> It = FindOrAddCategory(FileBuckets, DefaultBucket, CategoryID); // Add block to an existing or new category
 				It.Key.DataSize += BlockSize;
 				It.Value.DataSize += BlockSize;
@@ -2825,7 +2825,7 @@ namespace MutablePrivate
 		{
 			// Add RealTimeMorphs to the Default bucket
 			FFileCategory& RealTimeMorphCategory = DefaultBucket.Categories.AddDefaulted_GetRef();
-			RealTimeMorphCategory.Id.DataType = EDataType::RealTimeMorph;
+			RealTimeMorphCategory.Id.DataType = EStreamableDataType::RealTimeMorph;
 
 			const TMap<uint32, FRealTimeMorphStreamable>& RealTimeMorphStreamables = ModelStreamableBulkData.RealTimeMorphStreamables;
 			for (const TPair<uint32, FRealTimeMorphStreamable>& MorphStreamable : RealTimeMorphStreamables)
@@ -2840,7 +2840,7 @@ namespace MutablePrivate
 		{
 			// Add Clothing to the Default bucket
 			FFileCategory& ClothingCategory = DefaultBucket.Categories.AddDefaulted_GetRef();
-			ClothingCategory.Id.DataType = EDataType::Clothing;
+			ClothingCategory.Id.DataType = EStreamableDataType::Clothing;
 			
 			const TMap<uint32, FClothingStreamable>& ClothingStreamables = ModelStreamableBulkData.ClothingStreamables;
 			for (const TPair<uint32, FClothingStreamable>& ClothStreamable : ClothingStreamables)
@@ -3036,11 +3036,11 @@ namespace MutablePrivate
 				// TODO: Temp. Remove after unifying generated output files code between editor an package. UE-222777 
 				if (bUseRomTypeAndFlagsToFilter)
 				{
-					BlockResourceType = Model->IsMeshData(RomIndex) ? mu::DATATYPE::DT_MESH : mu::DATATYPE::DT_IMAGE;
+					BlockResourceType = Model->IsMeshData(RomIndex) ? uint16(mu::EDataType::Mesh) : uint16(mu::EDataType::Image);
 					BlockFlags = Model->IsRomHighRes(RomIndex) ? EMutableFileFlags::HighRes : EMutableFileFlags::None;
 				}
 
-				FFileCategoryID CurrentCategory = { EDataType::Model, BlockResourceType, uint16(BlockFlags) };
+				FFileCategoryID CurrentCategory = { EStreamableDataType::Model, BlockResourceType, uint16(BlockFlags) };
 				FBlock CurrentBlock = 
 				{ 
 					.Id = RomIndex, 
@@ -3056,7 +3056,7 @@ namespace MutablePrivate
 		{
 			const TMap<uint32, FRealTimeMorphStreamable>& RealTimeMorphStreamables = ModelStreamableBulkData.RealTimeMorphStreamables;
 
-			const FFileCategoryID RealTimeMorphCategory = { EDataType::RealTimeMorph, (uint16)mu::DATATYPE::DT_NONE, (uint16)EMutableFileFlags::None };
+			const FFileCategoryID RealTimeMorphCategory = { EStreamableDataType::RealTimeMorph, uint16(mu::EDataType::None), uint16(EMutableFileFlags::None) };
 
 			for (const TPair<uint32, FRealTimeMorphStreamable>& MorphStreamable : RealTimeMorphStreamables)
 			{
@@ -3077,7 +3077,7 @@ namespace MutablePrivate
 		{
 			const TMap<uint32, FClothingStreamable>& ClothingStreamables = ModelStreamableBulkData.ClothingStreamables;
 
-			const FFileCategoryID ClothingCategory = { EDataType::Clothing, (uint16)mu::DATATYPE::DT_NONE, (uint16)EMutableFileFlags::None };
+			const FFileCategoryID ClothingCategory = { EStreamableDataType::Clothing, (uint16)mu::EDataType::None, (uint16)EMutableFileFlags::None };
 
 			for (const TPair<uint32, FClothingStreamable>& ClothStreamable : ClothingStreamables)
 			{
@@ -3111,7 +3111,7 @@ namespace MutablePrivate
 
 			switch (File.DataType)
 			{
-			case EDataType::Model:
+			case EStreamableDataType::Model:
 			{
 				for (FBlock& Block : File.Blocks)
 				{
@@ -3124,7 +3124,7 @@ namespace MutablePrivate
 				}
 				break;
 			}
-			case EDataType::RealTimeMorph:
+			case EStreamableDataType::RealTimeMorph:
 			{
 				for (FBlock& Block : File.Blocks)
 				{
@@ -3137,7 +3137,7 @@ namespace MutablePrivate
 				}
 				break;
 			}
-			case EDataType::Clothing:
+			case EStreamableDataType::Clothing:
 			{
 				for (FBlock& Block : File.Blocks)
 				{
@@ -3521,7 +3521,7 @@ void UCustomizableObjectPrivate::LoadCompiledDataFromDDC(FCompilationOptions Opt
 							// Filter and cache the data by DataType
 							switch (File->DataType)
 							{
-							case MutablePrivate::EDataType::Model:
+							case MutablePrivate::EStreamableDataType::Model:
 							{
 								for (const MutablePrivate::FBlock& Block : File->Blocks)
 								{
@@ -3529,7 +3529,7 @@ void UCustomizableObjectPrivate::LoadCompiledDataFromDDC(FCompilationOptions Opt
 								}
 								break;
 							}
-							case MutablePrivate::EDataType::RealTimeMorph:
+							case MutablePrivate::EStreamableDataType::RealTimeMorph:
 							{
 								for (const MutablePrivate::FBlock& Block : File->Blocks)
 								{
@@ -3537,7 +3537,7 @@ void UCustomizableObjectPrivate::LoadCompiledDataFromDDC(FCompilationOptions Opt
 								}
 								break;
 							}
-							case MutablePrivate::EDataType::Clothing:
+							case MutablePrivate::EStreamableDataType::Clothing:
 							{
 								for (const MutablePrivate::FBlock& Block : File->Blocks)
 								{
diff --git a/Engine/Plugins/Mutable/Source/CustomizableObject/Private/MuCO/UnrealMutableModelDiskStreamer.cpp b/Engine/Plugins/Mutable/Source/CustomizableObject/Private/MuCO/UnrealMutableModelDiskStreamer.cpp
--- a/Engine/Plugins/Mutable/Source/CustomizableObject/Private/MuCO/UnrealMutableModelDiskStreamer.cpp
+++ b/Engine/Plugins/Mutable/Source/CustomizableObject/Private/MuCO/UnrealMutableModelDiskStreamer.cpp
@@ -344,7 +344,7 @@ mu::FModelReader::FOperationID FUnrealMutableModelBulkReader::BeginReadBlock(con
 			FValueId::ByteArray ValueIdBytes = {};
 			check(sizeof(FValueId::ByteArray) >= 4 + sizeof(Block->FileId));
 
-			MutablePrivate::EDataType DataType = MutablePrivate::EDataType::Model;
+			MutablePrivate::EStreamableDataType DataType = MutablePrivate::EStreamableDataType::Model;
 			FMemory::Memcpy(&ValueIdBytes, &DataType, sizeof(DataType));
 			FMemory::Memcpy(&ValueIdBytes[4], &Block->FileId, sizeof(Block->FileId));
 			FValueId ResourceId = FValueId(ValueIdBytes);
diff --git a/Engine/Plugins/Mutable/Source/CustomizableObjectEditor/Private/MuCOE/CustomizableObjectCompileRunnable.cpp b/Engine/Plugins/Mutable/Source/CustomizableObjectEditor/Private/MuCOE/CustomizableObjectCompileRunnable.cpp
--- a/Engine/Plugins/Mutable/Source/CustomizableObjectEditor/Private/MuCOE/CustomizableObjectCompileRunnable.cpp
+++ b/Engine/Plugins/Mutable/Source/CustomizableObjectEditor/Private/MuCOE/CustomizableObjectCompileRunnable.cpp
@@ -651,17 +651,17 @@ void FCustomizableObjectSaveDDRunnable::StoreCachedPlatformDataToDisk(bool& bSto
 			{
 				switch (File.DataType)
 				{
-				case MutablePrivate::EDataType::Model:
+				case MutablePrivate::EStreamableDataType::Model:
 				{
 					StreamableMemoryWriter->Serialize(FileBulkData.GetData(), FileBulkData.Num() * sizeof(uint8));
 					break;
 				}
-				case MutablePrivate::EDataType::RealTimeMorph:
+				case MutablePrivate::EStreamableDataType::RealTimeMorph:
 				{
 					// TODO: Store clothing streamable to disk. UE-222777
 					break;
 				}
-				case MutablePrivate::EDataType::Clothing:
+				case MutablePrivate::EStreamableDataType::Clothing:
 				{
 					// TODO: Store clothing streamable to disk. UE-222777
 					break;
diff --git a/Engine/Plugins/Mutable/Source/CustomizableObjectEditor/Private/MuCOE/SMutableCodeViewer.cpp b/Engine/Plugins/Mutable/Source/CustomizableObjectEditor/Private/MuCOE/SMutableCodeViewer.cpp
--- a/Engine/Plugins/Mutable/Source/CustomizableObjectEditor/Private/MuCOE/SMutableCodeViewer.cpp
+++ b/Engine/Plugins/Mutable/Source/CustomizableObjectEditor/Private/MuCOE/SMutableCodeViewer.cpp
@@ -1785,18 +1785,18 @@ void SMutableCodeViewer::OnSelectionChanged(TSharedPtr<FMutableCodeTreeElement>
 
 	SelectedOperationAddress = SelectedNodes[0]->MutableOperation;
 	const mu::EOpType OperationType = MutableModel->GetPrivate()->Program.GetOpType(SelectedOperationAddress);
-	const mu::DATATYPE OperationDataType = mu::GetOpDataType(OperationType);
+	const mu::EDataType OperationDataType = mu::GetOpDataType(OperationType);
 	
 	switch (OperationDataType)
 	{
-	case mu::DT_LAYOUT:
+	case mu::EDataType::Layout:
 	{
 		// Create or reuse the UI
 		PrepareLayoutViewer();
 		break;
 	}
 
-	case mu::DT_IMAGE:
+	case mu::EDataType::Image:
 	{
 		// Create or reuse the UI
 		bSelectedOperationIsImage = true;
@@ -1804,21 +1804,21 @@ void SMutableCodeViewer::OnSelectionChanged(TSharedPtr<FMutableCodeTreeElement>
 		break;
 	}
 
-	case mu::DT_MESH:
+	case mu::EDataType::Mesh:
 	{
 		// Create or reuse the UI
 		PrepareMeshViewer();
 		break;
 	}
 
-	case mu::DT_INSTANCE:
+	case mu::EDataType::Instance:
 	{
 		// Create or reuse the UI
 		PrepareInstanceViewer();
 		break;
 	}
 
-	case mu::DT_SCALAR:
+	case mu::EDataType::Scalar:
 	{
 		// Create or reuse the UI
 		if (!PreviewScalarViewer)
@@ -1830,14 +1830,14 @@ void SMutableCodeViewer::OnSelectionChanged(TSharedPtr<FMutableCodeTreeElement>
 		break;
 	}
 
-	case mu::DT_STRING:
+	case mu::EDataType::String:
 	{
 		// Create or reuse the UI
 		PrepareStringViewer();
 		break;
 	}
 
-	case mu::DT_COLOUR:
+	case mu::EDataType::Color:
 	{
 		// Create or reuse the UI
 		if (!PreviewColorViewer)
@@ -1849,7 +1849,7 @@ void SMutableCodeViewer::OnSelectionChanged(TSharedPtr<FMutableCodeTreeElement>
 		break;
 	}
 
-	case mu::DT_INT:
+	case mu::EDataType::Int:
 	{
 		// Create or reuse the UI
 		if (!PreviewIntViewer)
@@ -1861,7 +1861,7 @@ void SMutableCodeViewer::OnSelectionChanged(TSharedPtr<FMutableCodeTreeElement>
 		break;
 	}
 
-	case mu::DT_BOOL:
+	case mu::EDataType::Bool:
 	{
 		// Create or reuse the UI
 		if (!PreviewBoolViewer)
@@ -1873,7 +1873,7 @@ void SMutableCodeViewer::OnSelectionChanged(TSharedPtr<FMutableCodeTreeElement>
 		break;
 	}
 
-	case mu::DT_PROJECTOR:
+	case mu::EDataType::Projector:
 	{
 		// Create or reuse the UI
 		PrepareProjectorViewer();
@@ -2004,7 +2004,7 @@ void SMutableCodeViewer::ClearHighlightedItems()
 
 void SMutableCodeViewer::TreeExpandElements(TArray<TSharedPtr<FMutableCodeTreeElement>>& InElementsToExpand, 
 	bool bForceExpandDuplicates /*= false*/,
-	mu::DATATYPE FilteringDataType /*= mu::DATATYPE::DT_NONE*/,
+	mu::EDataType FilteringDataType /*= mu::EDataType::None*/,
 	TSharedPtr<FProcessedOperationsBuffer> InExpandedOperationsBuffer /* = nullptr */)
 {
 	if (InElementsToExpand.IsEmpty())
@@ -2039,11 +2039,11 @@ void SMutableCodeViewer::TreeExpandElements(TArray<TSharedPtr<FMutableCodeTreeEl
 		// Identifier of the element. May be repeated if there are elements duplicating another element
 		const mu::OP::ADDRESS OperationAddress = Item->MutableOperation;
 		
-		// Filter the elements being expanded if the user has defined a desired DATATYPE
-		if (FilteringDataType != mu::DATATYPE::DT_NONE)
+		// Filter the elements being expanded if the user has defined a desired EDataType
+		if (FilteringDataType != mu::EDataType::None)
 		{
 			const mu::EOpType OperationType = Item->MutableModel->GetPrivate()->Program.GetOpType(OperationAddress);
-			const mu::DATATYPE OperationDataType = mu::GetOpDataType(OperationType);
+			const mu::EDataType OperationDataType = mu::GetOpDataType(OperationType);
 
 			// If it is not of the desired type then ignore it and continue to the next pending element
 			if (OperationDataType != FilteringDataType)
@@ -2172,7 +2172,7 @@ void SMutableCodeViewer::TreeExpandUnique()
 void SMutableCodeViewer::TreeExpandInstance()
 {
 	// Expand only the items that match the datatype provided
-	TreeExpandElements(RootNodes,false, mu::DT_INSTANCE);
+	TreeExpandElements(RootNodes,false, mu::EDataType::Instance);
 }
 
 #pragma endregion
@@ -2197,7 +2197,7 @@ void SMutableCodeViewer::CacheRootNodeAddresses()
 	RootNodeAddresses = MoveTemp(FoundRootNodeAddresses);
 }
 
- void SMutableCodeViewer::CacheAddressesRelatedWithConstantResource(const mu::DATATYPE ConstantDataType,
+ void SMutableCodeViewer::CacheAddressesRelatedWithConstantResource(const mu::EDataType ConstantDataType,
 	const int32 IndexOnConstantsArray)
 {
 	check(MutableModel);
@@ -2241,7 +2241,7 @@ void SMutableCodeViewer::CacheRootNodeAddresses()
 
 
 void SMutableCodeViewer::GetOperationsReferencingConstantResource(
-	const mu::DATATYPE ConstantDataType,
+	const mu::EDataType ConstantDataType,
 	const int32 IndexOnConstantsArray,
 	FElementsSearchCache& InSearchPayload,
 	const mu::FProgram& InProgram)
@@ -2280,14 +2280,14 @@ void SMutableCodeViewer::GetOperationsReferencingConstantResource(
 }
 
 bool SMutableCodeViewer::IsConstantResourceUsedByOperation(const int32 IndexOnConstantsArray,
-	const mu::DATATYPE ConstantDataType, const mu::OP::ADDRESS OperationAddress, const mu::FProgram& InProgram) const
+	const mu::EDataType ConstantDataType, const mu::OP::ADDRESS OperationAddress, const mu::FProgram& InProgram) const
 {
 	// Cache the current operation type to know where to look and what to check
 	const mu::EOpType OperationType = InProgram.GetOpType(OperationAddress);
 
 	// Making usage of the operation data type is not valid since some operations while return one type do, in fact,
 	// contain data from other types (like the mesh constant for example that contains mesh, skeleton and physics asset)
-	// const mu::DATATYPE DataType = mu::GetOpDataType(OperationType);
+	// const mu::EDataType DataType = mu::GetOpDataType(OperationType);
 	// if (DataType != ConstantDataType)
 	// {
 	// 	return false;
@@ -2299,7 +2299,7 @@ bool SMutableCodeViewer::IsConstantResourceUsedByOperation(const int32 IndexOnCo
 	// Check if the operation data type is compatible with the type of resources we are providing
 	switch (ConstantDataType)
 	{
-		case mu::DATATYPE::DT_STRING:
+		case mu::EDataType::String:
 			{
 				// TIP: To know if they represent a constant value check the code on code runner to see if they read from the constants array
 				if (OperationType == mu::EOpType::ST_CONSTANT)
@@ -2363,7 +2363,7 @@ bool SMutableCodeViewer::IsConstantResourceUsedByOperation(const int32 IndexOnCo
 				break;
 			}
 
-		case mu::DATATYPE::DT_IMAGE:
+		case mu::EDataType::Image:
 			{
 				if (OperationType == mu::EOpType::IM_CONSTANT)
 				{
@@ -2373,7 +2373,7 @@ bool SMutableCodeViewer::IsConstantResourceUsedByOperation(const int32 IndexOnCo
 				break;
 			}
 
-		case mu::DATATYPE::DT_MESH:
+		case mu::EDataType::Mesh:
 			{
 				if (OperationType == mu::EOpType::ME_CONSTANT)
 				{
@@ -2382,7 +2382,7 @@ bool SMutableCodeViewer::IsConstantResourceUsedByOperation(const int32 IndexOnCo
 				break;
 			}
 
-		case mu::DATATYPE::DT_LAYOUT:
+		case mu::EDataType::Layout:
 			{
 				if (OperationType == mu::EOpType::LA_CONSTANT)
 				{
@@ -2391,7 +2391,7 @@ bool SMutableCodeViewer::IsConstantResourceUsedByOperation(const int32 IndexOnCo
 				break;
 			}
 
-		case mu::DATATYPE::DT_PROJECTOR:
+		case mu::EDataType::Projector:
 			{
 				if (OperationType == mu::EOpType::PR_CONSTANT)
 				{
@@ -2400,7 +2400,7 @@ bool SMutableCodeViewer::IsConstantResourceUsedByOperation(const int32 IndexOnCo
 				break;
 			}
 
-		case mu::DATATYPE::DT_MATRIX:
+		case mu::EDataType::Matrix:
 			{
 				if (OperationType == mu::EOpType::ME_TRANSFORM)
 				{
@@ -2413,7 +2413,7 @@ bool SMutableCodeViewer::IsConstantResourceUsedByOperation(const int32 IndexOnCo
 				break;
 			}
 
-		case mu::DATATYPE::DT_SHAPE:
+		case mu::EDataType::Shape:
 			{
 				if (OperationType == mu::EOpType::ME_CLIPMORPHPLANE)
 				{
@@ -2435,7 +2435,7 @@ bool SMutableCodeViewer::IsConstantResourceUsedByOperation(const int32 IndexOnCo
 				break;
 			}
 		
-		case mu::DATATYPE::DT_CURVE:
+		case mu::EDataType::Curve:
 			{
 				if (OperationType == mu::EOpType::SC_CURVE)
 				{
@@ -2444,7 +2444,7 @@ bool SMutableCodeViewer::IsConstantResourceUsedByOperation(const int32 IndexOnCo
 				break;
 			}
 
-		case mu::DATATYPE::DT_SKELETON:
+		case mu::EDataType::Skeleton:
 			{
 				if (OperationType == mu::EOpType::ME_CONSTANT)
 				{
@@ -2454,7 +2454,7 @@ bool SMutableCodeViewer::IsConstantResourceUsedByOperation(const int32 IndexOnCo
 			}
 
 		// Invalid types
-		case mu::DATATYPE::DT_NONE:
+		case mu::EDataType::None:
 		default:
 			{
 				checkNoEntry();
@@ -2643,7 +2643,7 @@ void SMutableCodeViewer::Tick(const FGeometry& AllottedGeometry, const double In
 	bIsPreviewPendingUpdate = false;
 
 	const mu::EOpType OperationType = MutableModel->GetPrivate()->Program.GetOpType(SelectedOperationAddress);
-	const mu::DATATYPE OperationDataType = mu::GetOpDataType(OperationType);
+	const mu::EDataType OperationDataType = mu::GetOpDataType(OperationType);
 
 	mu::FSettings Settings;
 	const TSharedPtr<mu::FSystem> System = MakeShared<mu::FSystem>(Settings);
@@ -2657,7 +2657,7 @@ void SMutableCodeViewer::Tick(const FGeometry& AllottedGeometry, const double In
 
 	switch (OperationDataType)
 	{
-	case mu::DT_LAYOUT:
+	case mu::EDataType::Layout:
 	{
 		check(PreviewLayoutViewer);
 		TSharedPtr<const mu::FLayout> MutableLayout = System->GetPrivate()->BuildLayout(MutableModel, PreviewParameters.Get(), SelectedOperationAddress);
@@ -2665,7 +2665,7 @@ void SMutableCodeViewer::Tick(const FGeometry& AllottedGeometry, const double In
 		break;
 	}
 
-	case mu::DT_IMAGE:
+	case mu::EDataType::Image:
 	{
 		check(PreviewImageViewer);
 		TSharedPtr<const mu::FImage> MutableImage = System->GetPrivate()->BuildImage(MutableModel, PreviewParameters.Get(), SelectedOperationAddress, MipsToSkip, 0);
@@ -2673,7 +2673,7 @@ void SMutableCodeViewer::Tick(const FGeometry& AllottedGeometry, const double In
 		break;
 	}
 
-	case mu::DT_MESH:
+	case mu::EDataType::Mesh:
 	{
 		check(PreviewMeshViewer);
 		TSharedPtr<const mu::FMesh> MutableMesh = System->GetPrivate()->BuildMesh(MutableModel, PreviewParameters.Get(), SelectedOperationAddress);
@@ -2681,7 +2681,7 @@ void SMutableCodeViewer::Tick(const FGeometry& AllottedGeometry, const double In
 		break;
 	}
 
-	case mu::DT_INSTANCE:
+	case mu::EDataType::Instance:
 	{
 		check(PreviewInstanceViewer);
 		TSharedPtr<const mu::FInstance> MutableInstance = System->GetPrivate()->BuildInstance(MutableModel, PreviewParameters.Get(), SelectedOperationAddress);
@@ -2689,7 +2689,7 @@ void SMutableCodeViewer::Tick(const FGeometry& AllottedGeometry, const double In
 		break;
 	}
 
-	case mu::DT_BOOL:
+	case mu::EDataType::Bool:
 	{
 		check(PreviewBoolViewer);
 		const bool MutableBool = System->GetPrivate()->BuildBool(MutableModel, PreviewParameters.Get(), SelectedOperationAddress);
@@ -2697,7 +2697,7 @@ void SMutableCodeViewer::Tick(const FGeometry& AllottedGeometry, const double In
 		break;
 	}
 
-	case mu::DT_INT:
+	case mu::EDataType::Int:
 	{
 		check(PreviewIntViewer);
 		const int32 MutableInt = System->GetPrivate()->BuildInt(MutableModel, PreviewParameters.Get(), SelectedOperationAddress);
@@ -2705,7 +2705,7 @@ void SMutableCodeViewer::Tick(const FGeometry& AllottedGeometry, const double In
 		break;
 	}
 
-	case mu::DT_SCALAR:
+	case mu::EDataType::Scalar:
 	{
 		check(PreviewScalarViewer);
 		const float MutableScalar = System->GetPrivate()->BuildScalar(MutableModel, PreviewParameters.Get(), SelectedOperationAddress);
@@ -2713,7 +2713,7 @@ void SMutableCodeViewer::Tick(const FGeometry& AllottedGeometry, const double In
 		break;
 	}
 
-	case mu::DT_STRING:
+	case mu::EDataType::String:
 	{
 		check(PreviewStringViewer);
 		TSharedPtr<const mu::String> MutableString = System->GetPrivate()->BuildString(MutableModel, PreviewParameters.Get(), SelectedOperationAddress);
@@ -2722,7 +2722,7 @@ void SMutableCodeViewer::Tick(const FGeometry& AllottedGeometry, const double In
 		break;
 	}
 
-	case mu::DT_COLOUR:
+	case mu::EDataType::Color:
 	{
 		check(PreviewColorViewer);
 		FVector4f Color = System->GetPrivate()->BuildColour(MutableModel, PreviewParameters.Get(), SelectedOperationAddress);
@@ -2730,7 +2730,7 @@ void SMutableCodeViewer::Tick(const FGeometry& AllottedGeometry, const double In
 		break;
 	}
 
-	case mu::DT_PROJECTOR:
+	case mu::EDataType::Projector:
 	{
 		check (PreviewProjectorViewer);
 		mu::FProjector Projector = System->GetPrivate()->BuildProjector(MutableModel, PreviewParameters.Get(), SelectedOperationAddress);
diff --git a/Engine/Plugins/Mutable/Source/CustomizableObjectEditor/Private/MuCOE/SMutableCodeViewer.h b/Engine/Plugins/Mutable/Source/CustomizableObjectEditor/Private/MuCOE/SMutableCodeViewer.h
--- a/Engine/Plugins/Mutable/Source/CustomizableObjectEditor/Private/MuCOE/SMutableCodeViewer.h
+++ b/Engine/Plugins/Mutable/Source/CustomizableObjectEditor/Private/MuCOE/SMutableCodeViewer.h
@@ -539,11 +539,11 @@ public:
 	 * @note The output of this operation will be cached onto NavigationOPAddresses so we can navigate over them.
 	 * @param ConstantDataType The constant data type to locate operations referencing it: It is required for the
 	 * system to know what operations should be checked for the usage of the provided constant resource. For example, a
-	 * type of mu::DATATYPE::DT_Curve will try to locate operations that we know operate with constant curves.
+	 * type of mu::EDataType::DT_Curve will try to locate operations that we know operate with constant curves.
 	 * @param IndexOnConstantsArray The index of the constant on its constants array. If it is, for example, a constant mesh
 	 * this value represents the index of the constant inside constantMeshes array on mu::program.
 	 */
-	void CacheAddressesRelatedWithConstantResource(const mu::DATATYPE ConstantDataType, const int32 IndexOnConstantsArray);
+	void CacheAddressesRelatedWithConstantResource(const mu::EDataType ConstantDataType, const int32 IndexOnConstantsArray);
 
 private:
 	
@@ -556,7 +556,7 @@ private:
 	 * resource
 	 * @param InProgram The mutable program containing all the operations of the graph.
 	 */
-	void GetOperationsReferencingConstantResource(const mu::DATATYPE ConstantDataType,
+	void GetOperationsReferencingConstantResource(const mu::EDataType ConstantDataType,
 	                                              const int32 IndexOnConstantsArray,
 	                                              FElementsSearchCache& InSearchPayload,
 	                                              const mu::FProgram& InProgram);
@@ -574,7 +574,7 @@ private:
 	 * @return True if the operation at OperationAddress does make use of the provided constant resource. False otherwise
 	 */
 	bool IsConstantResourceUsedByOperation(const int32 IndexOnConstantsArray,
-	                                       const mu::DATATYPE ConstantDataType, const mu::OP::ADDRESS OperationAddress,
+	                                       const mu::EDataType ConstantDataType, const mu::OP::ADDRESS OperationAddress,
 	                                       const mu::FProgram& InProgram) const;
 	
 
@@ -724,14 +724,14 @@ private:
 	*/
 	void TreeExpandElements(TArray<TSharedPtr<FMutableCodeTreeElement>>& InElementsToExpand,
 		bool bForceExpandDuplicates = false,
-		mu::DATATYPE FilteringDataType = mu::DATATYPE::DT_NONE,
+		mu::EDataType FilteringDataType = mu::EDataType::None,
 		TSharedPtr<FProcessedOperationsBuffer> InExpandedOperationsBuffer = nullptr);
 	
 
 	/** Expand only the unique elements (mot duplicates) */
 	void TreeExpandUnique();
 
-	/** Expand only the elements using as operation type DT_INSTANCE */
+	/** Expand only the elements using as operation type Instance */
 	void TreeExpandInstance();
 	
 	/** Grabs the selected element and expands all the elements inside the selected branch. Duplicates are ignored */
diff --git a/Engine/Plugins/Mutable/Source/CustomizableObjectEditor/Private/MuCOE/SMutableConstantsWidget.cpp b/Engine/Plugins/Mutable/Source/CustomizableObjectEditor/Private/MuCOE/SMutableConstantsWidget.cpp
--- a/Engine/Plugins/Mutable/Source/CustomizableObjectEditor/Private/MuCOE/SMutableConstantsWidget.cpp
+++ b/Engine/Plugins/Mutable/Source/CustomizableObjectEditor/Private/MuCOE/SMutableConstantsWidget.cpp
@@ -1407,41 +1407,41 @@ void SMutableConstantsWidget::LoadConstantCurves()
 #pragma region Previewer invocation methods
 
 
-void SMutableConstantsWidget::ClearSelectedConstantItems(mu::DATATYPE ExceptionDataType /* = mu::DATATYPE::DT_NONE */) const
+void SMutableConstantsWidget::ClearSelectedConstantItems(mu::EDataType ExceptionDataType /* = mu::EDataType::None */) const
 {
-	if (ExceptionDataType != mu::DATATYPE::DT_MESH)
+	if (ExceptionDataType != mu::EDataType::Mesh)
 	{
 		MeshListViewHandler->ClearSelection();
 	}
-	if (ExceptionDataType != mu::DATATYPE::DT_STRING)
+	if (ExceptionDataType != mu::EDataType::String)
 	{
 		ConstantStringsSlate->ClearSelection();
 	}
-	if (ExceptionDataType != mu::DATATYPE::DT_LAYOUT)
+	if (ExceptionDataType != mu::EDataType::Layout)
 	{
 		ConstantLayoutsSlate->ClearSelection();
 	}
-	if (ExceptionDataType != mu::DATATYPE::DT_PROJECTOR)
+	if (ExceptionDataType != mu::EDataType::Projector)
 	{
 		ConstantProjectorsSlate->ClearSelection();
 	}
-	if (ExceptionDataType != mu::DATATYPE::DT_MATRIX)
+	if (ExceptionDataType != mu::EDataType::Matrix)
 	{
 		ConstantMatricesSlate->ClearSelection();
 	}
-	if (ExceptionDataType != mu::DATATYPE::DT_SHAPE)
+	if (ExceptionDataType != mu::EDataType::Shape)
 	{
 		ConstantShapesSlate->ClearSelection();
 	}
-	if (ExceptionDataType != mu::DATATYPE::DT_CURVE)
+	if (ExceptionDataType != mu::EDataType::Curve)
 	{
 		ConstantCurvesSlate->ClearSelection();
 	}
-	if (ExceptionDataType != mu::DATATYPE::DT_SKELETON)
+	if (ExceptionDataType != mu::EDataType::Skeleton)
 	{
 		ConstantSkeletonsSlate->ClearSelection();
 	}
-	if (ExceptionDataType != mu::DATATYPE::DT_IMAGE)
+	if (ExceptionDataType != mu::EDataType::Image)
 	{
 		ImageListViewHandler->ClearSelection();
 		//ConstantImagesListView->ClearSelection();
@@ -1456,7 +1456,7 @@ void SMutableConstantsWidget::OnSelectedStringChanged(
 {
 	if (MutableConstantStringElement)
 	{
-		static constexpr mu::DATATYPE SlateDataType = mu::DATATYPE::DT_STRING;
+		static constexpr mu::EDataType SlateDataType = mu::EDataType::String;
 
 		// Clear the selected CodeViewer row and all other constant viewer slates that are not the type provided.
 		MutableCodeViewerPtr->ClearSelectedTreeRow();
@@ -1474,7 +1474,7 @@ void SMutableConstantsWidget::OnSelectedImageChanged(
 {
 	if (MutableConstantImageElement)
 	{
-		static constexpr mu::DATATYPE SlateDataType = mu::DATATYPE::DT_IMAGE;
+		static constexpr mu::EDataType SlateDataType = mu::EDataType::Image;
 
 		MutableCodeViewerPtr->ClearSelectedTreeRow();
 		ClearSelectedConstantItems(SlateDataType);
@@ -1490,7 +1490,7 @@ void SMutableConstantsWidget::OnSelectedMeshChanged(
 {
 	if (MutableConstantMeshElement)
 	{
-		static constexpr mu::DATATYPE SlateDataType = mu::DATATYPE::DT_MESH;
+		static constexpr mu::EDataType SlateDataType = mu::EDataType::Mesh;
 		
 		MutableCodeViewerPtr->ClearSelectedTreeRow();
 		ClearSelectedConstantItems(SlateDataType);
@@ -1506,7 +1506,7 @@ void SMutableConstantsWidget::OnSelectedLayoutChanged(
 {
 	if (MutableConstantLayoutElement)
 	{
-		static constexpr mu::DATATYPE SlateDataType = mu::DATATYPE::DT_LAYOUT;
+		static constexpr mu::EDataType SlateDataType = mu::EDataType::Layout;
 		
 		MutableCodeViewerPtr->ClearSelectedTreeRow();
 		ClearSelectedConstantItems(SlateDataType);
@@ -1522,7 +1522,7 @@ void SMutableConstantsWidget::OnSelectedProjectorChanged(
 {
 	if (MutableConstantProjectorElement)
 	{
-		static constexpr mu::DATATYPE SlateDataType = mu::DATATYPE::DT_PROJECTOR;
+		static constexpr mu::EDataType SlateDataType = mu::EDataType::Projector;
 
 		MutableCodeViewerPtr->ClearSelectedTreeRow();
 		ClearSelectedConstantItems(SlateDataType);
@@ -1538,7 +1538,7 @@ void SMutableConstantsWidget::OnSelectedMatrixChanged(
 {
 	if (MutableConstantMatrixElement)
 	{
-		static constexpr mu::DATATYPE SlateDataType = mu::DATATYPE::DT_MATRIX;
+		static constexpr mu::EDataType SlateDataType = mu::EDataType::Matrix;
 
 		MutableCodeViewerPtr->ClearSelectedTreeRow();
 		ClearSelectedConstantItems(SlateDataType);
@@ -1554,7 +1554,7 @@ void SMutableConstantsWidget::OnSelectedShapeChanged(
 {
 	if (MutableConstantShapeElement)
 	{
-		static constexpr mu::DATATYPE SlateDataType = mu::DATATYPE::DT_SHAPE;
+		static constexpr mu::EDataType SlateDataType = mu::EDataType::Shape;
 
 		MutableCodeViewerPtr->ClearSelectedTreeRow();
 		ClearSelectedConstantItems(SlateDataType);
@@ -1570,7 +1570,7 @@ void SMutableConstantsWidget::OnSelectedCurveChanged(
 {
 	if (MutableConstantCurveElement)
 	{
-		static constexpr mu::DATATYPE SlateDataType = mu::DATATYPE::DT_CURVE;
+		static constexpr mu::EDataType SlateDataType = mu::EDataType::Curve;
 
 		MutableCodeViewerPtr->ClearSelectedTreeRow();
 		ClearSelectedConstantItems(SlateDataType);
@@ -1586,7 +1586,7 @@ void SMutableConstantsWidget::OnSelectedSkeletonChanged(
 {
 	if (MutableConstantSkeletonElement)
 	{
-		static constexpr mu::DATATYPE SlateDataType = mu::DATATYPE::DT_SKELETON;
+		static constexpr mu::EDataType SlateDataType = mu::EDataType::Skeleton;
 
 		MutableCodeViewerPtr->ClearSelectedTreeRow();
 		ClearSelectedConstantItems(SlateDataType);
diff --git a/Engine/Plugins/Mutable/Source/CustomizableObjectEditor/Private/MuCOE/SMutableConstantsWidget.h b/Engine/Plugins/Mutable/Source/CustomizableObjectEditor/Private/MuCOE/SMutableConstantsWidget.h
--- a/Engine/Plugins/Mutable/Source/CustomizableObjectEditor/Private/MuCOE/SMutableConstantsWidget.h
+++ b/Engine/Plugins/Mutable/Source/CustomizableObjectEditor/Private/MuCOE/SMutableConstantsWidget.h
@@ -122,10 +122,10 @@ public:
 
 	/** It clears the selected element for all constant view slates (SListView and STileView) except the one that shows the data of the provided type
 	 * @param ExceptionDataType The type of data of the slate container that we want to skip when clearing the selected elements.
-	 * EX : We want to clear all selected elements but not the ones from the constant images. Provide as argument "mu::DATATYPE::DT_IMAGE"
+	 * EX : We want to clear all selected elements but not the ones from the constant images. Provide as argument "mu::EDataType::Image"
 	 * to clear all but the selected item/s of the Image type.
 	 */
-	void ClearSelectedConstantItems(mu::DATATYPE ExceptionDataType = mu::DATATYPE::DT_NONE) const;
+	void ClearSelectedConstantItems(mu::EDataType ExceptionDataType = mu::EDataType::None) const;
 
 private:
 
diff --git a/Engine/Plugins/Mutable/Source/MutableRuntime/Internal/MuR/CodeVisitor.h b/Engine/Plugins/Mutable/Source/MutableRuntime/Internal/MuR/CodeVisitor.h
--- a/Engine/Plugins/Mutable/Source/MutableRuntime/Internal/MuR/CodeVisitor.h
+++ b/Engine/Plugins/Mutable/Source/MutableRuntime/Internal/MuR/CodeVisitor.h
@@ -16,9 +16,7 @@
 namespace mu
 {
 
-    //---------------------------------------------------------------------------------------------
-    //! Decide what operations are an "add resource" since they are handled differently sometimes.
-    //---------------------------------------------------------------------------------------------
+    /** Decide what operations are an "add resource" since they are handled differently sometimes. */
     inline bool VisitorIsAddResource(const EOpType& type)
     {
         return  type == EOpType::IN_ADDIMAGE ||
@@ -26,24 +24,23 @@ namespace mu
     }
 
 
-    //---------------------------------------------------------------------------------------------
-    //! Code visitor that:
-    //! - is top-down
-    //! - cannot change the visited instructions.
-    //! - will not visit twice the same instruction with the same state.
-    //! - Its iterative
-    //---------------------------------------------------------------------------------------------
+    /** Code visitor that:
+    * - is top-down
+    * - cannot change the visited instructions.
+    * - will not visit twice the same instruction with the same state.
+    * - Its iterative
+	*/
     template<typename STATE=int32>
     class UniqueConstCodeVisitorIterative
     {
     public:
 
-        UniqueConstCodeVisitorIterative( bool skipResources=false )
+        UniqueConstCodeVisitorIterative( bool bInSkipResources=false )
         {
             // Default state
             States.Add(STATE());
             CurrentState = 0;
-            bSkipResources = skipResources;
+            bSkipResources = bInSkipResources;
         }
 
         //! Ensure virtual destruction
@@ -52,9 +49,9 @@ namespace mu
     protected:
 
         //!
-        void SetDefaultState(const STATE& s)
+        void SetDefaultState(const STATE& State)
         {
-            States[0] = s;
+            States[0] = State;
         }
 
         //!
@@ -71,78 +68,77 @@ namespace mu
         }
 
         //! For manual recursion that changes the state for a specific path.
-        void RecurseWithState(OP::ADDRESS at, const STATE& newState)
+        void RecurseWithState(OP::ADDRESS Address, const STATE& NewState)
         {
-			int32 it = States.Find(newState);
-            if (it==INDEX_NONE)
+			int32 StateIndex = States.Find(NewState);
+            if (StateIndex==INDEX_NONE)
             {
-                States.Add(newState);
+				StateIndex = States.Add(NewState);
             }
-            int32 StateIndex = States.IndexOfByKey(newState);
-
-            Pending.Add( FPending(at,StateIndex) );
+            Pending.Emplace( Address, StateIndex );
         }
 
         //! For manual recursion that doesn't change the state for a specific path.
-        void RecurseWithCurrentState(OP::ADDRESS at)
+        void RecurseWithCurrentState(OP::ADDRESS Address)
         {
-            Pending.Add( FPending(at,CurrentState) );
+            Pending.Emplace( Address, CurrentState );
         }
 
         //! Can be called from visit to set the state to visit all children ops
-        void SetCurrentState(const STATE& newState)
+        void SetCurrentState(const STATE& NewState)
         {
-            int32 it = States.Find(newState);
-            if (it==INDEX_NONE)
+			CurrentState = States.Find(NewState);
+            if (CurrentState ==INDEX_NONE)
             {
-                States.Add(newState);
+				CurrentState = States.Add(NewState);
             }
-            CurrentState = States.Find(newState);
         }
 
 
-        void Traverse( OP::ADDRESS root, FProgram& program )
+        void Traverse( OP::ADDRESS Root, FProgram& Program )
         {
-            Pending.Reserve( program.OpAddress.Num() );
+            Pending.Reserve( Program.OpAddress.Num() );
 
             // Visit the given root
-            Pending.Add( FPending(root,0) );
-            Recurse( program );
+            Pending.Emplace(Root, 0);
+            Recurse( Program );
          }
 
-        void FullTraverse( FProgram& program )
+        void FullTraverse( FProgram& Program )
         {
             // Visit all the state roots
-            for ( int32 p=0; p<program.States.Num(); ++p )
+            for ( int32 StateIndex=0; StateIndex<Program.States.Num(); ++StateIndex)
             {
-                Pending.Add(FPending(program.States[p].Root,0) );
-                Recurse( program );
+                Pending.Add(FPending(Program.States[StateIndex].Root,0) );
+                Recurse( Program );
             }
         }
 
 
     private:
 
-        //! Do the actual work by overriding this in the derived classes.
-        //! Return true if the traverse has to continue with the children of "at"
-        virtual bool Visit( OP::ADDRESS at, FProgram& program ) = 0;
+        /** Do the actual work by overriding this in the derived classes.
+         * Return true if the traverse has to continue with the children of the given address.
+		 */
+        virtual bool Visit( OP::ADDRESS, FProgram& ) = 0;
 
         //! Operations to be processed
         struct FPending
         {
-            FPending()
-            {
-                at = 0;
-                StateIndex = 0;
-            }
-            FPending(OP::ADDRESS _at, int32 _stateIndex)
-            {
-                at = _at;
-                StateIndex = _stateIndex;
-            }
+			FPending()
+			{
+				Address = 0;
+				StateIndex = 0;
+			}
 
-            OP::ADDRESS at;
-            int32 StateIndex;
+			FPending(OP::ADDRESS InAddress, int32 InStateIndex)
+			{
+				Address = InAddress;
+				StateIndex = InStateIndex;
+			}
+			
+			OP::ADDRESS Address=0;
+            int32 StateIndex=0;
         };
 		TArray<FPending> Pending;
 
@@ -158,38 +154,38 @@ namespace mu
 
         //! Array of states visited for each operation.
         //! Empty array means operation not visited at all.
-		TArray<TArray<int32>> m_visited;
+		TArray<TArray<int32>> Visited;
 
         //! Process all the pending operations and visit all children if necessary
-        void Recurse( FProgram& program )
+        void Recurse( FProgram& Program )
         {
-			m_visited.Empty();
-			m_visited.SetNum(program.OpAddress.Num());
+			Visited.Empty();
+			Visited.SetNum(Program.OpAddress.Num());
 
             while ( Pending.Num() )
             {
-                OP::ADDRESS at = Pending.Last().at;
+                OP::ADDRESS Address = Pending.Last().Address;
                 CurrentState = Pending.Last().StateIndex;
                 Pending.Pop();
 
                 bool bRecurse = false;
 
-                bool visitedInThisState = m_visited[at].Contains(CurrentState);
-                if (!visitedInThisState)
+                bool bVisitedInThisState = Visited[Address].Contains(CurrentState);
+                if (!bVisitedInThisState)
                 {
-                    m_visited[at].Add(CurrentState);
+                    Visited[Address].Add(CurrentState);
 
                     // Visit may change current state
-                    bRecurse = Visit( at, program );
+                    bRecurse = Visit(Address, Program );
                 }
 
                 if (bRecurse)
                 {
-                    ForEachReference( program, at, [&](OP::ADDRESS ref)
+                    ForEachReference( Program, Address, [&](OP::ADDRESS Ref)
                     {
-                        if (ref)
+                        if (Ref)
                         {
-                            Pending.Add( FPending(ref,CurrentState) );
+                            Pending.Emplace(Ref, CurrentState);
                         }
                     });
                 }
@@ -199,185 +195,6 @@ namespace mu
     };
 
 
-    //---------------------------------------------------------------------------------------------
-    //! Code visitor that:
-    //! - is top-down
-    //! - cannot change the instructions.
-    //! - will repeat visits to the instructions that have multiple references.
-    //! - Its iterative
-    //---------------------------------------------------------------------------------------------
-    template<class STATE=int32>
-    class RepeatConstCodeVisitorIterative
-    {
-    public:
-
-        RepeatConstCodeVisitorIterative( bool skipResources=false )
-        {
-            // Default state
-            States.Add(STATE());
-            CurrentState = 0;
-            bSkipResources = skipResources;
-        }
-
-        //! Ensure virtual destruction
-        virtual ~RepeatConstCodeVisitorIterative() {}
-
-    protected:
-
-        //!
-        void SetDefaultState(const STATE& s)
-        {
-            States[0] = s;
-        }
-
-        //!
-        const STATE& GetDefaultState() const
-        {
-            return States[0];
-        }
-
-        //! Use this from visit to access the state at the time of processing the current
-        //! instruction.
-        const STATE& GetCurrentState() const
-        {
-            return States[CurrentState];
-        }
-
-        //! For manual recursion that changes the state for a specific path.
-        void RecurseWithState(OP::ADDRESS At, const STATE& NewState)
-        {
-            int32 StateIndex = States.Find(NewState);
-            if (StateIndex ==INDEX_NONE)
-            {
-				StateIndex = States.Add(NewState);
-            }
-
-            //check(at<1000000);
-            Pending.Add( FPending(At,StateIndex) );
-        }
-
-        //! For manual recursion that doesn't change the state for a specific path.
-        void RecurseWithCurrentState(OP::ADDRESS at)
-        {
-            //check(at<1000000);
-            Pending.Add( FPending(at,CurrentState) );
-        }
-
-        //! Can be called from visit to set the state to visit all children ops
-        void SetCurrentState(const STATE& NewState)
-        {
-			int32 StateIndex = States.Find(NewState);
-			if (StateIndex == INDEX_NONE)
-			{
-				StateIndex = States.Add(NewState);
-			}
-
-			CurrentState = StateIndex;
-        }
-
-
-        void Traverse( OP::ADDRESS Root, FProgram& Program )
-        {
-            Pending.reserve(Program.OpAddress.Num() );
-
-            // Visit the given root
-            Pending.Add( FPending(Root,0) );
-            Recurse(Program);
-        }
-
-        void FullTraverse( FProgram& Program)
-        {
-            // Visit all the state roots
-            for ( int32 p=0; p<Program.States.Num(); ++p )
-            {
-                Pending.Add(FPending(Program.States[p].Root,0) );
-                Recurse(Program);
-            }
-        }
-
-
-    private:
-
-        //! Do the actual work by overriding this in the derived classes.
-        //! Return true if the traverse has to continue with the children of "at"
-        virtual bool Visit( OP::ADDRESS, FProgram& ) = 0;
-
-        //! Operations to be processed
-        struct FPending
-        {
-            FPending()
-            {
-                at = 0;
-                StateIndex = 0;
-            }
-            FPending(OP::ADDRESS InAt, int32 _stateIndex)
-            {
-                at = InAt;
-                StateIndex = _stateIndex;
-            }
-
-            OP::ADDRESS at;
-            int32 StateIndex;
-        };
-        TArray<FPending> Pending;
-
-        //! States found so far
-		TArray<STATE> States;
-
-        //! Index of the current state, from the States array.
-        int32 CurrentState;
-
-        //! If true, operations adding resources (meshes or images) will only
-        //! have the base operation recursed, but not the resources.
-        bool bSkipResources;
-
-
-        //! Process all the pending operations and visit all children if necessary
-        void Recurse( FProgram& program )
-        {
-            while ( Pending.Num() )
-            {
-                OP::ADDRESS at = Pending.back().at;
-                CurrentState = Pending.back().StateIndex;
-                Pending.pop_back();
-
-                bool bRecurse = false;
-
-                // Visit may change current state
-                bRecurse = Visit( at, program );
-
-                if (bRecurse)
-                {
-                    if (bSkipResources && VisitorIsAddResource(program.GetOpType(at)))
-                    {
-						OP::InstanceAddArgs args = program.GetOpArgs<OP::InstanceAddArgs>(at);
-
-                        // Recurse only the base
-                        OP::ADDRESS base = args.instance;
-                        if (base)
-                        {
-							check(base<program.OpAddress.Num());
-                            Pending.Add( FPending(base,CurrentState) );
-                        }
-                    }
-                    else
-                    {
-                        ForEachReference( program, at, [&](OP::ADDRESS ref)
-                        {
-                            if (ref)
-                            {
-								check(ref<program.OpAddress.Num());
-                                Pending.Add( FPending(ref,CurrentState) );
-                            }
-                        });
-                    }
-                }
-            }
-        }
-
-    };
-
-
     //---------------------------------------------------------------------------------------------
     //! Code visitor template for visitors that:
     //! - only traverses the operations that are relevant for a given set of parameter values. It
@@ -429,13 +246,13 @@ namespace mu
 
     protected:
 
-        virtual bool Visit( OP::ADDRESS at, FProgram& program )
+        virtual bool Visit( OP::ADDRESS Address, FProgram& Program )
         {
             bool bRecurse = true;
 
-            EOpType type = program.GetOpType(at);
+            EOpType Type = Program.GetOpType(Address);
 
-            switch ( type )
+            switch ( Type )
             {
             case EOpType::NU_CONDITIONAL:
             case EOpType::SC_CONDITIONAL:
@@ -446,27 +263,27 @@ namespace mu
             case EOpType::IN_CONDITIONAL:
 			case EOpType::ED_CONDITIONAL:
             {
-				OP::ConditionalArgs args = program.GetOpArgs<OP::ConditionalArgs>(at);
+				OP::ConditionalArgs Args = Program.GetOpArgs<OP::ConditionalArgs>(Address);
 
                 bRecurse = false;
 
-                PARENT::RecurseWithCurrentState( args.condition );
+                PARENT::RecurseWithCurrentState( Args.condition );
 
                 // If there is no expression, we'll assume true.
-                bool value = true;
+                bool bValue = true;
 
-                if (args.condition)
+                if (Args.condition)
                 {
-                    value = System->BuildBool(Model, Params, args.condition);
+					bValue = System->BuildBool(Model, Params, Args.condition);
                 }
 
-                if (value)
+                if (bValue)
                 {
-                    PARENT::RecurseWithCurrentState( args.yes );
+                    PARENT::RecurseWithCurrentState( Args.yes );
                 }
                 else
                 {
-                    PARENT::RecurseWithCurrentState( args.no );
+                    PARENT::RecurseWithCurrentState( Args.no );
                 }
                 break;
             }
@@ -482,36 +299,36 @@ namespace mu
             {
                 bRecurse = false;
 
-				const uint8* data = program.GetOpArgsPointer(at);
+				const uint8* Data = Program.GetOpArgsPointer(Address);
 				
 				OP::ADDRESS VarAddress;
-				FMemory::Memcpy( &VarAddress, data, sizeof(OP::ADDRESS));
-				data += sizeof(OP::ADDRESS);
+				FMemory::Memcpy( &VarAddress, Data, sizeof(OP::ADDRESS));
+				Data += sizeof(OP::ADDRESS);
 
                 if (VarAddress)
                 {
 					OP::ADDRESS DefAddress;
-					FMemory::Memcpy( &DefAddress, data, sizeof(OP::ADDRESS));
-					data += sizeof(OP::ADDRESS);
+					FMemory::Memcpy( &DefAddress, Data, sizeof(OP::ADDRESS));
+					Data += sizeof(OP::ADDRESS);
 
 					uint32 CaseCount;
-					FMemory::Memcpy( &CaseCount, data, sizeof(uint32));
-					data += sizeof(uint32);
+					FMemory::Memcpy( &CaseCount, Data, sizeof(uint32));
+					Data += sizeof(uint32);
 
                     PARENT::RecurseWithCurrentState( VarAddress );
 
                     int32 var = System->BuildInt( Model, Params, VarAddress );
 
 					OP::ADDRESS valueAt = DefAddress;
-					for (uint32 C = 0; C < CaseCount; ++C)
+					for (uint32 CaseIndex = 0; CaseIndex < CaseCount; ++CaseIndex)
 					{
 						int32 Condition;
-						FMemory::Memcpy( &Condition, data, sizeof(int32));		
-						data += sizeof(int32);
+						FMemory::Memcpy( &Condition, Data, sizeof(int32));		
+						Data += sizeof(int32);
 
 						OP::ADDRESS At;
-						FMemory::Memcpy( &At, data, sizeof(OP::ADDRESS));
-						data += sizeof(OP::ADDRESS);
+						FMemory::Memcpy( &At, Data, sizeof(OP::ADDRESS));
+						Data += sizeof(OP::ADDRESS);
 
 						if (At && var == (int32)Condition)
 						{
@@ -531,7 +348,7 @@ namespace mu
             {
                 bRecurse = false;
 
-				const uint8* Data = program.GetOpArgsPointer(at);
+				const uint8* Data = Program.GetOpArgsPointer(Address);
 
 				uint8 LODCount;
 				FMemory::Memcpy(&LODCount, Data, sizeof(uint8));
@@ -559,19 +376,19 @@ namespace mu
 
             case EOpType::IN_ADDMESH:
             {
-				OP::InstanceAddArgs args = program.GetOpArgs<OP::InstanceAddArgs>(at);
+				OP::InstanceAddArgs Args = Program.GetOpArgs<OP::InstanceAddArgs>(Address);
 
                 bRecurse = false;
 
-                PARENT::RecurseWithCurrentState(args.instance);
+                PARENT::RecurseWithCurrentState(Args.instance);
 
-                STATE newState = PARENT::GetCurrentState();
-                newState.m_underResourceCount=1;
+                STATE NewState = PARENT::GetCurrentState();
+                NewState.m_underResourceCount=1;
 
-                OP::ADDRESS meshAt = args.value;
-                if (meshAt)
+                OP::ADDRESS MeshAddress = Args.value;
+                if (MeshAddress)
                 {
-                    PARENT::RecurseWithState(meshAt, newState);
+                    PARENT::RecurseWithState(MeshAddress, NewState);
                 }
                 break;
             }
@@ -579,19 +396,19 @@ namespace mu
 
             case EOpType::IN_ADDIMAGE:
             {
-				OP::InstanceAddArgs args = program.GetOpArgs<OP::InstanceAddArgs>(at);
+				OP::InstanceAddArgs Args = Program.GetOpArgs<OP::InstanceAddArgs>(Address);
 
                 bRecurse = false;
 
-                PARENT::RecurseWithCurrentState(args.instance);
+                PARENT::RecurseWithCurrentState(Args.instance);
 
-                STATE newState = PARENT::GetCurrentState();
-                newState.m_underResourceCount=1;
+                STATE NewState = PARENT::GetCurrentState();
+				NewState.m_underResourceCount=1;
 
-                OP::ADDRESS imageAt = args.value;
-                if (imageAt)
+                OP::ADDRESS ImageAddress = Args.value;
+                if (ImageAddress)
                 {
-                    PARENT::RecurseWithState(imageAt, newState);
+                    PARENT::RecurseWithState(ImageAddress, NewState);
                 }
                 break;
             }
@@ -608,20 +425,19 @@ namespace mu
         FSystem::Private* System = nullptr;
 		TSharedPtr<const FModel> Model;
         const FParameters* Params = nullptr;
-        unsigned LODMask = 0;
+        uint32 LODMask = 0;
     };
 
 
-    //---------------------------------------------------------------------------------------------
-    //! Code visitor that:
-    //! - only traverses the operations that are relevant for a given set of parameter values. It
-    //! only considers the discrete parameters like integers and booleans. In the case of forks
-    //! caused by continuous parameters like float weights for interpolation, all the branches are
-    //! traversed.
-    //! - cannot change the instructions
-    //! - will not repeat visits to instructions with the same state
-    //! - the state has to be a compatible with COVERED_CODE_VISITOR_STATE
-    //---------------------------------------------------------------------------------------------
+    /** Code visitor that :
+    * - only traverses the operations that are relevant for a given set of parameter values. It
+    * only considers the discrete parameters like integers and booleans. In the case of forks
+    * caused by continuous parameters like float weights for interpolation, all the branches are
+    * traversed.
+    * - cannot change the instructions
+    * - will not repeat visits to instructions with the same state
+    * - the state has to be a compatible with COVERED_CODE_VISITOR_STATE
+    */
     template<typename COVERED_STATE=COVERED_CODE_VISITOR_STATE>
     class UniqueDiscreteCoveredCodeVisitor :
             public DiscreteCoveredCodeVisitorBase
@@ -649,10 +465,9 @@ namespace mu
 
 
 
-    //---------------------------------------------------------------------------------------------
-    //! Calculate all the parameters found under a particular operation
-    //! It has an internal cache, so don't reuse objects of this class if the program changes.
-    //---------------------------------------------------------------------------------------------
+    /** Calculate all the parameters found under a particular operation
+    * It has an internal cache, so don't reuse objects of this class if the Program changes.
+	*/
     class MUTABLERUNTIME_API SubtreeParametersVisitor
     {
     public:
diff --git a/Engine/Plugins/Mutable/Source/MutableRuntime/Internal/MuR/Operations.h b/Engine/Plugins/Mutable/Source/MutableRuntime/Internal/MuR/Operations.h
--- a/Engine/Plugins/Mutable/Source/MutableRuntime/Internal/MuR/Operations.h
+++ b/Engine/Plugins/Mutable/Source/MutableRuntime/Internal/MuR/Operations.h
@@ -1174,55 +1174,44 @@ namespace mu
 
 
 	//! Types of data handled by the Mutable runtime.
-	typedef enum
+	enum class EDataType : uint8
 	{
-		DT_NONE,
-		DT_BOOL,
-		DT_INT,
-		DT_SCALAR,
-		DT_COLOUR,
-		DT_IMAGE,
-		DT_VOLUME_DEPRECATED,
-		DT_LAYOUT,
-		DT_MESH,
-		DT_INSTANCE,
-		DT_PROJECTOR,
-		DT_STRING,
-		DT_EXTENSION_DATA,
-		DT_MATRIX,
+		None,
+		Bool,
+		Int,
+		Scalar,
+		Color,
+		Image,
+		Layout,
+		Mesh,
+		Instance,
+		Projector,
+		String,
+		ExtensionData,
+		Matrix,
 
 		// Supporting data types : Never returned as an actual data type for any operation.
-		DT_SHAPE,
-		DT_CURVE,
-		DT_SKELETON,
-		DT_PHYSICS_ASSET,
+		Shape,
+		Curve,
+		Skeleton,
+		PhysicsAsset,
 		
-		DT_COUNT
-	} DATATYPE;
-
-	static_assert(DT_COUNT < 256, "Unexpected enum size.");
+		Count
+	};
 
 
-	// Generic data about a Mutable runtime operation.
+	// Generic data about a Mutable operation that is needed at runtime.
     struct FOpDesc
     {
 		//! Type of data generated by the instruction
-		DATATYPE type;
-
-        //! True if the instruction is worth caching when generating models
-        bool cached;
-
-        //! For image instructions, for every image format, true if it is supported as the base
-        //! format of the operation.
-        //! TODO: Move to tools library?
-		bool supportedBasePixelFormats[uint8(EImageFormat::Count)];
+		EDataType type;
     };
 
 	MUTABLERUNTIME_API extern const FOpDesc& GetOpDesc( EOpType type );
 
 
     //!
-    inline static DATATYPE GetOpDataType( EOpType type )
+    inline static EDataType GetOpDataType( EOpType type )
     {
 		return GetOpDesc(type).type;
     }
@@ -1234,18 +1223,18 @@ namespace mu
 	MUTABLERUNTIME_API extern void ForEachReference( const struct FProgram& program, OP::ADDRESS at, const TFunctionRef<void(OP::ADDRESS)> );
 
 	//!
-	MUTABLERUNTIME_API inline EOpType GetSwitchForType( DATATYPE d )
+	MUTABLERUNTIME_API inline EOpType GetSwitchForType( EDataType d )
     {
         switch (d)
         {
-        case DT_INSTANCE: return EOpType::IN_SWITCH;
-        case DT_MESH: return EOpType::ME_SWITCH;
-		case DT_IMAGE: return EOpType::IM_SWITCH;
-		case DT_LAYOUT: return EOpType::LA_SWITCH;
-        case DT_COLOUR: return EOpType::CO_SWITCH;
-        case DT_SCALAR: return EOpType::SC_SWITCH;
-        case DT_INT: return EOpType::NU_SWITCH;
-		case DT_EXTENSION_DATA: return EOpType::ED_SWITCH;
+		case EDataType::Instance: return EOpType::IN_SWITCH;
+        case EDataType::Mesh: return EOpType::ME_SWITCH;
+		case EDataType::Image: return EOpType::IM_SWITCH;
+		case EDataType::Layout: return EOpType::LA_SWITCH;
+        case EDataType::Color: return EOpType::CO_SWITCH;
+        case EDataType::Scalar: return EOpType::SC_SWITCH;
+        case EDataType::Int: return EOpType::NU_SWITCH;
+		case EDataType::ExtensionData: return EOpType::ED_SWITCH;
 
         default:
 			check(false);
diff --git a/Engine/Plugins/Mutable/Source/MutableRuntime/Internal/MuR/SystemPrivate.h b/Engine/Plugins/Mutable/Source/MutableRuntime/Internal/MuR/SystemPrivate.h
--- a/Engine/Plugins/Mutable/Source/MutableRuntime/Internal/MuR/SystemPrivate.h
+++ b/Engine/Plugins/Mutable/Source/MutableRuntime/Internal/MuR/SystemPrivate.h
@@ -469,7 +469,7 @@ namespace mu
 			uint8 IsCacheLocked : 1;
 
 			/** The operation DATA_TYPE. */
-			uint8 DataType;
+			EDataType DataType;
 
 			/** The position in the type-specific array where the result data is stored. 
 			* 0 means index not valid.
@@ -581,22 +581,22 @@ namespace mu
 
 			if (DataTypeIndex)
 			{
-				check(Data.DataType < DATATYPE::DT_COUNT);
-				switch ((DATATYPE)Data.DataType)
+				check(uint8(Data.DataType) < uint8(EDataType::Count));
+				switch (Data.DataType)
 				{
-				case DATATYPE::DT_IMAGE:					
+				case EDataType::Image:					
 					ImageResults[DataTypeIndex].Value = nullptr;
 					break;
 
-				case DATATYPE::DT_MESH:
+				case EDataType::Mesh:
 					MeshResults[DataTypeIndex].Value = nullptr;
 					break;
 
-				case DATATYPE::DT_INSTANCE:
+				case EDataType::Instance:
 					InstanceResults[DataTypeIndex] = nullptr;
 					break;
 
-				case DATATYPE::DT_EXTENSION_DATA:
+				case EDataType::ExtensionData:
 					ExtensionDataResults[DataTypeIndex] = nullptr;
 					break;
 
@@ -691,7 +691,7 @@ namespace mu
 			const FOpExecutionData* Data = OpExecutionData.get_ptr(at);
 			if (!Data) return false;
 
-			check(Data->DataType==DATATYPE::DT_BOOL);
+			check(Data->DataType==EDataType::Bool);
 			return Data->DataTypeIndex!=0;
 		}
 
@@ -701,7 +701,7 @@ namespace mu
 			const FOpExecutionData* Data = OpExecutionData.get_ptr(at);
 			if (!Data) return 0.0f;
 
-			check(Data->DataType == DATATYPE::DT_SCALAR);
+			check(Data->DataType == EDataType::Scalar);
 			return Data->ScalarResult;
 		}
 
@@ -711,7 +711,7 @@ namespace mu
 			const FOpExecutionData* Data = OpExecutionData.get_ptr(at);
 			if (!Data) return 0;
 
-			check(Data->DataType == DATATYPE::DT_INT);
+			check(Data->DataType == EDataType::Int);
 			return Data->DataTypeIndex;
 		}
 
@@ -721,7 +721,7 @@ namespace mu
 			const FOpExecutionData* Data = OpExecutionData.get_ptr(at);
 			if (!Data) return FVector4f();
 
-			check(Data->DataType == DATATYPE::DT_COLOUR);
+			check(Data->DataType == EDataType::Color);
 			return ColorResults[Data->DataTypeIndex];
 		}
 
@@ -731,7 +731,7 @@ namespace mu
 			const FOpExecutionData* Data = OpExecutionData.get_ptr(at);
 			if (!Data) return FMatrix44f::Identity;
 
-			check(Data->DataType == DATATYPE::DT_MATRIX);
+			check(Data->DataType == EDataType::Matrix);
 			return MatrixResults[Data->DataTypeIndex];
 		}
 
@@ -741,7 +741,7 @@ namespace mu
 			const FOpExecutionData* Data = OpExecutionData.get_ptr(at);
 			if (!Data) return FProjector();
 
-			check(Data->DataType == DATATYPE::DT_PROJECTOR);
+			check(Data->DataType == EDataType::Projector);
 			return ProjectorResults[Data->DataTypeIndex];
 		}
 
@@ -751,7 +751,7 @@ namespace mu
 			FOpExecutionData* Data = OpExecutionData.get_ptr(at);
 			if (!Data) return nullptr;
 
-			check(Data->DataType == DATATYPE::DT_INSTANCE);
+			check(Data->DataType == EDataType::Instance);
 			TSharedPtr<const FInstance> Result = InstanceResults[Data->DataTypeIndex];
 
 			// We need to decrease the hit-count even if the result is null.
@@ -776,7 +776,7 @@ namespace mu
 			FOpExecutionData* Data = OpExecutionData.get_ptr(at);
 			if (!Data) return nullptr;
 
-			check(Data->DataType == DATATYPE::DT_IMAGE);
+			check(Data->DataType == EDataType::Image);
 			TSharedPtr<const FImage> Result = ImageResults[Data->DataTypeIndex].Value;
 
 			// We need to decrease the hit-count even if the result is null.
@@ -802,7 +802,7 @@ namespace mu
 			FOpExecutionData* Data = OpExecutionData.get_ptr(at);
 			if (!Data) return nullptr;
 
-			check(Data->DataType == DATATYPE::DT_MESH);
+			check(Data->DataType == EDataType::Mesh);
 			TSharedPtr<const FMesh> Result = MeshResults[Data->DataTypeIndex].Value;
 
 			// We need to decrease the hit-count even if the result is null.
@@ -825,7 +825,7 @@ namespace mu
 			FOpExecutionData* Data = OpExecutionData.get_ptr(at);
 			if (!Data) return nullptr;
 
-			check(Data->DataType == DATATYPE::DT_LAYOUT);
+			check(Data->DataType == EDataType::Layout);
 			return LayoutResults[Data->DataTypeIndex];
 		}
 
@@ -836,7 +836,7 @@ namespace mu
 			FOpExecutionData* Data = OpExecutionData.get_ptr(at);
 			if (!Data) return nullptr;
 
-			check(Data->DataType == DATATYPE::DT_STRING);
+			check(Data->DataType == EDataType::String);
 			return StringResults[Data->DataTypeIndex];
 		}
 
@@ -847,7 +847,7 @@ namespace mu
 			FOpExecutionData* Data = OpExecutionData.get_ptr(at);
 			if (!Data) return nullptr;
 
-			check(Data->DataType == DATATYPE::DT_EXTENSION_DATA);
+			check(Data->DataType == EDataType::ExtensionData);
 			return ExtensionDataResults[Data->DataTypeIndex];
 		}
 
@@ -863,8 +863,8 @@ namespace mu
 		{
 			check(at.At < OpExecutionData.size_code());
 			FOpExecutionData* Data = OpExecutionData.get_ptr(at);
-			check(Data->DataType == DATATYPE::DT_BOOL || Data->DataType == DATATYPE::DT_NONE);
-			Data->DataType = DATATYPE::DT_BOOL;
+			check(Data->DataType == EDataType::Bool || Data->DataType == EDataType::None);
+			Data->DataType = EDataType::Bool;
 			Data->DataTypeIndex = v;
 			Data->IsValueValid = true;
 		}
@@ -873,8 +873,8 @@ namespace mu
 		{
 			check(at.At < OpExecutionData.size_code());
 			FOpExecutionData* Data = OpExecutionData.get_ptr(at);
-			check(Data->DataType == DATATYPE::DT_INT || Data->DataType == DATATYPE::DT_NONE);
-			Data->DataType = DATATYPE::DT_INT;
+			check(Data->DataType == EDataType::Int || Data->DataType == EDataType::None);
+			Data->DataType = EDataType::Int;
 			Data->DataTypeIndex = v;
 			Data->IsValueValid = true;
 		}
@@ -883,8 +883,8 @@ namespace mu
 		{
 			check(at.At < OpExecutionData.size_code());
 			FOpExecutionData* Data = OpExecutionData.get_ptr(at);
-			check(Data->DataType == DATATYPE::DT_SCALAR || Data->DataType == DATATYPE::DT_NONE);
-			Data->DataType = DATATYPE::DT_SCALAR;
+			check(Data->DataType == EDataType::Scalar || Data->DataType == EDataType::None);
+			Data->DataType = EDataType::Scalar;
 			Data->ScalarResult = v;
 			Data->IsValueValid = true;
 		}
@@ -893,8 +893,8 @@ namespace mu
 		{
 			check(at.At < OpExecutionData.size_code());
 			FOpExecutionData* Data = OpExecutionData.get_ptr(at);
-			check(Data->DataType == DATATYPE::DT_COLOUR || Data->DataType == DATATYPE::DT_NONE);
-			Data->DataType = DATATYPE::DT_COLOUR;
+			check(Data->DataType == EDataType::Color || Data->DataType == EDataType::None);
+			Data->DataType = EDataType::Color;
 			Data->IsValueValid = true;
 
 			if (!Data->DataTypeIndex)
@@ -913,8 +913,8 @@ namespace mu
 		{
 			check(at.At < OpExecutionData.size_code());
 			FOpExecutionData* Data = OpExecutionData.get_ptr(at);
-			check(Data->DataType == DATATYPE::DT_MATRIX || Data->DataType == DATATYPE::DT_NONE);
-			Data->DataType = DATATYPE::DT_MATRIX;
+			check(Data->DataType == EDataType::Matrix || Data->DataType == EDataType::None);
+			Data->DataType = EDataType::Matrix;
 			Data->IsValueValid = true;
 
 			if (!Data->DataTypeIndex)
@@ -933,8 +933,8 @@ namespace mu
 		{
 			check(at.At < OpExecutionData.size_code());
 			FOpExecutionData* Data = OpExecutionData.get_ptr(at);
-			check(Data->DataType == DATATYPE::DT_PROJECTOR || Data->DataType == DATATYPE::DT_NONE);
-			Data->DataType = DATATYPE::DT_PROJECTOR;
+			check(Data->DataType == EDataType::Projector || Data->DataType == EDataType::None);
+			Data->DataType = EDataType::Projector;
 			Data->IsValueValid = true;
 
 			if (!Data->DataTypeIndex)
@@ -953,8 +953,8 @@ namespace mu
 		{
 			check(at.At < OpExecutionData.size_code());
 			FOpExecutionData* Data = OpExecutionData.get_ptr(at);
-			check(Data->DataType == DATATYPE::DT_INSTANCE || Data->DataType == DATATYPE::DT_NONE);
-			Data->DataType = DATATYPE::DT_INSTANCE;
+			check(Data->DataType == EDataType::Instance || Data->DataType == EDataType::None);
+			Data->DataType = EDataType::Instance;
 			Data->IsValueValid = true;
 
 			if (!Data->DataTypeIndex)
@@ -973,8 +973,8 @@ namespace mu
 		{
 			check(at.At < OpExecutionData.size_code());
 			FOpExecutionData* Data = OpExecutionData.get_ptr(at);
-			check(Data->DataType == DATATYPE::DT_EXTENSION_DATA || Data->DataType == DATATYPE::DT_NONE);
-			Data->DataType = DATATYPE::DT_EXTENSION_DATA;
+			check(Data->DataType == EDataType::ExtensionData || Data->DataType == EDataType::None);
+			Data->DataType = EDataType::ExtensionData;
 			Data->IsValueValid = true;
 
 			if (!Data->DataTypeIndex)
@@ -993,8 +993,8 @@ namespace mu
 		{
 			check(At.At < OpExecutionData.size_code());
 			FOpExecutionData* Data = OpExecutionData.get_ptr(At);
-			check(Data->DataType == DATATYPE::DT_IMAGE || Data->DataType == DATATYPE::DT_NONE);
-			Data->DataType = DATATYPE::DT_IMAGE;
+			check(Data->DataType == EDataType::Image || Data->DataType == EDataType::None);
+			Data->DataType = EDataType::Image;
 			Data->IsValueValid = true;
 
 			if (!Data->DataTypeIndex)
@@ -1017,8 +1017,8 @@ namespace mu
 
 			FOpExecutionData* Data = OpExecutionData.get_ptr(At);
 
-			check(Data->DataType == DATATYPE::DT_MESH || Data->DataType == DATATYPE::DT_NONE);
-			Data->DataType = DATATYPE::DT_MESH;
+			check(Data->DataType == EDataType::Mesh || Data->DataType == EDataType::None);
+			Data->DataType = EDataType::Mesh;
 			Data->IsValueValid = true;
 
 			if (!Data->DataTypeIndex)
@@ -1039,8 +1039,8 @@ namespace mu
 		{
 			check(at.At < OpExecutionData.size_code());
 			FOpExecutionData* Data = OpExecutionData.get_ptr(at);
-			check(Data->DataType == DATATYPE::DT_LAYOUT || Data->DataType == DATATYPE::DT_NONE);
-			Data->DataType = DATATYPE::DT_LAYOUT;
+			check(Data->DataType == EDataType::Layout || Data->DataType == EDataType::None);
+			Data->DataType = EDataType::Layout;
 			Data->IsValueValid = true;
 
 			if (!Data->DataTypeIndex)
@@ -1061,8 +1061,8 @@ namespace mu
 		{
 			check(at.At < OpExecutionData.size_code());
 			FOpExecutionData* Data = OpExecutionData.get_ptr(at);
-			check(Data->DataType == DATATYPE::DT_STRING || Data->DataType == DATATYPE::DT_NONE);
-			Data->DataType = DATATYPE::DT_STRING;
+			check(Data->DataType == EDataType::String || Data->DataType == EDataType::None);
+			Data->DataType = EDataType::String;
 			Data->IsValueValid = true;
 
 			if (!Data->DataTypeIndex)
@@ -1678,7 +1678,7 @@ namespace mu
 					const FResource* Value = nullptr;
 					switch ((*it).DataType)
 					{
-					case DATATYPE::DT_IMAGE:
+					case EDataType::Image:
 						Value = Instance.Cache->ImageResults[(*it).DataTypeIndex].Value.Get();
 						if (Value)
 						{
@@ -1686,7 +1686,7 @@ namespace mu
 						}
 						break;
 
-					case DATATYPE::DT_MESH:
+					case EDataType::Mesh:
 						Value = Instance.Cache->MeshResults[(*it).DataTypeIndex].Value.Get();
 						if (Value)
 						{
@@ -1734,7 +1734,7 @@ namespace mu
 
 				switch (Data.DataType)
 				{
-				case DATATYPE::DT_IMAGE:
+				case EDataType::Image:
 				{
 					TSharedPtr<const FImage> Value = CurrentInstanceCache->ImageResults[Data.DataTypeIndex].Value;
 					if (Value)
@@ -1745,7 +1745,7 @@ namespace mu
 					break;
 				}
 
-				case DATATYPE::DT_MESH:
+				case EDataType::Mesh:
 				{
 					TSharedPtr<const FMesh> Value = CurrentInstanceCache->MeshResults[Data.DataTypeIndex].Value;
 					if (Value)
@@ -1756,14 +1756,14 @@ namespace mu
 					break;
 				}
 
-				case DATATYPE::DT_LAYOUT:
+				case EDataType::Layout:
 				{
 					TSharedPtr<const FLayout> Value = CurrentInstanceCache->LayoutResults[Data.DataTypeIndex];
 					CurrentInstanceCache->LayoutResults[Data.DataTypeIndex] = nullptr;
 					break;
 				}
 
-				case DATATYPE::DT_INSTANCE:
+				case EDataType::Instance:
 				{
 					TSharedPtr<const FInstance> Value = CurrentInstanceCache->InstanceResults[Data.DataTypeIndex];
 					CurrentInstanceCache->InstanceResults[Data.DataTypeIndex] = nullptr;
@@ -1798,7 +1798,7 @@ namespace mu
 
 				switch (Data.DataType)
 				{
-				case DATATYPE::DT_IMAGE:
+				case EDataType::Image:
 				{
 					TSharedPtr<const FImage> Value = CurrentInstanceCache->ImageResults[Data.DataTypeIndex].Value;
 					CacheResources.Remove(Value);
@@ -1806,7 +1806,7 @@ namespace mu
 					break;
 				}
 
-				case DATATYPE::DT_MESH:
+				case EDataType::Mesh:
 				{
 					TSharedPtr<const FMesh> Value = CurrentInstanceCache->MeshResults[Data.DataTypeIndex].Value;
 					CacheResources.Remove(Value);
@@ -1814,14 +1814,14 @@ namespace mu
 					break;
 				}
 
-				case DATATYPE::DT_LAYOUT:
+				case EDataType::Layout:
 				{
 					TSharedPtr<const FLayout> Value = CurrentInstanceCache->LayoutResults[Data.DataTypeIndex];
 					CurrentInstanceCache->LayoutResults[Data.DataTypeIndex] = nullptr;
 					break;
 				}
 
-				case DATATYPE::DT_INSTANCE:
+				case EDataType::Instance:
 				{
 					TSharedPtr<const FInstance> Value = CurrentInstanceCache->InstanceResults[Data.DataTypeIndex];
 					CurrentInstanceCache->InstanceResults[Data.DataTypeIndex] = nullptr;
