[mutable] Remove unused static constant tables, and move some to the tools library.

diff --git a/Engine/Plugins/Mutable/Source/MutableRuntime/Private/MuR/CodeRunner.cpp b/Engine/Plugins/Mutable/Source/MutableRuntime/Private/MuR/CodeRunner.cpp
--- a/Engine/Plugins/Mutable/Source/MutableRuntime/Private/MuR/CodeRunner.cpp
+++ b/Engine/Plugins/Mutable/Source/MutableRuntime/Private/MuR/CodeRunner.cpp
@@ -246,25 +246,25 @@ namespace mu
 
 	
     //---------------------------------------------------------------------------------------------
-    void CodeRunner::RunCode_Conditional( const FScheduledOp& item, const FModel* pModel )
+    void CodeRunner::RunCode_Conditional( const FScheduledOp& Item, const FModel* InModel )
     {
         MUTABLE_CPUPROFILER_SCOPE(RunCode_Conditional);
 
 		const FProgram& Program = Model->GetPrivate()->Program;
 
-		EOpType type = Program.GetOpType(item.At);
-		OP::ConditionalArgs args = Program.GetOpArgs<OP::ConditionalArgs>(item.At);
+		EOpType type = Program.GetOpType(Item.At);
+		OP::ConditionalArgs Args = Program.GetOpArgs<OP::ConditionalArgs>(Item.At);
 
         // Conditionals have the following execution stages:
         // 0: we need to run the condition
         // 1: we need to run the branch
         // 2: we need to fetch the result and store it in this op
-        switch( item.Stage )
+        switch( Item.Stage )
         {
         case 0:
         {
-            AddOp( FScheduledOp( item.At,item,1 ),
-                   FScheduledOp( args.condition, item ) );
+            AddOp( FScheduledOp( Item.At,Item,1 ),
+                   FScheduledOp( Args.condition, Item ) );
             break;
         }
 
@@ -274,37 +274,37 @@ namespace mu
 
             // If there is no expression, we'll assume true.
             bool value = true;
-            value = LoadBool( FCacheAddress(args.condition, item.ExecutionIndex, item.ExecutionOptions) );
+            value = LoadBool( FCacheAddress(Args.condition, Item.ExecutionIndex, Item.ExecutionOptions) );
 
-            OP::ADDRESS resultAt = value ? args.yes : args.no;
+            OP::ADDRESS resultAt = value ? Args.yes : Args.no;
 
             // Schedule the end of this instruction if necessary
-            AddOp( FScheduledOp( item.At, item, 2, (uint32)value),
-				FScheduledOp( resultAt, item) );
+            AddOp( FScheduledOp( Item.At, Item, 2, (uint32)value),
+				FScheduledOp( resultAt, Item) );
 
             break;
         }
 
         case 2:
         {
-            OP::ADDRESS resultAt = item.CustomState ? args.yes : args.no;
+            OP::ADDRESS resultAt = Item.CustomState ? Args.yes : Args.no;
 
             // Store the final result
-            FCacheAddress cat( item );
-            FCacheAddress rat( resultAt, item );
+            FCacheAddress cat( Item );
+            FCacheAddress rat( resultAt, Item );
             switch (GetOpDataType(type))
             {
-            case DT_BOOL:       StoreBool( cat, LoadBool(rat) ); break;
-            case DT_INT:        StoreInt( cat, LoadInt(rat) ); break;
-            case DT_SCALAR:     StoreScalar( cat, LoadScalar(rat) ); break;
-			case DT_STRING:		StoreString( cat, LoadString( rat ) ); break;
-            case DT_COLOUR:		StoreColor( cat, LoadColor( rat ) ); break;
-            case DT_PROJECTOR:  StoreProjector( cat, LoadProjector(rat) ); break;
-            case DT_MESH:       StoreMesh( cat, LoadMesh(rat) ); break;
-            case DT_IMAGE:      StoreImage( cat, LoadImage(rat) ); break;
-            case DT_LAYOUT:     StoreLayout( cat, LoadLayout(rat) ); break;
-            case DT_INSTANCE:   StoreInstance( cat, LoadInstance(rat) ); break;
-			case DT_EXTENSION_DATA: StoreExtensionData( cat, LoadExtensionData(rat) ); break;
+            case EDataType::Bool:			StoreBool( cat, LoadBool(rat) ); break;
+            case EDataType::Int:			StoreInt( cat, LoadInt(rat) ); break;
+            case EDataType::Scalar:			StoreScalar( cat, LoadScalar(rat) ); break;
+			case EDataType::String:			StoreString( cat, LoadString( rat ) ); break;
+            case EDataType::Color:			StoreColor( cat, LoadColor( rat ) ); break;
+            case EDataType::Projector:		StoreProjector( cat, LoadProjector(rat) ); break;
+            case EDataType::Mesh:			StoreMesh( cat, LoadMesh(rat) ); break;
+            case EDataType::Image:			StoreImage( cat, LoadImage(rat) ); break;
+            case EDataType::Layout:			StoreLayout( cat, LoadLayout(rat) ); break;
+            case EDataType::Instance:		StoreInstance( cat, LoadInstance(rat) ); break;
+			case EDataType::ExtensionData:	StoreExtensionData( cat, LoadExtensionData(rat) ); break;
             default:
                 // Not implemented
                 check( false );
@@ -320,13 +320,13 @@ namespace mu
 
 
 	//---------------------------------------------------------------------------------------------
-	void CodeRunner::RunCode_Switch(const FScheduledOp& item, const FModel* pModel )
+	void CodeRunner::RunCode_Switch(const FScheduledOp& Item, const FModel* InModel )
 	{
 		const FProgram& Program = Model->GetPrivate()->Program;
 
-		EOpType type = Program.GetOpType(item.At);
+		EOpType type = Program.GetOpType(Item.At);
 
-		const uint8* data = Program.GetOpArgsPointer(item.At);
+		const uint8* data = Program.GetOpArgsPointer(Item.At);
 
 		OP::ADDRESS VarAddress;
 		FMemory::Memcpy(&VarAddress, data, sizeof(OP::ADDRESS));
@@ -340,29 +340,29 @@ namespace mu
 		FMemory::Memcpy(&CaseCount, data, sizeof(uint32));
 		data += sizeof(uint32);
 
-		switch (item.Stage)
+		switch (Item.Stage)
 		{
 		case 0:
 		{
 			if (VarAddress)
 			{
-				AddOp(FScheduledOp(item.At, item, 1), FScheduledOp(VarAddress, item));
+				AddOp(FScheduledOp(Item.At, Item, 1), FScheduledOp(VarAddress, Item));
 			}
 			else
 			{
 				switch (GetOpDataType(type))
 				{
-				case DT_BOOL:       StoreBool(item, false); break;
-				case DT_INT:        StoreInt(item, 0); break;
-				case DT_SCALAR:		StoreScalar(item, 0.0f); break;
-				case DT_STRING:		StoreString(item, nullptr); break;
-				case DT_COLOUR:		StoreColor(item, FVector4f()); break;
-				case DT_PROJECTOR:  StoreProjector(item, FProjector()); break;
-				case DT_MESH:       StoreMesh(item, nullptr); break;
-				case DT_IMAGE:      StoreImage(item, nullptr); break;
-				case DT_LAYOUT:     StoreLayout(item, nullptr); break;
-				case DT_INSTANCE:   StoreInstance(item, nullptr); break;
-				case DT_EXTENSION_DATA: StoreExtensionData(item, MakeShared<FExtensionData>()); break;
+				case EDataType::Bool:			StoreBool(Item, false); break;
+				case EDataType::Int:			StoreInt(Item, 0); break;
+				case EDataType::Scalar:			StoreScalar(Item, 0.0f); break;
+				case EDataType::String:			StoreString(Item, nullptr); break;
+				case EDataType::Color:			StoreColor(Item, FVector4f()); break;
+				case EDataType::Projector:		StoreProjector(Item, FProjector()); break;
+				case EDataType::Mesh:			StoreMesh(Item, nullptr); break;
+				case EDataType::Image:			StoreImage(Item, nullptr); break;
+				case EDataType::Layout:			StoreLayout(Item, nullptr); break;
+				case EDataType::Instance:		StoreInstance(Item, nullptr); break;
+				case EDataType::ExtensionData:	StoreExtensionData(Item, MakeShared<FExtensionData>()); break;
 				default:
 					// Not implemented
 					check(false);
@@ -374,7 +374,7 @@ namespace mu
 		case 1:
 		{
 			// Get the variable result
-			int32 var = LoadInt(FCacheAddress(VarAddress, item));
+			int32 var = LoadInt(FCacheAddress(VarAddress, Item));
 
 			OP::ADDRESS valueAt = DefAddress;
 			for (uint32 C = 0; C < CaseCount; ++C)
@@ -395,32 +395,32 @@ namespace mu
 			}
 
             // Schedule the end of this instruction if necessary
-            AddOp( FScheduledOp( item.At, item, 2, valueAt ),
-				   FScheduledOp( valueAt, item ) );
+            AddOp( FScheduledOp( Item.At, Item, 2, valueAt ),
+				   FScheduledOp( valueAt, Item ) );
 
             break;
         }
 
         case 2:
         {
-			OP::ADDRESS resultAt = OP::ADDRESS(item.CustomState);
+			OP::ADDRESS resultAt = OP::ADDRESS(Item.CustomState);
 
             // Store the final result
-            FCacheAddress cat( item );
-            FCacheAddress rat( resultAt, item );
+            FCacheAddress cat( Item );
+            FCacheAddress rat( resultAt, Item );
             switch (GetOpDataType(type))
             {
-            case DT_BOOL:       StoreBool( cat, LoadBool(rat) ); break;
-            case DT_INT:        StoreInt( cat, LoadInt(rat) ); break;
-            case DT_SCALAR:     StoreScalar( cat, LoadScalar(rat) ); break;
-            case DT_STRING:		StoreString( cat, LoadString( rat ) ); break;
-            case DT_COLOUR:		StoreColor( cat, LoadColor( rat ) ); break;
-            case DT_PROJECTOR:  StoreProjector( cat, LoadProjector(rat) ); break;
-			case DT_MESH:       StoreMesh( cat, LoadMesh(rat) ); break;
-            case DT_IMAGE:      StoreImage( cat, LoadImage(rat) ); break;
-            case DT_LAYOUT:     StoreLayout( cat, LoadLayout(rat) ); break;
-            case DT_INSTANCE:   StoreInstance( cat, LoadInstance(rat) ); break;
-			case DT_EXTENSION_DATA: StoreExtensionData( cat, LoadExtensionData(rat) ); break;
+            case EDataType::Bool:			StoreBool( cat, LoadBool(rat) ); break;
+            case EDataType::Int:			StoreInt( cat, LoadInt(rat) ); break;
+            case EDataType::Scalar:			StoreScalar( cat, LoadScalar(rat) ); break;
+            case EDataType::String:			StoreString( cat, LoadString( rat ) ); break;
+            case EDataType::Color:			StoreColor( cat, LoadColor( rat ) ); break;
+            case EDataType::Projector:		StoreProjector( cat, LoadProjector(rat) ); break;
+			case EDataType::Mesh:			StoreMesh( cat, LoadMesh(rat) ); break;
+            case EDataType::Image:			StoreImage( cat, LoadImage(rat) ); break;
+            case EDataType::Layout:			StoreLayout( cat, LoadLayout(rat) ); break;
+            case EDataType::Instance:		StoreInstance( cat, LoadInstance(rat) ); break;
+			case EDataType::ExtensionData:	StoreExtensionData( cat, LoadExtensionData(rat) ); break;
             default:
                 // Not implemented
                 check( false );
@@ -436,31 +436,31 @@ namespace mu
 
 
     //---------------------------------------------------------------------------------------------
-    void CodeRunner::RunCode_Instance(const FScheduledOp& item, const FModel* pModel, uint32 lodMask )
+    void CodeRunner::RunCode_Instance(const FScheduledOp& Item, const FModel* InModel, uint32 lodMask )
     {
         MUTABLE_CPUPROFILER_SCOPE(RunCode_Instance);
 
 		const FProgram& Program = Model->GetPrivate()->Program;
 
-		EOpType type = Program.GetOpType(item.At);
+		EOpType type = Program.GetOpType(Item.At);
         switch (type)
         {
 
         case EOpType::IN_ADDVECTOR:
         {
-			OP::InstanceAddArgs args = Program.GetOpArgs<OP::InstanceAddArgs>(item.At);
+			OP::InstanceAddArgs Args = Program.GetOpArgs<OP::InstanceAddArgs>(Item.At);
 
-            switch (item.Stage)
+            switch (Item.Stage)
             {
             case 0:
-                    AddOp( FScheduledOp( item.At, item, 1),
-                           FScheduledOp( args.instance, item),
-                           FScheduledOp( args.value, item) );
+                    AddOp( FScheduledOp( Item.At, Item, 1),
+                           FScheduledOp( Args.instance, Item),
+                           FScheduledOp( Args.value, Item) );
                 break;
 
             case 1:
             {
-                TSharedPtr<const FInstance> Base = LoadInstance( FCacheAddress(args.instance,item) );
+                TSharedPtr<const FInstance> Base = LoadInstance( FCacheAddress(Args.instance,Item) );
 				TSharedPtr<FInstance> Result;
                 if (!Base)
                 {
@@ -471,17 +471,17 @@ namespace mu
 					Result = mu::CloneOrTakeOver<FInstance>(Base);
                 }
 
-                if ( args.value )
+                if ( Args.value )
                 {
-					FVector4f value = LoadColor( FCacheAddress(args.value,item) );
+					FVector4f value = LoadColor( FCacheAddress(Args.value,Item) );
 
-                    OP::ADDRESS nameAd = args.name;
+                    OP::ADDRESS nameAd = Args.name;
                     check(  nameAd < (uint32)Program.ConstantStrings.Num() );
                     const FString& Name = Program.ConstantStrings[ nameAd ];
 
                     Result->GetPrivate()->AddVector( 0, 0, 0, value, FName(Name) );
                 }
-                StoreInstance( item, Result );
+                StoreInstance( Item, Result );
                 break;
             }
 
@@ -494,18 +494,18 @@ namespace mu
 
         case EOpType::IN_ADDSCALAR:
         {
-			OP::InstanceAddArgs args = Program.GetOpArgs<OP::InstanceAddArgs>(item.At);
-            switch (item.Stage)
+			OP::InstanceAddArgs Args = Program.GetOpArgs<OP::InstanceAddArgs>(Item.At);
+            switch (Item.Stage)
             {
             case 0:
-                    AddOp( FScheduledOp( item.At, item, 1),
-                           FScheduledOp( args.instance, item),
-                           FScheduledOp( args.value, item) );
+                    AddOp( FScheduledOp( Item.At, Item, 1),
+                           FScheduledOp( Args.instance, Item),
+                           FScheduledOp( Args.value, Item) );
                 break;
 
             case 1:
             {
-                TSharedPtr<const FInstance> Base = LoadInstance( FCacheAddress(args.instance,item) );
+                TSharedPtr<const FInstance> Base = LoadInstance( FCacheAddress(Args.instance,Item) );
                 TSharedPtr<FInstance> Result;
                 if (!Base)
                 {
@@ -516,17 +516,17 @@ namespace mu
                     Result = mu::CloneOrTakeOver<FInstance>(Base);
                 }
 
-                if ( args.value )
+                if ( Args.value )
                 {
-                    float value = LoadScalar( FCacheAddress(args.value,item) );
+                    float value = LoadScalar( FCacheAddress(Args.value,Item) );
 
-                    OP::ADDRESS nameAd = args.name;
+                    OP::ADDRESS nameAd = Args.name;
                     check(  nameAd < (uint32)Program.ConstantStrings.Num() );
                     const FString& Name = Program.ConstantStrings[ nameAd ];
 
                     Result->GetPrivate()->AddScalar( 0, 0, 0, value, FName(Name));
                 }
-                StoreInstance( item, Result );
+                StoreInstance( Item, Result );
                 break;
             }
 
@@ -539,18 +539,18 @@ namespace mu
 
         case EOpType::IN_ADDSTRING:
         {
-			OP::InstanceAddArgs args = Program.GetOpArgs<OP::InstanceAddArgs>( item.At );
-            switch ( item.Stage )
+			OP::InstanceAddArgs Args = Program.GetOpArgs<OP::InstanceAddArgs>( Item.At );
+            switch ( Item.Stage )
             {
             case 0:
-                    AddOp( FScheduledOp( item.At, item, 1 ), FScheduledOp( args.instance, item ),
-                           FScheduledOp( args.value, item ) );
+                    AddOp( FScheduledOp( Item.At, Item, 1 ), FScheduledOp( Args.instance, Item ),
+                           FScheduledOp( Args.value, Item ) );
                 break;
 
             case 1:
             {
                 TSharedPtr<const FInstance> Base =
-                    LoadInstance( FCacheAddress( args.instance, item ) );
+                    LoadInstance( FCacheAddress( Args.instance, Item ) );
                 TSharedPtr<FInstance> Result;
                 if ( !Base )
                 {
@@ -561,17 +561,17 @@ namespace mu
 					Result = mu::CloneOrTakeOver<FInstance>(Base);
 				}
 
-                if ( args.value )
+                if ( Args.value )
                 {
-                    TSharedPtr<const String> value = LoadString( FCacheAddress( args.value, item ) );
+                    TSharedPtr<const String> value = LoadString( FCacheAddress( Args.value, Item ) );
 
-                    OP::ADDRESS nameAd = args.name;
+                    OP::ADDRESS nameAd = Args.name;
                     check( nameAd < (uint32)Program.ConstantStrings.Num() );
                     const FString& Name = Program.ConstantStrings[nameAd];
 
                     Result->GetPrivate()->AddString( 0, 0, 0, value->GetValue(), FName(Name) );
                 }
-                StoreInstance( item, Result );
+                StoreInstance( Item, Result );
                 break;
             }
 
@@ -584,18 +584,18 @@ namespace mu
 
         case EOpType::IN_ADDCOMPONENT:
         {
-			OP::InstanceAddArgs args = Program.GetOpArgs<OP::InstanceAddArgs>(item.At);
-            switch (item.Stage)
+			OP::InstanceAddArgs Args = Program.GetOpArgs<OP::InstanceAddArgs>(Item.At);
+            switch (Item.Stage)
             {
             case 0:
-                    AddOp( FScheduledOp( item.At, item, 1),
-                           FScheduledOp( args.instance, item),
-                           FScheduledOp( args.value, item) );
+                    AddOp( FScheduledOp( Item.At, Item, 1),
+                           FScheduledOp( Args.instance, Item),
+                           FScheduledOp( Args.value, Item) );
                 break;
 
             case 1:
             {
-				TSharedPtr<const FInstance> Base = LoadInstance( FCacheAddress(args.instance,item) );
+				TSharedPtr<const FInstance> Base = LoadInstance( FCacheAddress(Args.instance,Item) );
 				TSharedPtr<FInstance> Result;
                 if (!Base)
                 {
@@ -606,9 +606,9 @@ namespace mu
 					Result = mu::CloneOrTakeOver<FInstance>(Base);
 				}
 
-                if ( args.value )
+                if ( Args.value )
                 {
-                    TSharedPtr<const FInstance> pComp = LoadInstance( FCacheAddress(args.value,item) );
+                    TSharedPtr<const FInstance> pComp = LoadInstance( FCacheAddress(Args.value,Item) );
 
                     int32 NewComponentIndex = Result->GetPrivate()->AddComponent();
 
@@ -617,10 +617,10 @@ namespace mu
                         Result->GetPrivate()->Components[NewComponentIndex] = pComp->GetPrivate()->Components[0];
 
 						// Id
-                    	Result->GetPrivate()->Components[NewComponentIndex].Id = args.ExternalId;
+                    	Result->GetPrivate()->Components[NewComponentIndex].Id = Args.ExternalId;
                     }
                 }
-                StoreInstance( item, Result );
+                StoreInstance( Item, Result );
                 break;
             }
 
@@ -633,18 +633,18 @@ namespace mu
 
         case EOpType::IN_ADDSURFACE:
         {
-			OP::InstanceAddArgs args = Program.GetOpArgs<OP::InstanceAddArgs>(item.At);
-            switch (item.Stage)
+			OP::InstanceAddArgs Args = Program.GetOpArgs<OP::InstanceAddArgs>(Item.At);
+            switch (Item.Stage)
             {
             case 0:
-                    AddOp( FScheduledOp( item.At, item, 1),
-                           FScheduledOp( args.instance, item),
-                           FScheduledOp( args.value, item) );
+                    AddOp( FScheduledOp( Item.At, Item, 1),
+                           FScheduledOp( Args.instance, Item),
+                           FScheduledOp( Args.value, Item) );
                 break;
 
             case 1:
             {
-				TSharedPtr<const FInstance> Base = LoadInstance( FCacheAddress(args.instance,item) );
+				TSharedPtr<const FInstance> Base = LoadInstance( FCacheAddress(Args.instance,Item) );
 
 				TSharedPtr<FInstance> Result;
 				if (Base)
@@ -658,9 +658,9 @@ namespace mu
 
                 // Empty surfaces are ok, they still need to be created, because they may contain
                 // additional information like internal or external IDs
-                //if ( args.value )
+                //if ( Args.value )
                 {
-					TSharedPtr<const FInstance> pSurf = LoadInstance( FCacheAddress(args.value,item) );
+					TSharedPtr<const FInstance> pSurf = LoadInstance( FCacheAddress(Args.value,Item) );
 
                     int32 sindex = Result->GetPrivate()->AddSurface( 0, 0 );
 
@@ -678,17 +678,17 @@ namespace mu
                     }
 
                     // Name
-                    OP::ADDRESS nameAd = args.name;
+                    OP::ADDRESS nameAd = Args.name;
                     check( nameAd < (uint32)Program.ConstantStrings.Num() );
                     const FString& Name = Program.ConstantStrings[ nameAd ];
                     Result->GetPrivate()->SetSurfaceName( 0, 0, sindex, FName(Name) );
 
                     // IDs
-                    Result->GetPrivate()->Components[0].LODs[0].Surfaces[sindex].InternalId = args.id;
-                    Result->GetPrivate()->Components[0].LODs[0].Surfaces[sindex].ExternalId = args.ExternalId;
-                    Result->GetPrivate()->Components[0].LODs[0].Surfaces[sindex].SharedId = args.SharedSurfaceId;
+                    Result->GetPrivate()->Components[0].LODs[0].Surfaces[sindex].InternalId = Args.id;
+                    Result->GetPrivate()->Components[0].LODs[0].Surfaces[sindex].ExternalId = Args.ExternalId;
+                    Result->GetPrivate()->Components[0].LODs[0].Surfaces[sindex].SharedId = Args.SharedSurfaceId;
                 }
-                StoreInstance( item, Result );
+                StoreInstance( Item, Result );
                 break;
             }
 
@@ -701,13 +701,13 @@ namespace mu
 
         case EOpType::IN_ADDLOD:
         {
-			const uint8* Data = Program.GetOpArgsPointer(item.At);
+			const uint8* Data = Program.GetOpArgsPointer(Item.At);
 
 			uint8 LODCount;
 			FMemory::Memcpy(&LODCount, Data, sizeof(uint8));
 			Data += sizeof(uint8);
 
-			switch (item.Stage)
+			switch (Item.Stage)
             {
             case 0:
             {                
@@ -724,12 +724,12 @@ namespace mu
 
                         if (bSelectedLod)
                         {
-                            deps.Emplace(LODAddress, item);
+                            deps.Emplace(LODAddress, Item);
                         }
                     }
                 }
 
-                AddOp( FScheduledOp( item.At, item, 1), deps );
+                AddOp( FScheduledOp( Item.At, Item, 1), deps );
 
                 break;
             }
@@ -755,7 +755,7 @@ namespace mu
 						
 						if (bIsSelectedLod)
                         {
-							TSharedPtr<const FInstance> pLOD = LoadInstance( FCacheAddress(LODAddress,item) );
+							TSharedPtr<const FInstance> pLOD = LoadInstance( FCacheAddress(LODAddress,Item) );
 
                             // In a degenerated case, the returned pLOD may not have an LOD inside
  							if (!pLOD->GetPrivate()->Components.IsEmpty()
@@ -768,7 +768,7 @@ namespace mu
                     }
                 }
 
-                StoreInstance( item, Result );
+                StoreInstance( Item, Result );
                 break;
             }
 
@@ -781,8 +781,8 @@ namespace mu
 
 		case EOpType::IN_ADDEXTENSIONDATA:
 		{
-			OP::InstanceAddExtensionDataArgs Args = Program.GetOpArgs<OP::InstanceAddExtensionDataArgs>(item.At);
-			switch (item.Stage)
+			OP::InstanceAddExtensionDataArgs Args = Program.GetOpArgs<OP::InstanceAddExtensionDataArgs>(Item.At);
+			switch (Item.Stage)
 			{
 				case 0:
 				{
@@ -791,10 +791,10 @@ namespace mu
 					check(Args.ExtensionData);
 
 					TArray<FScheduledOp> Dependencies;
-					Dependencies.Emplace(Args.Instance, item);
-					Dependencies.Emplace(Args.ExtensionData, item);
+					Dependencies.Emplace(Args.Instance, Item);
+					Dependencies.Emplace(Args.ExtensionData, Item);
 
-					AddOp(FScheduledOp(item.At, item, 1), Dependencies);
+					AddOp(FScheduledOp(Item.At, Item, 1), Dependencies);
 
 					break;
 				}
@@ -802,11 +802,11 @@ namespace mu
 				case 1:
 				{
 					// Assemble result
-					TSharedPtr<const FInstance> InstanceOpResult = LoadInstance(FCacheAddress(Args.Instance, item));
+					TSharedPtr<const FInstance> InstanceOpResult = LoadInstance(FCacheAddress(Args.Instance, Item));
 
 					TSharedPtr<FInstance> Result = mu::CloneOrTakeOver<FInstance>(InstanceOpResult);
 
-					if (TSharedPtr<const FExtensionData> ExtensionData = LoadExtensionData(FCacheAddress(Args.ExtensionData, item)))
+					if (TSharedPtr<const FExtensionData> ExtensionData = LoadExtensionData(FCacheAddress(Args.ExtensionData, Item)))
 					{
 						const OP::ADDRESS NameAddress = Args.ExtensionDataName;
 						check(NameAddress < (uint32)Program.ConstantStrings.Num());
@@ -815,7 +815,7 @@ namespace mu
 						Result->GetPrivate()->AddExtensionData(ExtensionData, FName(NameString) );
 					}
 
-					StoreInstance(item, Result);
+					StoreInstance(Item, Result);
 					break;
 				}
 
@@ -833,7 +833,7 @@ namespace mu
 
 
     //---------------------------------------------------------------------------------------------
-    void CodeRunner::RunCode_InstanceAddResource(const FScheduledOp& item, const TSharedPtr<const FModel>& InModel, const FParameters* InParams )
+    void CodeRunner::RunCode_InstanceAddResource(const FScheduledOp& Item, const TSharedPtr<const FModel>& InModel, const FParameters* InParams )
     {
 		MUTABLE_CPUPROFILER_SCOPE(RunCode_InstanceAddResource);
 
@@ -844,23 +844,23 @@ namespace mu
 
 		const FProgram& Program = Model->GetPrivate()->Program;
 
-		EOpType type = Program.GetOpType(item.At);
+		EOpType type = Program.GetOpType(Item.At);
         switch (type)
         {
 
         case EOpType::IN_ADDMESH:
         {
-			OP::InstanceAddArgs args = Program.GetOpArgs<OP::InstanceAddArgs>(item.At);
-            switch (item.Stage)
+			OP::InstanceAddArgs Args = Program.GetOpArgs<OP::InstanceAddArgs>(Item.At);
+            switch (Item.Stage)
             {
             case 0:
 				// We don't build the resources when building instance: just store ids for them.
-				AddOp( FScheduledOp( item.At, item, 1), FScheduledOp( args.instance, item) );
+				AddOp( FScheduledOp( Item.At, Item, 1), FScheduledOp( Args.instance, Item) );
                 break;
 
             case 1:
             {
-                TSharedPtr<const FInstance> Base = LoadInstance( FCacheAddress(args.instance,item) );
+                TSharedPtr<const FInstance> Base = LoadInstance( FCacheAddress(Args.instance,Item) );
                 TSharedPtr<FInstance> Result;
                 if (!Base)
                 {
@@ -871,15 +871,15 @@ namespace mu
 					Result = mu::CloneOrTakeOver<FInstance>(Base);
 				}
 
-                if ( args.value )
+                if ( Args.value )
                 {
-					FResourceID MeshId = System->WorkingMemoryManager.GetResourceKey(InModel,InParams,args.relevantParametersListIndex, args.value);
-					OP::ADDRESS NameAd = args.name;
+					FResourceID MeshId = System->WorkingMemoryManager.GetResourceKey(InModel,InParams,Args.relevantParametersListIndex, Args.value);
+					OP::ADDRESS NameAd = Args.name;
 					check(NameAd < (uint32)Program.ConstantStrings.Num());
 					const FString& Name = Program.ConstantStrings[NameAd];
 					Result->GetPrivate()->SetMesh(0, 0, MeshId, FName(Name));
                 }
-                StoreInstance( item, Result );
+                StoreInstance( Item, Result );
                 break;
             }
 
@@ -891,17 +891,17 @@ namespace mu
 
         case EOpType::IN_ADDIMAGE:
         {
-			OP::InstanceAddArgs args = Program.GetOpArgs<OP::InstanceAddArgs>(item.At);
-            switch (item.Stage)
+			OP::InstanceAddArgs Args = Program.GetOpArgs<OP::InstanceAddArgs>(Item.At);
+            switch (Item.Stage)
             {
             case 0:
 				// We don't build the resources when building instance: just store ids for them.
-				AddOp( FScheduledOp( item.At, item, 1), FScheduledOp( args.instance, item) );
+				AddOp( FScheduledOp( Item.At, Item, 1), FScheduledOp( Args.instance, Item) );
                 break;
 
             case 1:
             {
-                TSharedPtr<const FInstance> Base = LoadInstance( FCacheAddress(args.instance,item) );
+                TSharedPtr<const FInstance> Base = LoadInstance( FCacheAddress(Args.instance,Item) );
                 TSharedPtr<FInstance> Result;
                 if (!Base)
                 {
@@ -912,15 +912,15 @@ namespace mu
 					Result = mu::CloneOrTakeOver<FInstance>(Base);
 				}
 
-                if ( args.value )
+                if ( Args.value )
                 {
-					FResourceID ImageId = System->WorkingMemoryManager.GetResourceKey(InModel, InParams, args.relevantParametersListIndex, args.value);
-					OP::ADDRESS NameAd = args.name;
+					FResourceID ImageId = System->WorkingMemoryManager.GetResourceKey(InModel, InParams, Args.relevantParametersListIndex, Args.value);
+					OP::ADDRESS NameAd = Args.name;
 					check(NameAd < (uint32)Program.ConstantStrings.Num());
 					const FString& Name = Program.ConstantStrings[NameAd];
 					Result->GetPrivate()->AddImage(0, 0, 0, ImageId, FName(Name) );
                 }
-                StoreInstance( item, Result );
+                StoreInstance( Item, Result );
                 break;
             }
 
@@ -938,21 +938,21 @@ namespace mu
 
 
     //---------------------------------------------------------------------------------------------
-    bool CodeRunner::RunCode_ConstantResource(const FScheduledOp& item, const FModel* pModel )
+    bool CodeRunner::RunCode_ConstantResource(const FScheduledOp& Item, const FModel* InModel )
     {
 		MUTABLE_CPUPROFILER_SCOPE(RunCode_Constant);
 
 		const FProgram& Program = Model->GetPrivate()->Program;
 
-		EOpType type = Program.GetOpType(item.At);
+		EOpType type = Program.GetOpType(Item.At);
         switch (type)
         {
 
         case EOpType::ME_CONSTANT:
         {
-			OP::MeshConstantArgs args = Program.GetOpArgs<OP::MeshConstantArgs>(item.At);
+			OP::MeshConstantArgs Args = Program.GetOpArgs<OP::MeshConstantArgs>(Item.At);
 
-            FConstantResourceIndex ResourceIndex = *reinterpret_cast<const FConstantResourceIndex*>(&args.value);
+            FConstantResourceIndex ResourceIndex = *reinterpret_cast<const FConstantResourceIndex*>(&Args.value);
 
             TSharedPtr<const FMesh> SourceConst;
 			Program.GetConstant(ResourceIndex, SourceConst);
@@ -962,31 +962,31 @@ namespace mu
 			Source->CopyFrom(*SourceConst);
 
             // Set the separate skeleton if necessary
-            if (args.skeleton >= 0)
+            if (Args.skeleton >= 0)
             {
-                check(Program.ConstantSkeletons.Num() > args.skeleton);
-                TSharedPtr<const FSkeleton> pSkeleton = Program.ConstantSkeletons[args.skeleton];
+                check(Program.ConstantSkeletons.Num() > Args.skeleton);
+                TSharedPtr<const FSkeleton> pSkeleton = Program.ConstantSkeletons[Args.skeleton];
                 Source->SetSkeleton(pSkeleton);
             }
 
-			if (args.physicsBody >= 0)
+			if (Args.physicsBody >= 0)
 			{
-                check(Program.ConstantPhysicsBodies.Num() > args.physicsBody);
-                TSharedPtr<const FPhysicsBody> pPhysicsBody = Program.ConstantPhysicsBodies[args.physicsBody];
+                check(Program.ConstantPhysicsBodies.Num() > Args.physicsBody);
+                TSharedPtr<const FPhysicsBody> pPhysicsBody = Program.ConstantPhysicsBodies[Args.physicsBody];
                 Source->SetPhysicsBody(pPhysicsBody);
 			}
 
-            StoreMesh(item, Source);
-			//UE_LOG(LogMutableCore, Log, TEXT("Set mesh constant %d."), item.At);
+            StoreMesh(Item, Source);
+			//UE_LOG(LogMutableCore, Log, TEXT("Set mesh constant %d."), Item.At);
             break;
         }
 
         case EOpType::IM_CONSTANT:
         {
-			OP::ResourceConstantArgs args = Program.GetOpArgs<OP::ResourceConstantArgs>(item.At);
-            OP::ADDRESS cat = args.value;
+			OP::ResourceConstantArgs Args = Program.GetOpArgs<OP::ResourceConstantArgs>(Item.At);
+            OP::ADDRESS cat = Args.value;
 
-			int32 MipsToSkip = item.ExecutionOptions;
+			int32 MipsToSkip = Item.ExecutionOptions;
             TSharedPtr<const FImage> Source;
 			Program.GetConstant(cat, Source, MipsToSkip, [this](int32 x, int32 y, int32 m, EImageFormat f, EInitializationType i)
 				{
@@ -999,14 +999,14 @@ namespace mu
 				return false;
 			}
 
-            StoreImage( item, Source );
-			//UE_LOG(LogMutableCore, Log, TEXT("Set image constant %d."), item.At);
+            StoreImage( Item, Source );
+			//UE_LOG(LogMutableCore, Log, TEXT("Set image constant %d."), Item.At);
 			break;
         }
 
 		case EOpType::ED_CONSTANT:
 		{
-			OP::ResourceConstantArgs Args = Program.GetOpArgs<OP::ResourceConstantArgs>(item.At);
+			OP::ResourceConstantArgs Args = Program.GetOpArgs<OP::ResourceConstantArgs>(Item.At);
 
 			// Assume the ROM has been loaded previously
 			TSharedPtr<const mu::FExtensionData> SourceConst;
@@ -1014,7 +1014,7 @@ namespace mu
 
 			check(SourceConst);
 
-            StoreExtensionData(item, SourceConst);
+            StoreExtensionData(Item, SourceConst);
             break;
 		}
 
@@ -1032,21 +1032,21 @@ namespace mu
     }
 
     //---------------------------------------------------------------------------------------------
-	void CodeRunner::RunCode_Mesh(const FScheduledOp& item, const FModel* pModel )
+	void CodeRunner::RunCode_Mesh(const FScheduledOp& Item, const FModel* InModel )
     {
 		MUTABLE_CPUPROFILER_SCOPE(RunCode_Mesh);
 
 		const FProgram& Program = Model->GetPrivate()->Program;
 
-		EOpType type = Program.GetOpType(item.At);
+		EOpType type = Program.GetOpType(Item.At);
 
         switch (type)
         {
 
 		case EOpType::ME_REFERENCE:
 		{
-			OP::ResourceReferenceArgs Args = Program.GetOpArgs<OP::ResourceReferenceArgs>(item.At);
-			switch (item.Stage)
+			OP::ResourceReferenceArgs Args = Program.GetOpArgs<OP::ResourceReferenceArgs>(Item.At);
+			switch (Item.Stage)
 			{
 			case 0:
 			{
@@ -1060,7 +1060,7 @@ namespace mu
 				{
 					Result = FMesh::CreateAsReference(Args.ID, false);
 				}
-				StoreMesh(item, Result);
+				StoreMesh(Item, Result);
 				break;
 			}
 
@@ -1073,36 +1073,36 @@ namespace mu
 
         case EOpType::ME_APPLYLAYOUT:
         {
-			OP::MeshApplyLayoutArgs args = Program.GetOpArgs<OP::MeshApplyLayoutArgs>(item.At);
-			switch (item.Stage)
+			OP::MeshApplyLayoutArgs Args = Program.GetOpArgs<OP::MeshApplyLayoutArgs>(Item.At);
+			switch (Item.Stage)
 			{
 			case 0:
 			{
-				AddOp(FScheduledOp(item.At, item, 1),
-					FScheduledOp(args.mesh, item),
-					FScheduledOp(args.layout, item));
+				AddOp(FScheduledOp(Item.At, Item, 1),
+					FScheduledOp(Args.mesh, Item),
+					FScheduledOp(Args.layout, Item));
 				break;
 			}
             case 1:
             {
            		MUTABLE_CPUPROFILER_SCOPE(ME_APPLYLAYOUT)
             		
-                TSharedPtr<const FMesh> Base = LoadMesh(FCacheAddress(args.mesh, item));
+                TSharedPtr<const FMesh> Base = LoadMesh(FCacheAddress(Args.mesh, Item));
 
                 if (Base)
                 {
 					TSharedPtr<FMesh> Result = CloneOrTakeOver(Base);
 
-                    TSharedPtr<const FLayout> Layout = LoadLayout(FCacheAddress(args.layout, item));
-                    int32 texCoordsSet = args.channel;
+                    TSharedPtr<const FLayout> Layout = LoadLayout(FCacheAddress(Args.layout, Item));
+                    int32 texCoordsSet = Args.channel;
 
                     MeshApplyLayout(Result.Get(), Layout.Get(), texCoordsSet);
 					
-					StoreMesh(item, Result);
+					StoreMesh(Item, Result);
                 }
 				else
 				{
-					StoreMesh(item, nullptr);
+					StoreMesh(Item, nullptr);
 				}
 
 
@@ -1118,7 +1118,7 @@ namespace mu
 
         case EOpType::ME_DIFFERENCE:
         {
-			const uint8* data = Program.GetOpArgsPointer(item.At);
+			const uint8* data = Program.GetOpArgsPointer(Item.At);
 
 			OP::ADDRESS BaseAt = 0;
 			FMemory::Memcpy(&BaseAt, data, sizeof(OP::ADDRESS)); 
@@ -1128,19 +1128,19 @@ namespace mu
 			FMemory::Memcpy(&TargetAt, data, sizeof(OP::ADDRESS)); 
 			data += sizeof(OP::ADDRESS);
 
-            switch (item.Stage)
+            switch (Item.Stage)
             {
             case 0:
 			{
 				if (BaseAt && TargetAt)
 				{
-					AddOp(FScheduledOp(item.At, item, 1),
-						FScheduledOp(BaseAt, item),
-						FScheduledOp(TargetAt, item));
+					AddOp(FScheduledOp(Item.At, Item, 1),
+						FScheduledOp(BaseAt, Item),
+						FScheduledOp(TargetAt, Item));
 				}
 				else
 				{
-					StoreMesh(item, nullptr);
+					StoreMesh(Item, nullptr);
 				}
 				break;
 			}
@@ -1148,8 +1148,8 @@ namespace mu
             {
        	        MUTABLE_CPUPROFILER_SCOPE(ME_DIFFERENCE)
 
-                TSharedPtr<const FMesh> Base = LoadMesh(FCacheAddress(BaseAt,item));
-                TSharedPtr<const FMesh> pTarget = LoadMesh(FCacheAddress(TargetAt,item));
+                TSharedPtr<const FMesh> Base = LoadMesh(FCacheAddress(BaseAt,Item));
+                TSharedPtr<const FMesh> pTarget = LoadMesh(FCacheAddress(TargetAt,Item));
 
 				TArray<EMeshBufferSemantic, TInlineAllocator<8>> Semantics;
 				TArray<int32, TInlineAllocator<8>> SemanticIndices;
@@ -1184,7 +1184,7 @@ namespace mu
 				Release(Base);
 				Release(pTarget);
 
-                StoreMesh(item, Result);
+                StoreMesh(Item, Result);
                 break;
             }
 
@@ -1197,7 +1197,7 @@ namespace mu
 
         case EOpType::ME_MORPH:
         {
-			const uint8* Data = Program.GetOpArgsPointer(item.At);
+			const uint8* Data = Program.GetOpArgsPointer(Item.At);
 
 			OP::ADDRESS FactorAt = 0;
 			FMemory::Memcpy(&FactorAt, Data, sizeof(OP::ADDRESS)); 
@@ -1211,17 +1211,17 @@ namespace mu
 			FMemory::Memcpy(&TargetAt, Data, sizeof(OP::ADDRESS)); 
 			Data += sizeof(OP::ADDRESS);
 
-			switch (item.Stage)
+			switch (Item.Stage)
             {
             case 0:
                 if (BaseAt)
                 {
-                    AddOp(FScheduledOp(item.At, item, 1),
-                           FScheduledOp(FactorAt, item));
+                    AddOp(FScheduledOp(Item.At, Item, 1),
+                           FScheduledOp(FactorAt, Item));
                 }
                 else
                 {
-                    StoreMesh(item, nullptr);
+                    StoreMesh(Item, nullptr);
                 }
                 break;
 
@@ -1229,7 +1229,7 @@ namespace mu
             {
                 MUTABLE_CPUPROFILER_SCOPE(ME_MORPH_1)
 
-                float Factor = LoadScalar(FCacheAddress(FactorAt, item));
+                float Factor = LoadScalar(FCacheAddress(FactorAt, Item));
 
                 // Factor goes from -1 to 1 across all targets. [0 - 1] represents positive morphs, while [-1, 0) represent negative morphs.
 				Factor = FMath::Clamp(Factor, -1.0f, 1.0f); // Is the factor not in range [-1, 1], it will index a non existing morph.
@@ -1241,16 +1241,16 @@ namespace mu
                 // No morph
 				if (FMath::IsNearlyZero(Factor))
                 {                        
-                    AddOp(FScheduledOp(item.At, item, 2, DataAddress),
-						FScheduledOp(BaseAt, item));
+                    AddOp(FScheduledOp(Item.At, Item, 2, DataAddress),
+						FScheduledOp(BaseAt, Item));
                 }
                 // The Morph, partial or full
                 else
                 {
                     // We will need the base again
-                    AddOp(FScheduledOp(item.At, item, 2, DataAddress),
-						FScheduledOp(BaseAt, item),
-						FScheduledOp(TargetAt, item));
+                    AddOp(FScheduledOp(Item.At, Item, 2, DataAddress),
+						FScheduledOp(BaseAt, Item),
+						FScheduledOp(TargetAt, Item));
                 }
 
                 break;
@@ -1260,10 +1260,10 @@ namespace mu
             {
        		    MUTABLE_CPUPROFILER_SCOPE(ME_MORPH_2)
 
-                TSharedPtr<const FMesh> BaseMesh = LoadMesh(FCacheAddress(BaseAt, item));
+                TSharedPtr<const FMesh> BaseMesh = LoadMesh(FCacheAddress(BaseAt, Item));
 
                 // Factor from 0 to 1 between the two targets
-                const FScheduledOpData& OpHeapData = HeapData[item.CustomState];
+                const FScheduledOpData& OpHeapData = HeapData[Item.CustomState];
                 float Factor = OpHeapData.Interpolate.Bifactor;
 
                 if (BaseMesh)
@@ -1271,12 +1271,12 @@ namespace mu
 					// No morph
 					if (FMath::IsNearlyZero(Factor))
                     {
-						StoreMesh(item, BaseMesh);
+						StoreMesh(Item, BaseMesh);
                     }
 					// The Morph, partial or full
                     else 
                     {
-                        TSharedPtr<const FMesh> MorphMesh = LoadMesh(FCacheAddress(TargetAt, item));
+                        TSharedPtr<const FMesh> MorphMesh = LoadMesh(FCacheAddress(TargetAt, Item));
 						
 						if (MorphMesh)
 						{
@@ -1284,17 +1284,17 @@ namespace mu
 							MeshMorph(Result.Get(), MorphMesh.Get(), Factor);
 
 							Release(MorphMesh);
-							StoreMesh(item, Result);
+							StoreMesh(Item, Result);
 						}
 						else
 						{
-							StoreMesh(item, BaseMesh);
+							StoreMesh(Item, BaseMesh);
 						}
                     }
                 }
 				else
 				{
-					StoreMesh(item, nullptr);
+					StoreMesh(Item, nullptr);
 				}
 
                 break;
@@ -1309,22 +1309,22 @@ namespace mu
 
         case EOpType::ME_MERGE:
         {
-			OP::MeshMergeArgs args = Program.GetOpArgs<OP::MeshMergeArgs>(item.At);
-            switch (item.Stage)
+			OP::MeshMergeArgs Args = Program.GetOpArgs<OP::MeshMergeArgs>(Item.At);
+            switch (Item.Stage)
             {
             case 0:
 			{
-				AddOp(FScheduledOp(item.At, item, 1),
-					FScheduledOp(args.base, item),
-					FScheduledOp(args.added, item));
+				AddOp(FScheduledOp(Item.At, Item, 1),
+					FScheduledOp(Args.base, Item),
+					FScheduledOp(Args.added, Item));
 				break;
 			}
             case 1:
             {
 				MUTABLE_CPUPROFILER_SCOPE(ME_MERGE_1)
 
-                TSharedPtr<const FMesh> pA = LoadMesh(FCacheAddress(args.base, item));
-                TSharedPtr<const FMesh> pB = LoadMesh(FCacheAddress(args.added, item));
+                TSharedPtr<const FMesh> pA = LoadMesh(FCacheAddress(Args.base, Item));
+                TSharedPtr<const FMesh> pB = LoadMesh(FCacheAddress(Args.added, Item));
 
                 if (pA && pB && pA->GetVertexCount() && pB->GetVertexCount())
                 {
@@ -1336,25 +1336,25 @@ namespace mu
 
 					TSharedPtr<FMesh> Result = CreateMesh(pA->GetDataSize() + pB->GetDataSize());
 
-					MeshMerge(Result.Get(), pA.Get(), pB.Get(), !args.newSurfaceID, Scratch);
+					MeshMerge(Result.Get(), pA.Get(), pB.Get(), !Args.newSurfaceID, Scratch);
 
 					Release(Scratch.FirstReformat);
 					Release(Scratch.SecondReformat);
 
-                    if (args.newSurfaceID)
+                    if (Args.newSurfaceID)
                     {
 						check(pB->GetSurfaceCount() == 1);
-						Result->Surfaces.Last().Id = args.newSurfaceID;
+						Result->Surfaces.Last().Id = Args.newSurfaceID;
                     }
 
 					Release(pA);
 					Release(pB);
-					StoreMesh(item, Result);
+					StoreMesh(Item, Result);
                 }
                 else if (pA && (pA->GetVertexCount() || pA->IsReference()))
                 {
 					Release(pB);
-					StoreMesh(item, pA);
+					StoreMesh(Item, pA);
                 }
                 else if (pB && (pB->GetVertexCount() || pB->IsReference()))
                 {
@@ -1362,19 +1362,19 @@ namespace mu
 
                     check(Result->IsReference() || (Result->GetSurfaceCount() == 1) );
 
-                    if (Result->GetSurfaceCount() > 0 && args.newSurfaceID)
+                    if (Result->GetSurfaceCount() > 0 && Args.newSurfaceID)
                     {
-                        Result->Surfaces.Last().Id = args.newSurfaceID;
+                        Result->Surfaces.Last().Id = Args.newSurfaceID;
                     }
 
 					Release(pA);
-					StoreMesh(item, Result);
+					StoreMesh(Item, Result);
                 }
                 else
                 {
 					Release(pA);
 					Release(pB);
-					StoreMesh(item, CreateMesh());
+					StoreMesh(Item, CreateMesh());
                 }
 
                 break;
@@ -1389,22 +1389,22 @@ namespace mu
 
         case EOpType::ME_MASKCLIPMESH:
         {
-			OP::MeshMaskClipMeshArgs args = Program.GetOpArgs<OP::MeshMaskClipMeshArgs>(item.At);
-            switch (item.Stage)
+			OP::MeshMaskClipMeshArgs Args = Program.GetOpArgs<OP::MeshMaskClipMeshArgs>(Item.At);
+            switch (Item.Stage)
             {
             case 0:
 			{
-				AddOp(FScheduledOp(item.At, item, 1),
-					FScheduledOp(args.source, item),
-					FScheduledOp(args.clip, item));
+				AddOp(FScheduledOp(Item.At, Item, 1),
+					FScheduledOp(Args.source, Item),
+					FScheduledOp(Args.clip, Item));
 				break;
 			}
             case 1:
             {
             	MUTABLE_CPUPROFILER_SCOPE(ME_MASKCLIPMESH_1)
             		
-                TSharedPtr<const FMesh> Source = LoadMesh(FCacheAddress(args.source, item));
-                TSharedPtr<const FMesh> pClip = LoadMesh(FCacheAddress(args.clip, item));
+                TSharedPtr<const FMesh> Source = LoadMesh(FCacheAddress(Args.source, Item));
+                TSharedPtr<const FMesh> pClip = LoadMesh(FCacheAddress(Args.clip, Item));
 
                 // Only if both are valid.
                 if (Source.Get() && pClip.Get())
@@ -1419,18 +1419,18 @@ namespace mu
 					if (!bOutSuccess)
 					{
 						Release(Result);
-						StoreMesh(item, nullptr);
+						StoreMesh(Item, nullptr);
 					}
 					else
 					{
-						StoreMesh(item, Result);
+						StoreMesh(Item, Result);
 					}
                 }
 				else
 				{
 					Release(Source);
 					Release(pClip);
-					StoreMesh(item, nullptr);
+					StoreMesh(Item, nullptr);
 				}
 
                 break;
@@ -1445,26 +1445,26 @@ namespace mu
 
 		case EOpType::ME_MASKCLIPUVMASK:
 		{
-			OP::MeshMaskClipUVMaskArgs args = Program.GetOpArgs<OP::MeshMaskClipUVMaskArgs>(item.At);
-			switch (item.Stage)
+			OP::MeshMaskClipUVMaskArgs Args = Program.GetOpArgs<OP::MeshMaskClipUVMaskArgs>(Item.At);
+			switch (Item.Stage)
 			{
 			case 0:
 			{
-				AddOp(FScheduledOp(item.At, item, 1),
-					FScheduledOp(args.Source, item),
-					FScheduledOp(args.UVSource, item),
-					FScheduledOp(args.MaskImage, item),
-					FScheduledOp(args.MaskLayout, item));
+				AddOp(FScheduledOp(Item.At, Item, 1),
+					FScheduledOp(Args.Source, Item),
+					FScheduledOp(Args.UVSource, Item),
+					FScheduledOp(Args.MaskImage, Item),
+					FScheduledOp(Args.MaskLayout, Item));
 				break;
 			}
 			case 1:
 			{
 				MUTABLE_CPUPROFILER_SCOPE(ME_MASKCLIPUVMASK_1);
 
-				TSharedPtr<const FMesh> Source = LoadMesh(FCacheAddress(args.Source, item));
-				TSharedPtr<const FMesh> UVSource = LoadMesh(FCacheAddress(args.UVSource, item));
-				TSharedPtr<const FImage> MaskImage = LoadImage(FCacheAddress(args.MaskImage, item));
-				TSharedPtr<const FLayout> MaskLayout = LoadLayout(FCacheAddress(args.MaskLayout, item));
+				TSharedPtr<const FMesh> Source = LoadMesh(FCacheAddress(Args.Source, Item));
+				TSharedPtr<const FMesh> UVSource = LoadMesh(FCacheAddress(Args.UVSource, Item));
+				TSharedPtr<const FImage> MaskImage = LoadImage(FCacheAddress(Args.MaskImage, Item));
+				TSharedPtr<const FLayout> MaskLayout = LoadLayout(FCacheAddress(Args.MaskLayout, Item));
 
 				// Only if both are valid.
 				if (Source.Get() && MaskImage.Get())
@@ -1472,7 +1472,7 @@ namespace mu
 					TSharedPtr<FMesh> Result = CreateMesh();
 
 					bool bOutSuccess = false;
-					MakeMeshMaskFromUVMask(Result.Get(), Source.Get(), UVSource.Get(), MaskImage.Get(), args.LayoutIndex, bOutSuccess);
+					MakeMeshMaskFromUVMask(Result.Get(), Source.Get(), UVSource.Get(), MaskImage.Get(), Args.LayoutIndex, bOutSuccess);
 
 					Release(Source);
 					Release(UVSource);
@@ -1480,11 +1480,11 @@ namespace mu
 					if (!bOutSuccess)
 					{
 						Release(Result);
-						StoreMesh(item, nullptr);
+						StoreMesh(Item, nullptr);
 					}
 					else
 					{
-						StoreMesh(item, Result);
+						StoreMesh(Item, Result);
 					}
 				}
 				else if (Source.Get() && MaskLayout.Get())
@@ -1492,7 +1492,7 @@ namespace mu
 					TSharedPtr<FMesh> Result = CreateMesh();
 
 					bool bOutSuccess = false;
-					MakeMeshMaskFromLayout(Result.Get(), Source.Get(), UVSource.Get(), MaskLayout.Get(), args.LayoutIndex, bOutSuccess);
+					MakeMeshMaskFromLayout(Result.Get(), Source.Get(), UVSource.Get(), MaskLayout.Get(), Args.LayoutIndex, bOutSuccess);
 
 					Release(Source);
 					Release(UVSource);
@@ -1500,11 +1500,11 @@ namespace mu
 					if (!bOutSuccess)
 					{
 						Release(Result);
-						StoreMesh(item, nullptr);
+						StoreMesh(Item, nullptr);
 					}
 					else
 					{
-						StoreMesh(item, Result);
+						StoreMesh(Item, Result);
 					}
 				}
 				else
@@ -1512,7 +1512,7 @@ namespace mu
 					Release(Source);
 					Release(UVSource);
 					Release(MaskImage);
-					StoreMesh(item, nullptr);
+					StoreMesh(Item, nullptr);
 				}
 
 				break;
@@ -1527,22 +1527,22 @@ namespace mu
 
         case EOpType::ME_MASKDIFF:
         {
-			OP::MeshMaskDiffArgs args = Program.GetOpArgs<OP::MeshMaskDiffArgs>(item.At);
-            switch (item.Stage)
+			OP::MeshMaskDiffArgs Args = Program.GetOpArgs<OP::MeshMaskDiffArgs>(Item.At);
+            switch (Item.Stage)
             {
             case 0:
 			{
-				AddOp(FScheduledOp(item.At, item, 1),
-					FScheduledOp(args.source, item),
-					FScheduledOp(args.fragment, item));
+				AddOp(FScheduledOp(Item.At, Item, 1),
+					FScheduledOp(Args.source, Item),
+					FScheduledOp(Args.fragment, Item));
 				break;
 			}
             case 1:
             {
            		MUTABLE_CPUPROFILER_SCOPE(ME_MASKDIFF_1)
             		
-                TSharedPtr<const FMesh> Source = LoadMesh(FCacheAddress(args.source, item));
-                TSharedPtr<const FMesh> pClip = LoadMesh(FCacheAddress(args.fragment, item));
+                TSharedPtr<const FMesh> Source = LoadMesh(FCacheAddress(Args.source, Item));
+                TSharedPtr<const FMesh> pClip = LoadMesh(FCacheAddress(Args.fragment, Item));
 
                 // Only if both are valid.
                 if (Source.Get() && pClip.Get())
@@ -1558,18 +1558,18 @@ namespace mu
 					if (!bOutSuccess)
 					{
 						Release(Result);
-						StoreMesh(item, nullptr);
+						StoreMesh(Item, nullptr);
 					}
 					else
 					{
-						StoreMesh(item, Result);
+						StoreMesh(Item, Result);
 					}
                 }
 				else
 				{
 					Release(Source);
 					Release(pClip);
-					StoreMesh(item, nullptr);
+					StoreMesh(Item, nullptr);
 				}
 
                 break;
@@ -1584,20 +1584,20 @@ namespace mu
 
         case EOpType::ME_FORMAT:
         {
-			OP::MeshFormatArgs args = Program.GetOpArgs<OP::MeshFormatArgs>(item.At);
-            switch (item.Stage)
+			OP::MeshFormatArgs Args = Program.GetOpArgs<OP::MeshFormatArgs>(Item.At);
+            switch (Item.Stage)
             {
             case 0:
 			{
-				if (args.source && args.format)
+				if (Args.source && Args.format)
 				{
-					AddOp(FScheduledOp(item.At, item, 1),
-						FScheduledOp(args.source, item),
-						FScheduledOp(args.format, item));
+					AddOp(FScheduledOp(Item.At, Item, 1),
+						FScheduledOp(Args.source, Item),
+						FScheduledOp(Args.format, Item));
 				}
 				else
 				{
-					StoreMesh(item, nullptr);
+					StoreMesh(Item, nullptr);
 				}
 				break;
 			}
@@ -1605,20 +1605,20 @@ namespace mu
             {
             	MUTABLE_CPUPROFILER_SCOPE(ME_FORMAT_1)
             		
-                TSharedPtr<const FMesh> Source = LoadMesh(FCacheAddress(args.source,item));
-                TSharedPtr<const FMesh> Format = LoadMesh(FCacheAddress(args.format,item));
+                TSharedPtr<const FMesh> Source = LoadMesh(FCacheAddress(Args.source,Item));
+                TSharedPtr<const FMesh> Format = LoadMesh(FCacheAddress(Args.format,Item));
 
 				if (Source && Source->IsReference())
 				{
 					Release(Format);
-					StoreMesh(item, Source);
+					StoreMesh(Item, Source);
 				}
 				else if (Source)
 				{
-					uint8 Flags = args.Flags;
+					uint8 Flags = Args.Flags;
 					if (!Format && !(Flags & OP::MeshFormatArgs::ResetBufferIndices))
 					{
-						StoreMesh(item, Source);
+						StoreMesh(Item, Source);
 					}
 					else if (!Format)
 					{
@@ -1629,7 +1629,7 @@ namespace mu
 							Result->ResetBufferIndices();
 						}
 
-						StoreMesh(item, Result);
+						StoreMesh(Item, Result);
 					}
 					else
 					{
@@ -1658,13 +1658,13 @@ namespace mu
 
 						Release(Source);
 						Release(Format);
-						StoreMesh(item, Result);
+						StoreMesh(Item, Result);
 					}
 				}
 				else
 				{
 					Release(Format);
-					StoreMesh(item, nullptr);
+					StoreMesh(Item, nullptr);
 				}
                 break;
             }
@@ -1678,7 +1678,7 @@ namespace mu
 
         case EOpType::ME_EXTRACTLAYOUTBLOCK:
         {
-            const uint8* data = Program.GetOpArgsPointer(item.At);
+            const uint8* data = Program.GetOpArgsPointer(Item.At);
 
             OP::ADDRESS source;
             FMemory::Memcpy( &source, data, sizeof(OP::ADDRESS) );
@@ -1692,18 +1692,18 @@ namespace mu
 			FMemory::Memcpy( &blockCount, data, sizeof(uint16) );
             data += sizeof(uint16);
 
-            switch (item.Stage)
+            switch (Item.Stage)
             {
             case 0:
 			{
 				if (source)
 				{
-					AddOp(FScheduledOp(item.At, item, 1),
-						FScheduledOp(source, item));
+					AddOp(FScheduledOp(Item.At, Item, 1),
+						FScheduledOp(source, Item));
 				}
 				else
 				{
-					StoreMesh(item, nullptr);
+					StoreMesh(Item, nullptr);
 				}
 				break;
 			}
@@ -1711,7 +1711,7 @@ namespace mu
             {
 				MUTABLE_CPUPROFILER_SCOPE(ME_EXTRACTLAYOUTBLOCK_1)
 
-                TSharedPtr<const FMesh> Source = LoadMesh(FCacheAddress(source, item));
+                TSharedPtr<const FMesh> Source = LoadMesh(FCacheAddress(source, Item));
 
                 // Access with memcpy necessary for unaligned memory access issues.
                 uint64 blocks[512];
@@ -1735,17 +1735,17 @@ namespace mu
 					if (!bOutSuccess)
 					{
 						Release(Result);
-						StoreMesh(item, Source);
+						StoreMesh(Item, Source);
 					}
 					else
 					{
 						Release(Source);
-						StoreMesh(item, Result);
+						StoreMesh(Item, Result);
 					}
 				}
 				else
 				{
-					StoreMesh(item, nullptr);
+					StoreMesh(Item, nullptr);
 				}
                 break;
             }
@@ -1759,19 +1759,19 @@ namespace mu
 
         case EOpType::ME_TRANSFORM:
         {
-			OP::MeshTransformArgs args = Program.GetOpArgs<OP::MeshTransformArgs>(item.At);
-            switch (item.Stage)
+			OP::MeshTransformArgs Args = Program.GetOpArgs<OP::MeshTransformArgs>(Item.At);
+            switch (Item.Stage)
             {
             case 0:
 			{
-				if (args.source)
+				if (Args.source)
 				{
-					AddOp(FScheduledOp(item.At, item, 1),
-						FScheduledOp(args.source, item));
+					AddOp(FScheduledOp(Item.At, Item, 1),
+						FScheduledOp(Args.source, Item));
 				}
 				else
 				{
-					StoreMesh(item, nullptr);
+					StoreMesh(Item, nullptr);
 				}
 				break;
 			}
@@ -1779,9 +1779,9 @@ namespace mu
             {
            		MUTABLE_CPUPROFILER_SCOPE(ME_TRANSFORM_1)
             		
-                TSharedPtr<const FMesh> Source = LoadMesh(FCacheAddress(args.source,item));
+                TSharedPtr<const FMesh> Source = LoadMesh(FCacheAddress(Args.source,Item));
 
-                const FMatrix44f& mat = Program.ConstantMatrices[args.matrix];
+                const FMatrix44f& mat = Program.ConstantMatrices[Args.matrix];
 
 				TSharedPtr<FMesh> Result = CreateMesh(Source ? Source->GetDataSize() : 0);
 
@@ -1791,12 +1791,12 @@ namespace mu
 				if (!bOutSuccess)
 				{
 					Release(Result);
-					StoreMesh(item, Source);
+					StoreMesh(Item, Source);
 				}
 				else
 				{
 					Release(Source);
-					StoreMesh(item, Result);
+					StoreMesh(Item, Result);
 				}
                 break;
             }
@@ -1810,19 +1810,19 @@ namespace mu
 
         case EOpType::ME_CLIPMORPHPLANE:
         {
-			OP::MeshClipMorphPlaneArgs args = Program.GetOpArgs<OP::MeshClipMorphPlaneArgs>(item.At);
-            switch (item.Stage)
+			OP::MeshClipMorphPlaneArgs Args = Program.GetOpArgs<OP::MeshClipMorphPlaneArgs>(Item.At);
+            switch (Item.Stage)
             {
             case 0:
 			{
-				if (args.source)
+				if (Args.source)
 				{
-					AddOp(FScheduledOp(item.At, item, 1),
-						FScheduledOp(args.source, item));
+					AddOp(FScheduledOp(Item.At, Item, 1),
+						FScheduledOp(Args.source, Item));
 				}
 				else
 				{
-					StoreMesh(item, nullptr);
+					StoreMesh(Item, nullptr);
 				}
 				break;
 			}
@@ -1830,64 +1830,64 @@ namespace mu
             {
            		MUTABLE_CPUPROFILER_SCOPE(ME_CLIPMORPHPLANE_1)
             		
-                TSharedPtr<const FMesh> Source = LoadMesh(FCacheAddress(args.source, item));
+                TSharedPtr<const FMesh> Source = LoadMesh(FCacheAddress(Args.source, Item));
 
-                check(args.morphShape < (uint32)pModel->GetPrivate()->Program.ConstantShapes.Num());
+                check(Args.morphShape < (uint32)InModel->GetPrivate()->Program.ConstantShapes.Num());
 
                 // Should be an ellipse
-                const FShape& morphShape = Program.ConstantShapes[args.morphShape];
+                const FShape& morphShape = Program.ConstantShapes[Args.morphShape];
 
                 const FVector3f& origin = morphShape.position;
                 const FVector3f& normal = morphShape.up;
 
-				bool bRemoveFaceIfAllVerticesCulled = args.FaceCullStrategy==EFaceCullStrategy::AllVerticesCulled;
+				bool bRemoveFaceIfAllVerticesCulled = Args.FaceCullStrategy==EFaceCullStrategy::AllVerticesCulled;
 
-                if (args.VertexSelectionType == EClipVertexSelectionType::Shape)
+                if (Args.VertexSelectionType == EClipVertexSelectionType::Shape)
                 {
-                    check(args.vertexSelectionShapeOrBone < (uint32)pModel->GetPrivate()->Program.ConstantShapes.Num());
+                    check(Args.vertexSelectionShapeOrBone < (uint32)InModel->GetPrivate()->Program.ConstantShapes.Num());
 
                     // Should be None or an axis aligned box
-                    const FShape& selectionShape = Program.ConstantShapes[args.vertexSelectionShapeOrBone];
+                    const FShape& selectionShape = Program.ConstantShapes[Args.vertexSelectionShapeOrBone];
 
 					TSharedPtr<FMesh> Result = CreateMesh(Source ? Source->GetDataSize() : 0);
 
 					bool bOutSuccess = false;
-					MeshClipMorphPlane(Result.Get(), Source.Get(), origin, normal, args.dist, args.factor, morphShape.size[0], morphShape.size[1], morphShape.size[2], selectionShape, bRemoveFaceIfAllVerticesCulled, bOutSuccess, nullptr, -1);
+					MeshClipMorphPlane(Result.Get(), Source.Get(), origin, normal, Args.dist, Args.factor, morphShape.size[0], morphShape.size[1], morphShape.size[2], selectionShape, bRemoveFaceIfAllVerticesCulled, bOutSuccess, nullptr, -1);
 					
 					if (!bOutSuccess)
 					{
 						Release(Result);
-						StoreMesh(item, Source);
+						StoreMesh(Item, Source);
 					}
 					else
 					{
 						Release(Source);
-						StoreMesh(item, Result);
+						StoreMesh(Item, Result);
 					}
                 }
 
-				else if (args.VertexSelectionType == EClipVertexSelectionType::BoneHierarchy)
+				else if (Args.VertexSelectionType == EClipVertexSelectionType::BoneHierarchy)
 				{
 					FShape selectionShape;
 					selectionShape.type = (uint8)FShape::Type::None;
 
 					TSharedPtr<FMesh> Result = CreateMesh(Source->GetDataSize());
 
-					check(args.vertexSelectionShapeOrBone <= MAX_uint32);
-					const FBoneName Bone(args.vertexSelectionShapeOrBone);
+					check(Args.vertexSelectionShapeOrBone <= MAX_uint32);
+					const FBoneName Bone(Args.vertexSelectionShapeOrBone);
 
 					bool bOutSuccess = false;
-					MeshClipMorphPlane(Result.Get(), Source.Get(), origin, normal, args.dist, args.factor, morphShape.size[0], morphShape.size[1], morphShape.size[2], selectionShape, bRemoveFaceIfAllVerticesCulled, bOutSuccess, &Bone, args.maxBoneRadius);
+					MeshClipMorphPlane(Result.Get(), Source.Get(), origin, normal, Args.dist, Args.factor, morphShape.size[0], morphShape.size[1], morphShape.size[2], selectionShape, bRemoveFaceIfAllVerticesCulled, bOutSuccess, &Bone, Args.maxBoneRadius);
 
 					if (!bOutSuccess)
 					{
 						Release(Result);
-						StoreMesh(item, Source);
+						StoreMesh(Item, Source);
 					}
 					else
 					{
 						Release(Source);
-						StoreMesh(item, Result);
+						StoreMesh(Item, Result);
 					}
                 }
                 else
@@ -1899,17 +1899,17 @@ namespace mu
 					TSharedPtr<FMesh> Result = CreateMesh(Source ? Source->GetDataSize() : 0);
 
 					bool bOutSuccess = false;
-					MeshClipMorphPlane(Result.Get(), Source.Get(), origin, normal, args.dist, args.factor, morphShape.size[0], morphShape.size[1], morphShape.size[2], selectionShape, bRemoveFaceIfAllVerticesCulled, bOutSuccess, nullptr, -1.0f);
+					MeshClipMorphPlane(Result.Get(), Source.Get(), origin, normal, Args.dist, Args.factor, morphShape.size[0], morphShape.size[1], morphShape.size[2], selectionShape, bRemoveFaceIfAllVerticesCulled, bOutSuccess, nullptr, -1.0f);
 
 					if (!bOutSuccess)
 					{
 						Release(Result);
-						StoreMesh(item, Source);
+						StoreMesh(Item, Source);
 					}
 					else
 					{
 						Release(Source);
-						StoreMesh(item, Result);
+						StoreMesh(Item, Result);
 					}
                 }
 
@@ -1926,20 +1926,20 @@ namespace mu
 
         case EOpType::ME_CLIPWITHMESH:
         {
-			OP::MeshClipWithMeshArgs args = Program.GetOpArgs<OP::MeshClipWithMeshArgs>(item.At);
-            switch (item.Stage)
+			OP::MeshClipWithMeshArgs Args = Program.GetOpArgs<OP::MeshClipWithMeshArgs>(Item.At);
+            switch (Item.Stage)
             {
             case 0:
 			{
-				if (args.source)
+				if (Args.source)
 				{
-					AddOp(FScheduledOp(item.At, item, 1),
-						FScheduledOp(args.source, item),
-						FScheduledOp(args.clipMesh, item));
+					AddOp(FScheduledOp(Item.At, Item, 1),
+						FScheduledOp(Args.source, Item),
+						FScheduledOp(Args.clipMesh, Item));
 				}
 				else
 				{
-					StoreMesh(item, nullptr);
+					StoreMesh(Item, nullptr);
 				}
 
 				break;
@@ -1948,8 +1948,8 @@ namespace mu
             {
 				MUTABLE_CPUPROFILER_SCOPE(ME_CLIPWITHMESH_1)
 
-                TSharedPtr<const FMesh> Source = LoadMesh(FCacheAddress(args.source, item));
-                TSharedPtr<const FMesh> Clip = LoadMesh(FCacheAddress(args.clipMesh, item));
+                TSharedPtr<const FMesh> Source = LoadMesh(FCacheAddress(Args.source, Item));
+                TSharedPtr<const FMesh> Clip = LoadMesh(FCacheAddress(Args.clipMesh, Item));
 
                 // Only if both are valid.
                 if (Source && Clip)
@@ -1963,18 +1963,18 @@ namespace mu
 					if (!bOutSuccess)
 					{
 						Release(Result);
-						StoreMesh(item, Source);
+						StoreMesh(Item, Source);
 					}
 					else
 					{
 						Release(Source);
-						StoreMesh(item, Result);
+						StoreMesh(Item, Result);
 					}
                 }
                 else
                 {
 					Release(Clip);
-					StoreMesh(item, Source);
+					StoreMesh(Item, Source);
                 }
 
                 break;
@@ -1988,20 +1988,20 @@ namespace mu
         }
 		case EOpType::ME_CLIPDEFORM:
 		{
-			OP::MeshClipDeformArgs args = Program.GetOpArgs<OP::MeshClipDeformArgs>(item.At);
-			switch (item.Stage)
+			OP::MeshClipDeformArgs Args = Program.GetOpArgs<OP::MeshClipDeformArgs>(Item.At);
+			switch (Item.Stage)
 			{
 			case 0:
 			{
-				if (args.mesh)
+				if (Args.mesh)
 				{
-					AddOp(FScheduledOp(item.At, item, 1),
-						FScheduledOp(args.mesh, item),
-						FScheduledOp(args.clipShape, item));
+					AddOp(FScheduledOp(Item.At, Item, 1),
+						FScheduledOp(Args.mesh, Item),
+						FScheduledOp(Args.clipShape, Item));
 				}
 				else
 				{
-					StoreMesh(item, nullptr);
+					StoreMesh(Item, nullptr);
 				}
 				break;
 			}
@@ -2009,35 +2009,35 @@ namespace mu
 			{
 				MUTABLE_CPUPROFILER_SCOPE(ME_CLIPDEFORM_1)
 
-				TSharedPtr<const FMesh> BaseMesh = LoadMesh(FCacheAddress(args.mesh, item));
-				TSharedPtr<const FMesh> ClipShape = LoadMesh(FCacheAddress(args.clipShape, item));
+				TSharedPtr<const FMesh> BaseMesh = LoadMesh(FCacheAddress(Args.mesh, Item));
+				TSharedPtr<const FMesh> ClipShape = LoadMesh(FCacheAddress(Args.clipShape, Item));
 
 				if (BaseMesh && ClipShape)
 				{
 					TSharedPtr<FMesh> Result = CreateMesh(BaseMesh->GetDataSize());
 
-					bool bRemoveIfAllVerticesCulled = args.FaceCullStrategy == EFaceCullStrategy::AllVerticesCulled;
+					bool bRemoveIfAllVerticesCulled = Args.FaceCullStrategy == EFaceCullStrategy::AllVerticesCulled;
 
 					bool bOutSuccess = false;
-					MeshClipDeform(Result.Get(), BaseMesh.Get(), ClipShape.Get(), args.clipWeightThreshold, bRemoveIfAllVerticesCulled, bOutSuccess);
+					MeshClipDeform(Result.Get(), BaseMesh.Get(), ClipShape.Get(), Args.clipWeightThreshold, bRemoveIfAllVerticesCulled, bOutSuccess);
 
 					Release(ClipShape);
 
 					if (!bOutSuccess)
 					{
 						Release(Result);
-						StoreMesh(item, BaseMesh);
+						StoreMesh(Item, BaseMesh);
 					}
 					else
 					{
 						Release(BaseMesh);
-						StoreMesh(item, Result);
+						StoreMesh(Item, Result);
 					}
 				}
 				else
 				{
 					Release(ClipShape);
-					StoreMesh(item, BaseMesh);
+					StoreMesh(Item, BaseMesh);
 				}
 
 				break;
@@ -2052,20 +2052,20 @@ namespace mu
 
         case EOpType::ME_APPLYPOSE:
         {
-			OP::MeshApplyPoseArgs args = Program.GetOpArgs<OP::MeshApplyPoseArgs>(item.At);
-            switch (item.Stage)
+			OP::MeshApplyPoseArgs Args = Program.GetOpArgs<OP::MeshApplyPoseArgs>(Item.At);
+            switch (Item.Stage)
             {
             case 0:
 			{
-				if (args.base)
+				if (Args.base)
 				{
-					AddOp(FScheduledOp(item.At, item, 1),
-						FScheduledOp(args.base, item),
-						FScheduledOp(args.pose, item));
+					AddOp(FScheduledOp(Item.At, Item, 1),
+						FScheduledOp(Args.base, Item),
+						FScheduledOp(Args.pose, Item));
 				}
 				else
 				{
-					StoreMesh(item, nullptr);
+					StoreMesh(Item, nullptr);
 				}
 				break;
 			}
@@ -2073,8 +2073,8 @@ namespace mu
             {
           		MUTABLE_CPUPROFILER_SCOPE(ME_APPLYPOSE_1)
 
-                TSharedPtr<const FMesh> Base = LoadMesh(FCacheAddress(args.base, item));
-                TSharedPtr<const FMesh> pPose = LoadMesh(FCacheAddress(args.pose, item));
+                TSharedPtr<const FMesh> Base = LoadMesh(FCacheAddress(Args.base, Item));
+                TSharedPtr<const FMesh> pPose = LoadMesh(FCacheAddress(Args.pose, Item));
 
                 // Only if both are valid.
                 if (Base && pPose)
@@ -2088,18 +2088,18 @@ namespace mu
 					if (!bOutSuccess)
 					{
 						Release(Result);
-						StoreMesh(item, Base);
+						StoreMesh(Item, Base);
 					}
 					else
 					{
 						Release(Base);
-						StoreMesh(item, Result);
+						StoreMesh(Item, Result);
 					}
                 }
                 else
                 {
 					Release(pPose);
-					StoreMesh(item, Base);
+					StoreMesh(Item, Base);
                 }
 
                 break;
@@ -2114,36 +2114,36 @@ namespace mu
 
 		case EOpType::ME_BINDSHAPE:
 		{
-			OP::MeshBindShapeArgs Args = Program.GetOpArgs<OP::MeshBindShapeArgs>(item.At);
-			const uint8* Data = Program.GetOpArgsPointer(item.At);
+			OP::MeshBindShapeArgs Args = Program.GetOpArgs<OP::MeshBindShapeArgs>(Item.At);
+			const uint8* Data = Program.GetOpArgsPointer(Item.At);
 
-			switch (item.Stage)
+			switch (Item.Stage)
 			{
 			case 0:
 			{
 				if (Args.mesh)
 				{
-					AddOp(FScheduledOp(item.At, item, 1),
-						FScheduledOp(Args.mesh, item),
-						FScheduledOp(Args.shape, item));
+					AddOp(FScheduledOp(Item.At, Item, 1),
+						FScheduledOp(Args.mesh, Item),
+						FScheduledOp(Args.shape, Item));
 				}
 				else
 				{
-					StoreMesh(item, nullptr);
+					StoreMesh(Item, nullptr);
 				}
 				break;
 			}
 			case 1:
 			{
 				MUTABLE_CPUPROFILER_SCOPE(ME_BINDSHAPE_1)
-				TSharedPtr<const FMesh> BaseMesh = LoadMesh(FCacheAddress(Args.mesh, item));
-				TSharedPtr<const FMesh> Shape = LoadMesh(FCacheAddress(Args.shape, item));
+				TSharedPtr<const FMesh> BaseMesh = LoadMesh(FCacheAddress(Args.mesh, Item));
+				TSharedPtr<const FMesh> Shape = LoadMesh(FCacheAddress(Args.shape, Item));
 				
 				EShapeBindingMethod BindingMethod = static_cast<EShapeBindingMethod>(Args.bindingMethod); 
 
 				if (BindingMethod == EShapeBindingMethod::ReshapeClosestProject)
 				{ 
-					// Bones are stored after the args
+					// Bones are stored after the Args
 					Data += sizeof(Args);
 
 					// Rebuilding array of bone names ----
@@ -2181,7 +2181,7 @@ namespace mu
 					if (!bOutSuccess)
 					{
 						Release(BindMeshResult);
-						StoreMesh(item, BaseMesh);
+						StoreMesh(Item, BaseMesh);
 					}
 					else
 					{
@@ -2192,12 +2192,12 @@ namespace mu
 
 							Release(BaseMesh);
 							Release(BindMeshResult);
-							StoreMesh(item, BindMeshNoVertsResult);
+							StoreMesh(Item, BindMeshNoVertsResult);
 						}
 						else
 						{
 							Release(BaseMesh);
-							StoreMesh(item, BindMeshResult);
+							StoreMesh(Item, BindMeshResult);
 						}
 					}
 				}	
@@ -2212,12 +2212,12 @@ namespace mu
 					if (!bOutSuccess)
 					{
 						Release(Result);
-						StoreMesh(item, BaseMesh);
+						StoreMesh(Item, BaseMesh);
 					}
 					else
 					{
 						Release(BaseMesh);
-						StoreMesh(item, Result);
+						StoreMesh(Item, Result);
 					}
 				}
 
@@ -2234,20 +2234,20 @@ namespace mu
 
 		case EOpType::ME_APPLYSHAPE:
 		{
-			OP::MeshApplyShapeArgs args = Program.GetOpArgs<OP::MeshApplyShapeArgs>(item.At);
-			switch (item.Stage)
+			OP::MeshApplyShapeArgs Args = Program.GetOpArgs<OP::MeshApplyShapeArgs>(Item.At);
+			switch (Item.Stage)
 			{
 			case 0:
 			{
-				if (args.mesh)
+				if (Args.mesh)
 				{
-					AddOp(FScheduledOp(item.At, item, 1),
-						FScheduledOp(args.mesh, item),
-						FScheduledOp(args.shape, item));
+					AddOp(FScheduledOp(Item.At, Item, 1),
+						FScheduledOp(Args.mesh, Item),
+						FScheduledOp(Args.shape, Item));
 				}
 				else
 				{
-					StoreMesh(item, nullptr);
+					StoreMesh(Item, nullptr);
 				}
 				break;
 			}
@@ -2255,10 +2255,10 @@ namespace mu
 			{
 				MUTABLE_CPUPROFILER_SCOPE(ME_APPLYSHAPE_1)
 					
-				TSharedPtr<const FMesh> BaseMesh = LoadMesh(FCacheAddress(args.mesh, item));
-				TSharedPtr<const FMesh> Shape = LoadMesh(FCacheAddress(args.shape, item));
+				TSharedPtr<const FMesh> BaseMesh = LoadMesh(FCacheAddress(Args.mesh, Item));
+				TSharedPtr<const FMesh> Shape = LoadMesh(FCacheAddress(Args.shape, Item));
 
-				const EMeshBindShapeFlags ReshapeFlags = static_cast<EMeshBindShapeFlags>(args.flags);
+				const EMeshBindShapeFlags ReshapeFlags = static_cast<EMeshBindShapeFlags>(Args.flags);
 				const bool bReshapeVertices = EnumHasAnyFlags(ReshapeFlags, EMeshBindShapeFlags::ReshapeVertices);
 
 				TSharedPtr<FMesh> ReshapedMeshResult = CreateMesh(BaseMesh ? BaseMesh->GetDataSize() : 0);
@@ -2271,7 +2271,7 @@ namespace mu
 				if (!bOutSuccess)
 				{
 					Release(ReshapedMeshResult);
-					StoreMesh(item, BaseMesh);
+					StoreMesh(Item, BaseMesh);
 				}
 				else
 				{
@@ -2293,12 +2293,12 @@ namespace mu
 
 						Release(BaseMesh);
 						Release(ReshapedMeshResult);
-						StoreMesh(item, NoVerticesReshpedMesh);
+						StoreMesh(Item, NoVerticesReshpedMesh);
 					}
 					else
 					{
 						Release(BaseMesh);
-						StoreMesh(item, ReshapedMeshResult);
+						StoreMesh(Item, ReshapedMeshResult);
 					}
 				}
 				break;
@@ -2313,20 +2313,20 @@ namespace mu
 
 		case EOpType::ME_MORPHRESHAPE:
 		{
-			OP::MeshMorphReshapeArgs Args = Program.GetOpArgs<OP::MeshMorphReshapeArgs>(item.At);
-			switch(item.Stage)
+			OP::MeshMorphReshapeArgs Args = Program.GetOpArgs<OP::MeshMorphReshapeArgs>(Item.At);
+			switch(Item.Stage)
 			{
 			case 0:
 			{
 				if (Args.Morph)
 				{
-					AddOp(FScheduledOp(item.At, item, 1), 
-						FScheduledOp(Args.Morph, item),
-						FScheduledOp(Args.Reshape, item));
+					AddOp(FScheduledOp(Item.At, Item, 1), 
+						FScheduledOp(Args.Morph, Item),
+						FScheduledOp(Args.Reshape, Item));
 				}
 				else 
 				{
-					StoreMesh(item, nullptr);
+					StoreMesh(Item, nullptr);
 				}
 				break;
 			}
@@ -2334,8 +2334,8 @@ namespace mu
 			{
 				MUTABLE_CPUPROFILER_SCOPE(ME_MORPHRESHAPE_1)
 					
-				TSharedPtr<const FMesh> MorphedMesh = LoadMesh(FCacheAddress(Args.Morph, item));
-				TSharedPtr<const FMesh> ReshapeMesh = LoadMesh(FCacheAddress(Args.Reshape, item));
+				TSharedPtr<const FMesh> MorphedMesh = LoadMesh(FCacheAddress(Args.Morph, Item));
+				TSharedPtr<const FMesh> ReshapeMesh = LoadMesh(FCacheAddress(Args.Reshape, Item));
 
 				if (ReshapeMesh && MorphedMesh)
 				{
@@ -2354,11 +2354,11 @@ namespace mu
 
 					Release(MorphedMesh);
 					Release(ReshapeMesh);
-					StoreMesh(item, Result);
+					StoreMesh(Item, Result);
 				}
 				else
 				{
-					StoreMesh(item, MorphedMesh);
+					StoreMesh(Item, MorphedMesh);
 				}
 
 				break;
@@ -2373,20 +2373,20 @@ namespace mu
 
         case EOpType::ME_SETSKELETON:
         {
-			OP::MeshSetSkeletonArgs args = Program.GetOpArgs<OP::MeshSetSkeletonArgs>(item.At);
-            switch (item.Stage)
+			OP::MeshSetSkeletonArgs Args = Program.GetOpArgs<OP::MeshSetSkeletonArgs>(Item.At);
+            switch (Item.Stage)
             {
             case 0:
 			{
-				if (args.source)
+				if (Args.source)
 				{
-					AddOp(FScheduledOp(item.At, item, 1),
-						FScheduledOp(args.source, item),
-						FScheduledOp(args.skeleton, item));
+					AddOp(FScheduledOp(Item.At, Item, 1),
+						FScheduledOp(Args.source, Item),
+						FScheduledOp(Args.skeleton, Item));
 				}
 				else
 				{
-					StoreMesh(item, nullptr);
+					StoreMesh(Item, nullptr);
 				}
 				break;
 			}
@@ -2394,8 +2394,8 @@ namespace mu
             {
             	MUTABLE_CPUPROFILER_SCOPE(ME_SETSKELETON_1)
             		
-                TSharedPtr<const FMesh> Source = LoadMesh(FCacheAddress(args.source, item));
-                TSharedPtr<const FMesh> Skeleton = LoadMesh(FCacheAddress(args.skeleton, item));
+                TSharedPtr<const FMesh> Source = LoadMesh(FCacheAddress(Args.source, Item));
+                TSharedPtr<const FMesh> Skeleton = LoadMesh(FCacheAddress(Args.skeleton, Item));
 
                 // Only if both are valid.
                 if (Source && Skeleton)
@@ -2420,13 +2420,13 @@ namespace mu
                         if (!bOutSuccess)
                         {
 							Release(Result);
-							StoreMesh(item, Source);
+							StoreMesh(Item, Source);
                         }
 						else
 						{
 							//Result->GetPrivate()->CheckIntegrity();
 							Release(Source);
-							StoreMesh(item, Result);
+							StoreMesh(Item, Result);
 						}
                     }
                     else
@@ -2437,13 +2437,13 @@ namespace mu
 
 						//Result->GetPrivate()->CheckIntegrity();
 						Release(Skeleton);
-						StoreMesh(item, Result);
+						StoreMesh(Item, Result);
                     }
                 }
                 else
                 {
 					Release(Skeleton);
-					StoreMesh(item, Source);
+					StoreMesh(Item, Source);
                 }
 
                 break;
@@ -2462,7 +2462,7 @@ namespace mu
         		
             // Decode op
             // TODO: Partial decode for each stage
-            const uint8* data = Program.GetOpArgsPointer(item.At);
+            const uint8* data = Program.GetOpArgsPointer(Item.At);
 
             OP::ADDRESS source;
             FMemory::Memcpy(&source,data,sizeof(OP::ADDRESS)); 
@@ -2485,7 +2485,7 @@ namespace mu
 				FMemory::Memcpy(&condition,data,sizeof(OP::ADDRESS)); 
 				data += sizeof(OP::ADDRESS);
                 
-				conditions.Emplace(condition, item);
+				conditions.Emplace(condition, Item);
 
                 OP::ADDRESS mask;
 				FMemory::Memcpy(&mask,data,sizeof(OP::ADDRESS)); 
@@ -2496,18 +2496,18 @@ namespace mu
 
 
             // Schedule next stages
-            switch (item.Stage)
+            switch (Item.Stage)
             {
             case 0:
 			{
 				if (source)
 				{
 					// Request the conditions
-					AddOp(FScheduledOp(item.At, item, 1), conditions);
+					AddOp(FScheduledOp(Item.At, Item, 1), conditions);
 				}
 				else
 				{
-					StoreMesh(item, nullptr);
+					StoreMesh(Item, nullptr);
 				}
 				break;
 			}
@@ -2518,25 +2518,25 @@ namespace mu
                 // Request the source and the necessary masks
                 // \todo: store condition values in heap?
                 TArray<FScheduledOp> deps;
-                deps.Emplace( source, item );
+                deps.Emplace( source, Item );
                 for( int32 r=0; source && r<conditions.Num(); ++r )
                 {
                     // If there is no expression, we'll assume true.
                     bool value = true;
                     if (conditions[r].At)
                     {
-                        value = LoadBool(FCacheAddress(conditions[r].At, item));
+                        value = LoadBool(FCacheAddress(conditions[r].At, Item));
                     }
 
                     if (value)
                     {
-                        deps.Emplace(masks[r], item);
+                        deps.Emplace(masks[r], Item);
                     }
                 }
 
                 if (source)
                 {
-					AddOp(FScheduledOp(item.At, item, 2), deps);
+					AddOp(FScheduledOp(Item.At, Item, 2), deps);
 				}
                 break;
             }
@@ -2546,7 +2546,7 @@ namespace mu
             	MUTABLE_CPUPROFILER_SCOPE(ME_REMOVEMASK_2)
             	
                 // \todo: single remove operation with all masks?
-                TSharedPtr<const FMesh> Source = LoadMesh(FCacheAddress(source, item));
+                TSharedPtr<const FMesh> Source = LoadMesh(FCacheAddress(source, Item));
 
 				if (Source)
 				{
@@ -2561,12 +2561,12 @@ namespace mu
 						bool value = true;
 						if (conditions[r].At)
 						{
-							value = LoadBool(FCacheAddress(conditions[r].At, item));
+							value = LoadBool(FCacheAddress(conditions[r].At, Item));
 						}
 
 						if (value)
 						{
-							TSharedPtr<const FMesh> Mask = LoadMesh(FCacheAddress(masks[r], item));
+							TSharedPtr<const FMesh> Mask = LoadMesh(FCacheAddress(masks[r], Item));
 							if (Mask)
 							{
 								//MeshRemoveMask will make a copy of Result, try to make room for it. 
@@ -2591,11 +2591,11 @@ namespace mu
 						}
 					}
 
-					StoreMesh(item, Result);
+					StoreMesh(Item, Result);
 				}
 				else
 				{
-					StoreMesh(item, nullptr);
+					StoreMesh(Item, nullptr);
 				}
 
                 break;
@@ -2614,25 +2614,25 @@ namespace mu
 
 			// Decode op
 			// TODO: Partial decode for each stage
-			const uint8* Data = Program.GetOpArgsPointer(item.At);
+			const uint8* Data = Program.GetOpArgsPointer(Item.At);
 
 			OP::ADDRESS Source;
 			FMemory::Memcpy(&Source, Data, sizeof(OP::ADDRESS));
 			Data += sizeof(OP::ADDRESS);
 
 			// Schedule next stages
-			switch (item.Stage)
+			switch (Item.Stage)
 			{
 			case 0:
 			{
 				if (Source)
 				{
 					// Request the source
-					AddOp(FScheduledOp(item.At, item, 1), FScheduledOp(Source, item));
+					AddOp(FScheduledOp(Item.At, Item, 1), FScheduledOp(Source, Item));
 				}
 				else
 				{
-					StoreMesh(item, nullptr);
+					StoreMesh(Item, nullptr);
 				}
 				break;
 			}
@@ -2641,11 +2641,11 @@ namespace mu
 			{
 				MUTABLE_CPUPROFILER_SCOPE(ME_ADDTAGS_2)
 
-				TSharedPtr<const FMesh> SourceMesh = LoadMesh(FCacheAddress(Source, item));
+				TSharedPtr<const FMesh> SourceMesh = LoadMesh(FCacheAddress(Source, Item));
 
 				if (!SourceMesh)
 				{
-					StoreMesh(item, nullptr);
+					StoreMesh(Item, nullptr);
 				}
 				else
 				{
@@ -2664,12 +2664,12 @@ namespace mu
 						FMemory::Memcpy(&TagConstant, Data, sizeof(OP::ADDRESS));
 						Data += sizeof(OP::ADDRESS);
 
-						check(TagConstant < (uint32)pModel->GetPrivate()->Program.ConstantStrings.Num());
+						check(TagConstant < (uint32)InModel->GetPrivate()->Program.ConstantStrings.Num());
 						const FString& Name = Program.ConstantStrings[TagConstant];
 						Result->Tags[FirstMeshTagIndex+TagIndex] = Name;
 					}
 
-					StoreMesh(item, Result);
+					StoreMesh(Item, Result);
 				}
 
 				break;
@@ -2684,20 +2684,20 @@ namespace mu
 
         case EOpType::ME_PROJECT:
         {
-			OP::MeshProjectArgs args = Program.GetOpArgs<OP::MeshProjectArgs>(item.At);
-            switch (item.Stage)
+			OP::MeshProjectArgs Args = Program.GetOpArgs<OP::MeshProjectArgs>(Item.At);
+            switch (Item.Stage)
             {
             case 0:
 			{
-				if (args.mesh)
+				if (Args.mesh)
 				{
-					AddOp(FScheduledOp(item.At, item, 1),
-						FScheduledOp(args.mesh, item),
-						FScheduledOp(args.projector, item));
+					AddOp(FScheduledOp(Item.At, Item, 1),
+						FScheduledOp(Args.mesh, Item),
+						FScheduledOp(Args.projector, Item));
 				}
 				else
 				{
-					StoreMesh(item, nullptr);
+					StoreMesh(Item, nullptr);
 				}
 				break;
 			}
@@ -2705,8 +2705,8 @@ namespace mu
             {
 				MUTABLE_CPUPROFILER_SCOPE(ME_PROJECT_1)
 
-                TSharedPtr<const FMesh> Mesh = LoadMesh(FCacheAddress(args.mesh,item));
-                const FProjector Projector = LoadProjector(FCacheAddress(args.projector, item));
+                TSharedPtr<const FMesh> Mesh = LoadMesh(FCacheAddress(Args.mesh,Item));
+                const FProjector Projector = LoadProjector(FCacheAddress(Args.projector, Item));
 
                 // Only if both are valid.
                 if (Mesh && Mesh->GetVertexBuffers().GetBufferCount() > 0)
@@ -2719,19 +2719,19 @@ namespace mu
 					if (!bOutSuccess)
 					{
 						Release(Result);
-						StoreMesh(item, Mesh);
+						StoreMesh(Item, Mesh);
 					}
 					else
 					{	
 //						Result->GetPrivate()->CheckIntegrity();
 						Release(Mesh);
-						StoreMesh(item, Result);
+						StoreMesh(Item, Result);
 					}
                 }
 				else
 				{
 					Release(Mesh);
-					StoreMesh(item, nullptr);
+					StoreMesh(Item, nullptr);
 				}
 
                 break;
@@ -2746,18 +2746,18 @@ namespace mu
 
 		case EOpType::ME_OPTIMIZESKINNING:
 		{
-			OP::MeshOptimizeSkinningArgs args = Program.GetOpArgs<OP::MeshOptimizeSkinningArgs>(item.At);
-			switch (item.Stage)
+			OP::MeshOptimizeSkinningArgs Args = Program.GetOpArgs<OP::MeshOptimizeSkinningArgs>(Item.At);
+			switch (Item.Stage)
 			{
 			case 0:
 			{
-				if (args.source)
+				if (Args.source)
 				{
-					AddOp(FScheduledOp(item.At, item, 1), FScheduledOp(args.source, item));
+					AddOp(FScheduledOp(Item.At, Item, 1), FScheduledOp(Args.source, Item));
 				}
 				else
 				{
-					StoreMesh(item, nullptr);
+					StoreMesh(Item, nullptr);
 				}
 				break;
 			}
@@ -2765,11 +2765,11 @@ namespace mu
 			{
 				MUTABLE_CPUPROFILER_SCOPE(ME_OPTIMIZESKINNING_1)
 
-				TSharedPtr<const FMesh> Source = LoadMesh(FCacheAddress(args.source, item));
+				TSharedPtr<const FMesh> Source = LoadMesh(FCacheAddress(Args.source, Item));
 
 				if (Source && Source->IsReference())
 				{
-					StoreMesh(item, Source);
+					StoreMesh(Item, Source);
 				}
 
 				TSharedPtr<FMesh> Result = CreateMesh();
@@ -2780,12 +2780,12 @@ namespace mu
 				if (!bOutSuccess)
 				{
 					Release(Result);
-					StoreMesh(item, Source);
+					StoreMesh(Item, Source);
 				}
 				else
 				{
 					Release(Source);
-					StoreMesh(item, Result);
+					StoreMesh(Item, Result);
 				}
 
 				break;
@@ -2800,21 +2800,21 @@ namespace mu
 
         case EOpType::ME_TRANSFORMWITHMESH:
 		{
-			OP::MeshTransformWithinMeshArgs args = Program.GetOpArgs<OP::MeshTransformWithinMeshArgs>(item.At);
-			switch (item.Stage)
+			OP::MeshTransformWithinMeshArgs Args = Program.GetOpArgs<OP::MeshTransformWithinMeshArgs>(Item.At);
+			switch (Item.Stage)
 			{
 			case 0:
 			{
-				if (args.sourceMesh)
+				if (Args.sourceMesh)
 				{
-					AddOp(FScheduledOp(item.At, item, 1),
-						FScheduledOp(args.sourceMesh, item),
-						FScheduledOp(args.boundingMesh, item),
-						FScheduledOp(args.matrix, item));
+					AddOp(FScheduledOp(Item.At, Item, 1),
+						FScheduledOp(Args.sourceMesh, Item),
+						FScheduledOp(Args.boundingMesh, Item),
+						FScheduledOp(Args.matrix, Item));
 				}
 				else
 				{
-					StoreMesh(item, nullptr);
+					StoreMesh(Item, nullptr);
 				}
 				break;
 			}
@@ -2822,9 +2822,9 @@ namespace mu
 			{
 				MUTABLE_CPUPROFILER_SCOPE(ME_TRANSFORMWITHMESH_1)
             
-				TSharedPtr<const FMesh> SourceMesh = LoadMesh(FCacheAddress(args.sourceMesh,item));
-				TSharedPtr<const FMesh> BoundingMesh = LoadMesh(FCacheAddress(args.boundingMesh, item));
-				const FMatrix44f& Transform = LoadMatrix(FCacheAddress(args.matrix, item));
+				TSharedPtr<const FMesh> SourceMesh = LoadMesh(FCacheAddress(Args.sourceMesh,Item));
+				TSharedPtr<const FMesh> BoundingMesh = LoadMesh(FCacheAddress(Args.boundingMesh, Item));
+				const FMatrix44f& Transform = LoadMatrix(FCacheAddress(Args.matrix, Item));
 
 				if (SourceMesh)
 				{
@@ -2837,18 +2837,18 @@ namespace mu
 					if (!bOutSuccess)
 					{
 						Release(Result);
-						StoreMesh(item, SourceMesh);
+						StoreMesh(Item, SourceMesh);
 					}
 					else
 					{
 						Release(SourceMesh);
-						StoreMesh(item, Result);
+						StoreMesh(Item, Result);
 					}
 				}
 				else
 				{
 					Release(BoundingMesh);
-					StoreMesh(item, SourceMesh);
+					StoreMesh(Item, SourceMesh);
 				}
 				break;
 			}
@@ -2870,7 +2870,7 @@ namespace mu
     }
 
     //---------------------------------------------------------------------------------------------
-    void CodeRunner::RunCode_Image(const FScheduledOp& item, const FParameters* pParams, const FModel* pModel )
+    void CodeRunner::RunCode_Image(const FScheduledOp& Item, const FParameters* pParams, const FModel* InModel )
     {
 		MUTABLE_CPUPROFILER_SCOPE(RunCode_Image);
 
@@ -2878,20 +2878,20 @@ namespace mu
 
 		const FProgram& Program = Model->GetPrivate()->Program;
 
-		EOpType type = Program.GetOpType(item.At);
+		EOpType type = Program.GetOpType(Item.At);
 		switch (type)
         {
 
         case EOpType::IM_LAYERCOLOUR:
         {
-			OP::ImageLayerColourArgs args = Program.GetOpArgs<OP::ImageLayerColourArgs>(item.At);
-            switch (item.Stage)
+			OP::ImageLayerColourArgs Args = Program.GetOpArgs<OP::ImageLayerColourArgs>(Item.At);
+            switch (Item.Stage)
             {
             case 0:
-                    AddOp( FScheduledOp( item.At, item, 1),
-                           FScheduledOp( args.base, item),
-                           FScheduledOp::FromOpAndOptions( args.colour, item, 0),
-                           FScheduledOp( args.mask, item) );
+                    AddOp( FScheduledOp( Item.At, Item, 1),
+                           FScheduledOp( Args.base, Item),
+                           FScheduledOp::FromOpAndOptions( Args.colour, Item, 0),
+                           FScheduledOp( Args.mask, Item) );
                 break;
 
             case 1:
@@ -2908,22 +2908,22 @@ namespace mu
 
         case EOpType::IM_LAYER:
         {
-			OP::ImageLayerArgs args = Program.GetOpArgs<OP::ImageLayerArgs>(item.At);
+			OP::ImageLayerArgs Args = Program.GetOpArgs<OP::ImageLayerArgs>(Item.At);
 
 			if (ExecutionStrategy == EExecutionStrategy::MinimizeMemory)
 			{
-				switch (item.Stage)
+				switch (Item.Stage)
 				{
 				case 0:
-					AddOp(FScheduledOp(item.At, item, 1),
-						FScheduledOp(args.base, item));
+					AddOp(FScheduledOp(Item.At, Item, 1),
+						FScheduledOp(Args.base, Item));
 					break;
 
 				case 1:
 					// Request the rest of the data.
-					AddOp(FScheduledOp(item.At, item, 2),
-						FScheduledOp(args.blended, item),
-						FScheduledOp(args.mask, item));
+					AddOp(FScheduledOp(Item.At, Item, 2),
+						FScheduledOp(Args.blended, Item),
+						FScheduledOp(Args.mask, Item));
 					break;
 
 				case 2:
@@ -2937,13 +2937,13 @@ namespace mu
 			}
 			else
 			{
-				switch (item.Stage)
+				switch (Item.Stage)
 				{
 				case 0:
-					AddOp(FScheduledOp(item.At, item, 1),
-						FScheduledOp(args.base, item),
-						FScheduledOp(args.blended, item),
-						FScheduledOp(args.mask, item));
+					AddOp(FScheduledOp(Item.At, Item, 1),
+						FScheduledOp(Args.base, Item),
+						FScheduledOp(Args.blended, Item),
+						FScheduledOp(Args.mask, Item));
 					break;
 
 				case 1:
@@ -2961,13 +2961,13 @@ namespace mu
 
         case EOpType::IM_MULTILAYER:
         {
-			OP::ImageMultiLayerArgs args = Program.GetOpArgs<OP::ImageMultiLayerArgs>(item.At);
-            switch (item.Stage)
+			OP::ImageMultiLayerArgs Args = Program.GetOpArgs<OP::ImageMultiLayerArgs>(Item.At);
+            switch (Item.Stage)
             {
             case 0:
-                AddOp( FScheduledOp( item.At, item, 1),
-                       FScheduledOp( args.rangeSize, item ),
-					   FScheduledOp(args.base, item));
+                AddOp( FScheduledOp( Item.At, Item, 1),
+                       FScheduledOp( Args.rangeSize, Item ),
+					   FScheduledOp(Args.base, Item));
                 break;
 
             case 1:
@@ -2976,29 +2976,29 @@ namespace mu
             		
                 // We now know the number of iterations
                 int32 Iterations = 0;
-                if (args.rangeSize)
+                if (Args.rangeSize)
                 {
-                    FCacheAddress RangeAddress(args.rangeSize,item);
+                    FCacheAddress RangeAddress(Args.rangeSize,Item);
 
                     // We support both integers and scalars here, which is not common.
                     // \todo: review if this is necessary or we can enforce it at compile time.
-                    DATATYPE RangeSizeType = GetOpDataType( pModel->GetPrivate()->Program.GetOpType(args.rangeSize) );
-                    if (RangeSizeType == DT_INT)
+                    EDataType RangeSizeType = GetOpDataType( InModel->GetPrivate()->Program.GetOpType(Args.rangeSize) );
+                    if (RangeSizeType == EDataType::Int)
                     {
 						Iterations = LoadInt(RangeAddress);
                     }
-                    else if (RangeSizeType == DT_SCALAR)
+                    else if (RangeSizeType == EDataType::Scalar)
                     {
 						Iterations = int32( LoadScalar(RangeAddress) );
                     }
                 }
 
-				TSharedPtr<const FImage> Base = LoadImage(FCacheAddress(args.base, item));
+				TSharedPtr<const FImage> Base = LoadImage(FCacheAddress(Args.base, Item));
 
 				if (Iterations <= 0)
 				{
 					// There are no layers: return the base
-					StoreImage(item, Base);
+					StoreImage(Item, Base);
 				}
 				else
 				{
@@ -3033,11 +3033,11 @@ namespace mu
 
 					// Request the first layer
 					int32 CurrentIteration = 0;
-					FScheduledOp ItemCopy = item;
-					ExecutionIndex Index = GetMemory().GetRangeIndex(item.ExecutionIndex);
-					Index.SetFromModelRangeIndex(args.rangeId, CurrentIteration);
+					FScheduledOp ItemCopy = Item;
+					ExecutionIndex Index = GetMemory().GetRangeIndex(Item.ExecutionIndex);
+					Index.SetFromModelRangeIndex(Args.rangeId, CurrentIteration);
 					ItemCopy.ExecutionIndex = GetMemory().GetRangeIndexIndex(Index);
-					AddOp(FScheduledOp(item.At, item, 2, DataPos), FScheduledOp(args.base, item), FScheduledOp(args.blended, ItemCopy), FScheduledOp(args.mask, ItemCopy));
+					AddOp(FScheduledOp(Item.At, Item, 2, DataPos), FScheduledOp(Args.base, Item), FScheduledOp(Args.blended, ItemCopy), FScheduledOp(Args.mask, ItemCopy));
 				}
 
                 break;
@@ -3047,10 +3047,10 @@ namespace mu
             {
             	MUTABLE_CPUPROFILER_SCOPE(IM_MULTILAYER_default)
 
-				FScheduledOpData& Data = HeapData[item.CustomState];
+				FScheduledOpData& Data = HeapData[Item.CustomState];
 
 				int32 Iterations = Data.MultiLayer.Iterations;
-				int32 CurrentIteration = item.Stage - 2;
+				int32 CurrentIteration = Item.Stage - 2;
 				check(CurrentIteration >= 0 && CurrentIteration < 120);
 
 				TRACE_CPUPROFILER_EVENT_SCOPE_TEXT(*FString::Printf(TEXT("Layer %d of %d"), CurrentIteration, Iterations));
@@ -3059,15 +3059,15 @@ namespace mu
 
 				TSharedPtr<FImage> Base = StaticCastSharedPtr<FImage>(ConstCastSharedPtr<FResource>(Data.Resource));
  
-                FScheduledOp itemCopy = item;
-                ExecutionIndex index = GetMemory().GetRangeIndex( item.ExecutionIndex );
+                FScheduledOp itemCopy = Item;
+                ExecutionIndex index = GetMemory().GetRangeIndex( Item.ExecutionIndex );
 				
                 {
-                    index.SetFromModelRangeIndex( args.rangeId, CurrentIteration);
+                    index.SetFromModelRangeIndex( Args.rangeId, CurrentIteration);
                     itemCopy.ExecutionIndex = GetMemory().GetRangeIndexIndex(index);
 					itemCopy.CustomState = 0;
 
-                    TSharedPtr<const FImage> Blended = LoadImage( FCacheAddress(args.blended,itemCopy) );
+                    TSharedPtr<const FImage> Blended = LoadImage( FCacheAddress(Args.blended,itemCopy) );
 
                     // This shouldn't happen in optimised models, but it could happen in editors, etc.
                     // \todo: raise a performance warning?
@@ -3111,11 +3111,11 @@ namespace mu
 					bool bBlendOnlyOneMip = Data.MultiLayer.bBlendOnlyOneMip;
 					bool bUseBlendSourceFromBlendAlpha = false; // (Args.flags& OP::ImageLayerArgs::F_BLENDED_RGB_FROM_ALPHA) != 0;
 
-					if (!args.mask && args.bUseMaskFromBlended
+					if (!Args.mask && Args.bUseMaskFromBlended
 						&&
-						args.blendType == uint8(EBlendType::BT_BLEND)
+						Args.blendType == uint8(EBlendType::BT_BLEND)
 						&&
-						args.blendTypeAlpha == uint8(EBlendType::BT_LIGHTEN) )
+						Args.blendTypeAlpha == uint8(EBlendType::BT_LIGHTEN) )
 					{
 						// This is a frequent critical-path case because of multilayer projectors.
 						bDone = true;
@@ -3123,17 +3123,17 @@ namespace mu
 						constexpr bool bUseVectorImpl = false;
 						if constexpr (bUseVectorImpl)
 						{
-							BufferLayerCompositeVector<VectorBlendChannelMasked, VectorLightenChannel, false>(Base.Get(), Blended.Get(), bBlendOnlyOneMip, args.BlendAlphaSourceChannel);
+							BufferLayerCompositeVector<VectorBlendChannelMasked, VectorLightenChannel, false>(Base.Get(), Blended.Get(), bBlendOnlyOneMip, Args.BlendAlphaSourceChannel);
 						}
 						else
 						{
-							BufferLayerComposite<BlendChannelMasked, LightenChannel, false>(Base.Get(), Blended.Get(), bBlendOnlyOneMip, args.BlendAlphaSourceChannel);
+							BufferLayerComposite<BlendChannelMasked, LightenChannel, false>(Base.Get(), Blended.Get(), bBlendOnlyOneMip, Args.BlendAlphaSourceChannel);
 						}
 					}
 
-                    if (!bDone && args.mask)
+                    if (!bDone && Args.mask)
                     {
-                        TSharedPtr<const FImage> Mask = LoadImage( FCacheAddress(args.mask,itemCopy) );
+                        TSharedPtr<const FImage> Mask = LoadImage( FCacheAddress(Args.mask,itemCopy) );
 
 						// TODO: This shouldn't happen, but be defensive.
 						if (Mask && Mask->GetSize() != ResultSize)
@@ -3149,7 +3149,7 @@ namespace mu
 						// Not implemented yet
 						check(!bUseBlendSourceFromBlendAlpha);
 
-                        switch (EBlendType(args.blendType))
+                        switch (EBlendType(Args.blendType))
                         {
 						case EBlendType::BT_NORMAL_COMBINE: check(false); break;
                         case EBlendType::BT_SOFTLIGHT: BufferLayer<SoftLightChannelMasked, SoftLightChannel, false>(Base.Get(), Base.Get(), Mask.Get(), Blended.Get(), bApplyColorBlendToAlpha, bBlendOnlyOneMip); break;
@@ -3166,12 +3166,12 @@ namespace mu
 
 						Release(Mask);
                     }
-					else if (!bDone && args.bUseMaskFromBlended)
+					else if (!bDone && Args.bUseMaskFromBlended)
 					{
 						// Not implemented yet
 						check(!bUseBlendSourceFromBlendAlpha);
 
-						switch (EBlendType(args.blendType))
+						switch (EBlendType(Args.blendType))
 						{
 						case EBlendType::BT_NORMAL_COMBINE: check(false); break;
 						case EBlendType::BT_SOFTLIGHT: BufferLayerEmbeddedMask<SoftLightChannelMasked, SoftLightChannel, false>(Base.Get(), Base.Get(), Blended.Get(), bApplyColorBlendToAlpha, bBlendOnlyOneMip); break;
@@ -3188,7 +3188,7 @@ namespace mu
 					}
                     else if (!bDone)
                     {
-                        switch (EBlendType(args.blendType))
+                        switch (EBlendType(Args.blendType))
                         {
 						case EBlendType::BT_NORMAL_COMBINE: check(false); break;
                         case EBlendType::BT_SOFTLIGHT: BufferLayer<SoftLightChannel, false>(Base.Get(), Base.Get(), Blended.Get(), bApplyColorBlendToAlpha, bBlendOnlyOneMip, bUseBlendSourceFromBlendAlpha); break;
@@ -3205,20 +3205,20 @@ namespace mu
                     }
 
 					// Apply the separate blend operation for alpha
-					if (!bDone && !bApplyColorBlendToAlpha && args.blendTypeAlpha != uint8(EBlendType::BT_NONE) )
+					if (!bDone && !bApplyColorBlendToAlpha && Args.blendTypeAlpha != uint8(EBlendType::BT_NONE) )
 					{
 						// Separate alpha operation ignores the mask.
-						switch (EBlendType(args.blendTypeAlpha))
+						switch (EBlendType(Args.blendTypeAlpha))
 						{
-						case EBlendType::BT_SOFTLIGHT: BufferLayerInPlace<SoftLightChannel, false, 1>(Base.Get(), Blended.Get(), bBlendOnlyOneMip, 3, args.BlendAlphaSourceChannel); break;
-						case EBlendType::BT_HARDLIGHT: BufferLayerInPlace<HardLightChannel, false, 1>(Base.Get(), Blended.Get(), bBlendOnlyOneMip, 3, args.BlendAlphaSourceChannel); break;
-						case EBlendType::BT_BURN: BufferLayerInPlace<BurnChannel, false, 1>(Base.Get(), Blended.Get(), bBlendOnlyOneMip, 3, args.BlendAlphaSourceChannel); break;
-						case EBlendType::BT_DODGE: BufferLayerInPlace<DodgeChannel, false, 1>(Base.Get(), Blended.Get(), bBlendOnlyOneMip, 3, args.BlendAlphaSourceChannel); break;
-						case EBlendType::BT_SCREEN: BufferLayerInPlace<ScreenChannel, false, 1>(Base.Get(), Blended.Get(), bBlendOnlyOneMip, 3, args.BlendAlphaSourceChannel); break;
-						case EBlendType::BT_OVERLAY: BufferLayerInPlace<OverlayChannel, false, 1>(Base.Get(), Blended.Get(), bBlendOnlyOneMip, 3, args.BlendAlphaSourceChannel); break;
-						case EBlendType::BT_LIGHTEN: BufferLayerInPlace<LightenChannel, false, 1>(Base.Get(), Blended.Get(), bBlendOnlyOneMip, 3, args.BlendAlphaSourceChannel); break;
-						case EBlendType::BT_MULTIPLY: BufferLayerInPlace<MultiplyChannel, false, 1>(Base.Get(), Blended.Get(), bBlendOnlyOneMip, 3, args.BlendAlphaSourceChannel); break;
-						case EBlendType::BT_BLEND: BufferLayerInPlace<BlendChannel, false, 1>(Base.Get(), Blended.Get(), bBlendOnlyOneMip, 3, args.BlendAlphaSourceChannel); break;
+						case EBlendType::BT_SOFTLIGHT: BufferLayerInPlace<SoftLightChannel, false, 1>(Base.Get(), Blended.Get(), bBlendOnlyOneMip, 3, Args.BlendAlphaSourceChannel); break;
+						case EBlendType::BT_HARDLIGHT: BufferLayerInPlace<HardLightChannel, false, 1>(Base.Get(), Blended.Get(), bBlendOnlyOneMip, 3, Args.BlendAlphaSourceChannel); break;
+						case EBlendType::BT_BURN: BufferLayerInPlace<BurnChannel, false, 1>(Base.Get(), Blended.Get(), bBlendOnlyOneMip, 3, Args.BlendAlphaSourceChannel); break;
+						case EBlendType::BT_DODGE: BufferLayerInPlace<DodgeChannel, false, 1>(Base.Get(), Blended.Get(), bBlendOnlyOneMip, 3, Args.BlendAlphaSourceChannel); break;
+						case EBlendType::BT_SCREEN: BufferLayerInPlace<ScreenChannel, false, 1>(Base.Get(), Blended.Get(), bBlendOnlyOneMip, 3, Args.BlendAlphaSourceChannel); break;
+						case EBlendType::BT_OVERLAY: BufferLayerInPlace<OverlayChannel, false, 1>(Base.Get(), Blended.Get(), bBlendOnlyOneMip, 3, Args.BlendAlphaSourceChannel); break;
+						case EBlendType::BT_LIGHTEN: BufferLayerInPlace<LightenChannel, false, 1>(Base.Get(), Blended.Get(), bBlendOnlyOneMip, 3, Args.BlendAlphaSourceChannel); break;
+						case EBlendType::BT_MULTIPLY: BufferLayerInPlace<MultiplyChannel, false, 1>(Base.Get(), Blended.Get(), bBlendOnlyOneMip, 3, Args.BlendAlphaSourceChannel); break;
+						case EBlendType::BT_BLEND: BufferLayerInPlace<BlendChannel, false, 1>(Base.Get(), Blended.Get(), bBlendOnlyOneMip, 3, Args.BlendAlphaSourceChannel); break;
 						default: check(false);
 						}
 					}
@@ -3239,18 +3239,18 @@ namespace mu
 					// TODO: Reconvert to OriginalBaseFormat if necessary?
 
 					Data.Resource = nullptr;
-					StoreImage(item, Base);
+					StoreImage(Item, Base);
 					break;
 				}
 				else
 				{
 					// Request a new layer
 					++CurrentIteration;
-					FScheduledOp ItemCopy = item;
-					ExecutionIndex Index = GetMemory().GetRangeIndex(item.ExecutionIndex);
-					Index.SetFromModelRangeIndex(args.rangeId, CurrentIteration);
+					FScheduledOp ItemCopy = Item;
+					ExecutionIndex Index = GetMemory().GetRangeIndex(Item.ExecutionIndex);
+					Index.SetFromModelRangeIndex(Args.rangeId, CurrentIteration);
 					ItemCopy.ExecutionIndex = GetMemory().GetRangeIndexIndex(Index);
-					AddOp(FScheduledOp(item.At, item, 2+CurrentIteration, item.CustomState), FScheduledOp(args.blended, ItemCopy), FScheduledOp(args.mask, ItemCopy));
+					AddOp(FScheduledOp(Item.At, Item, 2+CurrentIteration, Item.CustomState), FScheduledOp(Args.blended, ItemCopy), FScheduledOp(Args.mask, ItemCopy));
 
 				}
 
@@ -3264,19 +3264,19 @@ namespace mu
 
 		case EOpType::IM_NORMALCOMPOSITE:
 		{
-			OP::ImageNormalCompositeArgs args = Program.GetOpArgs<OP::ImageNormalCompositeArgs>(item.At);
-			switch (item.Stage)
+			OP::ImageNormalCompositeArgs Args = Program.GetOpArgs<OP::ImageNormalCompositeArgs>(Item.At);
+			switch (Item.Stage)
 			{
 			case 0:
-				if (args.base && args.normal)
+				if (Args.base && Args.normal)
 				{
-					AddOp(FScheduledOp(item.At, item, 1),
-							FScheduledOp(args.base, item),
-							FScheduledOp(args.normal, item));
+					AddOp(FScheduledOp(Item.At, Item, 1),
+							FScheduledOp(Args.base, Item),
+							FScheduledOp(Args.normal, Item));
 				}
 				else
 				{
-					StoreImage(item, nullptr);
+					StoreImage(Item, nullptr);
 				}
 				break;
 
@@ -3284,8 +3284,8 @@ namespace mu
 			{
 				MUTABLE_CPUPROFILER_SCOPE(IM_NORMALCOMPOSITE_1)
 
-				TSharedPtr<const FImage> Base = LoadImage(FCacheAddress(args.base, item));
-				TSharedPtr<const FImage> Normal = LoadImage(FCacheAddress(args.normal, item));
+				TSharedPtr<const FImage> Base = LoadImage(FCacheAddress(Args.base, Item));
+				TSharedPtr<const FImage> Normal = LoadImage(FCacheAddress(Args.normal, Item));
 
 				if (Normal->GetLODCount() < Base->GetLODCount())
 				{
@@ -3304,11 +3304,11 @@ namespace mu
 
 
                 TSharedPtr<FImage> Result = CreateImage(Base->GetSizeX(), Base->GetSizeY(), Base->GetLODCount(), Base->GetFormat(), EInitializationType::NotInitialized);
-				ImageNormalComposite(Result.Get(), Base.Get(), Normal.Get(), args.mode, args.power);
+				ImageNormalComposite(Result.Get(), Base.Get(), Normal.Get(), Args.mode, Args.power);
 
 				Release(Base);
 				Release(Normal);
-				StoreImage(item, Result);
+				StoreImage(Item, Result);
 				break;
 			}
 
@@ -3321,11 +3321,11 @@ namespace mu
 
         case EOpType::IM_PIXELFORMAT:
         {
-			OP::ImagePixelFormatArgs args = Program.GetOpArgs<OP::ImagePixelFormatArgs>(item.At);
-            switch (item.Stage)
+			OP::ImagePixelFormatArgs Args = Program.GetOpArgs<OP::ImagePixelFormatArgs>(Item.At);
+            switch (Item.Stage)
             {
             case 0:
-                AddOp( FScheduledOp( item.At, item, 1), FScheduledOp( args.source, item) );
+                AddOp( FScheduledOp( Item.At, Item, 1), FScheduledOp( Args.source, Item) );
                 break;
 
             case 1:
@@ -3344,11 +3344,11 @@ namespace mu
 
         case EOpType::IM_MIPMAP:
         {
-			OP::ImageMipmapArgs args = Program.GetOpArgs<OP::ImageMipmapArgs>(item.At);
-            switch (item.Stage)
+			OP::ImageMipmapArgs Args = Program.GetOpArgs<OP::ImageMipmapArgs>(Item.At);
+            switch (Item.Stage)
             {
             case 0:
-				AddOp( FScheduledOp( item.At, item, 1), FScheduledOp( args.source, item) );
+				AddOp( FScheduledOp( Item.At, Item, 1), FScheduledOp( Args.source, Item) );
                 break;
 
             case 1:
@@ -3365,11 +3365,11 @@ namespace mu
 
         case EOpType::IM_RESIZE:
         {
-			OP::ImageResizeArgs args = Program.GetOpArgs<OP::ImageResizeArgs>(item.At);
-            switch (item.Stage)
+			OP::ImageResizeArgs Args = Program.GetOpArgs<OP::ImageResizeArgs>(Item.At);
+            switch (Item.Stage)
             {
             case 0:
-                AddOp( FScheduledOp( item.At, item, 1), FScheduledOp( args.source, item) );
+                AddOp( FScheduledOp( Item.At, Item, 1), FScheduledOp( Args.source, Item) );
                 break;
 
             case 1:
@@ -3387,21 +3387,21 @@ namespace mu
 
         case EOpType::IM_RESIZELIKE:
         {
-			OP::ImageResizeLikeArgs args = Program.GetOpArgs<OP::ImageResizeLikeArgs>(item.At);
-            switch (item.Stage)
+			OP::ImageResizeLikeArgs Args = Program.GetOpArgs<OP::ImageResizeLikeArgs>(Item.At);
+            switch (Item.Stage)
             {
             case 0:
-                AddOp(FScheduledOp(item.At, item, 1),
-                      	FScheduledOp(args.source, item),
-                        FScheduledOp(args.sizeSource, item));
+                AddOp(FScheduledOp(Item.At, Item, 1),
+                      	FScheduledOp(Args.source, Item),
+                        FScheduledOp(Args.sizeSource, Item));
                 break;
 
             case 1:
             {
             	MUTABLE_CPUPROFILER_SCOPE(IM_RESIZELIKE_1)
             	
-                TSharedPtr<const FImage> Base = LoadImage( FCacheAddress(args.source,item) );
-                TSharedPtr<const FImage> SizeBase = LoadImage( FCacheAddress(args.sizeSource,item) );
+                TSharedPtr<const FImage> Base = LoadImage( FCacheAddress(Args.source,Item) );
+                TSharedPtr<const FImage> SizeBase = LoadImage( FCacheAddress(Args.sizeSource,Item) );
 				FImageSize DestSize = SizeBase->GetSize();
 				Release(SizeBase);
 
@@ -3427,11 +3427,11 @@ namespace mu
 						ImOp.ImageMipmap(Settings.ImageCompressionQuality, Result.Get(), Result.Get(), 0, LevelCount, MipSettings);
                     }				
 
-					StoreImage(item, Result);
+					StoreImage(Item, Result);
 				}
                 else
                 {
-					StoreImage(item, Base);
+					StoreImage(Item, Base);
 				}
 				
                 break;
@@ -3446,11 +3446,11 @@ namespace mu
 
         case EOpType::IM_RESIZEREL:
         {
-			OP::ImageResizeRelArgs args = Program.GetOpArgs<OP::ImageResizeRelArgs>(item.At);
-            switch (item.Stage)
+			OP::ImageResizeRelArgs Args = Program.GetOpArgs<OP::ImageResizeRelArgs>(Item.At);
+            switch (Item.Stage)
             {
             case 0:
-                AddOp( FScheduledOp( item.At, item, 1), FScheduledOp( args.source, item) );
+                AddOp( FScheduledOp( Item.At, Item, 1), FScheduledOp( Args.source, Item) );
                 break;
 
             case 1:
@@ -3468,18 +3468,18 @@ namespace mu
         }
         case EOpType::IM_BLANKLAYOUT:
         {
-			OP::ImageBlankLayoutArgs Args = Program.GetOpArgs<OP::ImageBlankLayoutArgs>(item.At);
-            switch (item.Stage)
+			OP::ImageBlankLayoutArgs Args = Program.GetOpArgs<OP::ImageBlankLayoutArgs>(Item.At);
+            switch (Item.Stage)
             {
             case 0:
-                AddOp( FScheduledOp( item.At, item, 1), FScheduledOp::FromOpAndOptions(Args.layout, item, 0));
+                AddOp( FScheduledOp( Item.At, Item, 1), FScheduledOp::FromOpAndOptions(Args.layout, Item, 0));
                 break;
 
             case 1:
             {
             	MUTABLE_CPUPROFILER_SCOPE(IM_BLANKLAYOUT_1)
             		
-                TSharedPtr<const FLayout> pLayout = LoadLayout(FScheduledOp::FromOpAndOptions(Args.layout, item, 0));
+                TSharedPtr<const FLayout> pLayout = LoadLayout(FScheduledOp::FromOpAndOptions(Args.layout, Item, 0));
 
                 FIntPoint SizeInBlocks = pLayout->GetGridSize();
 
@@ -3490,7 +3490,7 @@ namespace mu
 				int32 FullImageMipCount = FImage::GetMipmapCount(FullImageSizeInPixels.X, FullImageSizeInPixels.Y);
 
 				FIntPoint ImageSizeInPixels = FullImageSizeInPixels;
-				int32 MipsToSkip = item.ExecutionOptions;
+				int32 MipsToSkip = Item.ExecutionOptions;
 				MipsToSkip = FMath::Min(MipsToSkip, FullImageMipCount);
 				if (MipsToSkip > 0)
 				{
@@ -3500,7 +3500,7 @@ namespace mu
 					// If we want to generate only a subset of mipmaps, reduce the layout block size accordingly.
 					//ReducedBlockSizeInPixels.X = BlockSizeInPixels.X >> MipsToSkip;
 					//ReducedBlockSizeInPixels.Y = BlockSizeInPixels.Y >> MipsToSkip;
-					//const FImageFormatData& FormatData = GetImageFormatData((EImageFormat)args.format);
+					//const FImageFormatData& FormatData = GetImageFormatData((EImageFormat)Args.format);
 					//int32 MinBlockSize = FMath::Max(FormatData.PixelsPerBlockX, FormatData.PixelsPerBlockY);
 					//ReducedBlockSizeInPixels.X = FMath::Max<int32>(ReducedBlockSizeInPixels.X, FormatData.PixelsPerBlockX);
 					//ReducedBlockSizeInPixels.Y = FMath::Max<int32>(ReducedBlockSizeInPixels.Y, FormatData.PixelsPerBlockY);
@@ -3530,7 +3530,7 @@ namespace mu
 
 				// It needs to be initialized in case it has gaps.
                 TSharedPtr<FImage> New = CreateImage(ImageSizeInPixels.X, ImageSizeInPixels.Y, MipsToGenerate, EImageFormat(Args.format), EInitializationType::Black );
-                StoreImage(item, New);
+                StoreImage(Item, New);
                 break;
             }
 
@@ -3544,19 +3544,19 @@ namespace mu
 
         case EOpType::IM_COMPOSE:
         {
-			OP::ImageComposeArgs Args = Program.GetOpArgs<OP::ImageComposeArgs>(item.At);
+			OP::ImageComposeArgs Args = Program.GetOpArgs<OP::ImageComposeArgs>(Item.At);
 
 			if (ExecutionStrategy == EExecutionStrategy::MinimizeMemory)
 			{
-            	switch (item.Stage)
+            	switch (Item.Stage)
 				{
 				case 0:
-					AddOp(FScheduledOp(item.At, item, 1), FScheduledOp::FromOpAndOptions(Args.layout, item, 0));
+					AddOp(FScheduledOp(Item.At, Item, 1), FScheduledOp::FromOpAndOptions(Args.layout, Item, 0));
 					break;
 				case 1:
 				{
 					TSharedPtr<const FLayout> ComposeLayout = 
-							LoadLayout(FCacheAddress(Args.layout, FScheduledOp::FromOpAndOptions(Args.layout, item, 0)));
+							LoadLayout(FCacheAddress(Args.layout, FScheduledOp::FromOpAndOptions(Args.layout, Item, 0)));
 
 					FScheduledOpData Data;
 					Data.Resource = ComposeLayout;
@@ -3566,21 +3566,21 @@ namespace mu
 
 					if (RelBlockIndex >= 0)
 					{
-						AddOp(FScheduledOp(item.At, item, 2, DataPos), FScheduledOp(Args.base, item));
+						AddOp(FScheduledOp(Item.At, Item, 2, DataPos), FScheduledOp(Args.base, Item));
 					}
 					else
 					{
 						// Jump directly to stage 3, no need to load mask or blockImage.
-						AddOp(FScheduledOp(item.At, item, 3, DataPos), FScheduledOp(Args.base, item));
+						AddOp(FScheduledOp(Item.At, Item, 3, DataPos), FScheduledOp(Args.base, Item));
 					}
 
 					break;
 				}
 				case 2:
 				{
-					AddOp(FScheduledOp(item.At, item, 3, item.CustomState),
-						  FScheduledOp(Args.blockImage, item),
-						  FScheduledOp(Args.mask, item));
+					AddOp(FScheduledOp(Item.At, Item, 3, Item.CustomState),
+						  FScheduledOp(Args.blockImage, Item),
+						  FScheduledOp(Args.mask, Item));
 					break;
 				}
 
@@ -3595,16 +3595,16 @@ namespace mu
 			}
 			else
 			{
-            	switch (item.Stage)
+            	switch (Item.Stage)
 				{
 				case 0:
-					AddOp(FScheduledOp(item.At, item, 1), FScheduledOp::FromOpAndOptions(Args.layout, item, 0));
+					AddOp(FScheduledOp(Item.At, Item, 1), FScheduledOp::FromOpAndOptions(Args.layout, Item, 0));
 					break;
 
 				case 1:
 				{	
 					TSharedPtr<const FLayout> ComposeLayout = 
-							LoadLayout(FCacheAddress(Args.layout, FScheduledOp::FromOpAndOptions(Args.layout, item, 0)));
+							LoadLayout(FCacheAddress(Args.layout, FScheduledOp::FromOpAndOptions(Args.layout, Item, 0)));
 
 					FScheduledOpData Data;
 					Data.Resource = ComposeLayout;
@@ -3613,14 +3613,14 @@ namespace mu
 					int32 RelBlockIndex = ComposeLayout->FindBlock(Args.BlockId);
 					if (RelBlockIndex >= 0)
 					{
-						AddOp(FScheduledOp(item.At, item, 2, DataPos),
-							  FScheduledOp(Args.base, item),
-							  FScheduledOp(Args.blockImage, item),
-							  FScheduledOp(Args.mask, item));
+						AddOp(FScheduledOp(Item.At, Item, 2, DataPos),
+							  FScheduledOp(Args.base, Item),
+							  FScheduledOp(Args.blockImage, Item),
+							  FScheduledOp(Args.mask, Item));
 					}
 					else
 					{
-						AddOp(FScheduledOp(item.At, item, 2, DataPos), FScheduledOp(Args.base, item));
+						AddOp(FScheduledOp(Item.At, Item, 2, DataPos), FScheduledOp(Args.base, Item));
 					}
 					break;
 				}
@@ -3640,12 +3640,12 @@ namespace mu
 
         case EOpType::IM_INTERPOLATE:
         {
-			OP::ImageInterpolateArgs args = Program.GetOpArgs<OP::ImageInterpolateArgs>(item.At);
-            switch (item.Stage)
+			OP::ImageInterpolateArgs Args = Program.GetOpArgs<OP::ImageInterpolateArgs>(Item.At);
+            switch (Item.Stage)
             {
             case 0:
-                AddOp( FScheduledOp( item.At, item, 1),
-                       FScheduledOp( args.factor, item) );
+                AddOp( FScheduledOp( Item.At, Item, 1),
+                       FScheduledOp( Args.factor, Item) );
                 break;
 
             case 1:
@@ -3655,13 +3655,13 @@ namespace mu
                 // Targets must be consecutive
                 int32 count = 0;
                 for ( int32 i=0
-                    ; i<MUTABLE_OP_MAX_INTERPOLATE_COUNT && args.targets[i]
+                    ; i<MUTABLE_OP_MAX_INTERPOLATE_COUNT && Args.targets[i]
                     ; ++i )
                 {
                     count++;
                 }
 
-                float factor = LoadScalar( FCacheAddress(args.factor,item) );
+                float factor = LoadScalar( FCacheAddress(Args.factor,Item) );
 
                 float delta = 1.0f/(count-1);
                 int32 min = (int32)floorf( factor/delta );
@@ -3677,19 +3677,19 @@ namespace mu
 
                 if ( bifactor < UE_SMALL_NUMBER )
                 {
-                    AddOp( FScheduledOp( item.At, item, 2, dataPos),
-                            FScheduledOp( args.targets[min], item) );
+                    AddOp( FScheduledOp( Item.At, Item, 2, dataPos),
+                            FScheduledOp( Args.targets[min], Item) );
                 }
                 else if ( bifactor > 1.0f-UE_SMALL_NUMBER )
                 {
-                    AddOp( FScheduledOp( item.At, item, 2, dataPos),
-                            FScheduledOp( args.targets[max], item) );
+                    AddOp( FScheduledOp( Item.At, Item, 2, dataPos),
+                            FScheduledOp( Args.targets[max], Item) );
                 }
                 else
                 {
-                    AddOp( FScheduledOp( item.At, item, 2, dataPos),
-                            FScheduledOp( args.targets[min], item),
-                            FScheduledOp( args.targets[max], item) );
+                    AddOp( FScheduledOp( Item.At, Item, 2, dataPos),
+                            FScheduledOp( Args.targets[min], Item),
+                            FScheduledOp( Args.targets[max], Item) );
                 }
                 break;
             }
@@ -3701,32 +3701,32 @@ namespace mu
                 // Targets must be consecutive
                 int32 count = 0;
                 for ( int32 i=0
-                    ; i<MUTABLE_OP_MAX_INTERPOLATE_COUNT && args.targets[i]
+                    ; i<MUTABLE_OP_MAX_INTERPOLATE_COUNT && Args.targets[i]
                     ; ++i )
                 {
                     count++;
                 }
 
                 // Factor from 0 to 1 between the two targets
-                const FScheduledOpData& Data = HeapData[item.CustomState];
+                const FScheduledOpData& Data = HeapData[Item.CustomState];
                 float Bifactor = Data.Interpolate.Bifactor;
                 int32 Min = Data.Interpolate.Min;
                 int32 Max = Data.Interpolate.Max;
 
                 if (Bifactor < UE_SMALL_NUMBER)
                 {
-                    TSharedPtr<const FImage> Source = LoadImage(FCacheAddress(args.targets[Min], item));
-					StoreImage(item, Source);
+                    TSharedPtr<const FImage> Source = LoadImage(FCacheAddress(Args.targets[Min], Item));
+					StoreImage(Item, Source);
 				}
                 else if (Bifactor > 1.0f - UE_SMALL_NUMBER)
                 {
-                    TSharedPtr<const FImage> Source = LoadImage(FCacheAddress(args.targets[Max], item));
-					StoreImage(item, Source);
+                    TSharedPtr<const FImage> Source = LoadImage(FCacheAddress(Args.targets[Max], Item));
+					StoreImage(Item, Source);
 				}
                 else
                 {
-					TSharedPtr<const FImage> pMin = LoadImage(FCacheAddress(args.targets[Min], item));
-                    TSharedPtr<const FImage> pMax = LoadImage(FCacheAddress(args.targets[Max], item));
+					TSharedPtr<const FImage> pMin = LoadImage(FCacheAddress(Args.targets[Min], Item));
+                    TSharedPtr<const FImage> pMax = LoadImage(FCacheAddress(Args.targets[Max], Item));
 
                     if (pMin && pMax)
                     {						
@@ -3790,15 +3790,15 @@ namespace mu
                         ImageInterpolate(pNew.Get(), pMax.Get(), Bifactor);
 
 						Release(pMax);
-						StoreImage(item, pNew);
+						StoreImage(Item, pNew);
 					}
                     else if (pMin)
                     {
-						StoreImage(item, pMin);
+						StoreImage(Item, pMin);
 					}
                     else if (pMax)
                     {
-						StoreImage(item, pMax);
+						StoreImage(Item, pMax);
 					}
 				}
 
@@ -3814,13 +3814,13 @@ namespace mu
 
         case EOpType::IM_SATURATE:
         {
-			OP::ImageSaturateArgs args = Program.GetOpArgs<OP::ImageSaturateArgs>(item.At);
-            switch (item.Stage)
+			OP::ImageSaturateArgs Args = Program.GetOpArgs<OP::ImageSaturateArgs>(Item.At);
+            switch (Item.Stage)
             {
             case 0:
-                AddOp( FScheduledOp( item.At, item, 1),
-                        FScheduledOp( args.base, item ),
-                        FScheduledOp::FromOpAndOptions( args.factor, item, 0 ));
+                AddOp( FScheduledOp( Item.At, Item, 1),
+                        FScheduledOp( Args.base, Item ),
+                        FScheduledOp::FromOpAndOptions( Args.factor, Item, 0 ));
                 break;
 
             case 1:
@@ -3838,25 +3838,25 @@ namespace mu
 
         case EOpType::IM_LUMINANCE:
         {
-			OP::ImageLuminanceArgs args = Program.GetOpArgs<OP::ImageLuminanceArgs>(item.At);
-            switch (item.Stage)
+			OP::ImageLuminanceArgs Args = Program.GetOpArgs<OP::ImageLuminanceArgs>(Item.At);
+            switch (Item.Stage)
             {
             case 0:
-                AddOp( FScheduledOp( item.At, item, 1),
-                        FScheduledOp( args.base, item ) );
+                AddOp( FScheduledOp( Item.At, Item, 1),
+                        FScheduledOp( Args.base, Item ) );
                 break;
 
             case 1:
             {
            		MUTABLE_CPUPROFILER_SCOPE(IM_LUMINANCE_1)
             		
-                TSharedPtr<const FImage> Base = LoadImage( FCacheAddress(args.base,item) );
+                TSharedPtr<const FImage> Base = LoadImage( FCacheAddress(Args.base,Item) );
 
 				TSharedPtr<FImage> Result = CreateImage(Base->GetSizeX(), Base->GetSizeY(), Base->GetLODCount(), EImageFormat::L_UByte, EInitializationType::NotInitialized);
                 ImageLuminance( Result.Get(),Base.Get() );
 
 				Release(Base);
-				StoreImage( item, Result );
+				StoreImage( Item, Result );
                 break;
             }
 
@@ -3869,15 +3869,15 @@ namespace mu
 
         case EOpType::IM_SWIZZLE:
         {
-			OP::ImageSwizzleArgs args = Program.GetOpArgs<OP::ImageSwizzleArgs>(item.At);
-            switch (item.Stage)
+			OP::ImageSwizzleArgs Args = Program.GetOpArgs<OP::ImageSwizzleArgs>(Item.At);
+            switch (Item.Stage)
             {
             case 0:
-                AddOp( FScheduledOp( item.At, item, 1),
-                        FScheduledOp( args.sources[0], item ),
-                        FScheduledOp( args.sources[1], item ),
-                        FScheduledOp( args.sources[2], item ),
-                        FScheduledOp( args.sources[3], item ) );
+                AddOp( FScheduledOp( Item.At, Item, 1),
+                        FScheduledOp( Args.sources[0], Item ),
+                        FScheduledOp( Args.sources[1], Item ),
+                        FScheduledOp( Args.sources[2], Item ),
+                        FScheduledOp( Args.sources[3], Item ) );
                 break;
 
             case 1:
@@ -3894,23 +3894,23 @@ namespace mu
 
         case EOpType::IM_COLOURMAP:
         {
-			OP::ImageColourMapArgs args = Program.GetOpArgs<OP::ImageColourMapArgs>(item.At);
-            switch (item.Stage)
+			OP::ImageColourMapArgs Args = Program.GetOpArgs<OP::ImageColourMapArgs>(Item.At);
+            switch (Item.Stage)
             {
             case 0:
-                    AddOp( FScheduledOp( item.At, item, 1),
-                           FScheduledOp( args.base, item ),
-                           FScheduledOp( args.mask, item ),
-                           FScheduledOp( args.map, item ) );
+                    AddOp( FScheduledOp( Item.At, Item, 1),
+                           FScheduledOp( Args.base, Item ),
+                           FScheduledOp( Args.mask, Item ),
+                           FScheduledOp( Args.map, Item ) );
                 break;
 
             case 1:
             {
             	MUTABLE_CPUPROFILER_SCOPE(IM_COLOURMAP_1)
             		
-                TSharedPtr<const FImage> Source = LoadImage( FCacheAddress(args.base,item) );
-                TSharedPtr<const FImage> Mask = LoadImage( FCacheAddress(args.mask,item) );
-                TSharedPtr<const FImage> Map = LoadImage( FCacheAddress(args.map,item) );
+                TSharedPtr<const FImage> Source = LoadImage( FCacheAddress(Args.base,Item) );
+                TSharedPtr<const FImage> Mask = LoadImage( FCacheAddress(Args.mask,Item) );
+                TSharedPtr<const FImage> Map = LoadImage( FCacheAddress(Args.map,Item) );
 
 				bool bOnlyOneMip = (Mask->GetLODCount() < Source->GetLODCount());
 
@@ -3937,7 +3937,7 @@ namespace mu
 				Release(Source);
 				Release(Mask);
 				Release(Map);
-				StoreImage( item, Result );
+				StoreImage( Item, Result );
                 break;
             }
 
@@ -3950,26 +3950,26 @@ namespace mu
 
         case EOpType::IM_GRADIENT:
         {
-			OP::ImageGradientArgs args = Program.GetOpArgs<OP::ImageGradientArgs>(item.At);
-            switch (item.Stage)
+			OP::ImageGradientArgs Args = Program.GetOpArgs<OP::ImageGradientArgs>(Item.At);
+            switch (Item.Stage)
             {
             case 0:
-                AddOp( FScheduledOp( item.At, item, 1),
-                        FScheduledOp::FromOpAndOptions( args.colour0, item, 0 ),
-                        FScheduledOp::FromOpAndOptions( args.colour1, item, 0 ) );
+                AddOp( FScheduledOp( Item.At, Item, 1),
+                        FScheduledOp::FromOpAndOptions( Args.colour0, Item, 0 ),
+                        FScheduledOp::FromOpAndOptions( Args.colour1, Item, 0 ) );
                 break;
 
             case 1:
             {
 				MUTABLE_CPUPROFILER_SCOPE(IM_GRADIENT_1)
             		
-				FVector4f colour0 = LoadColor(FScheduledOp::FromOpAndOptions(args.colour0, item, 0));
-				FVector4f colour1 = LoadColor(FScheduledOp::FromOpAndOptions(args.colour1, item, 0));
+				FVector4f colour0 = LoadColor(FScheduledOp::FromOpAndOptions(Args.colour0, Item, 0));
+				FVector4f colour1 = LoadColor(FScheduledOp::FromOpAndOptions(Args.colour1, Item, 0));
 
-				TSharedPtr<FImage> Result = CreateImage(args.size[0], args.size[1], 1, EImageFormat::RGB_UByte, EInitializationType::NotInitialized);
+				TSharedPtr<FImage> Result = CreateImage(Args.size[0], Args.size[1], 1, EImageFormat::RGB_UByte, EInitializationType::NotInitialized);
                 ImageGradient( Result.Get(), colour0, colour1 );
 
-				StoreImage( item, Result );
+				StoreImage( Item, Result );
                 break;
             }
 
@@ -3982,28 +3982,28 @@ namespace mu
 
         case EOpType::IM_BINARISE:
         {
-			OP::ImageBinariseArgs args = Program.GetOpArgs<OP::ImageBinariseArgs>(item.At);
-            switch (item.Stage)
+			OP::ImageBinariseArgs Args = Program.GetOpArgs<OP::ImageBinariseArgs>(Item.At);
+            switch (Item.Stage)
             {
             case 0:
-                AddOp( FScheduledOp( item.At, item, 1),
-                        FScheduledOp( args.base, item ),
-                        FScheduledOp::FromOpAndOptions( args.threshold, item, 0 ) );
+                AddOp( FScheduledOp( Item.At, Item, 1),
+                        FScheduledOp( Args.base, Item ),
+                        FScheduledOp::FromOpAndOptions( Args.threshold, Item, 0 ) );
                 break;
 
             case 1:
             {
             	MUTABLE_CPUPROFILER_SCOPE(IM_BINARISE_1)
             		
-                TSharedPtr<const FImage> pA = LoadImage( FCacheAddress(args.base,item) );
+                TSharedPtr<const FImage> pA = LoadImage( FCacheAddress(Args.base,Item) );
 
-                float c = LoadScalar(FScheduledOp::FromOpAndOptions(args.threshold, item, 0));
+                float c = LoadScalar(FScheduledOp::FromOpAndOptions(Args.threshold, Item, 0));
 
                 TSharedPtr<FImage> Result = CreateImage(pA->GetSizeX(), pA->GetSizeY(), pA->GetLODCount(), EImageFormat::L_UByte, EInitializationType::NotInitialized);
 				ImageBinarise( Result.Get(), pA.Get(), c );
 
 				Release(pA);
-				StoreImage( item, Result );
+				StoreImage( Item, Result );
                 break;
             }
 
@@ -4016,11 +4016,11 @@ namespace mu
 
 		case EOpType::IM_INVERT:
 		{
-			OP::ImageInvertArgs args = Program.GetOpArgs<OP::ImageInvertArgs>(item.At);
-			switch (item.Stage)
+			OP::ImageInvertArgs Args = Program.GetOpArgs<OP::ImageInvertArgs>(Item.At);
+			switch (Item.Stage)
 			{
 			case 0:
-				AddOp(FScheduledOp(item.At, item, 1), FScheduledOp(args.base, item));
+				AddOp(FScheduledOp(Item.At, Item, 1), FScheduledOp(Args.base, Item));
 				break;
 
 			case 1:
@@ -4038,22 +4038,22 @@ namespace mu
 
         case EOpType::IM_PLAINCOLOUR:
         {
-			OP::ImagePlainColourArgs args = Program.GetOpArgs<OP::ImagePlainColourArgs>(item.At);
-            switch (item.Stage)
+			OP::ImagePlainColourArgs Args = Program.GetOpArgs<OP::ImagePlainColourArgs>(Item.At);
+            switch (Item.Stage)
             {
             case 0:
-				AddOp( FScheduledOp( item.At, item, 1), FScheduledOp::FromOpAndOptions( args.colour, item, 0 ) );
+				AddOp( FScheduledOp( Item.At, Item, 1), FScheduledOp::FromOpAndOptions( Args.colour, Item, 0 ) );
                 break;
 
             case 1:
             {
             	MUTABLE_CPUPROFILER_SCOPE(IM_PLAINCOLOUR_1)
             		
-				FVector4f c = LoadColor(FScheduledOp::FromOpAndOptions(args.colour, item, 0));
+				FVector4f c = LoadColor(FScheduledOp::FromOpAndOptions(Args.colour, Item, 0));
 
-				uint16 SizeX = args.size[0];
-				uint16 SizeY = args.size[1];
-				int32 LODs = args.LODs;
+				uint16 SizeX = Args.size[0];
+				uint16 SizeY = Args.size[1];
+				int32 LODs = Args.LODs;
 				
 				// This means all the mip chain
 				if (LODs == 0)
@@ -4061,18 +4061,18 @@ namespace mu
 					LODs = FMath::CeilLogTwo(FMath::Max(SizeX,SizeY));
 				}
 
-				for (int32 l=0; l<item.ExecutionOptions; ++l)
+				for (int32 l=0; l<Item.ExecutionOptions; ++l)
 				{
 					SizeX = FMath::Max(uint16(1), FMath::DivideAndRoundUp(SizeX, uint16(2)));
 					SizeY = FMath::Max(uint16(1), FMath::DivideAndRoundUp(SizeY, uint16(2)));
 					--LODs;
 				}
 
-                TSharedPtr<FImage> pA = CreateImage( SizeX, SizeY, FMath::Max(LODs,1), EImageFormat(args.format), EInitializationType::NotInitialized );
+                TSharedPtr<FImage> pA = CreateImage( SizeX, SizeY, FMath::Max(LODs,1), EImageFormat(Args.format), EInitializationType::NotInitialized );
 
 				ImOp.FillColor(pA.Get(), c);
 
-				StoreImage( item, pA );
+				StoreImage( Item, pA );
                 break;
             }
 
@@ -4085,8 +4085,8 @@ namespace mu
 
 		case EOpType::IM_REFERENCE:
 		{
-			OP::ResourceReferenceArgs Args = Program.GetOpArgs<OP::ResourceReferenceArgs>(item.At);
-			switch (item.Stage)
+			OP::ResourceReferenceArgs Args = Program.GetOpArgs<OP::ResourceReferenceArgs>(Item.At);
+			switch (Item.Stage)
 			{
 			case 0:
 			{
@@ -4100,7 +4100,7 @@ namespace mu
 				{
 					Result = FImage::CreateAsReference(Args.ID, Args.ImageDesc, false);
 				}
-				StoreImage(item, Result);
+				StoreImage(Item, Result);
 				break;
 			}
 
@@ -4113,27 +4113,27 @@ namespace mu
 
         case EOpType::IM_CROP:
         {
-			OP::ImageCropArgs args = Program.GetOpArgs<OP::ImageCropArgs>(item.At);
-            switch (item.Stage)
+			OP::ImageCropArgs Args = Program.GetOpArgs<OP::ImageCropArgs>(Item.At);
+            switch (Item.Stage)
             {
             case 0:
-                AddOp( FScheduledOp( item.At, item, 1), FScheduledOp( args.source, item ) );
+                AddOp( FScheduledOp( Item.At, Item, 1), FScheduledOp( Args.source, Item ) );
                 break;
 
             case 1:
             {
             	MUTABLE_CPUPROFILER_SCOPE(IM_CROP_1)
             		
-                TSharedPtr<const FImage> pA = LoadImage( FCacheAddress(args.source,item) );
+                TSharedPtr<const FImage> pA = LoadImage( FCacheAddress(Args.source,Item) );
 
                 box< UE::Math::TIntVector2<int32> > rect;
-                rect.min[0] = args.minX;
-                rect.min[1] = args.minY;
-                rect.size[0] = args.sizeX;
-                rect.size[1] = args.sizeY;
+                rect.min[0] = Args.minX;
+                rect.min[1] = Args.minY;
+                rect.size[0] = Args.sizeX;
+                rect.size[1] = Args.sizeY;
 
 				// Apply ther mipmap reduction to the crop rectangle.
-				int32 MipsToSkip = item.ExecutionOptions;
+				int32 MipsToSkip = Item.ExecutionOptions;
 				while ( MipsToSkip>0 && rect.size[0]>0 && rect.size[1]>0 )
 				{
 					rect.min[0] /= 2;
@@ -4151,7 +4151,7 @@ namespace mu
 				}
 
 				Release(pA);
-				StoreImage( item, Result );
+				StoreImage( Item, Result );
                 break;
             }
 
@@ -4165,37 +4165,37 @@ namespace mu
         case EOpType::IM_PATCH:
         {
 			// TODO: This is optimized for memory-usage but base and patch could be requested at the same time
-			OP::ImagePatchArgs args = Program.GetOpArgs<OP::ImagePatchArgs>(item.At);
-            switch (item.Stage)
+			OP::ImagePatchArgs Args = Program.GetOpArgs<OP::ImagePatchArgs>(Item.At);
+            switch (Item.Stage)
             {
 			case 0:
-				AddOp(FScheduledOp(item.At, item, 1), FScheduledOp(args.base, item));
+				AddOp(FScheduledOp(Item.At, Item, 1), FScheduledOp(Args.base, Item));
 				break;
 
 			case 1:
-				AddOp(FScheduledOp(item.At, item, 2), FScheduledOp(args.patch, item));
+				AddOp(FScheduledOp(Item.At, Item, 2), FScheduledOp(Args.patch, Item));
 				break;
 
 			case 2:
             {
                 MUTABLE_CPUPROFILER_SCOPE(IM_PATCH_1)
 
-                TSharedPtr<const FImage> pA = LoadImage( FCacheAddress(args.base,item) );
-                TSharedPtr<const FImage> pB = LoadImage( FCacheAddress(args.patch,item) );
+                TSharedPtr<const FImage> pA = LoadImage( FCacheAddress(Args.base,Item) );
+                TSharedPtr<const FImage> pB = LoadImage( FCacheAddress(Args.patch,Item) );
 
 				// Failsafe
 				if (!pA || !pB)
 				{
 					Release(pB);
-					StoreImage(item, pA);
+					StoreImage(Item, pA);
 					break;
 				}
 
 				// Apply the mipmap reduction to the crop rectangle.
-				int32 MipsToSkip = item.ExecutionOptions;
+				int32 MipsToSkip = Item.ExecutionOptions;
 				box<FIntVector2> rect;
-				rect.min[0] = args.minX / (1 << MipsToSkip);
-				rect.min[1] = args.minY / (1 << MipsToSkip);
+				rect.min[0] = Args.minX / (1 << MipsToSkip);
+				rect.min[1] = Args.minY / (1 << MipsToSkip);
 				rect.size[0] = pB->GetSizeX();
 				rect.size[1] = pB->GetSizeY();
 
@@ -4263,7 +4263,7 @@ namespace mu
 				}
 
 				Release(pB);
-				StoreImage( item, Result );
+				StoreImage( Item, Result );
                 break;
             }
 
@@ -4276,20 +4276,20 @@ namespace mu
 
         case EOpType::IM_RASTERMESH:
         {
-			OP::ImageRasterMeshArgs args = Program.GetOpArgs<OP::ImageRasterMeshArgs>(item.At);
-            switch (item.Stage)
+			OP::ImageRasterMeshArgs Args = Program.GetOpArgs<OP::ImageRasterMeshArgs>(Item.At);
+            switch (Item.Stage)
             {
             case 0:
-				if (args.image)
+				if (Args.image)
 				{
-					AddOp(FScheduledOp(item.At, item, 1),
-						FScheduledOp::FromOpAndOptions(args.mesh, item, 0),
-						FScheduledOp::FromOpAndOptions(args.projector, item, 0));
+					AddOp(FScheduledOp(Item.At, Item, 1),
+						FScheduledOp::FromOpAndOptions(Args.mesh, Item, 0),
+						FScheduledOp::FromOpAndOptions(Args.projector, Item, 0));
 				}
 				else
 				{
-					AddOp(FScheduledOp(item.At, item, 1),
-						FScheduledOp::FromOpAndOptions(args.mesh, item, 0));
+					AddOp(FScheduledOp(Item.At, Item, 1),
+						FScheduledOp::FromOpAndOptions(Args.mesh, Item, 0));
 				}
                 break;
 
@@ -4297,18 +4297,18 @@ namespace mu
 			{
 				MUTABLE_CPUPROFILER_SCOPE(IM_RASTERMESH_1)
 
-				TSharedPtr<const FMesh> Mesh = LoadMesh(FScheduledOp::FromOpAndOptions(args.mesh, item, 0));
+				TSharedPtr<const FMesh> Mesh = LoadMesh(FScheduledOp::FromOpAndOptions(Args.mesh, Item, 0));
 
 				// If no image, we are generating a flat mesh UV raster. This is the final stage in this case.
-				if (!args.image)
+				if (!Args.image)
 				{
-					uint16 SizeX = args.sizeX;
-					uint16 SizeY = args.sizeY;
-					UE::Math::TIntVector2<uint16> CropMin(args.CropMinX, args.CropMinY);
-					UE::Math::TIntVector2<uint16> UncroppedSize(args.UncroppedSizeX, args.UncroppedSizeY);
+					uint16 SizeX = Args.sizeX;
+					uint16 SizeY = Args.sizeY;
+					UE::Math::TIntVector2<uint16> CropMin(Args.CropMinX, Args.CropMinY);
+					UE::Math::TIntVector2<uint16> UncroppedSize(Args.UncroppedSizeX, Args.UncroppedSizeY);
 
 					// Drop mips while possible
-					int32 MipsToDrop = item.ExecutionOptions;
+					int32 MipsToDrop = Item.ExecutionOptions;
 					bool bUseCrop = UncroppedSize[0] > 0;
 					while (MipsToDrop && !(SizeX % 2) && !(SizeY % 2))
 					{
@@ -4328,22 +4328,22 @@ namespace mu
 					TSharedPtr<FImage> ResultImage = CreateImage(SizeX, SizeY, 1, EImageFormat::L_UByte, EInitializationType::Black);
 					if (Mesh)
 					{
-						ImageRasterMesh(Mesh.Get(), ResultImage.Get(), args.LayoutIndex, args.BlockId, CropMin, UncroppedSize);
+						ImageRasterMesh(Mesh.Get(), ResultImage.Get(), Args.LayoutIndex, Args.BlockId, CropMin, UncroppedSize);
 						Release(Mesh);
 					}
 
 					// Stop execution.
-					StoreImage(item, ResultImage);
+					StoreImage(Item, ResultImage);
 					break;
 				}
 
-				const int32 MipsToSkip = item.ExecutionOptions;
+				const int32 MipsToSkip = Item.ExecutionOptions;
 				int32 ProjectionMip = MipsToSkip;
 
 				FScheduledOpData Data;
 				Data.RasterMesh.Mip = ProjectionMip;
 				Data.RasterMesh.MipValue = static_cast<float>(ProjectionMip);
-				FProjector Projector = LoadProjector(FScheduledOp::FromOpAndOptions(args.projector, item, 0));
+				FProjector Projector = LoadProjector(FScheduledOp::FromOpAndOptions(Args.projector, Item, 0));
 
 				EMinFilterMethod MinFilterMethod = Invoke([&]() -> EMinFilterMethod
 				{
@@ -4356,15 +4356,15 @@ namespace mu
 						return EMinFilterMethod::TotalAreaHeuristic;
 					}
 						
-					return static_cast<EMinFilterMethod>(args.MinFilterMethod);
+					return static_cast<EMinFilterMethod>(Args.MinFilterMethod);
 				});
 
 				if (MinFilterMethod == EMinFilterMethod::TotalAreaHeuristic)
 				{
 					FVector2f TargetImageSizeF = FVector2f(
-						FMath::Max(args.sizeX >> MipsToSkip, 1),
-						FMath::Max(args.sizeY >> MipsToSkip, 1));
-					FVector2f SourceImageSizeF = FVector2f(args.SourceSizeX, args.SourceSizeY);
+						FMath::Max(Args.sizeX >> MipsToSkip, 1),
+						FMath::Max(Args.sizeY >> MipsToSkip, 1));
+					FVector2f SourceImageSizeF = FVector2f(Args.SourceSizeX, Args.SourceSizeY);
 						
 					if (Mesh)
 					{ 
@@ -4380,11 +4380,11 @@ namespace mu
 				// Mesh is need again in the next stage, store it in the heap.
 				HeapData[DataHeapAddress].Resource = Mesh;
 
-				AddOp(FScheduledOp(item.At, item, 2, DataHeapAddress),
-					FScheduledOp::FromOpAndOptions(args.projector, item, 0),
-					FScheduledOp::FromOpAndOptions(args.image, item, Data.RasterMesh.Mip),
-					FScheduledOp(args.mask, item),
-					FScheduledOp::FromOpAndOptions(args.angleFadeProperties, item, 0));
+				AddOp(FScheduledOp(Item.At, Item, 2, DataHeapAddress),
+					FScheduledOp::FromOpAndOptions(Args.projector, Item, 0),
+					FScheduledOp::FromOpAndOptions(Args.image, Item, Data.RasterMesh.Mip),
+					FScheduledOp(Args.mask, Item),
+					FScheduledOp::FromOpAndOptions(Args.angleFadeProperties, Item, 0));
 
 				break;
 			}
@@ -4393,15 +4393,15 @@ namespace mu
             {
 				MUTABLE_CPUPROFILER_SCOPE(IM_RASTERMESH_2)
 
-				if (!args.image)
+				if (!Args.image)
 				{
 					// This case is treated at the previous stage.
 					check(false);
-					StoreImage(item, nullptr);
+					StoreImage(Item, nullptr);
 					break;
 				}
 
-				FScheduledOpData& Data = HeapData[item.CustomState];
+				FScheduledOpData& Data = HeapData[Item.CustomState];
 
 				// Unsafe downcast, should be fine as it is known to be a Mesh.
 				TSharedPtr<const FMesh> Mesh = StaticCastSharedPtr<FMesh>(ConstCastSharedPtr<FResource>(Data.Resource));
@@ -4410,17 +4410,17 @@ namespace mu
 				if (!Mesh)
 				{
 					check(false);
-					StoreImage(item, nullptr);
+					StoreImage(Item, nullptr);
 					break;
 				}
 
-				uint16 SizeX = args.sizeX;
-				uint16 SizeY = args.sizeY;
-				UE::Math::TIntVector2<uint16> CropMin(args.CropMinX, args.CropMinY);
-				UE::Math::TIntVector2<uint16> UncroppedSize(args.UncroppedSizeX, args.UncroppedSizeY);
+				uint16 SizeX = Args.sizeX;
+				uint16 SizeY = Args.sizeY;
+				UE::Math::TIntVector2<uint16> CropMin(Args.CropMinX, Args.CropMinY);
+				UE::Math::TIntVector2<uint16> UncroppedSize(Args.UncroppedSizeX, Args.UncroppedSizeY);
 
 				// Drop mips while possible
-				int32 MipsToDrop = item.ExecutionOptions;
+				int32 MipsToDrop = Item.ExecutionOptions;
 				bool bUseCrop = UncroppedSize[0] > 0;
 				while (MipsToDrop && !(SizeX % 2) && !(SizeY % 2))
 				{
@@ -4437,12 +4437,12 @@ namespace mu
 				}
 
 				// Raster with projection
-				TSharedPtr<const FImage> Source = LoadImage(FCacheAddress(args.image, item.ExecutionIndex, Data.RasterMesh.Mip));
+				TSharedPtr<const FImage> Source = LoadImage(FCacheAddress(Args.image, Item.ExecutionIndex, Data.RasterMesh.Mip));
 
 				TSharedPtr<const FImage> Mask = nullptr;
-				if (args.mask)
+				if (Args.mask)
 				{
-					Mask = LoadImage(FCacheAddress(args.mask, item));
+					Mask = LoadImage(FCacheAddress(Args.mask, Item));
 
 					// TODO: This shouldn't happen, but be defensive.
 					FImageSize ResultSize(SizeX, SizeY);
@@ -4459,9 +4459,9 @@ namespace mu
 
 				float fadeStart = 180.0f;
 				float fadeEnd = 180.0f;
-				if ( args.angleFadeProperties )
+				if ( Args.angleFadeProperties )
 				{
-					FVector4f fadeProperties = LoadColor(FScheduledOp::FromOpAndOptions(args.angleFadeProperties, item, 0));
+					FVector4f fadeProperties = LoadColor(FScheduledOp::FromOpAndOptions(Args.angleFadeProperties, Item, 0));
 					fadeStart = fadeProperties[0];
 					fadeEnd = fadeProperties[1];
 				}
@@ -4492,15 +4492,15 @@ namespace mu
 						return EMinFilterMethod::TotalAreaHeuristic;
 					}
 						
-					return static_cast<EMinFilterMethod>(args.MinFilterMethod);
+					return static_cast<EMinFilterMethod>(Args.MinFilterMethod);
 				});
 
 				if (MinFilterMethod == EMinFilterMethod::TotalAreaHeuristic)
 				{
 					const uint16 Mip = Data.RasterMesh.Mip;
 					const FImageSize ExpectedSourceSize = FImageSize(
-							FMath::Max<uint16>(args.SourceSizeX >> Mip, 1), 
-							FMath::Max<uint16>(args.SourceSizeY >> Mip, 1));
+							FMath::Max<uint16>(Args.SourceSizeX >> Mip, 1), 
+							FMath::Max<uint16>(Args.SourceSizeY >> Mip, 1));
 
 					if (Source->GetSize() != ExpectedSourceSize)
 					{
@@ -4530,7 +4530,7 @@ namespace mu
 						return ESamplingMethod::BiLinear;
 					}
 					
-					return static_cast<ESamplingMethod>(args.SamplingMethod);
+					return static_cast<ESamplingMethod>(Args.SamplingMethod);
 				});
 
 				if (SamplingMethod == ESamplingMethod::BiLinear)
@@ -4551,19 +4551,19 @@ namespace mu
 				// Allocate new image after bilinear mip generation to reduce operation memory peak.
 				TSharedPtr<FImage> New = CreateImage(SizeX, SizeY, 1, Format, EInitializationType::Black);
 
-				if (args.projector && Source && Source->GetSizeX() > 0 && Source->GetSizeY() > 0)
+				if (Args.projector && Source && Source->GetSizeX() > 0 && Source->GetSizeY() > 0)
 				{
-					FProjector Projector = LoadProjector(FScheduledOp::FromOpAndOptions(args.projector, item, 0));
+					FProjector Projector = LoadProjector(FScheduledOp::FromOpAndOptions(Args.projector, Item, 0));
 
 					switch (Projector.type)
 					{
 					case EProjectorType::Planar:
 						ImageRasterProjectedPlanar(Mesh.Get(), New.Get(),
 							Source.Get(), Mask.Get(),
-							args.bIsRGBFadingEnabled, args.bIsAlphaFadingEnabled,
+							Args.bIsRGBFadingEnabled, Args.bIsAlphaFadingEnabled,
 							SamplingMethod,
 							FadeStartRad, FadeEndRad, FMath::Frac(Data.RasterMesh.MipValue),
-							args.LayoutIndex, args.BlockId,
+							Args.LayoutIndex, Args.BlockId,
 							CropMin, UncroppedSize,
 							&Scratch, bUseProjectionVectorImpl);
 						break;
@@ -4571,10 +4571,10 @@ namespace mu
 					case EProjectorType::Wrapping:
 						ImageRasterProjectedWrapping(Mesh.Get(), New.Get(),
 							Source.Get(), Mask.Get(),
-							args.bIsRGBFadingEnabled, args.bIsAlphaFadingEnabled,
+							Args.bIsRGBFadingEnabled, Args.bIsAlphaFadingEnabled,
 							SamplingMethod,
 							FadeStartRad, FadeEndRad, FMath::Frac(Data.RasterMesh.MipValue),
-							args.LayoutIndex, args.BlockId,
+							Args.LayoutIndex, Args.BlockId,
 							CropMin, UncroppedSize,
 							&Scratch, bUseProjectionVectorImpl);
 						break;
@@ -4582,10 +4582,10 @@ namespace mu
 					case EProjectorType::Cylindrical:
 						ImageRasterProjectedCylindrical(Mesh.Get(), New.Get(),
 							Source.Get(), Mask.Get(),
-							args.bIsRGBFadingEnabled, args.bIsAlphaFadingEnabled,
+							Args.bIsRGBFadingEnabled, Args.bIsAlphaFadingEnabled,
 							SamplingMethod,
 							FadeStartRad, FadeEndRad, FMath::Frac(Data.RasterMesh.MipValue),
-							args.LayoutIndex,
+							Args.LayoutIndex,
 							Projector.projectionAngle,
 							CropMin, UncroppedSize,
 							&Scratch, bUseProjectionVectorImpl);
@@ -4600,7 +4600,7 @@ namespace mu
 				Release(Mesh);
 				Release(Source);
 				Release(Mask);
-				StoreImage(item, New);
+				StoreImage(Item, New);
 
                 break;
             }
@@ -4614,26 +4614,26 @@ namespace mu
 
         case EOpType::IM_MAKEGROWMAP:
         {
-			OP::ImageMakeGrowMapArgs args = Program.GetOpArgs<OP::ImageMakeGrowMapArgs>(item.At);
-            switch (item.Stage)
+			OP::ImageMakeGrowMapArgs Args = Program.GetOpArgs<OP::ImageMakeGrowMapArgs>(Item.At);
+            switch (Item.Stage)
             {
             case 0:
-                AddOp( FScheduledOp( item.At, item, 1), FScheduledOp( args.mask, item) );
+                AddOp( FScheduledOp( Item.At, Item, 1), FScheduledOp( Args.mask, Item) );
                 break;
 
             case 1:
             {
                 MUTABLE_CPUPROFILER_SCOPE(IM_MAKEGROWMAP_1)
 
-                TSharedPtr<const FImage> Mask = LoadImage( FCacheAddress(args.mask,item) );
+                TSharedPtr<const FImage> Mask = LoadImage( FCacheAddress(Args.mask,Item) );
 
                 TSharedPtr<FImage> Result = CreateImage( Mask->GetSizeX(), Mask->GetSizeY(), Mask->GetLODCount(), EImageFormat::L_UByte, EInitializationType::NotInitialized);
 
-                ImageMakeGrowMap(Result.Get(), Mask.Get(), args.border );
+                ImageMakeGrowMap(Result.Get(), Mask.Get(), Args.border );
 				Result->Flags |= FImage::IF_CANNOT_BE_SCALED;
 
 				Release(Mask);
-                StoreImage( item, Result);
+                StoreImage( Item, Result);
                 break;
             }
 
@@ -4646,26 +4646,26 @@ namespace mu
 
         case EOpType::IM_DISPLACE:
         {
-			OP::ImageDisplaceArgs args = Program.GetOpArgs<OP::ImageDisplaceArgs>(item.At);
-            switch (item.Stage)
+			OP::ImageDisplaceArgs Args = Program.GetOpArgs<OP::ImageDisplaceArgs>(Item.At);
+            switch (Item.Stage)
             {
             case 0:
-                AddOp( FScheduledOp( item.At, item, 1),
-                        FScheduledOp( args.source, item ),
-                        FScheduledOp( args.displacementMap, item ) );
+                AddOp( FScheduledOp( Item.At, Item, 1),
+                        FScheduledOp( Args.source, Item ),
+                        FScheduledOp( Args.displacementMap, Item ) );
 				break;
 
             case 1:
             {
                 MUTABLE_CPUPROFILER_SCOPE(IM_DISPLACE_1)
 
-                TSharedPtr<const FImage> Source = LoadImage( FCacheAddress(args.source,item) );
-                TSharedPtr<const FImage> pMap = LoadImage( FCacheAddress(args.displacementMap,item) );
+                TSharedPtr<const FImage> Source = LoadImage( FCacheAddress(Args.source,Item) );
+                TSharedPtr<const FImage> pMap = LoadImage( FCacheAddress(Args.displacementMap,Item) );
 
 				if (!Source)
 				{
 					Release(pMap);
-					StoreImage(item, nullptr);
+					StoreImage(Item, nullptr);
 					break;
 				}
 
@@ -4702,7 +4702,7 @@ namespace mu
 				}
 
 				Release(pMap);
-                StoreImage( item, Result);
+                StoreImage( Item, Result);
                 break;
             }
 
@@ -4715,19 +4715,19 @@ namespace mu
 
         case EOpType::IM_TRANSFORM:
         {
-            const OP::ImageTransformArgs Args = Program.GetOpArgs<OP::ImageTransformArgs>(item.At);
+            const OP::ImageTransformArgs Args = Program.GetOpArgs<OP::ImageTransformArgs>(Item.At);
 
-            switch (item.Stage)
+            switch (Item.Stage)
             {
             case 0:
 			{
 				const TArray<FScheduledOp, TFixedAllocator<2>> Deps = 
 				{
-					FScheduledOp(Args.ScaleX, item),
-					FScheduledOp(Args.ScaleY, item),
+					FScheduledOp(Args.ScaleX, Item),
+					FScheduledOp(Args.ScaleY, Item),
 				};
 
-                AddOp(FScheduledOp(item.At, item, 1), Deps);
+                AddOp(FScheduledOp(Item.At, Item, 1), Deps);
 
 				break;
 			}
@@ -4736,13 +4736,13 @@ namespace mu
             	MUTABLE_CPUPROFILER_SCOPE(IM_TRANSFORM_1)
 
 				FVector2f Scale = FVector2f(
-                        Args.ScaleX ? LoadScalar(FCacheAddress(Args.ScaleX, item)) : 1.0f,
-                        Args.ScaleY ? LoadScalar(FCacheAddress(Args.ScaleY, item)) : 1.0f);
+                        Args.ScaleX ? LoadScalar(FCacheAddress(Args.ScaleX, Item)) : 1.0f,
+                        Args.ScaleY ? LoadScalar(FCacheAddress(Args.ScaleY, Item)) : 1.0f);
 	
 				using FUint16Vector2 = UE::Math::TIntVector2<uint16>;
 				const FUint16Vector2 DestSizeI = Invoke([&]() 
 				{
-					int32 MipsToDrop = item.ExecutionOptions;
+					int32 MipsToDrop = Item.ExecutionOptions;
 					
 					FUint16Vector2 Size = FUint16Vector2(
 							Args.SizeX > 0 ? Args.SizeX : Args.SourceSizeX, 
@@ -4800,13 +4800,13 @@ namespace mu
 				const uint8 Mip = static_cast<uint8>(FMath::Max(0, FMath::FloorToInt(OpHeapData.ImageTransform.MipValue)));
 				const TArray<FScheduledOp, TFixedAllocator<4>> Deps = 
 				{
-					FScheduledOp::FromOpAndOptions(Args.Base, item, Mip),
-					FScheduledOp(Args.OffsetX,  item),
-					FScheduledOp(Args.OffsetY,  item),
-					FScheduledOp(Args.Rotation, item) 
+					FScheduledOp::FromOpAndOptions(Args.Base, Item, Mip),
+					FScheduledOp(Args.OffsetX,  Item),
+					FScheduledOp(Args.OffsetY,  Item),
+					FScheduledOp(Args.Rotation, Item) 
 				};
 				
-                AddOp(FScheduledOp(item.At, item, 2, HeapDataAddress), Deps);
+                AddOp(FScheduledOp(Item.At, Item, 2, HeapDataAddress), Deps);
 
 				break;
 			}
@@ -4814,14 +4814,14 @@ namespace mu
             {
 				MUTABLE_CPUPROFILER_SCOPE(IM_TRANSFORM_2);
 			
-				const FScheduledOpData OpHeapData = HeapData[item.CustomState];
+				const FScheduledOpData OpHeapData = HeapData[Item.CustomState];
 
 				const uint8 Mip = static_cast<uint8>(FMath::Max(0, FMath::FloorToInt(OpHeapData.ImageTransform.MipValue)));
-				TSharedPtr<const FImage> Source = LoadImage(FCacheAddress(Args.Base, item.ExecutionIndex, Mip));
+				TSharedPtr<const FImage> Source = LoadImage(FCacheAddress(Args.Base, Item.ExecutionIndex, Mip));
 
 				const FVector2f Offset = FVector2f(
-                        Args.OffsetX ? LoadScalar(FCacheAddress(Args.OffsetX, item)) : 0.0f,
-                        Args.OffsetY ? LoadScalar(FCacheAddress(Args.OffsetY, item)) : 0.0f);
+                        Args.OffsetX ? LoadScalar(FCacheAddress(Args.OffsetX, Item)) : 0.0f,
+                        Args.OffsetY ? LoadScalar(FCacheAddress(Args.OffsetY, Item)) : 0.0f);
 
                 FVector2f Scale = FVector2f(
 						FPlatformMath::LoadHalf(&OpHeapData.ImageTransform.ScaleXEncodedHalf),
@@ -4841,7 +4841,7 @@ namespace mu
 				}
 
 				// Map Range [0..1] to a full rotation
-                const float RotationRad = LoadScalar(FCacheAddress(Args.Rotation, item)) * UE_TWO_PI;
+                const float RotationRad = LoadScalar(FCacheAddress(Args.Rotation, Item)) * UE_TWO_PI;
 	
 				EImageFormat SourceFormat = Source->GetFormat();
 				EImageFormat Format = GetUncompressedFormat(SourceFormat);
@@ -4914,7 +4914,7 @@ namespace mu
 				ImageTransform(Result.Get(), Source.Get(), Transform, MipFactor, AddressMode, bUseImageTransformVectorImpl);
 
 				Release(Source);
-				StoreImage(item, Result);
+				StoreImage(Item, Result);
 
                 break;
             }
@@ -4937,9 +4937,9 @@ namespace mu
     }	
 	
 	//---------------------------------------------------------------------------------------------
-    TSharedPtr<FRangeIndex> CodeRunner::BuildCurrentOpRangeIndex( const FScheduledOp& item, const FParameters* pParams, const FModel* pModel, int32 parameterIndex )
+    TSharedPtr<FRangeIndex> CodeRunner::BuildCurrentOpRangeIndex( const FScheduledOp& Item, const FParameters* pParams, const FModel* InModel, int32 parameterIndex )
     {
-        if (!item.ExecutionIndex)
+        if (!Item.ExecutionIndex)
         {
             return nullptr;
         }
@@ -4959,7 +4959,7 @@ namespace mu
              ++rangeIndexInParam )
         {
             uint32 rangeIndexInModel = paramDesc.Ranges[rangeIndexInParam];
-            const ExecutionIndex& currentIndex = GetMemory().GetRangeIndex( item.ExecutionIndex );
+            const ExecutionIndex& currentIndex = GetMemory().GetRangeIndex( Item.ExecutionIndex );
             int32 Position = currentIndex.GetFromModelRangeIndex(rangeIndexInModel);
 			Index->Values[rangeIndexInParam] = Position;
         }
@@ -4969,50 +4969,50 @@ namespace mu
 
 
     //---------------------------------------------------------------------------------------------
-    void CodeRunner::RunCode_Bool(const FScheduledOp& item, const FParameters* pParams, const FModel* pModel )
+    void CodeRunner::RunCode_Bool(const FScheduledOp& Item, const FParameters* pParams, const FModel* InModel )
     {
         MUTABLE_CPUPROFILER_SCOPE(RunCode_Bool);
 
 		const FProgram& Program = Model->GetPrivate()->Program;
-		EOpType type = Program.GetOpType(item.At);
+		EOpType type = Program.GetOpType(Item.At);
         switch (type)
         {
 
         case EOpType::BO_CONSTANT:
         {
-			OP::BoolConstantArgs args = Program.GetOpArgs<OP::BoolConstantArgs>(item.At);
-            bool result = args.value;
-            StoreBool( item, result );
+			OP::BoolConstantArgs Args = Program.GetOpArgs<OP::BoolConstantArgs>(Item.At);
+            bool result = Args.value;
+            StoreBool( Item, result );
             break;
         }
 
         case EOpType::BO_PARAMETER:
         {
-			OP::ParameterArgs args = Program.GetOpArgs<OP::ParameterArgs>(item.At);
+			OP::ParameterArgs Args = Program.GetOpArgs<OP::ParameterArgs>(Item.At);
             bool result = false;
-			TSharedPtr<FRangeIndex> Index = BuildCurrentOpRangeIndex( item, pParams, pModel, args.variable );
-            result = pParams->GetBoolValue( args.variable, Index.Get() );
-            StoreBool( item, result );
+			TSharedPtr<FRangeIndex> Index = BuildCurrentOpRangeIndex( Item, pParams, InModel, Args.variable );
+            result = pParams->GetBoolValue( Args.variable, Index.Get() );
+            StoreBool( Item, result );
             break;
         }
 
         case EOpType::BO_LESS:
         {
-			OP::BoolLessArgs args = Program.GetOpArgs<OP::BoolLessArgs>(item.At);
-            switch (item.Stage)
+			OP::BoolLessArgs Args = Program.GetOpArgs<OP::BoolLessArgs>(Item.At);
+            switch (Item.Stage)
             {
             case 0:
-                AddOp( FScheduledOp( item.At, item, 1),
-                        FScheduledOp( args.a, item),
-                        FScheduledOp( args.b, item) );
+                AddOp( FScheduledOp( Item.At, Item, 1),
+                        FScheduledOp( Args.a, Item),
+                        FScheduledOp( Args.b, Item) );
                 break;
 
             case 1:
             {
-                float a = LoadScalar( FCacheAddress(args.a,item) );
-                float b = LoadScalar( FCacheAddress(args.b,item) );
+                float a = LoadScalar( FCacheAddress(Args.a,Item) );
+                float b = LoadScalar( FCacheAddress(Args.b,Item) );
                 bool result = a<b;
-                StoreBool( item, result );
+                StoreBool( Item, result );
                 break;
             }
 
@@ -5025,52 +5025,52 @@ namespace mu
 
         case EOpType::BO_AND:
         {
-			OP::BoolBinaryArgs args = Program.GetOpArgs<OP::BoolBinaryArgs>(item.At);
-            switch (item.Stage)
+			OP::BoolBinaryArgs Args = Program.GetOpArgs<OP::BoolBinaryArgs>(Item.At);
+            switch (Item.Stage)
             {
             case 0:
                 {
                     // Try to avoid the op entirely if we have some children cached
                     bool skip = false;
-                    if ( args.a && GetMemory().IsValid( FCacheAddress(args.a,item) ) )
+                    if ( Args.a && GetMemory().IsValid( FCacheAddress(Args.a,Item) ) )
                     {
-                         bool a = LoadBool( FCacheAddress(args.a,item) );
+                         bool a = LoadBool( FCacheAddress(Args.a,Item) );
                          if (!a)
                          {
-                            StoreBool( item, false );
+                            StoreBool( Item, false );
                             skip=true;
                          }
                     }
 
-                    if ( !skip && args.b && GetMemory().IsValid( FCacheAddress(args.b,item) ) )
+                    if ( !skip && Args.b && GetMemory().IsValid( FCacheAddress(Args.b,Item) ) )
                     {
-                         bool b = LoadBool( FCacheAddress(args.b,item) );
+                         bool b = LoadBool( FCacheAddress(Args.b,Item) );
                          if (!b)
                          {
-                            StoreBool( item, false );
+                            StoreBool( Item, false );
                             skip=true;
                          }
                     }
 
                     if (!skip)
                     {
-                        AddOp( FScheduledOp( item.At, item, 1),
-                               FScheduledOp( args.a, item));
+                        AddOp( FScheduledOp( Item.At, Item, 1),
+                               FScheduledOp( Args.a, Item));
                     }
 				break;
                 }
 
             case 1:
             {
-                bool a = args.a ? LoadBool( FCacheAddress(args.a,item) ) : true;
+                bool a = Args.a ? LoadBool( FCacheAddress(Args.a,Item) ) : true;
                 if (!a)
                 {
-                    StoreBool( item, false );
+                    StoreBool( Item, false );
                 }
                 else
                 {
-                    AddOp( FScheduledOp( item.At, item, 2),
-                           FScheduledOp( args.b, item));
+                    AddOp( FScheduledOp( Item.At, Item, 2),
+                           FScheduledOp( Args.b, Item));
                 }
                 break;
             }
@@ -5078,8 +5078,8 @@ namespace mu
             case 2:
             {
                 // We arrived here because a is true
-                bool b = args.b ? LoadBool( FCacheAddress(args.b,item) ) : true;
-                StoreBool( item, b );
+                bool b = Args.b ? LoadBool( FCacheAddress(Args.b,Item) ) : true;
+                StoreBool( Item, b );
                 break;
             }
 
@@ -5091,52 +5091,52 @@ namespace mu
 
         case EOpType::BO_OR:
         {
-			OP::BoolBinaryArgs args = Program.GetOpArgs<OP::BoolBinaryArgs>(item.At);
-            switch (item.Stage)
+			OP::BoolBinaryArgs Args = Program.GetOpArgs<OP::BoolBinaryArgs>(Item.At);
+            switch (Item.Stage)
             {
             case 0:
                 {
                     // Try to avoid the op entirely if we have some children cached
                     bool skip = false;
-                    if ( args.a && GetMemory().IsValid( FCacheAddress(args.a,item) ) )
+                    if ( Args.a && GetMemory().IsValid( FCacheAddress(Args.a,Item) ) )
                     {
-                         bool a = LoadBool( FCacheAddress(args.a,item) );
+                         bool a = LoadBool( FCacheAddress(Args.a,Item) );
                          if (a)
                          {
-                            StoreBool( item, true );
+                            StoreBool( Item, true );
                             skip=true;
                          }
                     }
 
-                    if ( !skip && args.b && GetMemory().IsValid( FCacheAddress(args.b,item) ) )
+                    if ( !skip && Args.b && GetMemory().IsValid( FCacheAddress(Args.b,Item) ) )
                     {
-                         bool b = LoadBool( FCacheAddress(args.b,item) );
+                         bool b = LoadBool( FCacheAddress(Args.b,Item) );
                          if (b)
                          {
-                            StoreBool( item, true );
+                            StoreBool( Item, true );
                             skip=true;
                          }
                     }
 
                     if (!skip)
                     {
-                        AddOp( FScheduledOp( item.At, item, 1),
-                               FScheduledOp( args.a, item));
+                        AddOp( FScheduledOp( Item.At, Item, 1),
+                               FScheduledOp( Args.a, Item));
                     }
 				break;
                 }
 
             case 1:
             {
-                bool a = args.a ? LoadBool( FCacheAddress(args.a,item) ) : false;
+                bool a = Args.a ? LoadBool( FCacheAddress(Args.a,Item) ) : false;
                 if (a)
                 {
-                    StoreBool( item, true );
+                    StoreBool( Item, true );
                 }
                 else
                 {
-                    AddOp( FScheduledOp( item.At, item, 2),
-                           FScheduledOp( args.b, item));
+                    AddOp( FScheduledOp( Item.At, Item, 2),
+                           FScheduledOp( Args.b, Item));
                 }
                 break;
             }
@@ -5144,8 +5144,8 @@ namespace mu
             case 2:
             {
                 // We arrived here because a is false
-                bool b = args.b ? LoadBool( FCacheAddress(args.b,item) ) : false;
-                StoreBool( item, b );
+                bool b = Args.b ? LoadBool( FCacheAddress(Args.b,Item) ) : false;
+                StoreBool( Item, b );
                 break;
             }
 
@@ -5157,18 +5157,18 @@ namespace mu
 
         case EOpType::BO_NOT:
         {
-			OP::BoolNotArgs args = Program.GetOpArgs<OP::BoolNotArgs>(item.At);
-            switch (item.Stage)
+			OP::BoolNotArgs Args = Program.GetOpArgs<OP::BoolNotArgs>(Item.At);
+            switch (Item.Stage)
             {
             case 0:
-                    AddOp( FScheduledOp( item.At, item, 1),
-                           FScheduledOp( args.source, item) );
+                    AddOp( FScheduledOp( Item.At, Item, 1),
+                           FScheduledOp( Args.source, Item) );
                 break;
 
             case 1:
             {
-                bool result = !LoadBool( FCacheAddress(args.source,item) );
-                StoreBool( item, result );
+                bool result = !LoadBool( FCacheAddress(Args.source,Item) );
+                StoreBool( Item, result );
                 break;
             }
 
@@ -5180,19 +5180,19 @@ namespace mu
 
         case EOpType::BO_EQUAL_INT_CONST:
         {
-			OP::BoolEqualScalarConstArgs args = Program.GetOpArgs<OP::BoolEqualScalarConstArgs>(item.At);
-            switch (item.Stage)
+			OP::BoolEqualScalarConstArgs Args = Program.GetOpArgs<OP::BoolEqualScalarConstArgs>(Item.At);
+            switch (Item.Stage)
             {
             case 0:
-                    AddOp( FScheduledOp( item.At, item, 1),
-                           FScheduledOp( args.value, item) );
+                    AddOp( FScheduledOp( Item.At, Item, 1),
+                           FScheduledOp( Args.value, Item) );
                 break;
 
             case 1:
             {
-                int32 a = LoadInt( FCacheAddress(args.value,item) );
-                bool result = a == args.constant;
-                StoreBool( item, result );
+                int32 a = LoadInt( FCacheAddress(Args.value,Item) );
+                bool result = a == Args.constant;
+                StoreBool( Item, result );
                 break;
             }
 
@@ -5209,39 +5209,39 @@ namespace mu
     }
 
     //---------------------------------------------------------------------------------------------
-    void CodeRunner::RunCode_Int(const FScheduledOp& item, const FParameters* pParams, const FModel* pModel )
+    void CodeRunner::RunCode_Int(const FScheduledOp& Item, const FParameters* pParams, const FModel* InModel )
     {
         MUTABLE_CPUPROFILER_SCOPE(RunCode_Int);
 
 		const FProgram& Program = Model->GetPrivate()->Program;
 
-		EOpType type = Program.GetOpType(item.At);
+		EOpType type = Program.GetOpType(Item.At);
         switch (type)
         {
 
         case EOpType::NU_CONSTANT:
         {
-			OP::IntConstantArgs args = Program.GetOpArgs<OP::IntConstantArgs>(item.At);
-            int32 result = args.value;
-            StoreInt( item, result );
+			OP::IntConstantArgs Args = Program.GetOpArgs<OP::IntConstantArgs>(Item.At);
+            int32 result = Args.value;
+            StoreInt( Item, result );
             break;
         }
 
         case EOpType::NU_PARAMETER:
         {
-			OP::ParameterArgs args = Program.GetOpArgs<OP::ParameterArgs>(item.At);
-			TSharedPtr<FRangeIndex> Index = BuildCurrentOpRangeIndex( item, pParams, pModel, args.variable );
-            int32 result = pParams->GetIntValue( args.variable, Index.Get());
+			OP::ParameterArgs Args = Program.GetOpArgs<OP::ParameterArgs>(Item.At);
+			TSharedPtr<FRangeIndex> Index = BuildCurrentOpRangeIndex( Item, pParams, InModel, Args.variable );
+            int32 result = pParams->GetIntValue( Args.variable, Index.Get());
 
             // Check that the value is actually valid. Otherwise set the default.
-            if ( pParams->GetIntPossibleValueCount( args.variable ) )
+            if ( pParams->GetIntPossibleValueCount( Args.variable ) )
             {
                 bool valid = false;
                 for ( int32 i=0;
-                      (!valid) && i<pParams->GetIntPossibleValueCount( args.variable );
+                      (!valid) && i<pParams->GetIntPossibleValueCount( Args.variable );
                       ++i )
                 {
-                    if ( result == pParams->GetIntPossibleValue( args.variable, i ) )
+                    if ( result == pParams->GetIntPossibleValue( Args.variable, i ) )
                     {
                         valid = true;
                     }
@@ -5249,11 +5249,11 @@ namespace mu
 
                 if (!valid)
                 {
-                    result = pParams->GetIntPossibleValue( args.variable, 0 );
+                    result = pParams->GetIntPossibleValue( Args.variable, 0 );
                 }
             }
 
-            StoreInt( item, result );
+            StoreInt( Item, result );
             break;
         }
 
@@ -5264,51 +5264,51 @@ namespace mu
     }
 
     //---------------------------------------------------------------------------------------------
-    void CodeRunner::RunCode_Scalar(const FScheduledOp& item, const FParameters* pParams, const FModel* pModel )
+    void CodeRunner::RunCode_Scalar(const FScheduledOp& Item, const FParameters* pParams, const FModel* InModel )
     {
         MUTABLE_CPUPROFILER_SCOPE(RunCode_Scalar);
 
 		const FProgram& Program = Model->GetPrivate()->Program;
 
-		EOpType type = Program.GetOpType(item.At);
+		EOpType type = Program.GetOpType(Item.At);
         switch (type)
         {
 
         case EOpType::SC_CONSTANT:
         {
-			OP::ScalarConstantArgs args = Program.GetOpArgs<OP::ScalarConstantArgs>(item.At);
-            float result = args.value;
-            StoreScalar( item, result );
+			OP::ScalarConstantArgs Args = Program.GetOpArgs<OP::ScalarConstantArgs>(Item.At);
+            float result = Args.value;
+            StoreScalar( Item, result );
             break;
         }
 
         case EOpType::SC_PARAMETER:
         {
-			OP::ParameterArgs args = Program.GetOpArgs<OP::ParameterArgs>(item.At);
-			TSharedPtr<FRangeIndex> Index = BuildCurrentOpRangeIndex( item, pParams, pModel, args.variable );
-            float result = pParams->GetFloatValue( args.variable, Index.Get());
-            StoreScalar( item, result );
+			OP::ParameterArgs Args = Program.GetOpArgs<OP::ParameterArgs>(Item.At);
+			TSharedPtr<FRangeIndex> Index = BuildCurrentOpRangeIndex( Item, pParams, InModel, Args.variable );
+            float result = pParams->GetFloatValue( Args.variable, Index.Get());
+            StoreScalar( Item, result );
             break;
         }
 
         case EOpType::SC_CURVE:
         {
-			OP::ScalarCurveArgs args = Program.GetOpArgs<OP::ScalarCurveArgs>(item.At);
-            switch (item.Stage)
+			OP::ScalarCurveArgs Args = Program.GetOpArgs<OP::ScalarCurveArgs>(Item.At);
+            switch (Item.Stage)
             {
             case 0:
-                AddOp( FScheduledOp( item.At, item, 1),
-                        FScheduledOp( args.time, item) );
+                AddOp( FScheduledOp( Item.At, Item, 1),
+                        FScheduledOp( Args.time, Item) );
                 break;
 
             case 1:
             {
-                float time = LoadScalar( FCacheAddress(args.time,item) );
+                float time = LoadScalar( FCacheAddress(Args.time,Item) );
 
-                const FRichCurve& Curve = Program.ConstantCurves[args.curve];
+                const FRichCurve& Curve = Program.ConstantCurves[Args.curve];
                 float Result = Curve.Eval(time);
 
-                StoreScalar( item, Result );
+                StoreScalar( Item, Result );
                 break;
             }
 
@@ -5325,22 +5325,22 @@ namespace mu
 
         case EOpType::SC_ARITHMETIC:
         {
-			OP::ArithmeticArgs args = Program.GetOpArgs<OP::ArithmeticArgs>(item.At);
-            switch (item.Stage)
+			OP::ArithmeticArgs Args = Program.GetOpArgs<OP::ArithmeticArgs>(Item.At);
+            switch (Item.Stage)
             {
             case 0:
-                    AddOp( FScheduledOp( item.At, item, 1),
-                           FScheduledOp( args.a, item),
-                           FScheduledOp( args.b, item) );
+                    AddOp( FScheduledOp( Item.At, Item, 1),
+                           FScheduledOp( Args.a, Item),
+                           FScheduledOp( Args.b, Item) );
                 break;
 
             case 1:
             {
-                float a = LoadScalar( FCacheAddress(args.a,item) );
-                float b = LoadScalar( FCacheAddress(args.b,item) );
+                float a = LoadScalar( FCacheAddress(Args.a,Item) );
+                float b = LoadScalar( FCacheAddress(Args.b,Item) );
 
                 float result = 1.0f;
-                switch (args.operation)
+                switch (Args.operation)
                 {
                 case OP::ArithmeticArgs::ADD:
                     result = a + b;
@@ -5363,7 +5363,7 @@ namespace mu
                     break;
                 }
 
-                StoreScalar( item, result );
+                StoreScalar( Item, result );
                 break;
             }
 
@@ -5380,34 +5380,34 @@ namespace mu
     }
 
     //---------------------------------------------------------------------------------------------
-    void CodeRunner::RunCode_String(const FScheduledOp& item, const FParameters* pParams, const FModel* pModel )
+    void CodeRunner::RunCode_String(const FScheduledOp& Item, const FParameters* pParams, const FModel* InModel )
     {
         MUTABLE_CPUPROFILER_SCOPE(RunCode_String );
 
 		const FProgram& Program = Model->GetPrivate()->Program;
 
-		EOpType type = Program.GetOpType( item.At );
+		EOpType type = Program.GetOpType( Item.At );
         switch ( type )
         {
 
         case EOpType::ST_CONSTANT:
         {
-			OP::ResourceConstantArgs args = Program.GetOpArgs<OP::ResourceConstantArgs>( item.At );
-            check( args.value < (uint32)pModel->GetPrivate()->Program.ConstantStrings.Num() );
+			OP::ResourceConstantArgs Args = Program.GetOpArgs<OP::ResourceConstantArgs>( Item.At );
+            check( Args.value < (uint32)InModel->GetPrivate()->Program.ConstantStrings.Num() );
 
-            const FString& result = Program.ConstantStrings[args.value];
-            StoreString( item, MakeShared<String>(result) );
+            const FString& result = Program.ConstantStrings[Args.value];
+            StoreString( Item, MakeShared<String>(result) );
 
             break;
         }
 
         case EOpType::ST_PARAMETER:
         {
-			OP::ParameterArgs args = Program.GetOpArgs<OP::ParameterArgs>( item.At );
-			TSharedPtr<FRangeIndex> Index = BuildCurrentOpRangeIndex( item, pParams, pModel, args.variable );
+			OP::ParameterArgs Args = Program.GetOpArgs<OP::ParameterArgs>( Item.At );
+			TSharedPtr<FRangeIndex> Index = BuildCurrentOpRangeIndex( Item, pParams, InModel, Args.variable );
 			FString result;
-			pParams->GetStringValue(args.variable, result, Index.Get());
-            StoreString( item, MakeShared<String>(result) );
+			pParams->GetStringValue(Args.variable, result, Index.Get());
+            StoreString( Item, MakeShared<String>(result) );
             break;
         }
 
@@ -5419,63 +5419,63 @@ namespace mu
 
 
     //---------------------------------------------------------------------------------------------
-    void CodeRunner::RunCode_Colour(const FScheduledOp& item, const FParameters* pParams, const FModel* pModel )
+    void CodeRunner::RunCode_Colour(const FScheduledOp& Item, const FParameters* pParams, const FModel* InModel )
     {
 		MUTABLE_CPUPROFILER_SCOPE(RunCode_Colour);
 
 		const FProgram& Program = Model->GetPrivate()->Program;
 
-		EOpType type = Program.GetOpType(item.At);
+		EOpType type = Program.GetOpType(Item.At);
 
         switch ( type )
         {
 
         case EOpType::CO_CONSTANT:
         {
-			OP::ColourConstantArgs args = Program.GetOpArgs<OP::ColourConstantArgs>(item.At);
+			OP::ColourConstantArgs Args = Program.GetOpArgs<OP::ColourConstantArgs>(Item.At);
 			FVector4f result;
-            result[0] = args.value[0];
-            result[1] = args.value[1];
-            result[2] = args.value[2];
-            result[3] = args.value[3];
-            StoreColor( item, result );
+            result[0] = Args.value[0];
+            result[1] = Args.value[1];
+            result[2] = Args.value[2];
+            result[3] = Args.value[3];
+            StoreColor( Item, result );
             break;
         }
 
         case EOpType::CO_PARAMETER:
         {
-			OP::ParameterArgs args = Program.GetOpArgs<OP::ParameterArgs>(item.At);
-			TSharedPtr<FRangeIndex> Index = BuildCurrentOpRangeIndex( item, pParams, pModel, args.variable );
+			OP::ParameterArgs Args = Program.GetOpArgs<OP::ParameterArgs>(Item.At);
+			TSharedPtr<FRangeIndex> Index = BuildCurrentOpRangeIndex( Item, pParams, InModel, Args.variable );
             FVector4f V;
-            pParams->GetColourValue( args.variable, V, Index.Get());
-            StoreColor( item, V );
+            pParams->GetColourValue( Args.variable, V, Index.Get());
+            StoreColor( Item, V );
             break;
         }
 
         case EOpType::CO_SAMPLEIMAGE:
         {
-			OP::ColourSampleImageArgs args = Program.GetOpArgs<OP::ColourSampleImageArgs>(item.At);
-            switch (item.Stage)
+			OP::ColourSampleImageArgs Args = Program.GetOpArgs<OP::ColourSampleImageArgs>(Item.At);
+            switch (Item.Stage)
             {
             case 0:
-                    AddOp( FScheduledOp( item.At, item, 1),
-                           FScheduledOp( args.x, item),
-                           FScheduledOp( args.y, item),
+                    AddOp( FScheduledOp( Item.At, Item, 1),
+                           FScheduledOp( Args.x, Item),
+                           FScheduledOp( Args.y, Item),
 						   // Don't skip mips for the texture to sample
-                           FScheduledOp::FromOpAndOptions( args.image, item, 0) );
+                           FScheduledOp::FromOpAndOptions( Args.image, Item, 0) );
                 break;
 
             case 1:
             {
-                float x = args.x ? LoadScalar( FCacheAddress(args.x,item) ) : 0.5f;
-                float y = args.y ? LoadScalar( FCacheAddress(args.y,item) ) : 0.5f;
+                float x = Args.x ? LoadScalar( FCacheAddress(Args.x,Item) ) : 0.5f;
+                float y = Args.y ? LoadScalar( FCacheAddress(Args.y,Item) ) : 0.5f;
 
-                TSharedPtr<const FImage> pImage = LoadImage(FScheduledOp::FromOpAndOptions(args.image, item, 0));
+                TSharedPtr<const FImage> pImage = LoadImage(FScheduledOp::FromOpAndOptions(Args.image, Item, 0));
 
 				FVector4f result;
                 if (pImage)
                 {
-                    if (args.filter)
+                    if (Args.filter)
                     {
                         // TODO
                         result = pImage->Sample(FVector2f(x, y));
@@ -5491,7 +5491,7 @@ namespace mu
                 }
 
 				Release(pImage);
-                StoreColor( item, result );
+                StoreColor( Item, result );
                 break;
             }
 
@@ -5504,15 +5504,15 @@ namespace mu
 
         case EOpType::CO_SWIZZLE:
         {
-			OP::ColourSwizzleArgs args = Program.GetOpArgs<OP::ColourSwizzleArgs>(item.At);
-            switch (item.Stage)
+			OP::ColourSwizzleArgs Args = Program.GetOpArgs<OP::ColourSwizzleArgs>(Item.At);
+            switch (Item.Stage)
             {
             case 0:
-                    AddOp( FScheduledOp( item.At, item, 1),
-                           FScheduledOp( args.sources[0], item),
-                           FScheduledOp( args.sources[1], item),
-                           FScheduledOp( args.sources[2], item),
-                           FScheduledOp( args.sources[3], item) );
+                    AddOp( FScheduledOp( Item.At, Item, 1),
+                           FScheduledOp( Args.sources[0], Item),
+                           FScheduledOp( Args.sources[1], Item),
+                           FScheduledOp( Args.sources[2], Item),
+                           FScheduledOp( Args.sources[3], Item) );
                 break;
 
             case 1:
@@ -5521,14 +5521,14 @@ namespace mu
 
                 for (int32 t=0;t<MUTABLE_OP_MAX_SWIZZLE_CHANNELS;++t)
                 {
-                    if ( args.sources[t] )
+                    if ( Args.sources[t] )
                     {
-                        FVector4f p = LoadColor( FCacheAddress(args.sources[t],item) );
-                        result[t] = p[ args.sourceChannels[t] ];
+                        FVector4f p = LoadColor( FCacheAddress(Args.sources[t],Item) );
+                        result[t] = p[ Args.sourceChannels[t] ];
                     }
                 }
 
-                StoreColor( item, result );
+                StoreColor( Item, result );
                 break;
             }
 
@@ -5541,15 +5541,15 @@ namespace mu
 
         case EOpType::CO_FROMSCALARS:
         {
-			OP::ColourFromScalarsArgs args = Program.GetOpArgs<OP::ColourFromScalarsArgs>(item.At);
-            switch (item.Stage)
+			OP::ColourFromScalarsArgs Args = Program.GetOpArgs<OP::ColourFromScalarsArgs>(Item.At);
+            switch (Item.Stage)
             {
             case 0:
-                    AddOp( FScheduledOp( item.At, item, 1),
-                           FScheduledOp( args.v[0], item),
-                           FScheduledOp( args.v[1], item),
-                           FScheduledOp( args.v[2], item),
-                           FScheduledOp( args.v[3], item));
+                    AddOp( FScheduledOp( Item.At, Item, 1),
+                           FScheduledOp( Args.v[0], Item),
+                           FScheduledOp( Args.v[1], Item),
+                           FScheduledOp( Args.v[2], Item),
+                           FScheduledOp( Args.v[3], Item));
                 break;
 
             case 1:
@@ -5558,13 +5558,13 @@ namespace mu
 
 				for (int32 t = 0; t < MUTABLE_OP_MAX_SWIZZLE_CHANNELS; ++t)
 				{
-					if (args.v[t])
+					if (Args.v[t])
 					{
-						Result[t] = LoadScalar(FCacheAddress(args.v[t], item));
+						Result[t] = LoadScalar(FCacheAddress(Args.v[t], Item));
 					}
 				}
 
-                StoreColor( item, Result );
+                StoreColor( Item, Result );
                 break;
             }
 
@@ -5577,30 +5577,30 @@ namespace mu
 
         case EOpType::CO_ARITHMETIC:
         {
-			OP::ArithmeticArgs args = Program.GetOpArgs<OP::ArithmeticArgs>(item.At);
-            switch (item.Stage)
+			OP::ArithmeticArgs Args = Program.GetOpArgs<OP::ArithmeticArgs>(Item.At);
+            switch (Item.Stage)
             {
             case 0:
-                    AddOp( FScheduledOp( item.At, item, 1),
-                           FScheduledOp( args.a, item),
-                           FScheduledOp( args.b, item));
+                    AddOp( FScheduledOp( Item.At, Item, 1),
+                           FScheduledOp( Args.a, Item),
+                           FScheduledOp( Args.b, Item));
                 break;
 
             case 1:
             {
-				EOpType otype = Program.GetOpType( args.a );
-                DATATYPE dtype = GetOpDataType( otype );
-                check( dtype == DT_COLOUR );
-                otype = Program.GetOpType( args.b );
+				EOpType otype = Program.GetOpType( Args.a );
+                EDataType dtype = GetOpDataType( otype );
+                check( dtype == EDataType::Color );
+                otype = Program.GetOpType( Args.b );
                 dtype = GetOpDataType( otype );
-                check( dtype == DT_COLOUR );
-				FVector4f a = args.a ? LoadColor( FCacheAddress( args.a, item ) )
+                check( dtype == EDataType::Color );
+				FVector4f a = Args.a ? LoadColor( FCacheAddress( Args.a, Item ) )
                                  : FVector4f( 0, 0, 0, 0 );
-				FVector4f b = args.b ? LoadColor( FCacheAddress( args.b, item ) )
+				FVector4f b = Args.b ? LoadColor( FCacheAddress( Args.b, Item ) )
                                  : FVector4f( 0, 0, 0, 0 );
 
 				FVector4f result = FVector4f(0,0,0,0);
-                switch (args.operation)
+                switch (Args.operation)
                 {
                 case OP::ArithmeticArgs::ADD:
                     result = a + b;
@@ -5623,7 +5623,7 @@ namespace mu
                     break;
                 }
 
-                StoreColor( item, result );
+                StoreColor( Item, result );
                 break;
             }
 
@@ -5642,35 +5642,35 @@ namespace mu
 
 
     //---------------------------------------------------------------------------------------------
-    void CodeRunner::RunCode_Projector(const FScheduledOp& item, const FParameters* pParams, const FModel* pModel )
+    void CodeRunner::RunCode_Projector(const FScheduledOp& Item, const FParameters* pParams, const FModel* InModel )
     {
         MUTABLE_CPUPROFILER_SCOPE(RunCode_Projector);
 
 		const FProgram& Program = Model->GetPrivate()->Program;
 
-		EOpType type = Program.GetOpType(item.At);
+		EOpType type = Program.GetOpType(Item.At);
         switch (type)
         {
 
         case EOpType::PR_CONSTANT:
         {
-			OP::ResourceConstantArgs args = Program.GetOpArgs<OP::ResourceConstantArgs>(item.At);
-            FProjector Result = Program.ConstantProjectors[args.value];
-            StoreProjector( item, Result );
+			OP::ResourceConstantArgs Args = Program.GetOpArgs<OP::ResourceConstantArgs>(Item.At);
+            FProjector Result = Program.ConstantProjectors[Args.value];
+            StoreProjector( Item, Result );
             break;
         }
 
         case EOpType::PR_PARAMETER:
         {
-			OP::ParameterArgs args = Program.GetOpArgs<OP::ParameterArgs>(item.At);
-			TSharedPtr<FRangeIndex> Index = BuildCurrentOpRangeIndex( item, pParams, pModel, args.variable );
-            FProjector Result = pParams->GetPrivate()->GetProjectorValue(args.variable,Index.Get());
+			OP::ParameterArgs Args = Program.GetOpArgs<OP::ParameterArgs>(Item.At);
+			TSharedPtr<FRangeIndex> Index = BuildCurrentOpRangeIndex( Item, pParams, InModel, Args.variable );
+            FProjector Result = pParams->GetPrivate()->GetProjectorValue(Args.variable,Index.Get());
 
             // The type cannot be changed, take it from the default value
-            const FProjector& def = Program.Parameters[args.variable].DefaultValue.Get<FParamProjectorType>();
+            const FProjector& def = Program.Parameters[Args.variable].DefaultValue.Get<FParamProjectorType>();
             Result.type = def.type;
 
-            StoreProjector( item, Result );
+            StoreProjector( Item, Result );
             break;
         }
 
@@ -5681,30 +5681,30 @@ namespace mu
     }
 
     //---------------------------------------------------------------------------------------------
-    void CodeRunner::RunCode_Matrix(const FScheduledOp& item, const FParameters* pParams, const FModel* pModel )
+    void CodeRunner::RunCode_Matrix(const FScheduledOp& Item, const FParameters* pParams, const FModel* InModel )
     {
 		MUTABLE_CPUPROFILER_SCOPE(RunCode_Transform);
 
 		const FProgram& Program = Model->GetPrivate()->Program;
 
-		EOpType type = Program.GetOpType(item.At);
+		EOpType type = Program.GetOpType(Item.At);
 
 		switch ( type )
 		{
 		case EOpType::MA_CONSTANT:
 			{
-				OP::MatrixConstantArgs args = Program.GetOpArgs<OP::MatrixConstantArgs>(item.At);
-				StoreMatrix( item, Program.ConstantMatrices[args.value] );
+				OP::MatrixConstantArgs Args = Program.GetOpArgs<OP::MatrixConstantArgs>(Item.At);
+				StoreMatrix( Item, Program.ConstantMatrices[Args.value] );
 				break;
 			}
 
 		case EOpType::MA_PARAMETER:
 			{
-				OP::ParameterArgs args = Program.GetOpArgs<OP::ParameterArgs>(item.At);
-				TSharedPtr<FRangeIndex> Index = BuildCurrentOpRangeIndex( item, pParams, pModel, args.variable );
+				OP::ParameterArgs Args = Program.GetOpArgs<OP::ParameterArgs>(Item.At);
+				TSharedPtr<FRangeIndex> Index = BuildCurrentOpRangeIndex( Item, pParams, InModel, Args.variable );
 				FMatrix44f Value;
-				pParams->GetMatrixValue( args.variable, Value, Index.Get());
-				StoreMatrix( item, Value );
+				pParams->GetMatrixValue( Args.variable, Value, Index.Get());
+				StoreMatrix( Item, Value );
 				break;
 			}
 		}
@@ -5712,42 +5712,42 @@ namespace mu
 
 
     //---------------------------------------------------------------------------------------------
-    void CodeRunner::RunCode_Layout(const FScheduledOp& item, const FModel* pModel )
+    void CodeRunner::RunCode_Layout(const FScheduledOp& Item, const FModel* InModel )
     {
         //MUTABLE_CPUPROFILER_SCOPE(RunCode_Layout);
 
 		const FProgram& Program = Model->GetPrivate()->Program;
 
-		EOpType type = Program.GetOpType(item.At);
+		EOpType type = Program.GetOpType(Item.At);
         switch (type)
         {
 
         case EOpType::LA_CONSTANT:
         {
-			OP::ResourceConstantArgs args = Program.GetOpArgs<OP::ResourceConstantArgs>(item.At);
-            check( args.value < (uint32)pModel->GetPrivate()->Program.ConstantLayouts.Num() );
+			OP::ResourceConstantArgs Args = Program.GetOpArgs<OP::ResourceConstantArgs>(Item.At);
+            check( Args.value < (uint32)InModel->GetPrivate()->Program.ConstantLayouts.Num() );
 
             TSharedPtr<const FLayout> Result = Program.ConstantLayouts
-                    [ args.value ];
-            StoreLayout( item, Result );
+                    [ Args.value ];
+            StoreLayout( Item, Result );
             break;
         }
 
         case EOpType::LA_MERGE:
         {
-			OP::LayoutMergeArgs args = Program.GetOpArgs<OP::LayoutMergeArgs>(item.At);
-            switch (item.Stage)
+			OP::LayoutMergeArgs Args = Program.GetOpArgs<OP::LayoutMergeArgs>(Item.At);
+            switch (Item.Stage)
             {
             case 0:
-                    AddOp( FScheduledOp( item.At, item, 1),
-                           FScheduledOp( args.Base, item),
-                           FScheduledOp( args.Added, item) );
+                    AddOp( FScheduledOp( Item.At, Item, 1),
+                           FScheduledOp( Args.Base, Item),
+                           FScheduledOp( Args.Added, Item) );
                 break;
 
             case 1:
             {
-                TSharedPtr<const FLayout> pA = LoadLayout( FCacheAddress(args.Base,item) );
-                TSharedPtr<const FLayout> pB = LoadLayout( FCacheAddress(args.Added,item) );
+                TSharedPtr<const FLayout> pA = LoadLayout( FCacheAddress(Args.Base,Item) );
+                TSharedPtr<const FLayout> pB = LoadLayout( FCacheAddress(Args.Added,Item) );
 
                 TSharedPtr<const FLayout> Result;
 
@@ -5764,7 +5764,7 @@ namespace mu
                     Result = pB->Clone();
                 }
 
-                StoreLayout( item, Result );
+                StoreLayout( Item, Result );
                 break;
             }
 
@@ -5777,17 +5777,17 @@ namespace mu
 
         case EOpType::LA_PACK:
         {
-			OP::LayoutPackArgs args = Program.GetOpArgs<OP::LayoutPackArgs>(item.At);
-            switch (item.Stage)
+			OP::LayoutPackArgs Args = Program.GetOpArgs<OP::LayoutPackArgs>(Item.At);
+            switch (Item.Stage)
             {
             case 0:
-                    AddOp( FScheduledOp( item.At, item, 1),
-                           FScheduledOp( args.Source, item) );
+                    AddOp( FScheduledOp( Item.At, Item, 1),
+                           FScheduledOp( Args.Source, Item) );
                 break;
 
             case 1:
             {
-                TSharedPtr<const FLayout> Source = LoadLayout( FCacheAddress(args.Source,item) );
+                TSharedPtr<const FLayout> Source = LoadLayout( FCacheAddress(Args.Source,Item) );
 
 				TSharedPtr<FLayout> Result;
 
@@ -5798,7 +5798,7 @@ namespace mu
 					LayoutPack3(Result.Get(), Source.Get() );
 				}
 
-                StoreLayout( item, Result );
+                StoreLayout( Item, Result );
                 break;
             }
 
@@ -5811,22 +5811,22 @@ namespace mu
 
 		case EOpType::LA_FROMMESH:
 		{
-			OP::LayoutFromMeshArgs args = Program.GetOpArgs<OP::LayoutFromMeshArgs>(item.At);
-			switch (item.Stage)
+			OP::LayoutFromMeshArgs Args = Program.GetOpArgs<OP::LayoutFromMeshArgs>(Item.At);
+			switch (Item.Stage)
 			{
 			case 0:
-				AddOp(FScheduledOp(item.At, item, 1),
-					FScheduledOp(args.Mesh, item));
+				AddOp(FScheduledOp(Item.At, Item, 1),
+					FScheduledOp(Args.Mesh, Item));
 				break;
 
 			case 1:
 			{
-				TSharedPtr<const FMesh> Mesh = LoadMesh(FCacheAddress(args.Mesh, item));
+				TSharedPtr<const FMesh> Mesh = LoadMesh(FCacheAddress(Args.Mesh, Item));
 
-				TSharedPtr<const FLayout> Result = LayoutFromMesh_RemoveBlocks(Mesh.Get(), args.LayoutIndex);
+				TSharedPtr<const FLayout> Result = LayoutFromMesh_RemoveBlocks(Mesh.Get(), Args.LayoutIndex);
 
 				Release(Mesh);
-				StoreLayout(item, Result);
+				StoreLayout(Item, Result);
 				break;
 			}
 
@@ -5839,19 +5839,19 @@ namespace mu
 
 		case EOpType::LA_REMOVEBLOCKS:
 		{
-			OP::LayoutRemoveBlocksArgs args = Program.GetOpArgs<OP::LayoutRemoveBlocksArgs>(item.At);
-			switch (item.Stage)
+			OP::LayoutRemoveBlocksArgs Args = Program.GetOpArgs<OP::LayoutRemoveBlocksArgs>(Item.At);
+			switch (Item.Stage)
 			{
 			case 0:
-				AddOp(FScheduledOp(item.At, item, 1),
-					FScheduledOp(args.Source, item),
-					FScheduledOp(args.ReferenceLayout, item));
+				AddOp(FScheduledOp(Item.At, Item, 1),
+					FScheduledOp(Args.Source, Item),
+					FScheduledOp(Args.ReferenceLayout, Item));
 				break;
 
 			case 1:
 			{
-				TSharedPtr<const FLayout> Source = LoadLayout(FCacheAddress(args.Source, item));
-				TSharedPtr<const FLayout> ReferenceLayout = LoadLayout(FCacheAddress(args.ReferenceLayout, item));
+				TSharedPtr<const FLayout> Source = LoadLayout(FCacheAddress(Args.Source, Item));
+				TSharedPtr<const FLayout> ReferenceLayout = LoadLayout(FCacheAddress(Args.ReferenceLayout, Item));
 
 				TSharedPtr<const FLayout> Result;
 
@@ -5864,7 +5864,7 @@ namespace mu
 					Result = Source;
 				}
 
-				StoreLayout(item, Result);
+				StoreLayout(Item, Result);
 				break;
 			}
 
@@ -5884,17 +5884,15 @@ namespace mu
 
 
     //---------------------------------------------------------------------------------------------
-    void CodeRunner::RunCode( const FScheduledOp& item, const FParameters* pParams, const TSharedPtr<const FModel>& InModel, uint32 lodMask)
+    void CodeRunner::RunCode( const FScheduledOp& Item, const FParameters* pParams, const TSharedPtr<const FModel>& InModel, uint32 LodMask)
     {
-		//UE_LOG( LogMutableCore, Log, TEXT("Running :%5d , %d "), item.At, item.Stage );
-		check( item.Type == FScheduledOp::EType::Full );
-
-		const FModel* pModel = InModel.Get();
+		//UE_LOG( LogMutableCore, Log, TEXT("Running :%5d , %d "), Item.At, Item.Stage );
+		check( Item.Type == FScheduledOp::EType::Full );
 		
-		const FProgram& Program = pModel->GetPrivate()->Program;
+		const FProgram& Program = InModel->GetPrivate()->Program;
 
-		EOpType type = Program.GetOpType(item.At);
-		//UE_LOG(LogMutableCore, Log, TEXT("Running :%5d , %d, of type %d "), item.At, item.Stage, type);
+		EOpType type = Program.GetOpType(Item.At);
+		//UE_LOG(LogMutableCore, Log, TEXT("Running :%5d , %d, of type %d "), Item.At, Item.Stage, type);
 
 		// Very spammy, for debugging purposes.
 		//if (System)
@@ -5915,13 +5913,13 @@ namespace mu
         case EOpType::LA_CONDITIONAL:
         case EOpType::IN_CONDITIONAL:
 		case EOpType::ED_CONDITIONAL:
-            RunCode_Conditional(item, pModel);
+            RunCode_Conditional(Item, InModel.Get());
             break;
 
         case EOpType::ME_CONSTANT:
 		case EOpType::IM_CONSTANT:
 		case EOpType::ED_CONSTANT:
-            RunCode_ConstantResource(item, pModel);
+            RunCode_ConstantResource(Item, InModel.Get());
             break;
 
         case EOpType::NU_SWITCH:
@@ -5932,61 +5930,61 @@ namespace mu
         case EOpType::LA_SWITCH:
         case EOpType::IN_SWITCH:
 		case EOpType::ED_SWITCH:
-            RunCode_Switch(item, pModel);
+            RunCode_Switch(Item, InModel.Get());
             break;
 
         case EOpType::IN_ADDMESH:
         case EOpType::IN_ADDIMAGE:
-            RunCode_InstanceAddResource(item, InModel, pParams);
+            RunCode_InstanceAddResource(Item, InModel, pParams);
             break;
 
 		default:
 		{
-			DATATYPE DataType = GetOpDataType(type);
+			EDataType DataType = GetOpDataType(type);
 			switch (DataType)
 			{
-			case DT_INSTANCE:
-				RunCode_Instance(item, pModel, lodMask);
+			case EDataType::Instance:
+				RunCode_Instance(Item, InModel.Get(), LodMask);
 				break;
 
-			case DT_MESH:
-				RunCode_Mesh(item, pModel);
+			case EDataType::Mesh:
+				RunCode_Mesh(Item, InModel.Get());
 				break;
 
-			case DT_IMAGE:
-				RunCode_Image(item, pParams, pModel);
+			case EDataType::Image:
+				RunCode_Image(Item, pParams, InModel.Get());
 				break;
 
-			case DT_LAYOUT:
-				RunCode_Layout(item, pModel);
+			case EDataType::Layout:
+				RunCode_Layout(Item, InModel.Get());
 				break;
 
-			case DT_BOOL:
-				RunCode_Bool(item, pParams, pModel);
+			case EDataType::Bool:
+				RunCode_Bool(Item, pParams, InModel.Get());
 				break;
 
-			case DT_SCALAR:
-				RunCode_Scalar(item, pParams, pModel);
+			case EDataType::Scalar:
+				RunCode_Scalar(Item, pParams, InModel.Get());
 				break;
 
-			case DT_STRING:
-				RunCode_String(item, pParams, pModel);
+			case EDataType::String:
+				RunCode_String(Item, pParams, InModel.Get());
 				break;
 
-			case DT_INT:
-				RunCode_Int(item, pParams, pModel);
+			case EDataType::Int:
+				RunCode_Int(Item, pParams, InModel.Get());
 				break;
 
-			case DT_PROJECTOR:
-				RunCode_Projector(item, pParams, pModel);
+			case EDataType::Projector:
+				RunCode_Projector(Item, pParams, InModel.Get());
 				break;
 
-			case DT_COLOUR:
-				RunCode_Colour(item, pParams, pModel);
+			case EDataType::Color:
+				RunCode_Colour(Item, pParams, InModel.Get());
 				break;
 
-			case DT_MATRIX:
-				RunCode_Matrix(item, pParams, pModel);
+			case EDataType::Matrix:
+				RunCode_Matrix(Item, pParams, InModel.Get());
 				break;
 				
 			default:
@@ -6000,7 +5998,7 @@ namespace mu
     }
 
 	//---------------------------------------------------------------------------------------------
-	void CodeRunner::RunCodeImageDesc(const FScheduledOp& Item, const FParameters* pParams, const FModel* pModel,  uint32 lodMask )
+	void CodeRunner::RunCodeImageDesc(const FScheduledOp& Item, const FParameters* pParams, const FModel* InModel,  uint32 LodMask )
 	{
 		MUTABLE_CPUPROFILER_SCOPE(RunCodeImageDesc);
 
@@ -6051,7 +6049,7 @@ namespace mu
 				else
 				{
 					uint32 RomId = CurrentIndex.Index;
-					bIsLODAvailable = System->StreamInterface->DoesBlockExist(pModel, RomId);
+					bIsLODAvailable = System->StreamInterface->DoesBlockExist(Model.Get(), RomId);
 				}
 
 				if (bIsLODAvailable)
@@ -6749,6 +6747,7 @@ namespace mu
 					}
 				}
 
+				HeapImageDesc[Item.CustomState].m_format = Args.format;
 				StoreValidDesc(Item);
 				break;
 			}
@@ -7005,4 +7004,33 @@ namespace mu
 			break;
 		}
 	}
+
+
+	FImageOperator MakeImageOperator(CodeRunner* Runner)
+	{
+		return FImageOperator(
+			// Create
+			[Runner](int32 x, int32 y, int32 m, EImageFormat f, EInitializationType i)
+			{
+				return Runner->CreateImage(x, y, m, f, i);
+			},
+
+			// Release
+			[Runner](TSharedPtr<FImage>& Image)
+			{
+				Runner->Release(Image);
+			},
+
+			// Clone
+			[Runner](const FImage* Image)
+			{
+				TSharedPtr<FImage> New = Runner->CreateImage(Image->GetSizeX(), Image->GetSizeY(), Image->GetLODCount(), Image->GetFormat(), EInitializationType::NotInitialized);
+				New->Copy(Image);
+				return New;
+			},
+
+			Runner->System->ImagePixelFormatOverride
+		);
+	}
+
 }
