[mutable] Remove unused static constant tables, and move some to the tools library.

diff --git a/Engine/Plugins/Mutable/Source/MutableRuntime/Private/MuR/CodeRunner.h b/Engine/Plugins/Mutable/Source/MutableRuntime/Private/MuR/CodeRunner.h
--- a/Engine/Plugins/Mutable/Source/MutableRuntime/Private/MuR/CodeRunner.h
+++ b/Engine/Plugins/Mutable/Source/MutableRuntime/Private/MuR/CodeRunner.h
@@ -681,54 +681,16 @@ namespace  mu
     	
     	FRomLoadOps RomLoadOps = FRomLoadOps(*this);
     	    	
-		inline void AddChildren(const FScheduledOp& dep)
-		{
-			FCacheAddress at(dep);
-			if (dep.At && !GetMemory().IsValid(at))
-			{
-				if (ScheduledStagePerOp.get(at) <= dep.Stage)
-				{
-					OpenTasks.Add(dep);
-					ScheduledStagePerOp[at] = dep.Stage + 1;
-				}
-			}
-
-			if (dep.Type == FScheduledOp::EType::Full)
-			{
-				System->WorkingMemoryManager.CurrentInstanceCache->IncreaseHitCount(at);
-			}
-		}
+		void AddChildren(const FScheduledOp& dep);
 
 		/** Try to create a concurrent task for the given op. Return null if not possible. */
 		TSharedPtr<FIssuedTask> IssueOp(FScheduledOp item);
 
-		/** Calculate an heuristic to select op execution based on memory usage. */
-		int32 GetOpEstimatedMemoryDelta(const FScheduledOp& Candidate, const FProgram& Program);
-
 		/** Update debug stats. */
 		void UpdateTraces();
 
 		/** */
-		inline bool ShouldIssueTask() const
-		{
-			// Can we afford to delay issued tasks?
-			bool bCanDelayTasks = IssuedTasks.Num() > 0 || OpenTasks.Num() > 0;
-			if (!bCanDelayTasks)
-			{
-				return true;
-			}
-			else
-			{
-				// We could wait. Let's see if we have enough memory to issue tasks anyway.
-				bool bHaveEnoughMemory = !System->WorkingMemoryManager.IsMemoryBudgetFull();
-				if (bHaveEnoughMemory)
-				{
-					return  true;
-				}
-			}
-
-			return false;
-		}
+		bool ShouldIssueTask() const;
 
 		/** */
 		void LaunchIssuedTask(const TSharedPtr<FIssuedTask>& TaskToIssue, bool& bOutFailed);
@@ -737,31 +699,6 @@ namespace  mu
 
 
 	/** Helper function to create the memory-tracked image operator. */
-	inline FImageOperator MakeImageOperator(CodeRunner* Runner)
-	{
-		return FImageOperator(
-			// Create
-			[Runner](int32 x, int32 y, int32 m, EImageFormat f, EInitializationType i)
-			{
-				return Runner->CreateImage(x, y, m, f, i);
-			},
-
-			// Release
-			[Runner](TSharedPtr<FImage>& i)
-			{
-				Runner->Release(i);
-			},
-
-			// Clone
-			[Runner](const FImage* i)
-			{
-				TSharedPtr<FImage> New = Runner->CreateImage(i->GetSizeX(), i->GetSizeY(), i->GetLODCount(), i->GetFormat(), EInitializationType::NotInitialized);
-				New->Copy(i);
-				return New;
-			},
-
-			Runner->System->ImagePixelFormatOverride
-		);
-	}
+	extern FImageOperator MakeImageOperator(CodeRunner* Runner);
 
 }
diff --git a/Engine/Plugins/Mutable/Source/MutableRuntime/Private/MuR/CodeRunnerTasked.cpp b/Engine/Plugins/Mutable/Source/MutableRuntime/Private/MuR/CodeRunnerTasked.cpp
--- a/Engine/Plugins/Mutable/Source/MutableRuntime/Private/MuR/CodeRunnerTasked.cpp
+++ b/Engine/Plugins/Mutable/Source/MutableRuntime/Private/MuR/CodeRunnerTasked.cpp
@@ -55,190 +55,62 @@ DECLARE_CYCLE_STAT(TEXT("MutableCoreTask"), STAT_MutableCoreTask, STATGROUP_Game
 namespace mu
 {
 
-	//---------------------------------------------------------------------------------------------
 	TRACE_DECLARE_INT_COUNTER(MutableRuntime_OpenTask, TEXT("MutableRuntime/OpenTask"));
 	TRACE_DECLARE_INT_COUNTER(MutableRuntime_ClosedTasks, TEXT("MutableRuntime/ClosedTasks"));
 	TRACE_DECLARE_INT_COUNTER(MutableRuntime_IssuedTasks, TEXT("MutableRuntime/IssuedTasks"));
 	TRACE_DECLARE_INT_COUNTER(MutableRuntime_IssuedTasksOnHold, TEXT("MutableRuntime/IssuedHoldTasks"));
 
-	void CodeRunner::UpdateTraces()
+	void CodeRunner::AddChildren(const FScheduledOp& dep)
 	{
-		// Code Runner status
-		TRACE_COUNTER_SET(MutableRuntime_OpenTask, OpenTasks.Num());
-		TRACE_COUNTER_SET(MutableRuntime_ClosedTasks, ClosedTasks.Num());
-		TRACE_COUNTER_SET(MutableRuntime_IssuedTasks, IssuedTasks.Num());
-		TRACE_COUNTER_SET(MutableRuntime_IssuedTasksOnHold, IssuedTasksOnHold.Num());
-	}
-
-
-	//---------------------------------------------------------------------------------------------
-	/** This table encodes an heuristic for every execution stage of every operation  that tries to 
-	* guess the memory delta of that stage: 
-	* - if negative: after the stage the working memory may have reduced
-	* - if positive: after the stage the working memory may have increased
-	* The actual number is just a unit-less weight that represents "amount of additional memory that will be used".
-	* Some operations don't have a constant memory delta and are dealt with specially, but still need 
-	* to be in this table.
-	*/
-	static constexpr int32 sMemoryWeightsStageCount = 4;
-	static const int8 sMemoryWeights[int32(EOpType::COUNT)][sMemoryWeightsStageCount] =
-	{ 
-		{   0,   0,   0,   0 },	// NONE
-		    
-		{   0,   0,   0,   0 },	// BO_CONSTANT
-		{   0,   0,   0,   0 },	// NU_CONSTANT
-		{   0,   0,   0,   0 },	// SC_CONSTANT
-		{   0,   0,   0,   0 },	// CO_CONSTANT
-		{  20,   0,   0,   0 },	// IM_CONSTANT
-		{  20,   0,   0,   0 },	// ME_CONSTANT
-		{   0,   0,   0,   0 },	// LA_CONSTANT
-		{   0,   0,   0,   0 },	// PR_CONSTANT
-		{   0,   0,   0,   0 },	// ST_CONSTANT
-		    
-		{   0,   0,   0,   0 },	// BO_PARAMETER
-		{   0,   0,   0,   0 },	// NU_PARAMETER
-		{   0,   0,   0,   0 },	// SC_PARAMETER
-		{   0,   0,   0,   0 },	// CO_PARAMETER
-		{   0,   0,   0,   0 },	// PR_PARAMETER
-		{  20,   0,   0,   0 },	// IM_PARAMETER
-		{   0,   0,   0,   0 },	// ST_PARAMETER
-		    
-		{   0,   0,   0,   0 },	// IM_REFERENCE
-		{   0,   0,   0,   0 },	// ME_REFERENCE
-
-		{   0,   0,   0,   0 },	// NU_CONDITIONAL
-		{   0,   0,   0,   0 },	// SC_CONDITIONAL
-		{   0,   0,   0,   0 },	// CO_CONDITIONAL
-		{   0,   0,   0,   0 },	// IM_CONDITIONAL
-		{   0,   0,   0,   0 },	// ME_CONDITIONAL
-		{   0,   0,   0,   0 },	// LA_CONDITIONAL
-		{   0,   0,   0,   0 },	// IN_CONDITIONAL
-		{   0,   0,   0,   0 },	// ED_CONDITIONAL
-		    
-		{   0,   0,   0,   0 },	// NU_SWITCH
-		{   0,   0,   0,   0 },	// SC_SWITCH
-		{   0,   0,   0,   0 },	// CO_SWITCH
-		{   0,   0,   0,   0 },	// IM_SWITCH
-		{   0,   0,   0,   0 },	// ME_SWITCH
-		{   0,   0,   0,   0 },	// LA_SWITCH
-		{   0,   0,   0,   0 },	// IN_SWITCH
-		{   0,   0,   0,   0 },	// ED_SWITCH
-		    
-		{   0,   0,   0,   0 },	// BO_LESS
-		{   0,   0,   0,   0 },	// BO_EQUAL_SC_CONST
-		{   0,   0,   0,   0 },	// BO_AND
-		{   0,   0,   0,   0 },	// BO_OR
-		{   0,   0,   0,   0 },	// BO_NOT
-		    
-		{   0,   0,   0,   0 },	// SC_MULTIPLYADD
-		{   0,   0,   0,   0 },	// SC_ARITHMETIC
-		{   0,   0,   0,   0 },	// SC_CURVE
-		    
-		{   0,   0,   0,   0 },	// CO_SAMPLEIMAGE
-		{   0,   0,   0,   0 },	// CO_SWIZZLE
-		{   0,   0,   0,   0 },	// CO_FROMSCALARS
-		{   0,   0,   0,   0 },	// CO_ARITHMETIC
-		    
-		{   0, -20, -20,   0 },	// IM_LAYER
-		{   0,   0,   0,   0 },	// IM_LAYERCOLOUR
-		{   0,   0,   0,   0 },	// IM_PIXELFORMAT	(special case)
-		{   0,  10,   0,   0 },	// IM_MIPMAP
-		{   0,   0,   0,   0 },	// IM_RESIZE		(special case)
-		{   0,   0,   0,   0 },	// IM_RESIZELIKE	(to be deprecated?)
-		{   0,   0,   0,   0 },	// IM_RESIZEREL		(special case)
-		{   0,  20,   0,   0 },	// IM_BLANKLAYOUT
-		{   0, -20, -20,   0 },	// IM_COMPOSE
-		{   0,   0, -20,   0 },	// IM_INTERPOLATE
-		{   0,   0,   0,   0 },	// IM_SATURATE
-		{   0,   0,   0,   0 },	// IM_LUMINANCE
-		{   0,   0,   0,   0 },	// IM_SWIZZLE
-		{   0,   0,   0,   0 },	// IM_COLOURMAP
-		{   0,   5,   0,   0 },	// IM_GRADIENT
-		{   0,   0,   0,   0 },	// IM_BINARISE
-		{   0,  20,   0,   0 },	// IM_PLAINCOLOUR
-		{   0, -10,   0,   0 },	// IM_CROP
-		{   0, -10,   0,   0 },	// IM_PATCH
-		{   0,  10,  10,   0 },	// IM_RASTERMESH
-		{   0,   0,   0,   0 },	// IM_MAKEGROWMAP
-		{   0, -10,   0,   0 },	// IM_DISPLACE
-		{   0,   0, -20, -20 },	// IM_MULTILAYER
-		{   0,   0,   0,   0 },	// IM_INVERT
-		{   0,   0,   0,   0 },	// IM_NORMALCOMPOSITE
-		{   0,   0,   0,   0 },	// IM_TRANSFORM
-		    
-		{   0,   0,   0,   0 },	// ME_APPLYLAYOUT
-		{   0, -20,   0,   0 },	// ME_DIFFERENCE
-		{   0,   0, -10,   0 },	// ME_MORPH2
-		{   0,   0,   0,   0 },	// ME_MERGE
-		{   0, -10,   0,   0 },	// ME_MASKCLIPMESH
-		{   0, -10,   0,   0 },	// ME_MASKCLIPUVMASK
-		{   0, -10,   0,   0 },	// ME_MASKDIFF
-		{   0,   0, -10,   0 },	// ME_REMOVEMASK
-		{   0,   0,   0,   0 },	// ME_FORMAT
-		{   0, -10,   0,   0 },	// ME_EXTRACTLAYOUTBLOCK
-		{   0,   0,   0,   0 },	// ME_TRANSFORM
-		{   0, -10,   0,   0 },	// ME_CLIPMORPHPLANE
-		{   0, -10,   0,   0 },	// ME_CLIPWITHMESH
-		{   0,   0,   0,   0 },	// ME_SETSKELETON
-		{   0,   0,   0,   0 },	// ME_PROJECT
-		{   0,   0,   0,   0 },	// ME_APPLYPOSE
-		{   0, -20,   0,   0 },	// ME_BINDSHAPE
-		{   0, -10,   0,   0 },	// ME_APPLYSHAPE
-		{   0, -10,   0,   0 },	// ME_CLIPDEFORM
-		{   0, -10,   0,   0 },	// ME_MORPHRESHAPE
-		{   0,   0,   0,   0 },	// ME_OPTIMIZESKINNING
-		{	0,   0,	  0,   0 },	// ME_ADDTAGS
-
-		{   0,   0,   0,   0 },	// IN_ADDMESH
-		{   0,   0,   0,   0 },	// IN_ADDIMAGE
-		{   0,   0,   0,   0 },	// IN_ADDVECTOR
-		{   0,   0,   0,   0 },	// IN_ADDSCALAR
-		{   0,   0,   0,   0 },	// IN_ADDSTRING
-		{   0,   0,   0,   0 },	// IN_ADDSURFACE
-		{   0,   0,   0,   0 },	// IN_ADDCOMPONENT
-		{	0,   0,   0,   0 },	// IN_ADDLOD
-		{	0,   0,   0,   0 },	// IN_ADDEXTENSIONDATA
-
-		{   0,   0,   0,   0 },	// LA_PACK
-		{   0,   0,   0,   0 },	// LA_MERGE
-		{   0,  -1,   0,   0 },	// LA_REMOVEBLOCKS
-		{   0,   0,   0,   0 },	// LA_FROMMESH
-	};
+		FCacheAddress at(dep);
+		if (dep.At && !GetMemory().IsValid(at))
+		{
+			if (ScheduledStagePerOp.get(at) <= dep.Stage)
+			{
+				OpenTasks.Add(dep);
+				ScheduledStagePerOp[at] = dep.Stage + 1;
+			}
+		}
 
-	static_assert(sizeof(sMemoryWeights)/sMemoryWeightsStageCount == int32(EOpType::COUNT));
+		if (dep.Type == FScheduledOp::EType::Full)
+		{
+			System->WorkingMemoryManager.CurrentInstanceCache->IncreaseHitCount(at);
+		}
+	}
 
 
-	int32 CodeRunner::GetOpEstimatedMemoryDelta( const FScheduledOp& Candidate, const FProgram& Program )
+	bool CodeRunner::ShouldIssueTask() const
 	{
-		int32 OpDelta = 0;
-		int32 Stage = FMath::Min(int32(Candidate.Stage), sMemoryWeightsStageCount - 1);
-		EOpType OpType = Program.GetOpType(Candidate.At);
-
-		if (OpType == EOpType::IM_PIXELFORMAT && Stage == 1)
+		// Can we afford to delay issued tasks?
+		bool bCanDelayTasks = IssuedTasks.Num() > 0 || OpenTasks.Num() > 0;
+		if (!bCanDelayTasks)
 		{
-			// TODO: We should get the actual format from the arguments for a more precise calculation.
-			OP::ImagePixelFormatArgs args = Program.GetOpArgs<OP::ImagePixelFormatArgs>(Candidate.At);
-			bool bIsCompressed = GetUncompressedFormat(args.format) != args.format;
-			bool bIsSmallish = GetImageFormatData(args.format).BytesPerBlock < 2;
-			if (bIsCompressed || bIsSmallish)
-			{
-				OpDelta = -10;
-			}
-			else
-			{
-				OpDelta = 10;
-			}
+			return true;
 		}
 		else
 		{
-			OpDelta = sMemoryWeights[int32(OpType)][Stage];
+			// We could wait. Let's see if we have enough memory to issue tasks anyway.
+			bool bHaveEnoughMemory = !System->WorkingMemoryManager.IsMemoryBudgetFull();
+			if (bHaveEnoughMemory)
+			{
+				return  true;
+			}
 		}
 
-		return OpDelta;
+		return false;
+	}
+
+
+	void CodeRunner::UpdateTraces()
+	{
+		// Code Runner status
+		TRACE_COUNTER_SET(MutableRuntime_OpenTask, OpenTasks.Num());
+		TRACE_COUNTER_SET(MutableRuntime_ClosedTasks, ClosedTasks.Num());
+		TRACE_COUNTER_SET(MutableRuntime_IssuedTasks, IssuedTasks.Num());
+		TRACE_COUNTER_SET(MutableRuntime_IssuedTasksOnHold, IssuedTasksOnHold.Num());
 	}
 
 
-	//---------------------------------------------------------------------------------------------
 	void CodeRunner::LaunchIssuedTask( const TSharedPtr<FIssuedTask>& TaskToIssue, bool& bOutFailed )
 	{
 		bool bFailed = false;
@@ -373,31 +245,9 @@ namespace mu
 				{
 				//case EExecutionStrategy::MinimizeMemory:
 				//{
-				//	// TODO: This should be done when an operation is added to the OpenTasks array instead of every time.
-				//	int32 BestOp = 0;
-				//	int8 BestDelta = TNumericLimits<int8>::Max();
-				//	int32 OpenOpCount = OpenTasks.Num();
-				//	const FProgram& Program = Model->GetPrivate()->Program;
-				//	for (int32 OpIndex = 0; OpIndex < OpenOpCount; ++OpIndex)
-				//	{
-				//		const FScheduledOp& Candidate = OpenTasks[OpIndex];
-				//		int32 OpDelta = GetOpEstimatedMemoryDelta(Candidate, Program);
-
-				//		if (OpDelta < BestDelta)
-				//		{
-				//			BestOp = OpIndex;
-				//			BestDelta = OpDelta;
-
-				//			// Shortcut: If we are freeing memory, we don't care how much: it is already the best op
-				//			if (BestDelta < 0)
-				//			{
-				//				break;
-				//			}
-				//		}
-				//	}
-
-				//	item = OpenTasks[BestOp];
-				//	OpenTasks.RemoveAtSwap(BestOp,EAllowShrinking::No);
+				//	// TODO: Prioritize operation strages that have a negative memory delta somehow (the result uses less memory than the inputs)
+				//  // An attempt to do this with an heurisitc estimation per-stage of each operation in a static table was tested in the past but maybe
+				//  // there are better ways. 
 				//	break;
 				//}
 
diff --git a/Engine/Plugins/Mutable/Source/MutableRuntime/Private/MuR/Operations.cpp b/Engine/Plugins/Mutable/Source/MutableRuntime/Private/MuR/Operations.cpp
--- a/Engine/Plugins/Mutable/Source/MutableRuntime/Private/MuR/Operations.cpp
+++ b/Engine/Plugins/Mutable/Source/MutableRuntime/Private/MuR/Operations.cpp
@@ -10,149 +10,144 @@ namespace mu
 {
     MUTABLE_IMPLEMENT_POD_SERIALISABLE(OP);
     MUTABLE_IMPLEMENT_POD_VECTOR_SERIALISABLE(OP);
-	
-	//---------------------------------------------------------------------------------------------
-	//---------------------------------------------------------------------------------------------
-	//---------------------------------------------------------------------------------------------
 
     // clang-format off
-    static const FOpDesc s_opDescs[] =
+    static const FOpDesc SOperationRuntimeDescs[] =
 	{ 
-		// type				cached	supported base image formats
-        { DT_NONE,			false,	{ 0, 0, 0, 0, 0, 0, 0, 0, 0, 0, 0, 0, 0, 0, 0, 0, 0, 0, 0, 0, 0, 0, 0, 0, 0, 0, 0, 0, 0 }		},	// NONE
-
-		{ DT_BOOL,			false,	{ 0, 0, 0, 0, 0, 0, 0, 0, 0, 0, 0, 0, 0, 0, 0, 0, 0, 0, 0, 0, 0, 0, 0, 0, 0, 0, 0, 0, 0 }		},	// BO_CONSTANT
-		{ DT_INT,			false,	{ 0, 0, 0, 0, 0, 0, 0, 0, 0, 0, 0, 0, 0, 0, 0, 0, 0, 0, 0, 0, 0, 0, 0, 0, 0, 0, 0, 0, 0 }		},	// NU_CONSTANT
-		{ DT_SCALAR,		false,	{ 0, 0, 0, 0, 0, 0, 0, 0, 0, 0, 0, 0, 0, 0, 0, 0, 0, 0, 0, 0, 0, 0, 0, 0, 0, 0, 0, 0, 0 }		},	// SC_CONSTANT
-		{ DT_COLOUR,		false,	{ 0, 0, 0, 0, 0, 0, 0, 0, 0, 0, 0, 0, 0, 0, 0, 0, 0, 0, 0, 0, 0, 0, 0, 0, 0, 0, 0, 0, 0 }		},	// CO_CONSTANT
-		{ DT_IMAGE,			true,   { 1, 1, 1, 1, 1, 1, 1, 1, 1, 1, 1, 1, 0, 0, 0, 0, 0, 0, 0, 0, 0, 0, 0, 0, 0, 0, 0, 0, 0 }		},	// IM_CONSTANT
-		{ DT_MESH,			true,   { 0, 0, 0, 0, 0, 0, 0, 0, 0, 0, 0, 0, 0, 0, 0, 0, 0, 0, 0, 0, 0, 0, 0, 0, 0, 0, 0, 0, 0 }		},	// ME_CONSTANT
-		{ DT_LAYOUT,		false,	{ 0, 0, 0, 0, 0, 0, 0, 0, 0, 0, 0, 0, 0, 0, 0, 0, 0, 0, 0, 0, 0, 0, 0, 0, 0, 0, 0, 0, 0 }		},	// LA_CONSTANT
-		{ DT_PROJECTOR,		false,	{ 0, 0, 0, 0, 0, 0, 0, 0, 0, 0, 0, 0, 0, 0, 0, 0, 0, 0, 0, 0, 0, 0, 0, 0, 0, 0, 0, 0, 0 }		},	// PR_CONSTANT
-		{ DT_STRING,		false,	{ 0, 0, 0, 0, 0, 0, 0, 0, 0, 0, 0, 0, 0, 0, 0, 0, 0, 0, 0, 0, 0, 0, 0, 0, 0, 0, 0, 0, 0 }		},	// ST_CONSTANT
-		{ DT_EXTENSION_DATA,true,	{ 0, 0, 0, 0, 0, 0, 0, 0, 0, 0, 0, 0, 0, 0, 0, 0, 0, 0, 0, 0, 0, 0, 0, 0, 0, 0, 0, 0, 0 }		},	// ED_CONSTANT
-		{ DT_MATRIX,		false,	{ 0, 0, 0, 0, 0, 0, 0, 0, 0, 0, 0, 0, 0, 0, 0, 0, 0, 0, 0, 0, 0, 0, 0, 0, 0, 0, 0, 0, 0 }		},	// MA_CONSTANT
-
-		{ DT_BOOL,			false,	{ 0, 0, 0, 0, 0, 0, 0, 0, 0, 0, 0, 0, 0, 0, 0, 0, 0, 0, 0, 0, 0, 0, 0, 0, 0, 0, 0, 0, 0 }		},	// BO_PARAMETER
-		{ DT_INT,			false,	{ 0, 0, 0, 0, 0, 0, 0, 0, 0, 0, 0, 0, 0, 0, 0, 0, 0, 0, 0, 0, 0, 0, 0, 0, 0, 0, 0, 0, 0 }		},	// NU_PARAMETER
-		{ DT_SCALAR,		false,	{ 0, 0, 0, 0, 0, 0, 0, 0, 0, 0, 0, 0, 0, 0, 0, 0, 0, 0, 0, 0, 0, 0, 0, 0, 0, 0, 0, 0, 0 }		},	// SC_PARAMETER
-		{ DT_COLOUR,		false,	{ 0, 0, 0, 0, 0, 0, 0, 0, 0, 0, 0, 0, 0, 0, 0, 0, 0, 0, 0, 0, 0, 0, 0, 0, 0, 0, 0, 0, 0 }		},	// CO_PARAMETER
-		{ DT_PROJECTOR,		false,	{ 0, 0, 0, 0, 0, 0, 0, 0, 0, 0, 0, 0, 0, 0, 0, 0, 0, 0, 0, 0, 0, 0, 0, 0, 0, 0, 0, 0, 0 }		},	// PR_PARAMETER
-		{ DT_IMAGE,			false,	{ 0, 0, 0, 0, 0, 0, 0, 0, 0, 0, 0, 0, 0, 0, 0, 0, 0, 0, 0, 0, 0, 0, 0, 0, 0, 0, 0, 0, 0 }		},	// IM_PARAMETER
-		{ DT_STRING,		false,	{ 0, 0, 0, 0, 0, 0, 0, 0, 0, 0, 0, 0, 0, 0, 0, 0, 0, 0, 0, 0, 0, 0, 0, 0, 0, 0, 0, 0, 0 }		},	// ST_PARAMETER
-		{ DT_MATRIX,		false,	{ 0, 0, 0, 0, 0, 0, 0, 0, 0, 0, 0, 0, 0, 0, 0, 0, 0, 0, 0, 0, 0, 0, 0, 0, 0, 0, 0, 0, 0 }		},	// MA_PARAMETER
-
-		{ DT_IMAGE,			true,	{ 0, 0, 0, 0, 0, 0, 0, 0, 0, 0, 0, 0, 0, 0, 0, 0, 0, 0, 0, 0, 0, 0, 0, 0, 0, 0, 0, 0, 0 }		},	// IM_REFERENCE
-		{ DT_MESH,			true,	{ 0, 0, 0, 0, 0, 0, 0, 0, 0, 0, 0, 0, 0, 0, 0, 0, 0, 0, 0, 0, 0, 0, 0, 0, 0, 0, 0, 0, 0 }		},	// ME_REFERENCE
-
-		{ DT_INT,			false,	{ 0, 0, 0, 0, 0, 0, 0, 0, 0, 0, 0, 0, 0, 0, 0, 0, 0, 0, 0, 0, 0, 0, 0, 0, 0, 0, 0, 0, 0 }		},	// NU_CONDITIONAL
-		{ DT_SCALAR,		false,	{ 0, 0, 0, 0, 0, 0, 0, 0, 0, 0, 0, 0, 0, 0, 0, 0, 0, 0, 0, 0, 0, 0, 0, 0, 0, 0, 0, 0, 0 }		},	// SC_CONDITIONAL
-		{ DT_COLOUR,		false,	{ 0, 0, 0, 0, 0, 0, 0, 0, 0, 0, 0, 0, 0, 0, 0, 0, 0, 0, 0, 0, 0, 0, 0, 0, 0, 0, 0, 0, 0 }		},	// CO_CONDITIONAL
-		{ DT_IMAGE,			true,	{ 1, 1, 1, 1, 1, 1, 1, 1, 1, 1, 1, 1, 1, 1, 1, 1, 1, 1, 1, 1, 1, 1, 0, 0, 1, 1, 0, 0, 0 }		},	// IM_CONDITIONAL
-		{ DT_MESH,			true,	{ 0, 0, 0, 0, 0, 0, 0, 0, 0, 0, 0, 0, 0, 0, 0, 0, 0, 0, 0, 0, 0, 0, 0, 0, 0, 0, 0, 0, 0 }		},	// ME_CONDITIONAL
-		{ DT_LAYOUT,		false,  { 0, 0, 0, 0, 0, 0, 0, 0, 0, 0, 0, 0, 0, 0, 0, 0, 0, 0, 0, 0, 0, 0, 0, 0, 0, 0, 0, 0, 0 }		},	// LA_CONDITIONAL
-		{ DT_INSTANCE,		false,	{ 0, 0, 0, 0, 0, 0, 0, 0, 0, 0, 0, 0, 0, 0, 0, 0, 0, 0, 0, 0, 0, 0, 0, 0, 0, 0, 0, 0, 0 }		},	// IN_CONDITIONAL
-		{ DT_EXTENSION_DATA,false,	{ 0, 0, 0, 0, 0, 0, 0, 0, 0, 0, 0, 0, 0, 0, 0, 0, 0, 0, 0, 0, 0, 0, 0, 0, 0, 0, 0, 0, 0 }		},	// ED_CONDITIONAL
-
-		{ DT_INT,			false,	{ 0, 0, 0, 0, 0, 0, 0, 0, 0, 0, 0, 0, 0, 0, 0, 0, 0, 0, 0, 0, 0, 0, 0, 0, 0, 0, 0, 0, 0 }		},	// NU_SWITCH
-		{ DT_SCALAR,		false,	{ 0, 0, 0, 0, 0, 0, 0, 0, 0, 0, 0, 0, 0, 0, 0, 0, 0, 0, 0, 0, 0, 0, 0, 0, 0, 0, 0, 0, 0 }		},	// SC_SWITCH
-		{ DT_COLOUR,		false,	{ 0, 0, 0, 0, 0, 0, 0, 0, 0, 0, 0, 0, 0, 0, 0, 0, 0, 0, 0, 0, 0, 0, 0, 0, 0, 0, 0, 0, 0 }		},	// CO_SWITCH
-		{ DT_IMAGE,			true,	{ 1, 1, 1, 1, 1, 1, 1, 1, 1, 1, 1, 1, 1, 1, 1, 1, 1, 1, 1, 1, 1, 1, 0, 0, 1, 1, 0, 0, 0 }		},	// IM_SWITCH
-		{ DT_MESH,			true,	{ 0, 0, 0, 0, 0, 0, 0, 0, 0, 0, 0, 0, 0, 0, 0, 0, 0, 0, 0, 0, 0, 0, 0, 0, 0, 0, 0, 0, 0 }		},	// ME_SWITCH
-		{ DT_LAYOUT,		false,	{ 0, 0, 0, 0, 0, 0, 0, 0, 0, 0, 0, 0, 0, 0, 0, 0, 0, 0, 0, 0, 0, 0, 0, 0, 0, 0, 0, 0, 0 }		},	// LA_SWITCH
-		{ DT_INSTANCE,		false,	{ 0, 0, 0, 0, 0, 0, 0, 0, 0, 0, 0, 0, 0, 0, 0, 0, 0, 0, 0, 0, 0, 0, 0, 0, 0, 0, 0, 0, 0 }		},	// IN_SWITCH
-		{ DT_EXTENSION_DATA,false,	{ 0, 0, 0, 0, 0, 0, 0, 0, 0, 0, 0, 0, 0, 0, 0, 0, 0, 0, 0, 0, 0, 0, 0, 0, 0, 0, 0, 0, 0 }		},	// ED_SWITCH
-
-		{ DT_BOOL,			false,  { 0, 0, 0, 0, 0, 0, 0, 0, 0, 0, 0, 0, 0, 0, 0, 0, 0, 0, 0, 0, 0, 0, 0, 0, 0, 0, 0, 0, 0 }		},	// BO_LESS
-		{ DT_BOOL,			false,  { 0, 0, 0, 0, 0, 0, 0, 0, 0, 0, 0, 0, 0, 0, 0, 0, 0, 0, 0, 0, 0, 0, 0, 0, 0, 0, 0, 0, 0 }		},	// BO_EQUAL_SC_CONST
-		{ DT_BOOL,			false,  { 0, 0, 0, 0, 0, 0, 0, 0, 0, 0, 0, 0, 0, 0, 0, 0, 0, 0, 0, 0, 0, 0, 0, 0, 0, 0, 0, 0, 0 }		},	// BO_AND
-		{ DT_BOOL,			false,  { 0, 0, 0, 0, 0, 0, 0, 0, 0, 0, 0, 0, 0, 0, 0, 0, 0, 0, 0, 0, 0, 0, 0, 0, 0, 0, 0, 0, 0 }		},	// BO_OR
-		{ DT_BOOL,			false,  { 0, 0, 0, 0, 0, 0, 0, 0, 0, 0, 0, 0, 0, 0, 0, 0, 0, 0, 0, 0, 0, 0, 0, 0, 0, 0, 0, 0, 0 }		},	// BO_NOT
-
-		{ DT_SCALAR,		true,   { 0, 0, 0, 0, 0, 0, 0, 0, 0, 0, 0, 0, 0, 0, 0, 0, 0, 0, 0, 0, 0, 0, 0, 0, 0, 0, 0, 0, 0 }		},	// SC_MULTIPLYADD
-		{ DT_SCALAR,		true,   { 0, 0, 0, 0, 0, 0, 0, 0, 0, 0, 0, 0, 0, 0, 0, 0, 0, 0, 0, 0, 0, 0, 0, 0, 0, 0, 0, 0, 0 }		},	// SC_ARITHMETIC
-		{ DT_SCALAR,		true,   { 0, 0, 0, 0, 0, 0, 0, 0, 0, 0, 0, 0, 0, 0, 0, 0, 0, 0, 0, 0, 0, 0, 0, 0, 0, 0, 0, 0, 0 }		},	// SC_CURVE
-
-		{ DT_COLOUR,		false,	{ 0, 0, 0, 0, 0, 0, 0, 0, 0, 0, 0, 0, 0, 0, 0, 0, 0, 0, 0, 0, 0, 0, 0, 0, 0, 0, 0, 0, 0 }		},	// CO_SAMPLEIMAGE
-		{ DT_COLOUR,		false,	{ 0, 0, 0, 0, 0, 0, 0, 0, 0, 0, 0, 0, 0, 0, 0, 0, 0, 0, 0, 0, 0, 0, 0, 0, 0, 0, 0, 0, 0 }		},	// CO_SWIZZLE
-		{ DT_COLOUR,		false,	{ 0, 0, 0, 0, 0, 0, 0, 0, 0, 0, 0, 0, 0, 0, 0, 0, 0, 0, 0, 0, 0, 0, 0, 0, 0, 0, 0, 0, 0 }		},	// CO_FROMSCALARS
-		{ DT_COLOUR,		false,	{ 0, 0, 0, 0, 0, 0, 0, 0, 0, 0, 0, 0, 0, 0, 0, 0, 0, 0, 0, 0, 0, 0, 0, 0, 0, 0, 0, 0, 0 }		},	// CO_ARITHMETIC
-
-		{ DT_IMAGE,			true,	{ 0, 1, 1, 1, 0, 0, 0, 0, 0, 0, 0, 0, 0, 0, 0, 0, 0, 0, 0, 0, 0, 0, 0, 0, 0, 0, 0, 0, 0 }		},	// IM_LAYER
-		{ DT_IMAGE,			true,	{ 0, 1, 1, 1, 0, 0, 0, 0, 0, 0, 0, 0, 0, 0, 0, 0, 0, 0, 0, 0, 0, 0, 0, 0, 0, 0, 0, 0, 0 }		},	// IM_LAYERCOLOUR
-		{ DT_IMAGE,			true,	{ 1, 1, 1, 1, 1, 1, 1, 1, 1, 1, 1, 1, 0, 0, 0, 0, 0, 0, 0, 0, 0, 0, 0, 0, 0, 0, 0, 0, 0 }		},	// IM_PIXELFORMAT
-		{ DT_IMAGE,			true,	{ 1, 1, 1, 1, 0, 0, 0, 0, 0, 0, 0, 0, 0, 0, 0, 0, 0, 0, 0, 0, 0, 0, 0, 0, 0, 0, 0, 0, 0 }		},	// IM_MIPMAP
-		{ DT_IMAGE,			true,	{ 0, 1, 1, 1, 0, 0, 0, 0, 0, 0, 0, 0, 0, 0, 0, 0, 0, 0, 0, 0, 0, 0, 0, 0, 0, 0, 0, 0, 0 }		},	// IM_RESIZE
-		{ DT_IMAGE,			true,	{ 0, 1, 1, 1, 0, 0, 0, 0, 0, 0, 0, 0, 0, 0, 0, 0, 0, 0, 0, 0, 0, 0, 0, 0, 0, 0, 0, 0, 0 }		},	// IM_RESIZELIKE
-		{ DT_IMAGE,			true,	{ 0, 1, 1, 1, 0, 0, 0, 0, 0, 0, 0, 0, 0, 0, 0, 0, 0, 0, 0, 0, 0, 0, 0, 0, 0, 0, 0, 0, 0 }		},	// IM_RESIZEREL
-		{ DT_IMAGE,			true,   { 1, 1, 1, 1, 1, 1, 1, 1, 1, 1, 1, 1, 0, 0, 0, 0, 0, 0, 0, 0, 0, 0, 0, 0, 0, 0, 0, 0, 0 }		},	// IM_BLANKLAYOUT
-		{ DT_IMAGE,			true,	{ 1, 1, 1, 1, 1, 1, 1, 1, 1, 1, 1, 0, 0, 0, 1, 1, 1, 1, 1, 1, 1, 1, 0, 0, 1, 1, 0, 0, 0 }		},	// IM_COMPOSE
-		{ DT_IMAGE,			true,	{ 0, 1, 1, 1, 0, 0, 0, 0, 0, 0, 0, 0, 0, 0, 0, 0, 0, 0, 0, 0, 0, 0, 0, 0, 0, 0, 0, 0, 0 }		},	// IM_INTERPOLATE
-		{ DT_IMAGE,			true,	{ 0, 1, 1, 0, 0, 0, 0, 0, 0, 0, 0, 0, 0, 0, 0, 0, 0, 0, 0, 0, 0, 0, 0, 0, 0, 0, 0, 0, 0 }		},	// IM_SATURATE
-		{ DT_IMAGE,			true,	{ 0, 1, 1, 0, 0, 0, 0, 0, 0, 0, 0, 0, 0, 0, 0, 0, 0, 0, 0, 0, 0, 0, 0, 0, 0, 0, 0, 0, 0 }		},	// IM_LUMINANCE
-		{ DT_IMAGE,			true,	{ 0, 1, 1, 1, 0, 0, 0, 0, 0, 0, 0, 0, 0, 0, 0, 0, 0, 0, 0, 0, 0, 0, 0, 0, 0, 0, 0, 0, 0 }		},	// IM_SWIZZLE
-		{ DT_IMAGE,			true,	{ 0, 1, 1, 1, 0, 0, 0, 0, 0, 0, 0, 0, 0, 0, 0, 0, 0, 0, 0, 0, 0, 0, 0, 0, 0, 0, 0, 0, 0 }		},	// IM_COLOURMAP
-		{ DT_IMAGE,			true,	{ 0, 0, 0, 0, 0, 0, 0, 0, 0, 0, 0, 0, 0, 0, 0, 0, 0, 0, 0, 0, 0, 0, 0, 0, 0, 0, 0, 0, 0 }		},	// IM_GRADIENT
-		{ DT_IMAGE,			true,	{ 0, 1, 1, 1, 0, 0, 0, 0, 0, 0, 0, 0, 0, 0, 0, 0, 0, 0, 0, 0, 0, 0, 0, 0, 0, 0, 0, 0, 0 }		},	// IM_BINARISE
-		{ DT_IMAGE,			true,	{ 0, 0, 0, 0, 0, 0, 0, 0, 0, 0, 0, 0, 0, 0, 0, 0, 0, 0, 0, 0, 0, 0, 0, 0, 0, 0, 0, 0, 0 }		},	// IM_PLAINCOLOUR
-		{ DT_IMAGE,			true,	{ 0, 0, 0, 0, 0, 0, 0, 0, 0, 0, 0, 0, 0, 0, 0, 0, 0, 0, 0, 0, 0, 0, 0, 0, 0, 0, 0, 0, 0 }		},	// IM_CROP
-		{ DT_IMAGE,			true,	{ 0, 0, 0, 0, 0, 0, 0, 0, 0, 0, 0, 0, 0, 0, 0, 0, 0, 0, 0, 0, 0, 0, 0, 0, 0, 0, 0, 0, 0 }		},	// IM_PATCH
-		{ DT_IMAGE,			true,	{ 0, 0, 0, 0, 0, 0, 0, 0, 0, 0, 0, 0, 0, 0, 0, 0, 0, 0, 0, 0, 0, 0, 0, 0, 0, 0, 0, 0, 0 }		},	// IM_RASTERMESH
-		{ DT_IMAGE,			true,	{ 0, 0, 0, 0, 0, 0, 0, 0, 0, 0, 0, 0, 0, 0, 0, 0, 0, 0, 0, 0, 0, 0, 0, 0, 0, 0, 0, 0, 0 }		},	// IM_MAKEGROWMAP
-		{ DT_IMAGE,			true,	{ 0, 0, 0, 0, 0, 0, 0, 0, 0, 0, 0, 0, 0, 0, 0, 0, 0, 0, 0, 0, 0, 0, 0, 0, 0, 0, 0, 0, 0 }		},	// IM_DISPLACE
-		{ DT_IMAGE,			true,	{ 0, 0, 0, 0, 0, 0, 0, 0, 0, 0, 0, 0, 0, 0, 0, 0, 0, 0, 0, 0, 0, 0, 0, 0, 0, 0, 0, 0, 0 }		},	// IM_MULTILAYER
-		{ DT_IMAGE,			true,	{ 0, 1, 1, 1, 0, 0, 0, 0, 0, 0, 0, 0, 0, 0, 0, 0, 0, 0, 0, 0, 0, 0, 0, 0, 0, 0, 0, 0, 0 }		},	// IM_INVERT
-		{ DT_IMAGE,			true,	{ 0, 1, 1, 1, 0, 0, 0, 0, 0, 0, 0, 0, 0, 0, 0, 0, 0, 0, 0, 0, 0, 0, 0, 0, 0, 0, 0, 0, 0 }		},	// IM_NORMALCOMPOSITE
-		{ DT_IMAGE,			true,	{ 0, 1, 1, 1, 0, 0, 0, 0, 0, 0, 0, 0, 0, 0, 0, 0, 0, 0, 0, 0, 0, 0, 0, 0, 0, 0, 0, 0, 0 }		},	// IM_TRANSFORM
-
-		{ DT_MESH,			true,	{ 0, 0, 0, 0, 0, 0, 0, 0, 0, 0, 0, 0, 0, 0, 0, 0, 0, 0, 0, 0, 0, 0, 0, 0, 0, 0, 0, 0, 0 }		},	// ME_APPLYLAYOUT
-		{ DT_MESH,			true,	{ 0, 0, 0, 0, 0, 0, 0, 0, 0, 0, 0, 0, 0, 0, 0, 0, 0, 0, 0, 0, 0, 0, 0, 0, 0, 0, 0, 0, 0 }		},	// ME_DIFFERENCE
-		{ DT_MESH,			true,	{ 0, 0, 0, 0, 0, 0, 0, 0, 0, 0, 0, 0, 0, 0, 0, 0, 0, 0, 0, 0, 0, 0, 0, 0, 0, 0, 0, 0, 0 }		},	// ME_MORPH
-		{ DT_MESH,			true,	{ 0, 0, 0, 0, 0, 0, 0, 0, 0, 0, 0, 0, 0, 0, 0, 0, 0, 0, 0, 0, 0, 0, 0, 0, 0, 0, 0, 0, 0 }		},	// ME_MERGE
-		{ DT_MESH,			true,	{ 0, 0, 0, 0, 0, 0, 0, 0, 0, 0, 0, 0, 0, 0, 0, 0, 0, 0, 0, 0, 0, 0, 0, 0, 0, 0, 0, 0, 0 }		},	// ME_MASKCLIPMESH
-		{ DT_MESH,			true,	{ 0, 0, 0, 0, 0, 0, 0, 0, 0, 0, 0, 0, 0, 0, 0, 0, 0, 0, 0, 0, 0, 0, 0, 0, 0, 0, 0, 0, 0 }		},	// ME_MASKCLIPUVMASK
-		{ DT_MESH,			true,	{ 0, 0, 0, 0, 0, 0, 0, 0, 0, 0, 0, 0, 0, 0, 0, 0, 0, 0, 0, 0, 0, 0, 0, 0, 0, 0, 0, 0, 0 }		},	// ME_MASKDIFF
-		{ DT_MESH,			true,	{ 0, 0, 0, 0, 0, 0, 0, 0, 0, 0, 0, 0, 0, 0, 0, 0, 0, 0, 0, 0, 0, 0, 0, 0, 0, 0, 0, 0, 0 }		},	// ME_REMOVEMASK
-		{ DT_MESH,			true,	{ 0, 0, 0, 0, 0, 0, 0, 0, 0, 0, 0, 0, 0, 0, 0, 0, 0, 0, 0, 0, 0, 0, 0, 0, 0, 0, 0, 0, 0 }		},	// ME_FORMAT
-		{ DT_MESH,			true,	{ 0, 0, 0, 0, 0, 0, 0, 0, 0, 0, 0, 0, 0, 0, 0, 0, 0, 0, 0, 0, 0, 0, 0, 0, 0, 0, 0, 0, 0 }		},	// ME_EXTRACTLAYOUTBLOCK
-		{ DT_MESH,			true,	{ 0, 0, 0, 0, 0, 0, 0, 0, 0, 0, 0, 0, 0, 0, 0, 0, 0, 0, 0, 0, 0, 0, 0, 0, 0, 0, 0, 0, 0 }		},	// ME_TRANSFORM
-		{ DT_MESH,			true,	{ 0, 0, 0, 0, 0, 0, 0, 0, 0, 0, 0, 0, 0, 0, 0, 0, 0, 0, 0, 0, 0, 0, 0, 0, 0, 0, 0, 0, 0 }		},	// ME_CLIPMORPHPLANE
-		{ DT_MESH,			true,	{ 0, 0, 0, 0, 0, 0, 0, 0, 0, 0, 0, 0, 0, 0, 0, 0, 0, 0, 0, 0, 0, 0, 0, 0, 0, 0, 0, 0, 0 }		},	// ME_CLIPWITHMESH
-		{ DT_MESH,			true,	{ 0, 0, 0, 0, 0, 0, 0, 0, 0, 0, 0, 0, 0, 0, 0, 0, 0, 0, 0, 0, 0, 0, 0, 0, 0, 0, 0, 0, 0 }		},	// ME_SETSKELETON
-		{ DT_MESH,			true,	{ 0, 0, 0, 0, 0, 0, 0, 0, 0, 0, 0, 0, 0, 0, 0, 0, 0, 0, 0, 0, 0, 0, 0, 0, 0, 0, 0, 0, 0 }		},	// ME_PROJECT
-		{ DT_MESH,			true,	{ 0, 0, 0, 0, 0, 0, 0, 0, 0, 0, 0, 0, 0, 0, 0, 0, 0, 0, 0, 0, 0, 0, 0, 0, 0, 0, 0, 0, 0 }		},	// ME_APPLYPOSE
-		{ DT_MESH,			true,	{ 0, 0, 0, 0, 0, 0, 0, 0, 0, 0, 0, 0, 0, 0, 0, 0, 0, 0, 0, 0, 0, 0, 0, 0, 0, 0, 0, 0, 0 }		},	// ME_BINDSHAPE
-		{ DT_MESH,			true,	{ 0, 0, 0, 0, 0, 0, 0, 0, 0, 0, 0, 0, 0, 0, 0, 0, 0, 0, 0, 0, 0, 0, 0, 0, 0, 0, 0, 0, 0 }		},	// ME_APPLYSHAPE
-		{ DT_MESH,			true,	{ 0, 0, 0, 0, 0, 0, 0, 0, 0, 0, 0, 0, 0, 0, 0, 0, 0, 0, 0, 0, 0, 0, 0, 0, 0, 0, 0, 0, 0 }		},	// ME_CLIPDEFORM
-		{ DT_MESH,			true,	{ 0, 0, 0, 0, 0, 0, 0, 0, 0, 0, 0, 0, 0, 0, 0, 0, 0, 0, 0, 0, 0, 0, 0, 0, 0, 0, 0, 0, 0 }		},	// ME_MORPHRESHAPE
-		{ DT_MESH,			true,	{ 0, 0, 0, 0, 0, 0, 0, 0, 0, 0, 0, 0, 0, 0, 0, 0, 0, 0, 0, 0, 0, 0, 0, 0, 0, 0, 0, 0, 0 }		},	// ME_OPTIMIZESKINNING
-		{ DT_MESH,			true,	{ 0, 0, 0, 0, 0, 0, 0, 0, 0, 0, 0, 0, 0, 0, 0, 0, 0, 0, 0, 0, 0, 0, 0, 0, 0, 0, 0, 0, 0 }		},	// ME_ADDTAGS
-		{ DT_MESH,			true,	{ 0, 0, 0, 0, 0, 0, 0, 0, 0, 0, 0, 0, 0, 0, 0, 0, 0, 0, 0, 0, 0, 0, 0, 0, 0, 0, 0, 0, 0 }		},	// ME_TRANSFORMWITHMESH
-
-		{ DT_INSTANCE,		false,	{ 0, 0, 0, 0, 0, 0, 0, 0, 0, 0, 0, 0, 0, 0, 0, 0, 0, 0, 0, 0, 0, 0, 0, 0, 0, 0, 0, 0, 0 }		},	// IN_ADDMESH
-		{ DT_INSTANCE,		false,	{ 0, 0, 0, 0, 0, 0, 0, 0, 0, 0, 0, 0, 0, 0, 0, 0, 0, 0, 0, 0, 0, 0, 0, 0, 0, 0, 0, 0, 0 }		},	// IN_ADDIMAGE
-		{ DT_INSTANCE,		false,	{ 0, 0, 0, 0, 0, 0, 0, 0, 0, 0, 0, 0, 0, 0, 0, 0, 0, 0, 0, 0, 0, 0, 0, 0, 0, 0, 0, 0, 0 }		},	// IN_ADDVECTOR
-		{ DT_INSTANCE,		false,	{ 0, 0, 0, 0, 0, 0, 0, 0, 0, 0, 0, 0, 0, 0, 0, 0, 0, 0, 0, 0, 0, 0, 0, 0, 0, 0, 0, 0, 0 }		},	// IN_ADDSCALAR
-		{ DT_INSTANCE,		false,	{ 0, 0, 0, 0, 0, 0, 0, 0, 0, 0, 0, 0, 0, 0, 0, 0, 0, 0, 0, 0, 0, 0, 0, 0, 0, 0, 0, 0, 0 }		},	// IN_ADDSTRING
-		{ DT_INSTANCE,		false,	{ 0, 0, 0, 0, 0, 0, 0, 0, 0, 0, 0, 0, 0, 0, 0, 0, 0, 0, 0, 0, 0, 0, 0, 0, 0, 0, 0, 0, 0 }		},	// IN_ADDSURFACE
-		{ DT_INSTANCE,		false,	{ 0, 0, 0, 0, 0, 0, 0, 0, 0, 0, 0, 0, 0, 0, 0, 0, 0, 0, 0, 0, 0, 0, 0, 0, 0, 0, 0, 0, 0 }		},	// IN_ADDCOMPONENT
-		{ DT_INSTANCE,		false,	{ 0, 0, 0, 0, 0, 0, 0, 0, 0, 0, 0, 0, 0, 0, 0, 0, 0, 0, 0, 0, 0, 0, 0, 0, 0, 0, 0, 0, 0 }		},	// IN_ADDLOD
-		{ DT_INSTANCE,		false,	{ 0, 0, 0, 0, 0, 0, 0, 0, 0, 0, 0, 0, 0, 0, 0, 0, 0, 0, 0, 0, 0, 0, 0, 0, 0, 0, 0, 0, 0 }		},	// IN_ADDEXTENSIONDATA
-
-		{ DT_LAYOUT,		true,   { 0, 0, 0, 0, 0, 0, 0, 0, 0, 0, 0, 0, 0, 0, 0, 0, 0, 0, 0, 0, 0, 0, 0, 0, 0, 0, 0, 0, 0 }		},	// LA_PACK
-		{ DT_LAYOUT,		true,   { 0, 0, 0, 0, 0, 0, 0, 0, 0, 0, 0, 0, 0, 0, 0, 0, 0, 0, 0, 0, 0, 0, 0, 0, 0, 0, 0, 0, 0 }		},	// LA_MERGE
-		{ DT_LAYOUT,		true,   { 0, 0, 0, 0, 0, 0, 0, 0, 0, 0, 0, 0, 0, 0, 0, 0, 0, 0, 0, 0, 0, 0, 0, 0, 0, 0, 0, 0, 0 }		},	// LA_REMOVEBLOCKS
-		{ DT_LAYOUT,		true,   { 0, 0, 0, 0, 0, 0, 0, 0, 0, 0, 0, 0, 0, 0, 0, 0, 0, 0, 0, 0, 0, 0, 0, 0, 0, 0, 0, 0, 0 }		},	// LA_FROMMESH
+		// DataType				
+        { EDataType::None,			},	// NONE
+
+		{ EDataType::Bool,			},	// BO_CONSTANT
+		{ EDataType::Int,			},	// NU_CONSTANT
+		{ EDataType::Scalar,		},	// SC_CONSTANT
+		{ EDataType::Color,			},	// CO_CONSTANT
+		{ EDataType::Image,			},	// IM_CONSTANT
+		{ EDataType::Mesh,			},	// ME_CONSTANT
+		{ EDataType::Layout,		},	// LA_CONSTANT
+		{ EDataType::Projector,		},	// PR_CONSTANT
+		{ EDataType::String,		},	// ST_CONSTANT
+		{ EDataType::ExtensionData,	},	// ED_CONSTANT
+		{ EDataType::Matrix,		},	// MA_CONSTANT
+
+		{ EDataType::Bool,			},	// BO_PARAMETER
+		{ EDataType::Int,			},	// NU_PARAMETER
+		{ EDataType::Scalar,		},	// SC_PARAMETER
+		{ EDataType::Color,			},	// CO_PARAMETER
+		{ EDataType::Projector,		},	// PR_PARAMETER
+		{ EDataType::Image,			},	// IM_PARAMETER
+		{ EDataType::String,		},	// ST_PARAMETER
+		{ EDataType::Matrix,		},	// MA_PARAMETER
+
+		{ EDataType::Image,			},	// IM_REFERENCE
+		{ EDataType::Mesh,			},	// ME_REFERENCE
+
+		{ EDataType::Int,			},	// NU_CONDITIONAL
+		{ EDataType::Scalar,		},	// SC_CONDITIONAL
+		{ EDataType::Color,			},	// CO_CONDITIONAL
+		{ EDataType::Image,			},	// IM_CONDITIONAL
+		{ EDataType::Mesh,			},	// ME_CONDITIONAL
+		{ EDataType::Layout,		},	// LA_CONDITIONAL
+		{ EDataType::Instance,		},	// IN_CONDITIONAL
+		{ EDataType::ExtensionData,	},	// ED_CONDITIONAL
+
+		{ EDataType::Int,			},	// NU_SWITCH
+		{ EDataType::Scalar,		},	// SC_SWITCH
+		{ EDataType::Color,			},	// CO_SWITCH
+		{ EDataType::Image,			},	// IM_SWITCH
+		{ EDataType::Mesh,			},	// ME_SWITCH
+		{ EDataType::Layout,		},	// LA_SWITCH
+		{ EDataType::Instance,		},	// IN_SWITCH
+		{ EDataType::ExtensionData,	},	// ED_SWITCH
+
+		{ EDataType::Bool,			},	// BO_LESS
+		{ EDataType::Bool,			},	// BO_EQUAL_SC_CONST
+		{ EDataType::Bool,			},	// BO_AND
+		{ EDataType::Bool,			},	// BO_OR
+		{ EDataType::Bool,			},	// BO_NOT
+
+		{ EDataType::Scalar,		},	// SC_MULTIPLYADD
+		{ EDataType::Scalar,		},	// SC_ARITHMETIC
+		{ EDataType::Scalar,		},	// SC_CURVE
+
+		{ EDataType::Color,			},	// CO_SAMPLEIMAGE
+		{ EDataType::Color,			},	// CO_SWIZZLE
+		{ EDataType::Color,			},	// CO_FROMSCALARS
+		{ EDataType::Color,			},	// CO_ARITHMETIC
+
+		{ EDataType::Image,			},	// IM_LAYER
+		{ EDataType::Image,			},	// IM_LAYERCOLOUR
+		{ EDataType::Image,			},	// IM_PIXELFORMAT
+		{ EDataType::Image,			},	// IM_MIPMAP
+		{ EDataType::Image,			},	// IM_RESIZE
+		{ EDataType::Image,			},	// IM_RESIZELIKE
+		{ EDataType::Image,			},	// IM_RESIZEREL
+		{ EDataType::Image,			},	// IM_BLANKLAYOUT
+		{ EDataType::Image,			},	// IM_COMPOSE
+		{ EDataType::Image,			},	// IM_INTERPOLATE
+		{ EDataType::Image,			},	// IM_SATURATE
+		{ EDataType::Image,			},	// IM_LUMINANCE
+		{ EDataType::Image,			},	// IM_SWIZZLE
+		{ EDataType::Image,			},	// IM_COLOURMAP
+		{ EDataType::Image,			},	// IM_GRADIENT
+		{ EDataType::Image,			},	// IM_BINARISE
+		{ EDataType::Image,			},	// IM_PLAINCOLOUR
+		{ EDataType::Image,			},	// IM_CROP
+		{ EDataType::Image,			},	// IM_PATCH
+		{ EDataType::Image,			},	// IM_RASTERMESH
+		{ EDataType::Image,			},	// IM_MAKEGROWMAP
+		{ EDataType::Image,			},	// IM_DISPLACE
+		{ EDataType::Image,			},	// IM_MULTILAYER
+		{ EDataType::Image,			},	// IM_INVERT
+		{ EDataType::Image,			},	// IM_NORMALCOMPOSITE
+		{ EDataType::Image,			},	// IM_TRANSFORM
+
+		{ EDataType::Mesh,			},	// ME_APPLYLAYOUT
+		{ EDataType::Mesh,			},	// ME_DIFFERENCE
+		{ EDataType::Mesh,			},	// ME_MORPH
+		{ EDataType::Mesh,			},	// ME_MERGE
+		{ EDataType::Mesh,			},	// ME_MASKCLIPMESH
+		{ EDataType::Mesh,			},	// ME_MASKCLIPUVMASK
+		{ EDataType::Mesh,			},	// ME_MASKDIFF
+		{ EDataType::Mesh,			},	// ME_REMOVEMASK
+		{ EDataType::Mesh,			},	// ME_FORMAT
+		{ EDataType::Mesh,			},	// ME_EXTRACTLAYOUTBLOCK
+		{ EDataType::Mesh,			},	// ME_TRANSFORM
+		{ EDataType::Mesh,			},	// ME_CLIPMORPHPLANE
+		{ EDataType::Mesh,			},	// ME_CLIPWITHMESH
+		{ EDataType::Mesh,			},	// ME_SETSKELETON
+		{ EDataType::Mesh,			},	// ME_PROJECT
+		{ EDataType::Mesh,			},	// ME_APPLYPOSE
+		{ EDataType::Mesh,			},	// ME_BINDSHAPE
+		{ EDataType::Mesh,			},	// ME_APPLYSHAPE
+		{ EDataType::Mesh,			},	// ME_CLIPDEFORM
+		{ EDataType::Mesh,			},	// ME_MORPHRESHAPE
+		{ EDataType::Mesh,			},	// ME_OPTIMIZESKINNING
+		{ EDataType::Mesh,			},	// ME_ADDTAGS
+		{ EDataType::Mesh,			},	// ME_TRANSFORMWITHMESH
+
+		{ EDataType::Instance,		},	// IN_ADDMESH
+		{ EDataType::Instance,		},	// IN_ADDIMAGE
+		{ EDataType::Instance,		},	// IN_ADDVECTOR
+		{ EDataType::Instance,		},	// IN_ADDSCALAR
+		{ EDataType::Instance,		},	// IN_ADDSTRING
+		{ EDataType::Instance,		},	// IN_ADDSURFACE
+		{ EDataType::Instance,		},	// IN_ADDCOMPONENT
+		{ EDataType::Instance,		},	// IN_ADDLOD
+		{ EDataType::Instance,		},	// IN_ADDEXTENSIONDATA
+
+		{ EDataType::Layout,		},	// LA_PACK
+		{ EDataType::Layout,		},	// LA_MERGE
+		{ EDataType::Layout,		},	// LA_REMOVEBLOCKS
+		{ EDataType::Layout,		},	// LA_FROMMESH
 	};
 
     // clang-format on
 
-    static_assert( sizeof(s_opDescs)/sizeof(FOpDesc) == (int32)EOpType::COUNT, "OperationDescMismatch");
+    static_assert(UE_ARRAY_COUNT(SOperationRuntimeDescs) == int32(EOpType::COUNT), "OperationDescMismatch");
 
-	//---------------------------------------------------------------------------------------------
 	const FOpDesc& GetOpDesc( EOpType type )
 	{
-        return s_opDescs[ (int32)type ];
+        return SOperationRuntimeDescs[ (int32)type ];
 	}
 
 
diff --git a/Engine/Plugins/Mutable/Source/MutableRuntime/Private/MuR/System.cpp b/Engine/Plugins/Mutable/Source/MutableRuntime/Private/MuR/System.cpp
--- a/Engine/Plugins/Mutable/Source/MutableRuntime/Private/MuR/System.cpp
+++ b/Engine/Plugins/Mutable/Source/MutableRuntime/Private/MuR/System.cpp
@@ -217,17 +217,17 @@ namespace mu
 		LLM_SCOPE_BYNAME(TEXT("MutableRuntime"));
 		MUTABLE_CPUPROFILER_SCOPE(NewInstance);
 
-		FLiveInstance instanceData;
-		instanceData.InstanceID = ++m_pD->LastInstanceID;
-		instanceData.Instance = nullptr;
-		instanceData.Model = InModel;
-		instanceData.State = -1;
-		instanceData.Cache = MakeShared<FProgramCache>();
-		m_pD->WorkingMemoryManager.LiveInstances.Add(instanceData);
+		FLiveInstance InstanceData;
+		InstanceData.InstanceID = ++m_pD->LastInstanceID;
+		InstanceData.Instance = nullptr;
+		InstanceData.Model = InModel;
+		InstanceData.State = -1;
+		InstanceData.Cache = MakeShared<FProgramCache>();
+		m_pD->WorkingMemoryManager.LiveInstances.Add(InstanceData);
 
 		TRACE_COUNTER_SET(MutableRuntime_LiveInstances, m_pD->WorkingMemoryManager.LiveInstances.Num());
 
-		return instanceData.InstanceID;
+		return InstanceData.InstanceID;
 	}
 
 
@@ -247,66 +247,66 @@ namespace mu
 			return nullptr;
 		}
 
-		FLiveInstance* pLiveInstance = m_pD->FindLiveInstance(InInstanceID);
-		if (!pLiveInstance)
+		FLiveInstance* LiveInstance = m_pD->FindLiveInstance(InInstanceID);
+		if (!LiveInstance)
 		{
 			UE_LOG(LogMutableCore, Error, TEXT("Invalid instance id in mutable update."));
 			return nullptr;
 		}
 
-		m_pD->WorkingMemoryManager.CurrentInstanceCache = pLiveInstance->Cache;
+		m_pD->WorkingMemoryManager.CurrentInstanceCache = LiveInstance->Cache;
 
-		FProgram& program = pLiveInstance->Model->GetPrivate()->Program;
+		FProgram& Program = LiveInstance->Model->GetPrivate()->Program;
 
-		bool validState = InStateIndex >= 0 && InStateIndex < (int)program.States.Num();
-		if (!validState)
+		bool bValidState = InStateIndex >= 0 && InStateIndex < (int)Program.States.Num();
+		if (!bValidState)
 		{
 			UE_LOG(LogMutableCore, Error, TEXT("Invalid state in mutable update."));
 			return nullptr;
 		}
 
 		// This may free resources that allow us to use less memory.
-		pLiveInstance->Instance = nullptr;
+		LiveInstance->Instance = nullptr;
 
-		bool fullBuild = (InStateIndex != pLiveInstance->State);
+		bool bFullBuild = (InStateIndex != LiveInstance->State);
 
-		pLiveInstance->State = InStateIndex;
+		LiveInstance->State = InStateIndex;
 
 		// If we changed parameters that are not in this state, we need to rebuild all.
-		if (!fullBuild)
+		if (!bFullBuild)
 		{
-			fullBuild = m_pD->CheckUpdatedParameters(pLiveInstance, InParams, pLiveInstance->UpdatedParameters);
+			bFullBuild = m_pD->CheckUpdatedParameters(LiveInstance, InParams, LiveInstance->UpdatedParameters);
 		}
 
 		// Remove cached data
 		m_pD->WorkingMemoryManager.ClearCacheLayer0();
-		if (fullBuild)
+		if (bFullBuild)
 		{
 			m_pD->WorkingMemoryManager.ClearCacheLayer1();
 		}
 
 		m_pD->WorkingMemoryManager.BeginRunnerThread();
 
-		OP::ADDRESS rootAt = pLiveInstance->Model->GetPrivate()->Program.States[InStateIndex].Root;
+		OP::ADDRESS RootAt = LiveInstance->Model->GetPrivate()->Program.States[InStateIndex].Root;
 
 		// Prepare instance cache
-		m_pD->PrepareCache(pLiveInstance->Model.Get(), InStateIndex);
-		pLiveInstance->OldParameters = InParams->Clone();
+		m_pD->PrepareCache(LiveInstance->Model.Get(), InStateIndex);
+		LiveInstance->OldParameters = InParams->Clone();
 
 		// Ensure the model cache has been created
-		m_pD->WorkingMemoryManager.FindOrAddModelCache(pLiveInstance->Model);
+		m_pD->WorkingMemoryManager.FindOrAddModelCache(LiveInstance->Model);
 
-		m_pD->RunCode(pLiveInstance->Model, InParams.Get(), rootAt, InLodMask);
+		m_pD->RunCode(LiveInstance->Model, InParams.Get(), RootAt, InLodMask);
 
-		TSharedPtr<const FInstance> Result = pLiveInstance->Cache->GetInstance(FCacheAddress(rootAt, 0, 0));
+		TSharedPtr<const FInstance> Result = LiveInstance->Cache->GetInstance(FCacheAddress(RootAt, 0, 0));
 
 		// Debug check to see if we managed the op-hit-counts correctly
-		pLiveInstance->Cache->CheckHitCountsCleared();
+		LiveInstance->Cache->CheckHitCountsCleared();
 
-		pLiveInstance->Instance = Result;
+		LiveInstance->Instance = Result;
 		if (Result)
 		{
-			Result->GetPrivate()->Id = pLiveInstance->InstanceID;
+			Result->GetPrivate()->Id = LiveInstance->InstanceID;
 		}
 		else
 		{
@@ -330,12 +330,12 @@ namespace mu
 		TSharedPtr<const FImage> pResult;
 
 		// Find the live instance
-		FLiveInstance* pLiveInstance = m_pD->FindLiveInstance(instanceID);
-		check(pLiveInstance);
-		m_pD->WorkingMemoryManager.CurrentInstanceCache = pLiveInstance->Cache;
+		FLiveInstance* LiveInstance = m_pD->FindLiveInstance(instanceID);
+		check(LiveInstance);
+		m_pD->WorkingMemoryManager.CurrentInstanceCache = LiveInstance->Cache;
 
 		OP::ADDRESS RootAddress = GetResourceIDRoot(ImageId);
-		pResult = m_pD->BuildImage(pLiveInstance->Model, pLiveInstance->OldParameters.Get(), RootAddress, MipsToSkip, InImageLOD);
+		pResult = m_pD->BuildImage(LiveInstance->Model, LiveInstance->OldParameters.Get(), RootAddress, MipsToSkip, InImageLOD);
 
 		// We always need to return something valid.
 		if (!pResult)
@@ -356,16 +356,16 @@ namespace mu
 		TSharedPtr<const FImage> pResult;
 
 		// Find the live instance
-		FLiveInstance* pLiveInstance = m_pD->FindLiveInstance(instanceID);
-		check(pLiveInstance);
-		m_pD->WorkingMemoryManager.CurrentInstanceCache = pLiveInstance->Cache;
+		FLiveInstance* LiveInstance = m_pD->FindLiveInstance(instanceID);
+		check(LiveInstance);
+		m_pD->WorkingMemoryManager.CurrentInstanceCache = LiveInstance->Cache;
 
 		OP::ADDRESS RootAddress = GetResourceIDRoot(ImageId);
 		
 		m_pD->WorkingMemoryManager.BeginRunnerThread();
 		
-		mu::EOpType OpType = pLiveInstance->Model->GetPrivate()->Program.GetOpType(RootAddress);
-		if (GetOpDataType(OpType) != DT_IMAGE)
+		mu::EOpType OpType = LiveInstance->Model->GetPrivate()->Program.GetOpType(RootAddress);
+		if (GetOpDataType(OpType) != EDataType::Image)
 		{
 			m_pD->WorkingMemoryManager.EndRunnerThread();
 			m_pD->WorkingMemoryManager.CurrentInstanceCache = nullptr;
@@ -375,7 +375,7 @@ namespace mu
 		}
 		
 		TSharedRef<CodeRunner> Runner = CodeRunner::Create(
-				m_pD->Settings, m_pD, EExecutionStrategy::MinimizeMemory, pLiveInstance->Model, pLiveInstance->OldParameters.Get(), RootAddress, FSystem::AllLODs, MipsToSkip, InImageLOD, FScheduledOp::EType::Full);
+				m_pD->Settings, m_pD, EExecutionStrategy::MinimizeMemory, LiveInstance->Model, LiveInstance->OldParameters.Get(), RootAddress, FSystem::AllLODs, MipsToSkip, InImageLOD, FScheduledOp::EType::Full);
 	
 		
 		constexpr bool bForceInlineExecution = false;
@@ -424,14 +424,14 @@ namespace mu
 		OutDesc = FExtendedImageDesc();
 
 		// Find the live instance
-		FLiveInstance* pLiveInstance = m_pD->FindLiveInstance(instanceID);
-		check(pLiveInstance);
-		m_pD->WorkingMemoryManager.CurrentInstanceCache = pLiveInstance->Cache;
+		FLiveInstance* LiveInstance = m_pD->FindLiveInstance(instanceID);
+		check(LiveInstance);
+		m_pD->WorkingMemoryManager.CurrentInstanceCache = LiveInstance->Cache;
 
 		OP::ADDRESS RootAddress = GetResourceIDRoot(ImageId);
 
-		const mu::FModel* Model = pLiveInstance->Model.Get();
-		const mu::FProgram& program = Model->GetPrivate()->Program;
+		const mu::FModel* Model = LiveInstance->Model.Get();
+		const mu::FProgram& Program = Model->GetPrivate()->Program;
 
 		// TODO: It should be possible to reuse this data if cleared in the correct places only, together with HeapImageDesc.
 		int32 VarValue = CVarClearImageDescCache.GetValueOnAnyThread();
@@ -440,15 +440,15 @@ namespace mu
 			m_pD->WorkingMemoryManager.CurrentInstanceCache->ClearDescCache();
 		}
 
-		mu::EOpType opType = program.GetOpType(RootAddress);
-		if (GetOpDataType(opType) == DT_IMAGE)
+		mu::EOpType OpType = Program.GetOpType(RootAddress);
+		if (GetOpDataType(OpType) == EDataType::Image)
 		{
 			// GetImageDesc may call normal execution paths where meshes are computed.
 			m_pD->WorkingMemoryManager.BeginRunnerThread();
 					
 			int8 executionOptions = 0;
 			TSharedRef<CodeRunner> Runner = CodeRunner::Create(
-					m_pD->Settings, m_pD, EExecutionStrategy::MinimizeMemory, pLiveInstance->Model, pLiveInstance->OldParameters.Get(), RootAddress, FSystem::AllLODs, executionOptions, 0, FScheduledOp::EType::ImageDesc);
+					m_pD->Settings, m_pD, EExecutionStrategy::MinimizeMemory, LiveInstance->Model, LiveInstance->OldParameters.Get(), RootAddress, FSystem::AllLODs, executionOptions, 0, FScheduledOp::EType::ImageDesc);
 			
 			constexpr bool bForceInlineExecution = true;
 			UE::Tasks::FTask CompletionEvent = Runner->StartRun(bForceInlineExecution);
@@ -469,14 +469,14 @@ namespace mu
 		MUTABLE_CPUPROFILER_SCOPE(SystemGetImageDesc);
 
 		// Find the live instance
-		FLiveInstance* pLiveInstance = m_pD->FindLiveInstance(instanceID);
-		check(pLiveInstance);
-		m_pD->WorkingMemoryManager.CurrentInstanceCache = pLiveInstance->Cache;
+		FLiveInstance* LiveInstance = m_pD->FindLiveInstance(instanceID);
+		check(LiveInstance);
+		m_pD->WorkingMemoryManager.CurrentInstanceCache = LiveInstance->Cache;
 
 		OP::ADDRESS RootAddress = GetResourceIDRoot(ImageId);
 
-		const mu::FModel* Model = pLiveInstance->Model.Get();
-		const mu::FProgram& program = Model->GetPrivate()->Program;
+		const mu::FModel* Model = LiveInstance->Model.Get();
+		const mu::FProgram& Program = Model->GetPrivate()->Program;
 
 		// TODO: It should be possible to reuse this data if cleared in the correct places only, together with HeapImageDesc.
 		int32 VarValue = CVarClearImageDescCache.GetValueOnAnyThread();
@@ -486,8 +486,8 @@ namespace mu
 		}
 
 		m_pD->WorkingMemoryManager.BeginRunnerThread();
-		mu::EOpType OpType = program.GetOpType(RootAddress);
-		if (GetOpDataType(OpType) != DT_IMAGE)
+		mu::EOpType OpType = Program.GetOpType(RootAddress);
+		if (GetOpDataType(OpType) != EDataType::Image)
 		{
 			m_pD->WorkingMemoryManager.EndRunnerThread();
 			m_pD->WorkingMemoryManager.CurrentInstanceCache = nullptr;
@@ -498,7 +498,7 @@ namespace mu
 		// GetImageDesc may call normal execution paths where meshes are computed.
 		int8 ExecutionOptions = 0;
 		TSharedRef<CodeRunner> Runner = CodeRunner::Create(
-				m_pD->Settings, m_pD, EExecutionStrategy::MinimizeMemory, pLiveInstance->Model, pLiveInstance->OldParameters.Get(), RootAddress, FSystem::AllLODs, ExecutionOptions, 0, FScheduledOp::EType::ImageDesc);
+				m_pD->Settings, m_pD, EExecutionStrategy::MinimizeMemory, LiveInstance->Model, LiveInstance->OldParameters.Get(), RootAddress, FSystem::AllLODs, ExecutionOptions, 0, FScheduledOp::EType::ImageDesc);
 		
 		constexpr bool bForceInlineExecution = false;
 		UE::Tasks::FTask RunnerCompletionEvent = Runner->StartRun(bForceInlineExecution);
@@ -529,12 +529,12 @@ namespace mu
 		TSharedPtr<const FMesh> Result;
 
 		// Find the live instance
-		FLiveInstance* pLiveInstance = m_pD->FindLiveInstance(instanceID);
-		check(pLiveInstance);
-		m_pD->WorkingMemoryManager.CurrentInstanceCache = pLiveInstance->Cache;
+		FLiveInstance* LiveInstance = m_pD->FindLiveInstance(instanceID);
+		check(LiveInstance);
+		m_pD->WorkingMemoryManager.CurrentInstanceCache = LiveInstance->Cache;
 
 		OP::ADDRESS RootAddress = GetResourceIDRoot(MeshId);
-		Result = m_pD->BuildMesh(pLiveInstance->Model, pLiveInstance->OldParameters.Get(), RootAddress);
+		Result = m_pD->BuildMesh(LiveInstance->Model, LiveInstance->OldParameters.Get(), RootAddress);
 
 		// If the mesh is null it means empty, but we still need to return a valid one
 		if (!Result)
@@ -554,15 +554,15 @@ namespace mu
 		TSharedPtr<const FMesh> ResultMesh;
 
 		// Find the live instance
-		FLiveInstance* pLiveInstance = m_pD->FindLiveInstance(instanceID);
-		check(pLiveInstance);
-		m_pD->WorkingMemoryManager.CurrentInstanceCache = pLiveInstance->Cache;
+		FLiveInstance* LiveInstance = m_pD->FindLiveInstance(instanceID);
+		check(LiveInstance);
+		m_pD->WorkingMemoryManager.CurrentInstanceCache = LiveInstance->Cache;
 		m_pD->WorkingMemoryManager.BeginRunnerThread();
 		
 		OP::ADDRESS RootAddress = GetResourceIDRoot(MeshId);
 			
-		mu::EOpType OpType = pLiveInstance->Model->GetPrivate()->Program.GetOpType(RootAddress);
-		if (GetOpDataType(OpType) != DT_MESH)
+		mu::EOpType OpType = LiveInstance->Model->GetPrivate()->Program.GetOpType(RootAddress);
+		if (GetOpDataType(OpType) != EDataType::Mesh)
 		{
 			m_pD->WorkingMemoryManager.EndRunnerThread();
 			m_pD->WorkingMemoryManager.CurrentInstanceCache = nullptr;
@@ -571,7 +571,7 @@ namespace mu
 		}
 		
 		TSharedRef<CodeRunner> Runner = CodeRunner::Create(
-				m_pD->Settings, m_pD, EExecutionStrategy::MinimizeMemory, pLiveInstance->Model, pLiveInstance->OldParameters.Get(), RootAddress, FSystem::AllLODs, 0, 0, FScheduledOp::EType::Full);
+				m_pD->Settings, m_pD, EExecutionStrategy::MinimizeMemory, LiveInstance->Model, LiveInstance->OldParameters.Get(), RootAddress, FSystem::AllLODs, 0, 0, FScheduledOp::EType::Full);
 		
 		constexpr bool bForceInlineExecution = false;
 		UE::Tasks::FTask RunnerCompletionEvent = Runner->StartRun(bForceInlineExecution);
@@ -604,15 +604,15 @@ namespace mu
 		LLM_SCOPE_BYNAME(TEXT("MutableRuntime"));
 		MUTABLE_CPUPROFILER_SCOPE(EndUpdate);
 
-		FLiveInstance* pLiveInstance = m_pD->FindLiveInstance(instanceID);
-		if (pLiveInstance)
+		FLiveInstance* LiveInstance = m_pD->FindLiveInstance(instanceID);
+		if (LiveInstance)
 		{
-			pLiveInstance->Instance = nullptr;
+			LiveInstance->Instance = nullptr;
 
 			// Debug check to see if we managed the op-hit-counts correctly
-			pLiveInstance->Cache->CheckHitCountsCleared();
+			LiveInstance->Cache->CheckHitCountsCleared();
 
-			m_pD->WorkingMemoryManager.CurrentInstanceCache = pLiveInstance->Cache;
+			m_pD->WorkingMemoryManager.CurrentInstanceCache = LiveInstance->Cache;
 
 			// We don't want to clear the cache layer 1 because it contains data that can be useful for a 
 			// future update (same states, just runtime parameters changed).
@@ -706,9 +706,9 @@ namespace mu
         }
 
 
-        bool Visit( OP::ADDRESS at, FProgram& program ) override
+        bool Visit( OP::ADDRESS at, FProgram& Program ) override
         {
-            switch ( program.GetOpType(at) )
+            switch ( Program.GetOpType(at) )
             {
             case EOpType::BO_PARAMETER:
             case EOpType::NU_PARAMETER:
@@ -718,7 +718,7 @@ namespace mu
             case EOpType::IM_PARAMETER:
             case EOpType::MA_PARAMETER:
             {
-				OP::ParameterArgs args = program.GetOpArgs<OP::ParameterArgs>(at);
+				OP::ParameterArgs args = Program.GetOpArgs<OP::ParameterArgs>(at);
 				OP::ADDRESS ParamIndex = args.variable;
                 Flags[ParamIndex] = true;
                 break;
@@ -729,7 +729,7 @@ namespace mu
 
             }
 
-            return UniqueDiscreteCoveredCodeVisitor<>::Visit( at, program );
+            return UniqueDiscreteCoveredCodeVisitor<>::Visit( at, Program );
         }
 
     private:
@@ -747,11 +747,11 @@ namespace mu
 		LLM_SCOPE_BYNAME(TEXT("MutableRuntime"));
 
 		// Find the live instance
-		FLiveInstance* pLiveInstance = m_pD->FindLiveInstance(InstanceID);
-		check(pLiveInstance);
-		m_pD->WorkingMemoryManager.CurrentInstanceCache = pLiveInstance->Cache;
+		FLiveInstance* LiveInstance = m_pD->FindLiveInstance(InstanceID);
+		check(LiveInstance);
+		m_pD->WorkingMemoryManager.CurrentInstanceCache = LiveInstance->Cache;
 		
-		RelevantParameterVisitor visitor( m_pD, pLiveInstance->Model, FParameters, Flags );
+		RelevantParameterVisitor visitor( m_pD, LiveInstance->Model, FParameters, Flags );
 
 		m_pD->WorkingMemoryManager.CurrentInstanceCache = nullptr;
     }
@@ -786,15 +786,15 @@ namespace mu
 
         // check what parameters have changed
 		UpdatedParameters = 0;
-        const FProgram& program = LiveInstance->Model->GetPrivate()->Program;
-        const TArray<int>& runtimeParams = program.States[ LiveInstance->State ].m_runtimeParameters;
+        const FProgram& Program = LiveInstance->Model->GetPrivate()->Program;
+        const TArray<int>& runtimeParams = Program.States[ LiveInstance->State ].m_runtimeParameters;
 
-        check( Params->GetCount() == (int)program.Parameters.Num() );
+        check( Params->GetCount() == (int)Program.Parameters.Num() );
         check( !LiveInstance->OldParameters
 			||
 			Params->GetCount() == LiveInstance->OldParameters->GetCount() );
 
-        for ( int32 p=0; p<program.Parameters.Num() && !bFullBuild; ++p )
+        for ( int32 p=0; p<Program.Parameters.Num() && !bFullBuild; ++p )
         {
             bool isRuntime = runtimeParams.Contains( p );
             bool changed = !Params->HasSameValue( p, LiveInstance->OldParameters, p );
@@ -918,8 +918,8 @@ namespace mu
 
 		FVector4f Result(0,0,0,1);
 
-		mu::EOpType opType = pModel->GetPrivate()->Program.GetOpType(at);
-		if (GetOpDataType(opType) == DT_COLOUR)
+		mu::EOpType OpType = pModel->GetPrivate()->Program.GetOpType(at);
+		if (GetOpDataType(OpType) == EDataType::Color)
 		{
 			RunCode(pModel, Params, at);
 			if (!bUnrecoverableError)
@@ -959,8 +959,8 @@ namespace mu
 
 		TSharedPtr<const FImage> Result;
 
-		mu::EOpType opType = pModel->GetPrivate()->Program.GetOpType(at);
-		if (GetOpDataType(opType) == DT_IMAGE)
+		mu::EOpType OpType = pModel->GetPrivate()->Program.GetOpType(at);
+		if (GetOpDataType(OpType) == EDataType::Image)
 		{
 			RunCode(pModel, Params, at, FSystem::AllLODs, uint8(MipsToSkip), InImageLOD);
 			if (!bUnrecoverableError)
@@ -981,8 +981,8 @@ namespace mu
 
 		TSharedPtr<const FMesh> Result;
 
-		mu::EOpType opType = pModel->GetPrivate()->Program.GetOpType(at);
-		if (GetOpDataType(opType) == DT_MESH)
+		mu::EOpType OpType = pModel->GetPrivate()->Program.GetOpType(at);
+		if (GetOpDataType(OpType) == EDataType::Mesh)
 		{
 			RunCode(pModel, Params, at);
 			if (!bUnrecoverableError)
@@ -1003,8 +1003,8 @@ namespace mu
 
 		TSharedPtr<const FInstance> Result;
 
-		mu::EOpType opType = pModel->GetPrivate()->Program.GetOpType(at);
-		if (GetOpDataType(opType) == DT_INSTANCE)
+		mu::EOpType OpType = pModel->GetPrivate()->Program.GetOpType(at);
+		if (GetOpDataType(OpType) == EDataType::Instance)
 		{
 			RunCode(pModel, Params, at);
 			if (!bUnrecoverableError)
@@ -1028,8 +1028,8 @@ namespace mu
 
 		if (pModel->GetPrivate()->Program.States[0].Root)
 		{
-			mu::EOpType opType = pModel->GetPrivate()->Program.GetOpType(at);
-			if (GetOpDataType(opType) == DT_LAYOUT)
+			mu::EOpType OpType = pModel->GetPrivate()->Program.GetOpType(at);
+			if (GetOpDataType(OpType) == EDataType::Layout)
 			{
 				RunCode(pModel, Params, at);
 				if (!bUnrecoverableError)
@@ -1054,8 +1054,8 @@ namespace mu
 
 		if (pModel->GetPrivate()->Program.States[0].Root)
 		{
-			mu::EOpType opType = pModel->GetPrivate()->Program.GetOpType(at);
-			if (GetOpDataType(opType) == DT_STRING)
+			mu::EOpType OpType = pModel->GetPrivate()->Program.GetOpType(at);
+			if (GetOpDataType(OpType) == EDataType::String)
 			{
 				RunCode(pModel, Params, at);
 				if (!bUnrecoverableError)
diff --git a/Engine/Plugins/Mutable/Source/MutableTools/Private/MuT/ASTOpImageLayer.cpp b/Engine/Plugins/Mutable/Source/MutableTools/Private/MuT/ASTOpImageLayer.cpp
--- a/Engine/Plugins/Mutable/Source/MutableTools/Private/MuT/ASTOpImageLayer.cpp
+++ b/Engine/Plugins/Mutable/Source/MutableTools/Private/MuT/ASTOpImageLayer.cpp
@@ -553,7 +553,7 @@ namespace mu
 			// moved to generic sink in the default.
 //            case EOpType::IM_CONDITIONAL:
 //            {
-//                m_modified = true;
+//                bModified = true;
 
 //                OP op = program.m_code[baseAt];
 
@@ -640,13 +640,13 @@ namespace mu
 //                auto newRoot = neutral.Apply
 //                         (
 //                            program.m_code[at].args.ImageLayer.base, program,
-//                            DT_IMAGE, supportedOps,
+//                            Image, supportedOps,
 //                            &templateOp, &templateOp.args.ImageLayer.base);
 
 //                // If there is a change
 //                if (newRoot!=program.m_code[at].args.ImageLayer.base)
 //                {
-//                    m_modified = true;
+//                    bModified = true;
 //                    at = newRoot;
 //                }
 
diff --git a/Engine/Plugins/Mutable/Source/MutableTools/Private/MuT/ASTOpImageLayerColor.cpp b/Engine/Plugins/Mutable/Source/MutableTools/Private/MuT/ASTOpImageLayerColor.cpp
--- a/Engine/Plugins/Mutable/Source/MutableTools/Private/MuT/ASTOpImageLayerColor.cpp
+++ b/Engine/Plugins/Mutable/Source/MutableTools/Private/MuT/ASTOpImageLayerColor.cpp
@@ -335,7 +335,7 @@ namespace mu
 			// moved to generic sink in the default.
 //            case EOpType::IM_CONDITIONAL:
 //            {
-//                m_modified = true;
+//                bModified = true;
 
 //                OP op = program.m_code[baseAt];
 
diff --git a/Engine/Plugins/Mutable/Source/MutableTools/Private/MuT/ASTOpImagePixelFormat.cpp b/Engine/Plugins/Mutable/Source/MutableTools/Private/MuT/ASTOpImagePixelFormat.cpp
--- a/Engine/Plugins/Mutable/Source/MutableTools/Private/MuT/ASTOpImagePixelFormat.cpp
+++ b/Engine/Plugins/Mutable/Source/MutableTools/Private/MuT/ASTOpImagePixelFormat.cpp
@@ -18,6 +18,7 @@
 #include "MuT/ASTOpImageLayerColor.h"
 #include "MuT/ASTOpImageRasterMesh.h"
 #include "MuT/ASTOpSwitch.h"
+#include "MuT/CompilerPrivate.h"
 
 namespace mu
 {
@@ -484,7 +485,7 @@ namespace mu
 
 		case EOpType::IM_LAYER:
 		{
-			if (GetOpDesc(at->GetOpType()).supportedBasePixelFormats[(size_t)format])
+			if (GetOpToolsDesc(at->GetOpType()).bSupportedBasePixelFormats[(size_t)format])
 			{
 				// We move the format down the two paths
 				Ptr<ASTOpImageLayer> nop = mu::Clone<ASTOpImageLayer>(at);
@@ -499,7 +500,7 @@ namespace mu
 
 		case EOpType::IM_LAYERCOLOUR:
 		{
-			if (GetOpDesc(at->GetOpType()).supportedBasePixelFormats[(size_t)format])
+			if (GetOpToolsDesc(at->GetOpType()).bSupportedBasePixelFormats[(size_t)format])
 			{
 				// We move the format down the base
 				Ptr<ASTOpImageLayerColor> nop = mu::Clone<ASTOpImageLayerColor>(at);
diff --git a/Engine/Plugins/Mutable/Source/MutableTools/Private/MuT/CodeOptimiser.cpp b/Engine/Plugins/Mutable/Source/MutableTools/Private/MuT/CodeOptimiser.cpp
--- a/Engine/Plugins/Mutable/Source/MutableTools/Private/MuT/CodeOptimiser.cpp
+++ b/Engine/Plugins/Mutable/Source/MutableTools/Private/MuT/CodeOptimiser.cpp
@@ -342,12 +342,17 @@ namespace mu
 	//---------------------------------------------------------------------------------------------
 	//---------------------------------------------------------------------------------------------
 	//---------------------------------------------------------------------------------------------
-	bool FindOpTypeVisitor::Apply( FProgram& program, OP::ADDRESS rootAt )
+	FFindOpTypeVisitor::FFindOpTypeVisitor(const TArray<EOpType>& Types)
+		: TypesToFind(Types)
+	{
+	}
+
+	bool FFindOpTypeVisitor::Apply( FProgram& program, OP::ADDRESS rootAt )
 	{
 		bool found = false;
 
 		// If the program added new operations, we haven't visited them.
-		m_visited.SetNumZeroed( program.OpAddress.Num() );
+		Visited.SetNumZeroed( program.OpAddress.Num() );
 
 		Pending.Reset();
 		Pending.Reserve( program.OpAddress.Num()/4 );
@@ -360,38 +365,38 @@ namespace mu
 			OP::ADDRESS at = item.Value;
 
 			// Not cached?
-			if (m_visited[at]<=1)
+			if (Visited[at]<=1)
 			{
 				if (item.Key)
 				{
 					// Item indicating we finished with all the children of a parent
-					check(m_visited[at]==1);
+					check(Visited[at]==1);
 
-					bool subtreeFound = false;
+					bool bSubtreeFound = false;
 					ForEachReference( program, at, [&](OP::ADDRESS ref)
 					{
-						subtreeFound = subtreeFound || m_visited[ref]==3;
+						bSubtreeFound = bSubtreeFound || Visited[ref]==3;
 					});
 
-					m_visited[at] = subtreeFound ? 3 : 2;
+					Visited[at] = bSubtreeFound ? 3 : 2;
 				}
 				else if (!found)
 				{
-					if (m_typesToFind.Find(program.GetOpType(at)) != INDEX_NONE )
+					if (TypesToFind.Find(program.GetOpType(at)) != INDEX_NONE )
 					{
-						m_visited[at] = 3; // visited, ops found
+						Visited[at] = 3; // visited, ops found
 						found = true;
 					}
 					else
 					{
-						check(m_visited[at]==0);
+						check(Visited[at]==0);
 
-						m_visited[at] = 1;
+						Visited[at] = 1;
 						Pending.Add({ true,at });
 
 						ForEachReference( program, at, [&](OP::ADDRESS ref)
 						{
-							if (ref && m_visited[ref]==0)
+							if (ref && Visited[ref]==0)
 							{
 								Pending.Add({ false,ref });
 							}
@@ -401,12 +406,12 @@ namespace mu
 				else
 				{
 					// We won't process it.
-					m_visited[at] = 0;
+					Visited[at] = 0;
 				}
 			}
 		}
 
-		return m_visited[rootAt]==3;
+		return Visited[rootAt]==3;
 	}
 
 
@@ -415,21 +420,21 @@ namespace mu
 	//---------------------------------------------------------------------------------------------
 	IsConstantVisitor::IsConstantVisitor()
 	{
-		TArray<EOpType> parameterTypes;
-		parameterTypes.Add( EOpType::BO_PARAMETER );
-		parameterTypes.Add( EOpType::NU_PARAMETER );
-		parameterTypes.Add( EOpType::SC_PARAMETER );
-		parameterTypes.Add( EOpType::CO_PARAMETER );
-		parameterTypes.Add( EOpType::PR_PARAMETER );
-		parameterTypes.Add( EOpType::IM_PARAMETER );
-		m_findOpTypeVisitor = MakeUnique<FindOpTypeVisitor>(parameterTypes);
+		TArray<EOpType> ParameterTypes;
+		ParameterTypes.Add( EOpType::BO_PARAMETER );
+		ParameterTypes.Add( EOpType::NU_PARAMETER );
+		ParameterTypes.Add( EOpType::SC_PARAMETER );
+		ParameterTypes.Add( EOpType::CO_PARAMETER );
+		ParameterTypes.Add( EOpType::PR_PARAMETER );
+		ParameterTypes.Add( EOpType::IM_PARAMETER );
+		FindOpTypeVisitor = MakeUnique<FFindOpTypeVisitor>(ParameterTypes);
 	}
 
 
 	//---------------------------------------------------------------------------------------------
 	bool IsConstantVisitor::Apply( FProgram& program, OP::ADDRESS at )
 	{
-		return !m_findOpTypeVisitor->Apply(program,at);
+		return !FindOpTypeVisitor->Apply(program,at);
 	}
 
 
@@ -473,7 +478,7 @@ namespace mu
 			// This runs in a worker thread
 
 			EOpType type = SourceCloned->GetOpType();
-			DATATYPE DataType = GetOpDataType(type);
+			EDataType DataType = GetOpDataType(type);
 
 			FSettings Settings;
 			Settings.SetProfile( false );
@@ -483,7 +488,7 @@ namespace mu
 			System->GetPrivate()->ImagePixelFormatOverride = ImOp.FormatImageOverride;
 
 			FSourceDataDescriptor SourceDataDescriptor;
-			if (DataType == DT_IMAGE || DataType == DT_MESH)
+			if (DataType == EDataType::Image || DataType == EDataType::Mesh)
 			{
 				SourceDataDescriptor = SourceCloned->GetSourceDataDescriptor();
 				check(!SourceDataDescriptor.IsInvalid());
@@ -507,7 +512,7 @@ namespace mu
 			// Calculate the value and replace this op by a constant
 			switch( DataType )
 			{
-			case DT_MESH:
+			case EDataType::Mesh:
 			{
 				MUTABLE_CPUPROFILER_SCOPE(ConstantMesh);
 
@@ -524,7 +529,7 @@ namespace mu
 				break;
 			}
 
-			case DT_IMAGE:
+			case EDataType::Image:
 			{
 				MUTABLE_CPUPROFILER_SCOPE(ConstantImage);
 
@@ -541,7 +546,7 @@ namespace mu
 				break;
 			}
 
-			case DT_LAYOUT:
+			case EDataType::Layout:
 			{
 				MUTABLE_CPUPROFILER_SCOPE(ConstantLayout);
 
@@ -557,7 +562,7 @@ namespace mu
 				break;
 			}
 
-			case DT_BOOL:
+			case EDataType::Bool:
 			{
 				MUTABLE_CPUPROFILER_SCOPE(ConstantBool);
 
@@ -571,7 +576,7 @@ namespace mu
 				break;
 			}
 
-			case DT_COLOUR:
+			case EDataType::Color:
 			{
 				MUTABLE_CPUPROFILER_SCOPE(ConstantBool);
 
@@ -590,10 +595,10 @@ namespace mu
 				break;
 			}
 
-			case DT_INT:
-			case DT_SCALAR:
-			case DT_STRING:
-			case DT_PROJECTOR:
+			case EDataType::Int:
+			case EDataType::Scalar:
+			case EDataType::String:
+			case EDataType::Projector:
 				// TODO
 				break;
 
@@ -721,14 +726,14 @@ namespace mu
 					SubgraphRoot->bHasSpecialOpInSubgraph = bHasSpecialOpInSubgraph;
 
 					bool bIsDataTypeThanCanTurnIntoConst = false;
-					DATATYPE DataType = GetOpDataType(SubgraphType);
+					EDataType DataType = GetOpDataType(SubgraphType);
 					switch (DataType)
 					{
-					case DT_MESH:
-					case DT_IMAGE:
-					case DT_LAYOUT:
-					case DT_BOOL:
-					case DT_COLOUR:
+					case EDataType::Mesh:
+					case EDataType::Image:
+					case EDataType::Layout:
+					case EDataType::Bool:
+					case EDataType::Color:
 						bIsDataTypeThanCanTurnIntoConst = true;
 						break;
 					default:
@@ -748,7 +753,7 @@ namespace mu
 						bool bCanBeGenerated = true;
 
 						// Check source data incompatiblities: when generating constants don't mix data that has different source descriptors (tags and other properties).
-						if (DataType == DT_IMAGE || DataType == DT_MESH)
+						if (DataType == EDataType::Image || DataType == EDataType::Mesh)
 						{
 							FSourceDataDescriptor SourceDescriptor = SubgraphRoot->GetSourceDataDescriptor();
 							if (SourceDescriptor.IsInvalid())
@@ -1089,43 +1094,43 @@ namespace mu
 	//-------------------------------------------------------------------------------------------------
 	void CodeOptimiser::FullOptimiseAST( ASTOpList& roots, int32 Pass )
 	{
-		bool modified = true;
-		int32 numIterations = 0;
-		while (modified && (!m_optimizeIterationsMax || (m_optimizeIterationsLeft>0) || !numIterations) )
+		bool bModified = true;
+		int32 NumIterations = 0;
+		while (bModified && (!OptimizeIterationsMax || (OptimizeIterationsLeft>0) || !NumIterations) )
 		{
-			--m_optimizeIterationsLeft;
-			++numIterations;
+			--OptimizeIterationsLeft;
+			++NumIterations;
 			UE_LOG(LogMutableCore, Verbose, TEXT("Main optimise iteration %d, max %d, left %d"),
-					numIterations, m_optimizeIterationsMax, m_optimizeIterationsLeft);
+					NumIterations, OptimizeIterationsMax, OptimizeIterationsLeft);
 
-			modified = false;
+			bModified = false;
 
 			// All kind of optimisations that depend on the meaning of each operation
 			// \TODO: We are doing it for all states.
 			UE_LOG(LogMutableCore, Verbose, TEXT(" - semantic optimiser"));
-			modified |= SemanticOptimiserAST( roots, Options->GetPrivate()->OptimisationOptions, Pass );
+			bModified |= SemanticOptimiserAST( roots, Options->GetPrivate()->OptimisationOptions, Pass );
 			//AXE_INT_VALUE("Mutable", Verbose, "ast size", (int64_t)ASTOp::CountNodes(roots));
 			//UE_LOG(LogMutableCore, Verbose, TEXT("(int) %s : %ld"), TEXT("ast size"), int64(ASTOp::CountNodes(roots)));
 			ASTOp::LogHistogram(roots);
 
 			UE_LOG(LogMutableCore, Verbose, TEXT(" - sink optimiser"));
-			modified |= SinkOptimiserAST( roots, Options->GetPrivate()->OptimisationOptions );
+			bModified |= SinkOptimiserAST( roots, Options->GetPrivate()->OptimisationOptions );
 			//AXE_INT_VALUE("Mutable", Verbose, "ast size", (int64_t)ASTOp::CountNodes(roots));
 			ASTOp::LogHistogram(roots);
 
 			// Image size operations are treated separately
 			UE_LOG(LogMutableCore, Verbose, TEXT(" - size optimiser"));
-			modified |= SizeOptimiserAST( roots );
+			bModified |= SizeOptimiserAST( roots );
 			//AXE_INT_VALUE("Mutable", Verbose, "ast size", (int64_t)ASTOp::CountNodes(roots));
 
 			// We cannot run the logic optimiser here because the constant memory explodes.
 			// TODO: Conservative logic optimiser
-			if (!modified)
+			if (!bModified)
 			{
 	//                AXE_INT_VALUE("Mutable", Verbose, "program size", (int64_t)program.OpAddress.Num());
 	//                UE_LOG(LogMutableCore, Verbose, " - logic optimiser");
 	//                LocalLogicOptimiser log;
-	//                modified |= log.Apply( program, (int)s );
+	//                bModified |= log.Apply( program, (int)s );
 			}
 		}
 
@@ -1150,7 +1155,7 @@ namespace mu
 			UE_LOG(LogMutableCore, Verbose, TEXT(" - constant generator"));
 
 			// Constant subtree generation
-			modified = ConstantGeneratorAST( Options->GetPrivate(), Root, Pass );
+			bModified = ConstantGeneratorAST( Options->GetPrivate(), Root, Pass );
 		}
 
 		//UE_LOG(LogMutableCore, Verbose, TEXT("(int) %s : %ld"), TEXT("ast size"), int64(ASTOp::CountNodes(roots)));
@@ -1168,13 +1173,13 @@ namespace mu
 		DuplicatedDataRemoverAST( roots );
 		//UE_LOG(LogMutableCore, Verbose, TEXT("(int) %s : %ld"), TEXT("ast size"), int64(ASTOp::CountNodes(roots)));
 
-		//if (!modified)
+		//if (!bModified)
 		{
 			//for ( size_t s=0;  s<program.States.size(); ++s )
 			{
 				ASTOp::LogHistogram(roots);
 				UE_LOG(LogMutableCore, Verbose, TEXT(" - logic optimiser"));
-				modified |= LocalLogicOptimiserAST( roots );
+				bModified |= LocalLogicOptimiserAST( roots );
 				ASTOp::LogHistogram(roots);
 			}
 
@@ -1314,7 +1319,7 @@ namespace mu
 		ASTOp::Traverse_TopDown_Unique_Imprecise( roots, [&](mu::Ptr<ASTOp>& at )
 		{
 			// Only recurse instance construction ops.
-			bool processChildren = GetOpDataType(at->GetOpType())==DT_INSTANCE;
+			bool processChildren = GetOpDataType(at->GetOpType())== EDataType::Instance;
 
 			if ( at->GetOpType() == EOpType::IN_ADDMESH )
 			{
@@ -1442,8 +1447,8 @@ namespace mu
 		{
 			// We use 4 times the count because at the time we moved to sharing this count it
 			// was being used 4 times, and we want to keep the tests consistent.
-			m_optimizeIterationsLeft = Options->GetPrivate()->OptimisationOptions.MaxOptimisationLoopCount * 4;
-			m_optimizeIterationsMax = m_optimizeIterationsLeft;
+			OptimizeIterationsLeft = Options->GetPrivate()->OptimisationOptions.MaxOptimisationLoopCount * 4;
+			OptimizeIterationsMax = OptimizeIterationsLeft;
 
 			// The first duplicated data remover has the special mission of removing
 			// duplicated data (meshes) that may have been specified in the source
@@ -1473,21 +1478,21 @@ namespace mu
 			// First optimisation stage. It tries to resolve all the image sizes. This is necessary
 			// because some operations cannot be applied correctly until the image size is known
 			// like the grow-map generation.
-			bool modified = true;
-			int32 numIterations = 0;
-			while (modified)
+			bool bModified = true;
+			int32 NumIterations = 0;
+			while (bModified)
 			{
 				MUTABLE_CPUPROFILER_SCOPE(FirstStage);
 
-				--m_optimizeIterationsLeft;
-				++numIterations;
+				--OptimizeIterationsLeft;
+				++NumIterations;
 				UE_LOG(LogMutableCore, Verbose, TEXT("First optimise iteration %d, max %d, left %d"),
-						numIterations, m_optimizeIterationsMax, m_optimizeIterationsLeft);
+						NumIterations, OptimizeIterationsMax, OptimizeIterationsLeft);
 
-				modified = false;
+				bModified = false;
 
 				UE_LOG(LogMutableCore, Verbose, TEXT(" - size optimiser"));
-				modified |= SizeOptimiserAST( roots );
+				bModified |= SizeOptimiserAST( roots );
 			}
 
 			//UE_LOG(LogMutableCore, Verbose, TEXT("(int) %s : %ld"), TEXT("ast size"), int64(ASTOp::CountNodes(roots)));
@@ -1517,7 +1522,7 @@ namespace mu
 				for (Ptr<ASTOp>& Root : roots)
 				{
 					// Constant subtree generation
-					modified = ConstantGeneratorAST(Options->GetPrivate(), Root, Pass);
+					bModified = ConstantGeneratorAST(Options->GetPrivate(), Root, Pass);
 				}
 
 				DuplicatedDataRemoverAST(roots);
diff --git a/Engine/Plugins/Mutable/Source/MutableTools/Private/MuT/CodeOptimiser.h b/Engine/Plugins/Mutable/Source/MutableTools/Private/MuT/CodeOptimiser.h
--- a/Engine/Plugins/Mutable/Source/MutableTools/Private/MuT/CodeOptimiser.h
+++ b/Engine/Plugins/Mutable/Source/MutableTools/Private/MuT/CodeOptimiser.h
@@ -10,59 +10,46 @@
 namespace mu
 {
 
-    //---------------------------------------------------------------------------------------------
-    //! Code optimiser
-    //---------------------------------------------------------------------------------------------
+    /** Code optimiser: Rebuild an ASTOp graph so that it is more efficient. */
     class CodeOptimiser
     {
     public:
-        CodeOptimiser( Ptr<CompilerOptions> options, TArray<FStateCompilationData>& states );
+        CodeOptimiser( Ptr<CompilerOptions>, TArray<FStateCompilationData>& States );
 
-        //! Optimise the virtual machine code, using several transforms.
         void OptimiseAST();
 
     private:
 
 		Ptr<CompilerOptions> Options;
-
-        //!
 		TArray<FStateCompilationData>& States;
 
-        //! The max number of optimize iterations is shared across several stages now.
-        //! This is how many are left
-        int m_optimizeIterationsMax=0;
-        int m_optimizeIterationsLeft=0;
+        /** The max number of optimize iterations is shared across several stages now. This is how many are left. */
+        int32 OptimizeIterationsMax=0;
+        int32 OptimizeIterationsLeft=0;
 
         //! Full optimisation pass
         void FullOptimiseAST( ASTOpList& roots, int32 Pass);
 
-        //! Optimise the code of a model for a specific state, generating new instructions and
-        //! state information.
+        /** Optimise the code for each specific state, generating new instructions and state information. */
         void OptimiseStatesAST();
 
     };
 
 
-    //---------------------------------------------------------------------------------------------
-    //! Scan a subtree and return true if any op of a type in the provided subset is found.
-    //! Intermediate data is used between calls to apply, so don't remove program code or directly
-    //! change the instructions. Adding new instructions is ok.
-    //---------------------------------------------------------------------------------------------
-    class FindOpTypeVisitor
+    /** Scan a subtree and return true if any op of a type in the provided subset is found.
+    * Intermediate data is used between calls to apply, so don't remove program code or directly
+    * change the instructions. Adding new instructions is ok.
+    */
+    class FFindOpTypeVisitor
     {
     public:
 
-        FindOpTypeVisitor(const TArray<EOpType>& types)
-            : m_typesToFind(types)
-        {
-        }
-
-        bool Apply( FProgram& program, OP::ADDRESS at );
-
+		FFindOpTypeVisitor(const TArray<EOpType>&);
+        bool Apply( FProgram&, OP::ADDRESS );
 
     private:
 
-        TArray<EOpType> m_typesToFind;
+        TArray<EOpType> TypesToFind;
 
 		TArray< TPair<bool,OP::ADDRESS> > Pending;
 
@@ -70,80 +57,60 @@ namespace mu
         // 1 children pending
         // 2 visited and not found
         // 3 visited and found
-		TArray<uint8> m_visited;
+		TArray<uint8> Visited;
     };
 
 
-    //---------------------------------------------------------------------------------------------
-    //! Scan a subtree and return true if the expression is constant.
-    //! It has the same restrictions that FindOpTypeVisitor
-    //---------------------------------------------------------------------------------------------
+    /** Scan a subtree and return true if the expression is constant.
+    * It has the same restrictions that FindOpTypeVisitor
+    */
     class IsConstantVisitor
     {
     public:
 
         IsConstantVisitor();
 
-        bool Apply( FProgram& program, OP::ADDRESS at );
+        bool Apply( FProgram&, OP::ADDRESS );
 
     private:
 
-        TUniquePtr<FindOpTypeVisitor> m_findOpTypeVisitor;
+        TUniquePtr<FFindOpTypeVisitor> FindOpTypeVisitor;
     };
 
 
-    //---------------------------------------------------------------------------------------------
-    //! ConstantGenerator replaces constant subtrees of operations with an equivalent single
-	//! constant value operation. 
-    //---------------------------------------------------------------------------------------------
-    extern bool ConstantGeneratorAST( const CompilerOptions::Private* options, Ptr<ASTOp>& root, int32 Pass );
-
-    //---------------------------------------------------------------------------------------------
-    //! \TODO: shapes, projectors, others? but not switches (they must be unique)
-    //---------------------------------------------------------------------------------------------
-    extern void DuplicatedDataRemoverAST( ASTOpList& roots );
-
+    /** ConstantGenerator replaces constant subtrees of operations with an equivalent single */
+    extern bool ConstantGeneratorAST( const CompilerOptions::Private*, Ptr<ASTOp>& Root, int32 Pass );
 
-    //---------------------------------------------------------------------------------------------
-    //! Mark all the duplicated code instructions to point at the same operation, leaving the
-    //! copies unreachable.
-    //---------------------------------------------------------------------------------------------
-    extern void DuplicatedCodeRemoverAST( ASTOpList& roots );
+    /** \TODO: shapes, projectors, others? but not switches (they must be unique) */
+    extern void DuplicatedDataRemoverAST( ASTOpList& Roots );
 
+    /** Mark all the duplicated code instructions to point at the same operation, leaving the copies unreachable. */
+    extern void DuplicatedCodeRemoverAST( ASTOpList& Roots );
 
-    //---------------------------------------------------------------------------------------------
-    //! All kinds of optimisations that depend on the meaning of each operation
-    //---------------------------------------------------------------------------------------------
-    extern bool SemanticOptimiserAST( ASTOpList& roots, const FModelOptimizationOptions& optimisationOptions, int32 Pass);
+    /** All kinds of optimisations that depend on the meaning of each operation. */
+    extern bool SemanticOptimiserAST( ASTOpList& Roots, const FModelOptimizationOptions&, int32 Pass);
 
-    //---------------------------------------------------------------------------------------------
-    //! Semantic operator that reorders instructions moving expensive ones down to the
-    //! leaves of the expressions trying to turn them into constants.
-    //---------------------------------------------------------------------------------------------
-    extern bool SinkOptimiserAST( ASTOpList& roots, const FModelOptimizationOptions& );
+    /** Semantic operator that reorders instructions moving expensive ones down to the
+    * leaves of the expressions trying to turn them into constants.
+    */
+    extern bool SinkOptimiserAST( ASTOpList& Roots, const FModelOptimizationOptions& );
 
-    //---------------------------------------------------------------------------------------------
-    //!
-    //---------------------------------------------------------------------------------------------
-    extern bool SizeOptimiserAST( ASTOpList& roots );
+    /** */
+    extern bool SizeOptimiserAST( ASTOpList& Roots );
 
-    //---------------------------------------------------------------------------------------------
-    //!
-    //---------------------------------------------------------------------------------------------
+    /** */
     extern bool LocalLogicOptimiserAST(ASTOpList& roots);
 
-    //---------------------------------------------------------------------------------------------
-    //! Discard all LODs beyond the given lod count.
-    //---------------------------------------------------------------------------------------------
+    /** Discard all LODs beyond the given lod count. */
     class LODCountReducerAST : public Visitor_TopDown_Unique_Cloning
     {
     public:
 
-        LODCountReducerAST( Ptr<ASTOp>& root, uint8 NumExtraLODsToBuildAfterFirstLOD );
+        LODCountReducerAST( Ptr<ASTOp>& Root, uint8 NumExtraLODsToBuildAfterFirstLOD );
 
     protected:
 
-        Ptr<ASTOp> Visit( Ptr<ASTOp> at, bool& processChildren ) override;
+        Ptr<ASTOp> Visit( Ptr<ASTOp>, bool& bOutProcessChildren ) override;
 
         uint8 NumExtraLODs = 0;
     };
@@ -157,9 +124,9 @@ namespace mu
     {
     public:
 
-        RuntimeParameterVisitorAST(const FStateCompilationData* pState);
+        RuntimeParameterVisitorAST(const FStateCompilationData* );
 
-        bool HasAny( const Ptr<ASTOp>& root );
+        bool HasAny( const Ptr<ASTOp>& Root );
 
     private:
 
@@ -176,8 +143,8 @@ namespace mu
             //! 1: only layouts are relevant
             uint8 OnlyLayoutsRelevant;
 
-            //! Address to visit
-            Ptr<ASTOp> At;
+            //! Operation to visit
+            Ptr<ASTOp> Op;
         };
 
         //!
@@ -197,67 +164,56 @@ namespace mu
         TMap<Ptr<ASTOp>, EOpState> Visited;
 
         //!
-        void AddIfNeeded( const FPendingItem& item );
+        void AddIfNeeded( const FPendingItem& );
 
     };
 
 
-    //---------------------------------------------------------------------------------------------
-    //! Remove all texture compression operations that would happen for runtime parameter changes.
-    //---------------------------------------------------------------------------------------------
+    /** Remove all texture compression operations that would happen for runtime parameter changes. */
     class RuntimeTextureCompressionRemoverAST : public Visitor_TopDown_Unique_Cloning
     {
     public:
 
-        RuntimeTextureCompressionRemoverAST(FStateCompilationData* state, bool bInAlwaysUncompress );
+        RuntimeTextureCompressionRemoverAST(FStateCompilationData*, bool bInAlwaysUncompress );
 
     protected:
 
-        Ptr<ASTOp> Visit( Ptr<ASTOp> at, bool& processChildren ) override;
+        Ptr<ASTOp> Visit( Ptr<ASTOp>, bool& bOutProcessChildren ) override;
 
     private:
 
-        RuntimeParameterVisitorAST m_hasRuntimeParamVisitor;
+        RuntimeParameterVisitorAST HasRuntimeParamVisitor;
 		bool bAlwaysUncompress = false;
 
     };
 
-    //---------------------------------------------------------------------------------------------
-    //! Restructure the code to move operations involving runtime parameters as high as possible
-    //! in the code tree.
-    //---------------------------------------------------------------------------------------------
+    /** Restructure the code to move operations involving runtime parameters as high as possible. */
     class ParameterOptimiserAST : public Visitor_TopDown_Unique_Cloning
     {
     private:
 
-		FStateCompilationData& m_stateProps;
-        // unused const GPU_PLATFORM_PROPS& m_gpuPlatformProps;
+		FStateCompilationData& StateProps;
 
     public:
 
-        ParameterOptimiserAST(FStateCompilationData&s,
-                               const FModelOptimizationOptions& optimisationOptions );
+        ParameterOptimiserAST(FStateCompilationData&, const FModelOptimizationOptions& );
 
         bool Apply();
 
-
     private:
 
-        Ptr<ASTOp> Visit( Ptr<ASTOp> at, bool& processChildren ) override;
+        Ptr<ASTOp> Visit( Ptr<ASTOp>, bool& bOutProcessChildren ) override;
 
-        bool m_modified;
+        bool bModified;
 
         FModelOptimizationOptions OptimisationOptions;
 
-        RuntimeParameterVisitorAST m_hasRuntimeParamVisitor;
+        RuntimeParameterVisitorAST HasRuntimeParamVisitor;
 
     };
 
 
-    //---------------------------------------------------------------------------------------------
-    //! Some masks are optional. If they are null, replace them by a white plain image of the right
-    //! size.
-    //---------------------------------------------------------------------------------------------
+    /** Some masks are optional.If they are null, replace them by a white plain image of the right size. */
     extern Ptr<ASTOp> EnsureValidMask( Ptr<ASTOp> mask, Ptr<ASTOp> base );
 
 
@@ -286,7 +242,7 @@ namespace mu
 
             bool operator==(const FState& o) const
             {
-                return  Op == o.Op &&
+                return Op == o.Op &&
 					bOnlyLayoutIsRelevant == o.bOnlyLayoutIsRelevant;
             }
 
@@ -302,7 +258,7 @@ namespace mu
 
         // Result cache
         // \todo optimise by storing unique lists separately and an index here.
-        TMap< FState, TSet<FString> > m_resultCache;
+        TMap< FState, TSet<FString> > ResultCache;
     };
 
 }
diff --git a/Engine/Plugins/Mutable/Source/MutableTools/Private/MuT/CodeOptimiser_Logic.cpp b/Engine/Plugins/Mutable/Source/MutableTools/Private/MuT/CodeOptimiser_Logic.cpp
--- a/Engine/Plugins/Mutable/Source/MutableTools/Private/MuT/CodeOptimiser_Logic.cpp
+++ b/Engine/Plugins/Mutable/Source/MutableTools/Private/MuT/CodeOptimiser_Logic.cpp
@@ -246,10 +246,10 @@ namespace mu
 //        // Gather all the conditions for every op
 //        m_callstacks.Apply( program, state );
 
-//        m_modified = false;
+//        bModified = false;
 //        Traverse( program.States[state].Root, program );
 
-//        return m_modified;
+//        return bModified;
 //    }
 
 
@@ -293,14 +293,14 @@ namespace mu
 
 //            if ( result==BoolExpression::BET_TRUE )
 //            {
-//                m_modified = true;
+//                bModified = true;
 //                m_callstacks.m_factsPerCallstack[at] = m_callstacks.m_factsPerCallstack[program.m_code[at].args.Conditional.yes];
 //                at = program.m_code[at].args.Conditional.yes;
 //                at = Visit( at, program );
 //            }
 //            else if ( result==BoolExpression::BET_FALSE )
 //            {
-//                m_modified = true;
+//                bModified = true;
 //                m_callstacks.m_factsPerCallstack[at] = m_callstacks.m_factsPerCallstack[program.m_code[at].args.Conditional.no];
 //                at = program.m_code[at].args.Conditional.no;
 //                at = Visit( at, program );
@@ -350,7 +350,7 @@ namespace mu
 ////                    if ( result==BoolExpression::BET_TRUE )
 ////                    {
 ////                        // This is the only possible branch.
-////                        m_modified = true;
+////                        bModified = true;
 ////                        at = args.Switch.values[i];
 ////                        //at = Visit( at, program );
 ////                        return at;
@@ -398,10 +398,10 @@ namespace mu
 //    bool LogicOptimiser::Apply( FProgram& program, int state )
 //    {
 //        m_initialCodeSize = program.OpAddress.Num();
-//        m_modified = false;
+//        bModified = false;
 //        Traverse( program.States[state].Root, program );
 
-//        return m_modified;
+//        return bModified;
 //    }
 
 
@@ -421,7 +421,7 @@ namespace mu
 //        OP::ARGS args = program.m_code[at].args;
 
 //        // We only apply it to instances, for performance reasons.
-//        if (GetOpType(type)!=DT_INSTANCE)
+//        if (GetOpType(type)!=Instance)
 //        {
 //            return at;
 //        }
@@ -445,13 +445,13 @@ namespace mu
 //            BoolExpression::TYPE result = pCond->Evaluate( m_facts );
 //            if ( result==BoolExpression::BET_TRUE )
 //            {
-//                m_modified = true;
+//                bModified = true;
 //                at = program.m_code[at].args.Conditional.yes;
 //                at = Visit( at, program );
 //            }
 //            else if ( result==BoolExpression::BET_FALSE )
 //            {
-//                m_modified = true;
+//                bModified = true;
 //                at = program.m_code[at].args.Conditional.no;
 //                at = Visit( at, program );
 //            }
@@ -511,7 +511,7 @@ namespace mu
 ////                    if ( result==BoolExpression::BET_TRUE )
 ////                    {
 ////                        // This is the only possible branch.
-////                        m_modified = true;
+////                        bModified = true;
 ////                        at = args.Switch.values[i];
 ////                        at = Visit( at, program );
 ////                        return at;
@@ -543,7 +543,7 @@ namespace mu
 ////            // update op if necessary
 ////            if ( FMemory::Memcmp( &args, &program.m_code[at].args, sizeof(OP::ARGS) ) )
 ////            {
-////                m_modified = true;
+////                bModified = true;
 ////                OP op = program.m_code[at];
 ////                op.args = args;
 ////                at = program.AddOp( op );
