[mutable] Remove unused static constant tables, and move some to the tools library.

diff --git a/Engine/Plugins/Mutable/Source/MutableTools/Private/MuT/CodeOptimiser_Parameters.cpp b/Engine/Plugins/Mutable/Source/MutableTools/Private/MuT/CodeOptimiser_Parameters.cpp
--- a/Engine/Plugins/Mutable/Source/MutableTools/Private/MuT/CodeOptimiser_Parameters.cpp
+++ b/Engine/Plugins/Mutable/Source/MutableTools/Private/MuT/CodeOptimiser_Parameters.cpp
@@ -62,25 +62,25 @@ namespace mu
 
         Pending.Reset();
 
-        FPendingItem start;
-        start.At = Root;
-        start.ItemType = 0;
-        start.OnlyLayoutsRelevant = 0;
-		Pending.Add( start );
+        FPendingItem Start;
+		Start.Op = Root;
+		Start.ItemType = 0;
+		Start.OnlyLayoutsRelevant = 0;
+		Pending.Add(Start);
 
         // Don't early out to be able to complete parent op cached flags
         while (!Pending.IsEmpty())
         {
-			FPendingItem item = Pending.Pop(EAllowShrinking::No);
-            Ptr<ASTOp> At = item.At;
+			FPendingItem Item = Pending.Pop(EAllowShrinking::No);
+            Ptr<ASTOp> Op = Item.Op;
 
-			if (!At)
+			if (!Op)
 			{
 				continue;
 			}
 
             // Not cached?
-			EOpState* FoundCached = Visited.Find(At);
+			EOpState* FoundCached = Visited.Find(Op);
             if ( 
 				!FoundCached
 				||
@@ -91,7 +91,7 @@ namespace mu
 					) 
 				)
             {
-                if (item.ItemType)
+                if (Item.ItemType)
                 {
                     // Item indicating we finished with all the children of a parent
                     check(*FoundCached ==EOpState::ChildrenPendingFull
@@ -101,7 +101,7 @@ namespace mu
 						*FoundCached ==EOpState::VisitedPartialDoesntHaveRuntime );
 
                     bool bSubtreeFound = false;
-                    At->ForEachChild( [&](ASTChild& ref)
+					Op->ForEachChild( [&](ASTChild& ref)
                     {
 						EOpState* FoundChild = Visited.Find(ref.child());
 						if (FoundChild && *FoundChild == EOpState::VisitedHasRuntime)
@@ -116,7 +116,7 @@ namespace mu
                     }
                     else
                     {
-						*FoundCached = item.OnlyLayoutsRelevant
+						*FoundCached = Item.OnlyLayoutsRelevant
                                 ? EOpState::VisitedPartialDoesntHaveRuntime
                                 : EOpState::VisitedFullDoesntHaveRuntime;
                     }
@@ -131,19 +131,19 @@ namespace mu
 						||
 						(*FoundCached == EOpState::VisitedPartialDoesntHaveRuntime
 							&&
-							item.OnlyLayoutsRelevant==0 )
+							Item.OnlyLayoutsRelevant==0 )
 					);
 
                     // Request the processing of the end of this instruction
-					FPendingItem endItem = item;
+					FPendingItem endItem = Item;
                     endItem.ItemType = 1;
                     Pending.Add( endItem );
-					Visited.Add(At, item.OnlyLayoutsRelevant
+					Visited.Add(Op, Item.OnlyLayoutsRelevant
 						? EOpState::ChildrenPendingPartial
 						: EOpState::ChildrenPendingFull);
 
                     // Is it a special op type?
-                    switch ( At->GetOpType() )
+                    switch ( Op->GetOpType() )
                     {
 
                     case EOpType::BO_PARAMETER:
@@ -153,26 +153,26 @@ namespace mu
                     case EOpType::PR_PARAMETER:
                     case EOpType::IM_PARAMETER:
                     {
-						const ASTOpParameter* typed = static_cast<const ASTOpParameter*>(At.get());
-                        const TArray<FString>& params = State->nodeState.RuntimeParams;
-                        if ( params.Find( typed->Parameter.Name)
+						const ASTOpParameter* Typed = static_cast<const ASTOpParameter*>(Op.get());
+                        const TArray<FString>& Params = State->nodeState.RuntimeParams;
+                        if (Params.Find(Typed->Parameter.Name)
                              !=
                              INDEX_NONE )
                         {
 							bRuntimeFound = true;
-                            Visited.Add(At,EOpState::VisitedHasRuntime);
+                            Visited.Add(Op,EOpState::VisitedHasRuntime);
                         }
                         break;
                     }
 
                     default:
                     {
-                        At->ForEachChild([&](ASTChild& ref)
+                        Op->ForEachChild([&](ASTChild& ref)
                         {
 							FPendingItem ChildItem;
                             ChildItem.ItemType = 0;
-                            ChildItem.At = ref.child();
-                            ChildItem.OnlyLayoutsRelevant = item.OnlyLayoutsRelevant;
+                            ChildItem.Op = ref.child();
+                            ChildItem.OnlyLayoutsRelevant = Item.OnlyLayoutsRelevant;
                             AddIfNeeded(ChildItem);
                         });
                         break;
@@ -185,7 +185,7 @@ namespace mu
                 else
                 {
                     // We won't process it.
-					Visited.Add(At,EOpState::NotVisited);
+					Visited.Add(Op,EOpState::NotVisited);
                 }
             }
         }
@@ -196,9 +196,9 @@ namespace mu
 
     void RuntimeParameterVisitorAST::AddIfNeeded( const FPendingItem& Item )
     {
-        if (Item.At)
+        if (Item.Op)
         {
-			const EOpState* Found = Visited.Find(Item.At);
+			const EOpState* Found = Visited.Find(Item.Op);
             if (!Found || *Found==EOpState::NotVisited)
             {
                 Pending.Add(Item);
@@ -260,10 +260,10 @@ namespace mu
 		FStateCompilationData& s,
             const FModelOptimizationOptions& optimisationOptions
             )
-            : m_stateProps(s)
-            , m_modified(false)
+            : StateProps(s)
+            , bModified(false)
             , OptimisationOptions(optimisationOptions)
-            , m_hasRuntimeParamVisitor(&s)
+            , HasRuntimeParamVisitor(&s)
     {
     }
 
@@ -272,19 +272,19 @@ namespace mu
     {
         MUTABLE_CPUPROFILER_SCOPE(ParameterOptimiserAST);
 
-        m_modified = false;
+        bModified = false;
 
         // Optimise the cloned tree
-        Traverse( m_stateProps.root );
+        Traverse( StateProps.root );
 
-        return m_modified;
+        return bModified;
     }
 
 
     Ptr<ASTOp> ParameterOptimiserAST::Visit( Ptr<ASTOp> At, bool& processChildren )
     {
         // Only process children if there are runtime parameters in the subtree
-        processChildren = m_hasRuntimeParamVisitor.HasAny(At);
+        processChildren = HasRuntimeParamVisitor.HasAny(At);
 
         EOpType type = At->GetOpType();
         switch ( type )
@@ -310,7 +310,7 @@ namespace mu
 //                {
 //                    // TODO: this may unfold mesh combinations of some models increasing the size of
 //                    // the model data. Make this optimisation optional.
-//                    m_modified = true;
+//                    bModified = true;
 
 //                    OP yesOp = program.m_code[At];
 //                    yesOp.args.MeshMerge.base = program.m_code[ mergeArgs.base ].args.Conditional.yes;
@@ -345,7 +345,7 @@ namespace mu
 
 //                if ( !childBaseConst && childAddConst && addConst )
 //                {
-//                    m_modified = true;
+//                    bModified = true;
 
 //                    OP bottom = program.m_code[At];
 //                    bottom.args.MeshMerge.base = childAddAt;
@@ -375,13 +375,13 @@ namespace mu
             // conditional down
             bool optimised = false;
 
-            if ( !m_hasRuntimeParamVisitor.HasAny( TypedOp->condition.child() ) )
+            if ( !HasRuntimeParamVisitor.HasAny( TypedOp->condition.child() ) )
             {
                 EOpType yesType = TypedOp->yes->GetOpType();
                 EOpType noType = TypedOp->no->GetOpType();
 
-                bool yesHasAny = m_hasRuntimeParamVisitor.HasAny( TypedOp->yes.child() );
-                bool noHasAny = m_hasRuntimeParamVisitor.HasAny( TypedOp->no.child() );
+                bool yesHasAny = HasRuntimeParamVisitor.HasAny( TypedOp->yes.child() );
+                bool noHasAny = HasRuntimeParamVisitor.HasAny( TypedOp->no.child() );
 
                 if ( !optimised && yesHasAny && noHasAny && yesType==noType)
                 {
@@ -602,7 +602,7 @@ namespace mu
                 }
             }
 
-            m_modified |= optimised;
+            bModified |= optimised;
 
             break;
         }
@@ -617,7 +617,7 @@ namespace mu
 
 //            OP::ADDRESS variable = program.m_code[At].args.Switch.variable;
 
-//            if ( !m_hasRuntimeParamVisitor.HasAny( variable, program ) )
+//            if ( !HasRuntimeParamVisitor.HasAny( variable, program ) )
 //            {
 //                SWITCH_CHAIN chain = GetSwitchChain( program, At );
 
@@ -635,7 +635,7 @@ namespace mu
 //                    {
 //                        if (!branchHasAny)
 //                        {
-//                            branchHasAny = m_hasRuntimeParamVisitor.HasAny( it->second, program );
+//                            branchHasAny = HasRuntimeParamVisitor.HasAny( it->second, program );
 //                        }
 //                    }
 //                }
@@ -832,7 +832,7 @@ namespace mu
 
 //            }
 
-//            m_modified |= optimised;
+//            bModified |= optimised;
 
             break;
         }
@@ -856,9 +856,9 @@ namespace mu
             EOpType blockType = blockAt->GetOpType();
             EOpType baseType = baseAt->GetOpType();
 
-            bool baseHasRuntime = m_hasRuntimeParamVisitor.HasAny( baseAt );
-            bool blockHasRuntime = m_hasRuntimeParamVisitor.HasAny( blockAt );
-            bool layoutHasRuntime = m_hasRuntimeParamVisitor.HasAny( layoutAt );
+            bool baseHasRuntime = HasRuntimeParamVisitor.HasAny( baseAt );
+            bool blockHasRuntime = HasRuntimeParamVisitor.HasAny( blockAt );
+            bool layoutHasRuntime = HasRuntimeParamVisitor.HasAny( layoutAt );
 
             bool optimised = false;
 
@@ -972,7 +972,7 @@ namespace mu
 				const ASTOpImageCompose* typedBaseAt = static_cast<const ASTOpImageCompose*>(baseAt.get());
 
 				Ptr<ASTOp> baseBlockAt = typedBaseAt->BlockImage.child();
-                bool baseBlockHasAny = m_hasRuntimeParamVisitor.HasAny( baseBlockAt );
+                bool baseBlockHasAny = HasRuntimeParamVisitor.HasAny( baseBlockAt );
                 if ( baseBlockHasAny )
                 {
                     optimised = true;
@@ -993,7 +993,7 @@ namespace mu
             // This optimisation requires a lot of memory for every target. Use only if
             // we are optimising for GPU processing.
             if ( !optimised && blockHasRuntime && !baseHasRuntime
-                 //&& m_stateProps.m_gpu.m_external
+                 //&& StateProps.m_gpu.m_external
                  // TODO BLEH
                  // Only worth in case of more than one block using the same operation. Move this
                  // optimisation to that test.
@@ -1179,7 +1179,7 @@ namespace mu
 
 
             // Try to optimise the base
-            if ( !optimised && baseHasRuntime /*&& m_stateProps.nodeState.m_optimisation.m_gpu.m_external*/ )
+            if ( !optimised && baseHasRuntime /*&& StateProps.nodeState.m_optimisation.m_gpu.m_external*/ )
             {
                 switch ( baseType )
                 {
@@ -1310,7 +1310,7 @@ namespace mu
             }
 
 
-            m_modified = m_modified || optimised;
+            bModified = bModified || optimised;
             break;
         }
 
@@ -1321,7 +1321,7 @@ namespace mu
         {
             OP op = program.m_code[At];
 
-            if ( !m_hasRuntimeParamVisitor.HasAny( op.args.ImageLayer.mask, program ) )
+            if ( !HasRuntimeParamVisitor.HasAny( op.args.ImageLayer.mask, program ) )
             {
                 // If both the base and the blended have the same image layer operation with
                 // similar parameters, we can move that operation up.
@@ -1359,7 +1359,7 @@ namespace mu
                                  && !AreMasksOverlapping( program, baseMaskAt, maskAt ) )
                             {
                                 // We can apply the transform
-                                m_modified = true;
+                                bModified = true;
 
                                 OP baseOp;
                                 baseOp.type = EOpType::IM_BLEND;
@@ -1418,11 +1418,11 @@ namespace mu
             {
                 const ASTOpImageLayerColor* typedSource = static_cast<const ASTOpImageLayerColor*>(sourceOp.get());
 
-                bool colourHasRuntime = m_hasRuntimeParamVisitor.HasAny( typedSource->color.child() );
+                bool colourHasRuntime = HasRuntimeParamVisitor.HasAny( typedSource->color.child() );
 
                 if (colourHasRuntime)
                 {
-                    m_modified = true;
+                    bModified = true;
 
 					Ptr<ASTOpImageLayerColor> top = mu::Clone<ASTOpImageLayerColor>(sourceOp);
 
@@ -1470,11 +1470,11 @@ namespace mu
             MUTABLE_CPUPROFILER_SCOPE(AccumulateAllImageFormatsOpAST);
 
             // Initially, all formats are supported
-            m_allSupported.fill(1);
+            AllSupported.fill(1);
 
             // The initial traversal state is no format supported
-            m_initialState.fill(0);
-            Traverse( roots, m_initialState );
+            InitialState.fill(0);
+            Traverse( roots, InitialState );
         }
 
 
@@ -1482,22 +1482,22 @@ namespace mu
         {
             bool recurse = false;
 
-            const std::array<uint8_t, (size_t)EImageFormat::Count>& currentFormats = GetCurrentState();
+            const std::array<uint8, SIZE_T(EImageFormat::Count)>& CurrentFormats = GetCurrentState();
 
             // Remove unsupported formats
-            if (GetOpDataType( At->GetOpType() )==DT_IMAGE)
+            if (GetOpDataType( At->GetOpType() )==EDataType::Image)
             {
-				std::array<uint8_t, (size_t)EImageFormat::Count>* it = m_supportedFormats.Find(At);
+				std::array<uint8, SIZE_T(EImageFormat::Count)>* it = SupportedFormats.Find(At);
                 if (!it)
                 {
                     // Default to all supported
-                    m_supportedFormats.Add(At, m_allSupported);
-                    it = m_supportedFormats.Find(At);
+                    SupportedFormats.Add(At, AllSupported);
+                    it = SupportedFormats.Find(At);
                 }
 
-                for ( unsigned f=0; f< (unsigned)EImageFormat::Count; ++f )
+                for ( int32 f=0; f< int32(EImageFormat::Count); ++f )
                 {
-                    if ( !currentFormats[f] )
+                    if ( !CurrentFormats[f] )
                     {
                         (*it)[f] = 0;
                     }
@@ -1516,12 +1516,12 @@ namespace mu
 
                 if ( TypedOp->mask )
                 {
-                    std::array<uint8_t, (size_t)EImageFormat::Count> newState;
-                    newState.fill(0);
-                    newState[(size_t)EImageFormat::L_UByte ] = 1;
-                    newState[(size_t)EImageFormat::L_UByteRLE ] = 1;
+                    std::array<uint8, SIZE_T(EImageFormat::Count)> NewState;
+                    NewState.fill(0);
+                    NewState[SIZE_T(EImageFormat::L_UByte) ] = 1;
+                    NewState[SIZE_T(EImageFormat::L_UByteRLE) ] = 1;
 
-                    RecurseWithState( TypedOp->mask.child(), newState );
+                    RecurseWithState( TypedOp->mask.child(), NewState );
                 }
                 break;
             }
@@ -1533,15 +1533,15 @@ namespace mu
                 RecurseWithCurrentState( TypedOp->base.child() );
                 RecurseWithCurrentState( TypedOp->blend.child() );
 
-                std::array<uint8_t, (size_t)EImageFormat::Count> newState;
-                newState.fill(0);
+                std::array<uint8, SIZE_T(EImageFormat::Count)> NewState;
+                NewState.fill(0);
                 // TODO
-                //newState[ L_UByte ] = 1;
-                //newState[ L_UByteRLE ] = 1;
+                //NewState[ L_UByte ] = 1;
+                //NewState[ L_UByteRLE ] = 1;
 
                 if ( TypedOp->mask )
                 {
-                    RecurseWithState( TypedOp->mask.child(), newState );
+                    RecurseWithState( TypedOp->mask.child(), NewState );
                 }
                 break;
             }
@@ -1552,16 +1552,16 @@ namespace mu
 
                 RecurseWithCurrentState( TypedOp->children[TypedOp->op.args.ImageDisplace.source].child() );
 
-                std::array<uint8_t, (size_t)EImageFormat::Count> newState;
-                newState.fill(0);
-                newState[(size_t)EImageFormat::L_UByte ] = 1;
-                newState[(size_t)EImageFormat::L_UByteRLE ] = 1;
-                RecurseWithState( TypedOp->children[TypedOp->op.args.ImageDisplace.displacementMap].child(), newState );
+                std::array<uint8, SIZE_T(EImageFormat::Count)> NewState;
+                NewState.fill(0);
+                NewState[SIZE_T(EImageFormat::L_UByte) ] = 1;
+                NewState[SIZE_T(EImageFormat::L_UByteRLE) ] = 1;
+                RecurseWithState( TypedOp->children[TypedOp->op.args.ImageDisplace.displacementMap].child(), NewState );
                 break;
             }
 
             default:
-                SetCurrentState( m_initialState );
+                SetCurrentState( InitialState );
                 recurse = true;
                 break;
             }
@@ -1570,15 +1570,15 @@ namespace mu
         }
 
 
-        bool IsSupportedFormat( Ptr<ASTOp> At, EImageFormat format ) const
+        bool IsSupportedFormat( Ptr<ASTOp> Op, EImageFormat Format ) const
         {
-            const std::array<uint8_t, (size_t)EImageFormat::Count>* it = m_supportedFormats.Find(At);
+            const std::array<uint8, SIZE_T(EImageFormat::Count)>* it = SupportedFormats.Find(Op);
             if (!it)
             {
                 return false;
             }
 
-            return (*it)[(size_t)format]!=0;
+            return (*it)[SIZE_T(Format)]!=0;
         }
 
 
@@ -1586,13 +1586,13 @@ namespace mu
 
         //! Formats known to be supported for every instruction.
         //! Count*code.size() entries
-        TMap< Ptr<ASTOp>, std::array<uint8_t, (size_t)EImageFormat::Count> > m_supportedFormats;
+        TMap< Ptr<ASTOp>, std::array<uint8, SIZE_T(EImageFormat::Count)> > SupportedFormats;
 
         //! Constant convenience initial value
-        std::array<uint8_t, (size_t)EImageFormat::Count> m_initialState;
+        std::array<uint8, SIZE_T(EImageFormat::Count)> InitialState;
 
         //! Constant convenience initial value
-        std::array<uint8_t, (size_t)EImageFormat::Count> m_allSupported;
+        std::array<uint8, SIZE_T(EImageFormat::Count)> AllSupported;
 
     };
 
@@ -1601,7 +1601,7 @@ namespace mu
     void SubtreeRelevantParametersVisitorAST::Run( Ptr<ASTOp> root )
     {
         // Cached?
-		TSet<FString>* it = m_resultCache.Find( FState(root,false) );
+		TSet<FString>* it = ResultCache.Find( FState(root,false) );
         if (it)
         {
             Parameters = *it;
@@ -1686,7 +1686,7 @@ namespace mu
                 return true;
             });
 
-            m_resultCache.Add( FState(root,false), Parameters );
+            ResultCache.Add( FState(root,false), Parameters );
         }
     }
 
@@ -1703,25 +1703,25 @@ namespace mu
     {
     public:
 
-        StateCacheDetectorAST(FStateCompilationData* pState )
-            : m_hasRuntimeParamVisitor( pState )
+        StateCacheDetectorAST(FStateCompilationData* State )
+            : HasRuntimeParamVisitor( State )
         {
             ASTOpList roots;
-            roots.Add(pState->root);
+            roots.Add(State->root);
 			Traverse(roots, { false,false });
 
-            pState->m_updateCache.Empty();
-            pState->m_dynamicResources.Empty();
+            State->m_updateCache.Empty();
+            State->m_dynamicResources.Empty();
 
-            for( const TPair<Ptr<ASTOp>, bool>& i : m_cache )
+            for( const TPair<Ptr<ASTOp>, bool>& i : Cache )
             {
                 if ( i.Value )
                 {
-                    pState->m_updateCache.Add( i.Key );
+                    State->m_updateCache.Add( i.Key );
                 }
             }
 
-            for(const TPair<Ptr<ASTOp>, bool>& i : m_dynamicResourceRoot )
+            for(const TPair<Ptr<ASTOp>, bool>& i : DynamicResourceRoot )
             {
                 if ( i.Value )
                 {
@@ -1736,7 +1736,7 @@ namespace mu
 						ParamCopy.Add(e);
 					}
 
-                    pState->m_dynamicResources.Emplace( i.Key, MoveTemp(ParamCopy) );
+                    State->m_dynamicResources.Emplace( i.Key, MoveTemp(ParamCopy) );
                 }
             }
         }
@@ -1744,37 +1744,37 @@ namespace mu
 
         bool Visit( const Ptr<ASTOp>& At ) override
         {
-            bool thisIsRuntime = m_hasRuntimeParamVisitor.HasAny( At );
+            bool bThisIsRuntime = HasRuntimeParamVisitor.HasAny( At );
 
-            bool resourceRoot = GetCurrentState().Key;
-            bool parentIsRuntime = GetCurrentState().Value;
+            bool bResourceRoot = GetCurrentState().Key;
+            bool bParentIsRuntime = GetCurrentState().Value;
 
-			m_cache.FindOrAdd(At, false);
+			Cache.FindOrAdd(At, false);
 
-            EOpType type = At->GetOpType();
-            if ( GetOpDesc( type ).cached )
+            EOpType Type = At->GetOpType();
+            if ( GetOpToolsDesc( Type ).bCached )
             {
                 // If parent is runtime, but we are not
-                if ( (!thisIsRuntime) && parentIsRuntime
+                if ( (!bThisIsRuntime) && bParentIsRuntime
                      &&
                      // Resource roots are special, and they don't need to be marked as updateCache
                      // since the dynamicResource flag takes care of everything.
-                     !resourceRoot )
+                     !bResourceRoot )
                 {
                     // We want to cache this result to update the instances.
                     // Mark this as update cache
-                    m_cache.Add(At, true);
+                    Cache.Add(At, true);
                 }
             }
 
-            if ( !m_cache[At] && resourceRoot && thisIsRuntime )
+            if ( !Cache[At] && bResourceRoot && bThisIsRuntime )
             {
-                m_dynamicResourceRoot.Add(At, true);
+                DynamicResourceRoot.Add(At, true);
             }
 
-            if ( !m_cache[At] && thisIsRuntime )
+            if ( !Cache[At] && bThisIsRuntime )
             {
-                switch( type )
+                switch(Type)
                 {
                 case EOpType::IN_ADDIMAGE:
                 case EOpType::IN_ADDMESH:
@@ -1784,26 +1784,26 @@ namespace mu
                 {
 					const ASTOpInstanceAdd* TypedOp = static_cast<const ASTOpInstanceAdd*>(At.get());
 
-                    TPair<bool,bool> newState;
-                    newState.Key = false; //resource root
-                    newState.Value = thisIsRuntime;
-                    RecurseWithState( TypedOp->instance.child(), newState );
+                    TPair<bool,bool> NewState;
+                    NewState.Key = false; //resource root
+                    NewState.Value = bThisIsRuntime;
+                    RecurseWithState( TypedOp->instance.child(), NewState );
 
                     if ( TypedOp->value )
                     {
-                        newState.Key = true; //resource root
-                        newState.Value = thisIsRuntime;
-                        RecurseWithState( TypedOp->value.child(), newState );
+                        NewState.Key = true; //resource root
+                        NewState.Value = bThisIsRuntime;
+                        RecurseWithState( TypedOp->value.child(), NewState );
                     }
                     return false;
                 }
 
                 default:
                 {
-                    TPair<bool,bool> newState;
-                    newState.Key = false; //resource root
-                    newState.Value = thisIsRuntime;
-                    SetCurrentState(newState);
+                    TPair<bool,bool> NewState;
+                    NewState.Key = false; //resource root
+                    NewState.Value = bThisIsRuntime;
+                    SetCurrentState(NewState);
                     return true;
                 }
 
@@ -1816,10 +1816,10 @@ namespace mu
     private:
 
         //!
-        TMap<Ptr<ASTOp>,bool> m_cache;
-		TMap<Ptr<ASTOp>,bool> m_dynamicResourceRoot;
+        TMap<Ptr<ASTOp>,bool> Cache;
+		TMap<Ptr<ASTOp>,bool> DynamicResourceRoot;
 
-        RuntimeParameterVisitorAST m_hasRuntimeParamVisitor;
+        RuntimeParameterVisitorAST HasRuntimeParamVisitor;
     };
 
 
@@ -1854,7 +1854,7 @@ namespace mu
                 processChildren = false;
 
                 // See if we can convert it to a more efficient format
-                if ( GetOpDataType( At->GetOpType() )==DT_IMAGE )
+                if ( GetOpDataType( At->GetOpType() )== EDataType::Image )
                 {
                     FImageDesc desc = At->GetImageDesc();
 
@@ -1890,37 +1890,37 @@ namespace mu
     //---------------------------------------------------------------------------------------------
     //---------------------------------------------------------------------------------------------
     RuntimeTextureCompressionRemoverAST::RuntimeTextureCompressionRemoverAST(
-		FStateCompilationData* pState,
+		FStateCompilationData* State,
 		bool bInAlwaysUncompress
 	)
-        : m_hasRuntimeParamVisitor(pState)
+        : HasRuntimeParamVisitor(State)
 		, bAlwaysUncompress(bInAlwaysUncompress)
     {
-        Traverse( pState->root );
+        Traverse( State->root );
     }
 
 
     Ptr<ASTOp> RuntimeTextureCompressionRemoverAST::Visit( Ptr<ASTOp> At, bool& processChildren )
     {
         EOpType type = At->GetOpType();
-        processChildren = GetOpDataType(type)==DT_INSTANCE;
+        processChildren = GetOpDataType(type)== EDataType::Instance;
 
         // TODO: Finer grained: what if the runtime parameter just selects between compressed
         // textures? We don't want them uncompressed.
         if( type==EOpType::IN_ADDIMAGE )
         {
-			ASTOpInstanceAdd* typedAt = static_cast<ASTOpInstanceAdd*>(At.get());
+			ASTOpInstanceAdd* TypedOp = static_cast<ASTOpInstanceAdd*>(At.get());
 			
-			if (typedAt && typedAt->value)
+			if (TypedOp && TypedOp->value)
 			{			
-				Ptr<ASTOp> imageAt = typedAt->value.child();
+				Ptr<ASTOp> ImageOp = TypedOp->value.child();
 
 				// Does it have a runtime parameter in its subtree?
-				bool hasRuntimeParameter = m_hasRuntimeParamVisitor.HasAny( imageAt );
+				bool hasRuntimeParameter = HasRuntimeParamVisitor.HasAny(ImageOp);
 
 				if (bAlwaysUncompress || hasRuntimeParameter)
 				{
-					FImageDesc imageDesc = imageAt->GetImageDesc( true );
+					FImageDesc imageDesc = ImageOp->GetImageDesc( true );
 
 					// Is it a compressed format?
 					EImageFormat format = imageDesc.m_format;
@@ -1929,16 +1929,16 @@ namespace mu
 
 					if (isCompressedFormat)
 					{
-						Ptr<ASTOpInstanceAdd> newAt = mu::Clone<ASTOpInstanceAdd>(At);
+						Ptr<ASTOpInstanceAdd> NewOp = mu::Clone<ASTOpInstanceAdd>(At);
 
 						// Add a new format operation to uncompress the image
 						Ptr<ASTOpImagePixelFormat> fop = new ASTOpImagePixelFormat;
 						fop->Format = uncompressedFormat;
 						fop->FormatIfAlpha = uncompressedFormat;
-						fop->Source = imageAt;
+						fop->Source = ImageOp;
 
-						newAt->value = fop;
-						At = newAt;
+						NewOp->value = fop;
+						At = NewOp;
 					}
 				}
 			}
@@ -1979,7 +1979,7 @@ namespace mu
 	    		return false;
 	    	}
 	    	
-	    	if (GetOpDataType(OpType) != DT_INSTANCE)
+	    	if (GetOpDataType(OpType) != EDataType::Instance)
 	    	{
 	    		return false;
 	    	}
@@ -2019,12 +2019,12 @@ namespace mu
         	
             if (TypedOp->lods.Num() > NumLODs)
             {
-                Ptr<ASTOpAddLOD> newAt = mu::Clone<ASTOpAddLOD>(At);
-                while (newAt->lods.Num() > NumLODs)
+                Ptr<ASTOpAddLOD> NewOp = mu::Clone<ASTOpAddLOD>(At);
+                while (NewOp->lods.Num() > NumLODs)
                 {
-                    newAt->lods.Pop();
+                    NewOp->lods.Pop();
                 }
-                At = newAt;
+                At = NewOp;
             }
 
             processChildren = false;
@@ -2081,14 +2081,14 @@ namespace mu
                 // format instructions.
                 bool modified = true;
                 int32 numIterations = 0;
-                while (modified && (!m_optimizeIterationsMax || m_optimizeIterationsLeft>0 || !numIterations ))
+                while (modified && (!OptimizeIterationsMax || OptimizeIterationsLeft>0 || !numIterations ))
                 {
                     modified = false;
 
                     ++numIterations;
-                    --m_optimizeIterationsLeft;
+                    --OptimizeIterationsLeft;
                     UE_LOG(LogMutableCore, Verbose, TEXT("State optimise iteration %d, max %d, left %d"),
-                            numIterations, m_optimizeIterationsMax, m_optimizeIterationsLeft);
+                            numIterations, OptimizeIterationsMax, OptimizeIterationsLeft);
 
                     UE_LOG(LogMutableCore, Verbose, TEXT(" - before parameter optimiser"));
 
@@ -2177,12 +2177,12 @@ namespace mu
             bool modified = true;
             int32 numIterations = 0;
 			int32 Pass = 1;
-            while (modified && (!m_optimizeIterationsMax || m_optimizeIterationsLeft>0 || !numIterations ))
+            while (modified && (!OptimizeIterationsMax || OptimizeIterationsLeft>0 || !numIterations ))
             {
                 ++numIterations;
-                --m_optimizeIterationsLeft;
+                --OptimizeIterationsLeft;
                 UE_LOG(LogMutableCore, Verbose, TEXT("State reoptimise iteration %d, max %d, left %d"),
-                        numIterations, m_optimizeIterationsMax, m_optimizeIterationsLeft);
+                        numIterations, OptimizeIterationsMax, OptimizeIterationsLeft);
 
                 modified = false;
 
diff --git a/Engine/Plugins/Mutable/Source/MutableTools/Private/MuT/CodeOptimiser_Semantic.cpp b/Engine/Plugins/Mutable/Source/MutableTools/Private/MuT/CodeOptimiser_Semantic.cpp
--- a/Engine/Plugins/Mutable/Source/MutableTools/Private/MuT/CodeOptimiser_Semantic.cpp
+++ b/Engine/Plugins/Mutable/Source/MutableTools/Private/MuT/CodeOptimiser_Semantic.cpp
@@ -1506,7 +1506,7 @@ namespace mu
                 // Skip intermediate ops until the layout
 //                    if ( program.m_code[sizeSourceAt].type != EOpType::IM_BLANKLAYOUT )
 //                    {
-//                        m_modified = true;
+//                        bModified = true;
 
 //                        OP blackLayoutOp;
 //                        blackLayoutOp.type = EOpType::IM_BLANKLAYOUT;
diff --git a/Engine/Plugins/Mutable/Source/MutableTools/Private/MuT/Compiler.cpp b/Engine/Plugins/Mutable/Source/MutableTools/Private/MuT/Compiler.cpp
--- a/Engine/Plugins/Mutable/Source/MutableTools/Private/MuT/Compiler.cpp
+++ b/Engine/Plugins/Mutable/Source/MutableTools/Private/MuT/Compiler.cpp
@@ -3,12 +3,12 @@
 
 #include "MuT/Compiler.h"
 
+#include "MuT/CompilerPrivate.h"
 #include "MuT/AST.h"
 #include "MuT/ASTOpParameter.h"
 #include "MuT/ASTOpInstanceAdd.h"
 #include "MuT/CodeGenerator.h"
 #include "MuT/CodeOptimiser.h"
-#include "MuT/CompilerPrivate.h"
 #include "MuT/ErrorLog.h"
 #include "MuT/Node.h"
 #include "MuR/Image.h"
@@ -50,6 +50,146 @@ namespace mu
     }
 
 
+	// clang-format off
+	static const FOpToolsDesc SOpToolsDescs[] =
+	{
+		// cached	supported base image formats
+		{ false,	{ 0, 0, 0, 0, 0, 0, 0, 0, 0, 0, 0, 0, 0, 0, 0, 0, 0, 0, 0, 0, 0, 0, 0, 0, 0, 0, 0, 0, 0 }		},	// NONE
+
+		{ false,	{ 0, 0, 0, 0, 0, 0, 0, 0, 0, 0, 0, 0, 0, 0, 0, 0, 0, 0, 0, 0, 0, 0, 0, 0, 0, 0, 0, 0, 0 }		},	// BO_CONSTANT
+		{ false,	{ 0, 0, 0, 0, 0, 0, 0, 0, 0, 0, 0, 0, 0, 0, 0, 0, 0, 0, 0, 0, 0, 0, 0, 0, 0, 0, 0, 0, 0 }		},	// NU_CONSTANT
+		{ false,	{ 0, 0, 0, 0, 0, 0, 0, 0, 0, 0, 0, 0, 0, 0, 0, 0, 0, 0, 0, 0, 0, 0, 0, 0, 0, 0, 0, 0, 0 }		},	// SC_CONSTANT
+		{ false,	{ 0, 0, 0, 0, 0, 0, 0, 0, 0, 0, 0, 0, 0, 0, 0, 0, 0, 0, 0, 0, 0, 0, 0, 0, 0, 0, 0, 0, 0 }		},	// CO_CONSTANT
+		{ true,		{ 1, 1, 1, 1, 1, 1, 1, 1, 1, 1, 1, 1, 0, 0, 0, 0, 0, 0, 0, 0, 0, 0, 0, 0, 0, 0, 0, 0, 0 }		},	// IM_CONSTANT
+		{ true,		{ 0, 0, 0, 0, 0, 0, 0, 0, 0, 0, 0, 0, 0, 0, 0, 0, 0, 0, 0, 0, 0, 0, 0, 0, 0, 0, 0, 0, 0 }		},	// ME_CONSTANT
+		{ false,	{ 0, 0, 0, 0, 0, 0, 0, 0, 0, 0, 0, 0, 0, 0, 0, 0, 0, 0, 0, 0, 0, 0, 0, 0, 0, 0, 0, 0, 0 }		},	// LA_CONSTANT
+		{ false,	{ 0, 0, 0, 0, 0, 0, 0, 0, 0, 0, 0, 0, 0, 0, 0, 0, 0, 0, 0, 0, 0, 0, 0, 0, 0, 0, 0, 0, 0 }		},	// PR_CONSTANT
+		{ false,	{ 0, 0, 0, 0, 0, 0, 0, 0, 0, 0, 0, 0, 0, 0, 0, 0, 0, 0, 0, 0, 0, 0, 0, 0, 0, 0, 0, 0, 0 }		},	// ST_CONSTANT
+		{ true,		{ 0, 0, 0, 0, 0, 0, 0, 0, 0, 0, 0, 0, 0, 0, 0, 0, 0, 0, 0, 0, 0, 0, 0, 0, 0, 0, 0, 0, 0 }		},	// ED_CONSTANT
+		{ false,	{ 0, 0, 0, 0, 0, 0, 0, 0, 0, 0, 0, 0, 0, 0, 0, 0, 0, 0, 0, 0, 0, 0, 0, 0, 0, 0, 0, 0, 0 }		},	// MA_CONSTANT
+
+		{ false,	{ 0, 0, 0, 0, 0, 0, 0, 0, 0, 0, 0, 0, 0, 0, 0, 0, 0, 0, 0, 0, 0, 0, 0, 0, 0, 0, 0, 0, 0 }		},	// BO_PARAMETER
+		{ false,	{ 0, 0, 0, 0, 0, 0, 0, 0, 0, 0, 0, 0, 0, 0, 0, 0, 0, 0, 0, 0, 0, 0, 0, 0, 0, 0, 0, 0, 0 }		},	// NU_PARAMETER
+		{ false,	{ 0, 0, 0, 0, 0, 0, 0, 0, 0, 0, 0, 0, 0, 0, 0, 0, 0, 0, 0, 0, 0, 0, 0, 0, 0, 0, 0, 0, 0 }		},	// SC_PARAMETER
+		{ false,	{ 0, 0, 0, 0, 0, 0, 0, 0, 0, 0, 0, 0, 0, 0, 0, 0, 0, 0, 0, 0, 0, 0, 0, 0, 0, 0, 0, 0, 0 }		},	// CO_PARAMETER
+		{ false,	{ 0, 0, 0, 0, 0, 0, 0, 0, 0, 0, 0, 0, 0, 0, 0, 0, 0, 0, 0, 0, 0, 0, 0, 0, 0, 0, 0, 0, 0 }		},	// PR_PARAMETER
+		{ false,	{ 0, 0, 0, 0, 0, 0, 0, 0, 0, 0, 0, 0, 0, 0, 0, 0, 0, 0, 0, 0, 0, 0, 0, 0, 0, 0, 0, 0, 0 }		},	// IM_PARAMETER
+		{ false,	{ 0, 0, 0, 0, 0, 0, 0, 0, 0, 0, 0, 0, 0, 0, 0, 0, 0, 0, 0, 0, 0, 0, 0, 0, 0, 0, 0, 0, 0 }		},	// ST_PARAMETER
+		{ false,	{ 0, 0, 0, 0, 0, 0, 0, 0, 0, 0, 0, 0, 0, 0, 0, 0, 0, 0, 0, 0, 0, 0, 0, 0, 0, 0, 0, 0, 0 }		},	// MA_PARAMETER
+
+		{ true,		{ 0, 0, 0, 0, 0, 0, 0, 0, 0, 0, 0, 0, 0, 0, 0, 0, 0, 0, 0, 0, 0, 0, 0, 0, 0, 0, 0, 0, 0 }		},	// IM_REFERENCE
+		{ true,		{ 0, 0, 0, 0, 0, 0, 0, 0, 0, 0, 0, 0, 0, 0, 0, 0, 0, 0, 0, 0, 0, 0, 0, 0, 0, 0, 0, 0, 0 }		},	// ME_REFERENCE
+
+		{ false,	{ 0, 0, 0, 0, 0, 0, 0, 0, 0, 0, 0, 0, 0, 0, 0, 0, 0, 0, 0, 0, 0, 0, 0, 0, 0, 0, 0, 0, 0 }		},	// NU_CONDITIONAL
+		{ false,	{ 0, 0, 0, 0, 0, 0, 0, 0, 0, 0, 0, 0, 0, 0, 0, 0, 0, 0, 0, 0, 0, 0, 0, 0, 0, 0, 0, 0, 0 }		},	// SC_CONDITIONAL
+		{ false,	{ 0, 0, 0, 0, 0, 0, 0, 0, 0, 0, 0, 0, 0, 0, 0, 0, 0, 0, 0, 0, 0, 0, 0, 0, 0, 0, 0, 0, 0 }		},	// CO_CONDITIONAL
+		{ true,		{ 1, 1, 1, 1, 1, 1, 1, 1, 1, 1, 1, 1, 1, 1, 1, 1, 1, 1, 1, 1, 1, 1, 0, 0, 1, 1, 0, 0, 0 }		},	// IM_CONDITIONAL
+		{ true,		{ 0, 0, 0, 0, 0, 0, 0, 0, 0, 0, 0, 0, 0, 0, 0, 0, 0, 0, 0, 0, 0, 0, 0, 0, 0, 0, 0, 0, 0 }		},	// ME_CONDITIONAL
+		{ false,	{ 0, 0, 0, 0, 0, 0, 0, 0, 0, 0, 0, 0, 0, 0, 0, 0, 0, 0, 0, 0, 0, 0, 0, 0, 0, 0, 0, 0, 0 }		},	// LA_CONDITIONAL
+		{ false,	{ 0, 0, 0, 0, 0, 0, 0, 0, 0, 0, 0, 0, 0, 0, 0, 0, 0, 0, 0, 0, 0, 0, 0, 0, 0, 0, 0, 0, 0 }		},	// IN_CONDITIONAL
+		{ false,	{ 0, 0, 0, 0, 0, 0, 0, 0, 0, 0, 0, 0, 0, 0, 0, 0, 0, 0, 0, 0, 0, 0, 0, 0, 0, 0, 0, 0, 0 }		},	// ED_CONDITIONAL
+
+		{ false,	{ 0, 0, 0, 0, 0, 0, 0, 0, 0, 0, 0, 0, 0, 0, 0, 0, 0, 0, 0, 0, 0, 0, 0, 0, 0, 0, 0, 0, 0 }		},	// NU_SWITCH
+		{ false,	{ 0, 0, 0, 0, 0, 0, 0, 0, 0, 0, 0, 0, 0, 0, 0, 0, 0, 0, 0, 0, 0, 0, 0, 0, 0, 0, 0, 0, 0 }		},	// SC_SWITCH
+		{ false,	{ 0, 0, 0, 0, 0, 0, 0, 0, 0, 0, 0, 0, 0, 0, 0, 0, 0, 0, 0, 0, 0, 0, 0, 0, 0, 0, 0, 0, 0 }		},	// CO_SWITCH
+		{ true,		{ 1, 1, 1, 1, 1, 1, 1, 1, 1, 1, 1, 1, 1, 1, 1, 1, 1, 1, 1, 1, 1, 1, 0, 0, 1, 1, 0, 0, 0 }		},	// IM_SWITCH
+		{ true,		{ 0, 0, 0, 0, 0, 0, 0, 0, 0, 0, 0, 0, 0, 0, 0, 0, 0, 0, 0, 0, 0, 0, 0, 0, 0, 0, 0, 0, 0 }		},	// ME_SWITCH
+		{ false,	{ 0, 0, 0, 0, 0, 0, 0, 0, 0, 0, 0, 0, 0, 0, 0, 0, 0, 0, 0, 0, 0, 0, 0, 0, 0, 0, 0, 0, 0 }		},	// LA_SWITCH
+		{ false,	{ 0, 0, 0, 0, 0, 0, 0, 0, 0, 0, 0, 0, 0, 0, 0, 0, 0, 0, 0, 0, 0, 0, 0, 0, 0, 0, 0, 0, 0 }		},	// IN_SWITCH
+		{ false,	{ 0, 0, 0, 0, 0, 0, 0, 0, 0, 0, 0, 0, 0, 0, 0, 0, 0, 0, 0, 0, 0, 0, 0, 0, 0, 0, 0, 0, 0 }		},	// ED_SWITCH
+
+		{ false,	{ 0, 0, 0, 0, 0, 0, 0, 0, 0, 0, 0, 0, 0, 0, 0, 0, 0, 0, 0, 0, 0, 0, 0, 0, 0, 0, 0, 0, 0 }		},	// BO_LESS
+		{ false,	{ 0, 0, 0, 0, 0, 0, 0, 0, 0, 0, 0, 0, 0, 0, 0, 0, 0, 0, 0, 0, 0, 0, 0, 0, 0, 0, 0, 0, 0 }		},	// BO_EQUAL_SC_CONST
+		{ false,	{ 0, 0, 0, 0, 0, 0, 0, 0, 0, 0, 0, 0, 0, 0, 0, 0, 0, 0, 0, 0, 0, 0, 0, 0, 0, 0, 0, 0, 0 }		},	// BO_AND
+		{ false,	{ 0, 0, 0, 0, 0, 0, 0, 0, 0, 0, 0, 0, 0, 0, 0, 0, 0, 0, 0, 0, 0, 0, 0, 0, 0, 0, 0, 0, 0 }		},	// BO_OR
+		{ false,	{ 0, 0, 0, 0, 0, 0, 0, 0, 0, 0, 0, 0, 0, 0, 0, 0, 0, 0, 0, 0, 0, 0, 0, 0, 0, 0, 0, 0, 0 }		},	// BO_NOT
+
+		{ true,		{ 0, 0, 0, 0, 0, 0, 0, 0, 0, 0, 0, 0, 0, 0, 0, 0, 0, 0, 0, 0, 0, 0, 0, 0, 0, 0, 0, 0, 0 }		},	// SC_MULTIPLYADD
+		{ true,		{ 0, 0, 0, 0, 0, 0, 0, 0, 0, 0, 0, 0, 0, 0, 0, 0, 0, 0, 0, 0, 0, 0, 0, 0, 0, 0, 0, 0, 0 }		},	// SC_ARITHMETIC
+		{ true,		{ 0, 0, 0, 0, 0, 0, 0, 0, 0, 0, 0, 0, 0, 0, 0, 0, 0, 0, 0, 0, 0, 0, 0, 0, 0, 0, 0, 0, 0 }		},	// SC_CURVE
+
+		{ false,	{ 0, 0, 0, 0, 0, 0, 0, 0, 0, 0, 0, 0, 0, 0, 0, 0, 0, 0, 0, 0, 0, 0, 0, 0, 0, 0, 0, 0, 0 }		},	// CO_SAMPLEIMAGE
+		{ false,	{ 0, 0, 0, 0, 0, 0, 0, 0, 0, 0, 0, 0, 0, 0, 0, 0, 0, 0, 0, 0, 0, 0, 0, 0, 0, 0, 0, 0, 0 }		},	// CO_SWIZZLE
+		{ false,	{ 0, 0, 0, 0, 0, 0, 0, 0, 0, 0, 0, 0, 0, 0, 0, 0, 0, 0, 0, 0, 0, 0, 0, 0, 0, 0, 0, 0, 0 }		},	// CO_FROMSCALARS
+		{ false,	{ 0, 0, 0, 0, 0, 0, 0, 0, 0, 0, 0, 0, 0, 0, 0, 0, 0, 0, 0, 0, 0, 0, 0, 0, 0, 0, 0, 0, 0 }		},	// CO_ARITHMETIC
+
+		{ true,		{ 0, 1, 1, 1, 0, 0, 0, 0, 0, 0, 0, 0, 0, 0, 0, 0, 0, 0, 0, 0, 0, 0, 0, 0, 0, 0, 0, 0, 0 }		},	// IM_LAYER
+		{ true,		{ 0, 1, 1, 1, 0, 0, 0, 0, 0, 0, 0, 0, 0, 0, 0, 0, 0, 0, 0, 0, 0, 0, 0, 0, 0, 0, 0, 0, 0 }		},	// IM_LAYERCOLOUR
+		{ true,		{ 1, 1, 1, 1, 1, 1, 1, 1, 1, 1, 1, 1, 0, 0, 0, 0, 0, 0, 0, 0, 0, 0, 0, 0, 0, 0, 0, 0, 0 }		},	// IM_PIXELFORMAT
+		{ true,		{ 1, 1, 1, 1, 0, 0, 0, 0, 0, 0, 0, 0, 0, 0, 0, 0, 0, 0, 0, 0, 0, 0, 0, 0, 0, 0, 0, 0, 0 }		},	// IM_MIPMAP
+		{ true,		{ 0, 1, 1, 1, 0, 0, 0, 0, 0, 0, 0, 0, 0, 0, 0, 0, 0, 0, 0, 0, 0, 0, 0, 0, 0, 0, 0, 0, 0 }		},	// IM_RESIZE
+		{ true,		{ 0, 1, 1, 1, 0, 0, 0, 0, 0, 0, 0, 0, 0, 0, 0, 0, 0, 0, 0, 0, 0, 0, 0, 0, 0, 0, 0, 0, 0 }		},	// IM_RESIZELIKE
+		{ true,		{ 0, 1, 1, 1, 0, 0, 0, 0, 0, 0, 0, 0, 0, 0, 0, 0, 0, 0, 0, 0, 0, 0, 0, 0, 0, 0, 0, 0, 0 }		},	// IM_RESIZEREL
+		{ true,		{ 1, 1, 1, 1, 1, 1, 1, 1, 1, 1, 1, 1, 0, 0, 0, 0, 0, 0, 0, 0, 0, 0, 0, 0, 0, 0, 0, 0, 0 }		},	// IM_BLANKLAYOUT
+		{ true,		{ 1, 1, 1, 1, 1, 1, 1, 1, 1, 1, 1, 0, 0, 0, 1, 1, 1, 1, 1, 1, 1, 1, 0, 0, 1, 1, 0, 0, 0 }		},	// IM_COMPOSE
+		{ true,		{ 0, 1, 1, 1, 0, 0, 0, 0, 0, 0, 0, 0, 0, 0, 0, 0, 0, 0, 0, 0, 0, 0, 0, 0, 0, 0, 0, 0, 0 }		},	// IM_INTERPOLATE
+		{ true,		{ 0, 1, 1, 0, 0, 0, 0, 0, 0, 0, 0, 0, 0, 0, 0, 0, 0, 0, 0, 0, 0, 0, 0, 0, 0, 0, 0, 0, 0 }		},	// IM_SATURATE
+		{ true,		{ 0, 1, 1, 0, 0, 0, 0, 0, 0, 0, 0, 0, 0, 0, 0, 0, 0, 0, 0, 0, 0, 0, 0, 0, 0, 0, 0, 0, 0 }		},	// IM_LUMINANCE
+		{ true,		{ 0, 1, 1, 1, 0, 0, 0, 0, 0, 0, 0, 0, 0, 0, 0, 0, 0, 0, 0, 0, 0, 0, 0, 0, 0, 0, 0, 0, 0 }		},	// IM_SWIZZLE
+		{ true,		{ 0, 1, 1, 1, 0, 0, 0, 0, 0, 0, 0, 0, 0, 0, 0, 0, 0, 0, 0, 0, 0, 0, 0, 0, 0, 0, 0, 0, 0 }		},	// IM_COLOURMAP
+		{ true,		{ 0, 0, 0, 0, 0, 0, 0, 0, 0, 0, 0, 0, 0, 0, 0, 0, 0, 0, 0, 0, 0, 0, 0, 0, 0, 0, 0, 0, 0 }		},	// IM_GRADIENT
+		{ true,		{ 0, 1, 1, 1, 0, 0, 0, 0, 0, 0, 0, 0, 0, 0, 0, 0, 0, 0, 0, 0, 0, 0, 0, 0, 0, 0, 0, 0, 0 }		},	// IM_BINARISE
+		{ true,		{ 0, 0, 0, 0, 0, 0, 0, 0, 0, 0, 0, 0, 0, 0, 0, 0, 0, 0, 0, 0, 0, 0, 0, 0, 0, 0, 0, 0, 0 }		},	// IM_PLAINCOLOUR
+		{ true,		{ 0, 0, 0, 0, 0, 0, 0, 0, 0, 0, 0, 0, 0, 0, 0, 0, 0, 0, 0, 0, 0, 0, 0, 0, 0, 0, 0, 0, 0 }		},	// IM_CROP
+		{ true,		{ 0, 0, 0, 0, 0, 0, 0, 0, 0, 0, 0, 0, 0, 0, 0, 0, 0, 0, 0, 0, 0, 0, 0, 0, 0, 0, 0, 0, 0 }		},	// IM_PATCH
+		{ true,		{ 0, 0, 0, 0, 0, 0, 0, 0, 0, 0, 0, 0, 0, 0, 0, 0, 0, 0, 0, 0, 0, 0, 0, 0, 0, 0, 0, 0, 0 }		},	// IM_RASTERMESH
+		{ true,		{ 0, 0, 0, 0, 0, 0, 0, 0, 0, 0, 0, 0, 0, 0, 0, 0, 0, 0, 0, 0, 0, 0, 0, 0, 0, 0, 0, 0, 0 }		},	// IM_MAKEGROWMAP
+		{ true,		{ 0, 0, 0, 0, 0, 0, 0, 0, 0, 0, 0, 0, 0, 0, 0, 0, 0, 0, 0, 0, 0, 0, 0, 0, 0, 0, 0, 0, 0 }		},	// IM_DISPLACE
+		{ true,		{ 0, 0, 0, 0, 0, 0, 0, 0, 0, 0, 0, 0, 0, 0, 0, 0, 0, 0, 0, 0, 0, 0, 0, 0, 0, 0, 0, 0, 0 }		},	// IM_MULTILAYER
+		{ true,		{ 0, 1, 1, 1, 0, 0, 0, 0, 0, 0, 0, 0, 0, 0, 0, 0, 0, 0, 0, 0, 0, 0, 0, 0, 0, 0, 0, 0, 0 }		},	// IM_INVERT
+		{ true,		{ 0, 1, 1, 1, 0, 0, 0, 0, 0, 0, 0, 0, 0, 0, 0, 0, 0, 0, 0, 0, 0, 0, 0, 0, 0, 0, 0, 0, 0 }		},	// IM_NORMALCOMPOSITE
+		{ true,		{ 0, 1, 1, 1, 0, 0, 0, 0, 0, 0, 0, 0, 0, 0, 0, 0, 0, 0, 0, 0, 0, 0, 0, 0, 0, 0, 0, 0, 0 }		},	// IM_TRANSFORM
+
+		{ true,		{ 0, 0, 0, 0, 0, 0, 0, 0, 0, 0, 0, 0, 0, 0, 0, 0, 0, 0, 0, 0, 0, 0, 0, 0, 0, 0, 0, 0, 0 }		},	// ME_APPLYLAYOUT
+		{ true,		{ 0, 0, 0, 0, 0, 0, 0, 0, 0, 0, 0, 0, 0, 0, 0, 0, 0, 0, 0, 0, 0, 0, 0, 0, 0, 0, 0, 0, 0 }		},	// ME_DIFFERENCE
+		{ true,		{ 0, 0, 0, 0, 0, 0, 0, 0, 0, 0, 0, 0, 0, 0, 0, 0, 0, 0, 0, 0, 0, 0, 0, 0, 0, 0, 0, 0, 0 }		},	// ME_MORPH
+		{ true,		{ 0, 0, 0, 0, 0, 0, 0, 0, 0, 0, 0, 0, 0, 0, 0, 0, 0, 0, 0, 0, 0, 0, 0, 0, 0, 0, 0, 0, 0 }		},	// ME_MERGE
+		{ true,		{ 0, 0, 0, 0, 0, 0, 0, 0, 0, 0, 0, 0, 0, 0, 0, 0, 0, 0, 0, 0, 0, 0, 0, 0, 0, 0, 0, 0, 0 }		},	// ME_MASKCLIPMESH
+		{ true,		{ 0, 0, 0, 0, 0, 0, 0, 0, 0, 0, 0, 0, 0, 0, 0, 0, 0, 0, 0, 0, 0, 0, 0, 0, 0, 0, 0, 0, 0 }		},	// ME_MASKCLIPUVMASK
+		{ true,		{ 0, 0, 0, 0, 0, 0, 0, 0, 0, 0, 0, 0, 0, 0, 0, 0, 0, 0, 0, 0, 0, 0, 0, 0, 0, 0, 0, 0, 0 }		},	// ME_MASKDIFF
+		{ true,		{ 0, 0, 0, 0, 0, 0, 0, 0, 0, 0, 0, 0, 0, 0, 0, 0, 0, 0, 0, 0, 0, 0, 0, 0, 0, 0, 0, 0, 0 }		},	// ME_REMOVEMASK
+		{ true,		{ 0, 0, 0, 0, 0, 0, 0, 0, 0, 0, 0, 0, 0, 0, 0, 0, 0, 0, 0, 0, 0, 0, 0, 0, 0, 0, 0, 0, 0 }		},	// ME_FORMAT
+		{ true,		{ 0, 0, 0, 0, 0, 0, 0, 0, 0, 0, 0, 0, 0, 0, 0, 0, 0, 0, 0, 0, 0, 0, 0, 0, 0, 0, 0, 0, 0 }		},	// ME_EXTRACTLAYOUTBLOCK
+		{ true,		{ 0, 0, 0, 0, 0, 0, 0, 0, 0, 0, 0, 0, 0, 0, 0, 0, 0, 0, 0, 0, 0, 0, 0, 0, 0, 0, 0, 0, 0 }		},	// ME_TRANSFORM
+		{ true,		{ 0, 0, 0, 0, 0, 0, 0, 0, 0, 0, 0, 0, 0, 0, 0, 0, 0, 0, 0, 0, 0, 0, 0, 0, 0, 0, 0, 0, 0 }		},	// ME_CLIPMORPHPLANE
+		{ true,		{ 0, 0, 0, 0, 0, 0, 0, 0, 0, 0, 0, 0, 0, 0, 0, 0, 0, 0, 0, 0, 0, 0, 0, 0, 0, 0, 0, 0, 0 }		},	// ME_CLIPWITHMESH
+		{ true,		{ 0, 0, 0, 0, 0, 0, 0, 0, 0, 0, 0, 0, 0, 0, 0, 0, 0, 0, 0, 0, 0, 0, 0, 0, 0, 0, 0, 0, 0 }		},	// ME_SETSKELETON
+		{ true,		{ 0, 0, 0, 0, 0, 0, 0, 0, 0, 0, 0, 0, 0, 0, 0, 0, 0, 0, 0, 0, 0, 0, 0, 0, 0, 0, 0, 0, 0 }		},	// ME_PROJECT
+		{ true,		{ 0, 0, 0, 0, 0, 0, 0, 0, 0, 0, 0, 0, 0, 0, 0, 0, 0, 0, 0, 0, 0, 0, 0, 0, 0, 0, 0, 0, 0 }		},	// ME_APPLYPOSE
+		{ true,		{ 0, 0, 0, 0, 0, 0, 0, 0, 0, 0, 0, 0, 0, 0, 0, 0, 0, 0, 0, 0, 0, 0, 0, 0, 0, 0, 0, 0, 0 }		},	// ME_BINDSHAPE
+		{ true,		{ 0, 0, 0, 0, 0, 0, 0, 0, 0, 0, 0, 0, 0, 0, 0, 0, 0, 0, 0, 0, 0, 0, 0, 0, 0, 0, 0, 0, 0 }		},	// ME_APPLYSHAPE
+		{ true,		{ 0, 0, 0, 0, 0, 0, 0, 0, 0, 0, 0, 0, 0, 0, 0, 0, 0, 0, 0, 0, 0, 0, 0, 0, 0, 0, 0, 0, 0 }		},	// ME_CLIPDEFORM
+		{ true,		{ 0, 0, 0, 0, 0, 0, 0, 0, 0, 0, 0, 0, 0, 0, 0, 0, 0, 0, 0, 0, 0, 0, 0, 0, 0, 0, 0, 0, 0 }		},	// ME_MORPHRESHAPE
+		{ true,		{ 0, 0, 0, 0, 0, 0, 0, 0, 0, 0, 0, 0, 0, 0, 0, 0, 0, 0, 0, 0, 0, 0, 0, 0, 0, 0, 0, 0, 0 }		},	// ME_OPTIMIZESKINNING
+		{ true,		{ 0, 0, 0, 0, 0, 0, 0, 0, 0, 0, 0, 0, 0, 0, 0, 0, 0, 0, 0, 0, 0, 0, 0, 0, 0, 0, 0, 0, 0 }		},	// ME_ADDTAGS
+		{ true,		{ 0, 0, 0, 0, 0, 0, 0, 0, 0, 0, 0, 0, 0, 0, 0, 0, 0, 0, 0, 0, 0, 0, 0, 0, 0, 0, 0, 0, 0 }		},	// ME_TRANSFORMWITHMESH
+
+		{ false,	{ 0, 0, 0, 0, 0, 0, 0, 0, 0, 0, 0, 0, 0, 0, 0, 0, 0, 0, 0, 0, 0, 0, 0, 0, 0, 0, 0, 0, 0 }		},	// IN_ADDMESH
+		{ false,	{ 0, 0, 0, 0, 0, 0, 0, 0, 0, 0, 0, 0, 0, 0, 0, 0, 0, 0, 0, 0, 0, 0, 0, 0, 0, 0, 0, 0, 0 }		},	// IN_ADDIMAGE
+		{ false,	{ 0, 0, 0, 0, 0, 0, 0, 0, 0, 0, 0, 0, 0, 0, 0, 0, 0, 0, 0, 0, 0, 0, 0, 0, 0, 0, 0, 0, 0 }		},	// IN_ADDVECTOR
+		{ false,	{ 0, 0, 0, 0, 0, 0, 0, 0, 0, 0, 0, 0, 0, 0, 0, 0, 0, 0, 0, 0, 0, 0, 0, 0, 0, 0, 0, 0, 0 }		},	// IN_ADDSCALAR
+		{ false,	{ 0, 0, 0, 0, 0, 0, 0, 0, 0, 0, 0, 0, 0, 0, 0, 0, 0, 0, 0, 0, 0, 0, 0, 0, 0, 0, 0, 0, 0 }		},	// IN_ADDSTRING
+		{ false,	{ 0, 0, 0, 0, 0, 0, 0, 0, 0, 0, 0, 0, 0, 0, 0, 0, 0, 0, 0, 0, 0, 0, 0, 0, 0, 0, 0, 0, 0 }		},	// IN_ADDSURFACE
+		{ false,	{ 0, 0, 0, 0, 0, 0, 0, 0, 0, 0, 0, 0, 0, 0, 0, 0, 0, 0, 0, 0, 0, 0, 0, 0, 0, 0, 0, 0, 0 }		},	// IN_ADDCOMPONENT
+		{ false,	{ 0, 0, 0, 0, 0, 0, 0, 0, 0, 0, 0, 0, 0, 0, 0, 0, 0, 0, 0, 0, 0, 0, 0, 0, 0, 0, 0, 0, 0 }		},	// IN_ADDLOD
+		{ false,	{ 0, 0, 0, 0, 0, 0, 0, 0, 0, 0, 0, 0, 0, 0, 0, 0, 0, 0, 0, 0, 0, 0, 0, 0, 0, 0, 0, 0, 0 }		},	// IN_ADDEXTENSIONDATA
+
+		{ true,		{ 0, 0, 0, 0, 0, 0, 0, 0, 0, 0, 0, 0, 0, 0, 0, 0, 0, 0, 0, 0, 0, 0, 0, 0, 0, 0, 0, 0, 0 }		},	// LA_PACK
+		{ true,		{ 0, 0, 0, 0, 0, 0, 0, 0, 0, 0, 0, 0, 0, 0, 0, 0, 0, 0, 0, 0, 0, 0, 0, 0, 0, 0, 0, 0, 0 }		},	// LA_MERGE
+		{ true,		{ 0, 0, 0, 0, 0, 0, 0, 0, 0, 0, 0, 0, 0, 0, 0, 0, 0, 0, 0, 0, 0, 0, 0, 0, 0, 0, 0, 0, 0 }		},	// LA_REMOVEBLOCKS
+		{ true,		{ 0, 0, 0, 0, 0, 0, 0, 0, 0, 0, 0, 0, 0, 0, 0, 0, 0, 0, 0, 0, 0, 0, 0, 0, 0, 0, 0, 0, 0 }		},	// LA_FROMMESH
+	};
+
+	// clang-format on
+
+	static_assert(UE_ARRAY_COUNT(SOpToolsDescs) == (int32)EOpType::COUNT, "OperationDescMismatch");
+
+	const FOpToolsDesc& GetOpToolsDesc(EOpType type)
+	{
+		return SOpToolsDescs[(int32)type];
+	}
+
+
 	//---------------------------------------------------------------------------------------------
 	FProxyFileContext::FProxyFileContext()
 	{
@@ -124,9 +264,9 @@ namespace mu
 
 
     //---------------------------------------------------------------------------------------------
-    void CompilerOptions::SetOptimisationMaxIteration( int32 maxIterations )
+    void CompilerOptions::SetOptimisationMaxIteration( int32 MaxIterations )
     {
-        m_pD->OptimisationOptions.MaxOptimisationLoopCount = maxIterations;
+        m_pD->OptimisationOptions.MaxOptimisationLoopCount = MaxIterations;
     }
 
 
@@ -246,7 +386,7 @@ namespace mu
                 //ASTOp::Traverse_TopRandom_Unique_NonReentrant(Roots, [&](Ptr<ASTOp> Node)
                 {
 			        
-                    bool bProcessChildren = GetOpDataType(Node->GetOpType()) == DT_INSTANCE;
+                    bool bProcessChildren = GetOpDataType(Node->GetOpType()) == EDataType::Instance;
                     if (!bProcessChildren)
                     {
                         return false;
diff --git a/Engine/Plugins/Mutable/Source/MutableTools/Private/MuT/CompilerPrivate.h b/Engine/Plugins/Mutable/Source/MutableTools/Private/MuT/CompilerPrivate.h
--- a/Engine/Plugins/Mutable/Source/MutableTools/Private/MuT/CompilerPrivate.h
+++ b/Engine/Plugins/Mutable/Source/MutableTools/Private/MuT/CompilerPrivate.h
@@ -13,13 +13,24 @@
 
 namespace mu
 {
+
+	/** Generic data about a Mutable operation that is needed at compile time. */
+	struct FOpToolsDesc
+	{
+		/** True if the instruction is worth caching when generating models. */
+		bool bCached;
+
+		/** For image instructions, for every image format, true if it is supported as the base format of the operation. */
+		bool bSupportedBasePixelFormats[uint8(EImageFormat::Count)];
+	};
+
+	extern const FOpToolsDesc& GetOpToolsDesc(EOpType type);
+
 	/** Statistics about the proxy file usage. */
 	struct FProxyFileContext
 	{
 		FProxyFileContext();
 
-		/** Options */
-
 		/** Minimum data size in bytes to dump it to the disk. */
 		uint64 MinProxyFileSize = 64 * 1024;
 
