[mutable] Fixed Export/Import node issues.
- Fixed Export node not broadcasting name changes
- Fixed Import node crashing on target CO change while having a node already selected of the previous CO
- Now the title of the Import node will expose the name of the Export node linked.
- Added missing "virtual" keywords

#rnx
#tests Editor

diff --git a/Engine/Plugins/Mutable/Source/CustomizableObjectEditor/Private/MuCOE/Nodes/CustomizableObjectNodeExposePin.cpp b/Engine/Plugins/Mutable/Source/CustomizableObjectEditor/Private/MuCOE/Nodes/CustomizableObjectNodeExposePin.cpp
--- a/Engine/Plugins/Mutable/Source/CustomizableObjectEditor/Private/MuCOE/Nodes/CustomizableObjectNodeExposePin.cpp
+++ b/Engine/Plugins/Mutable/Source/CustomizableObjectEditor/Private/MuCOE/Nodes/CustomizableObjectNodeExposePin.cpp
@@ -2,6 +2,7 @@
 
 #include "MuCOE/Nodes/CustomizableObjectNodeExposePin.h"
 
+#include "ObjectEditorUtils.h"
 #include "MuCO/CustomizableObjectCustomVersion.h"
 #include "MuCO/ICustomizableObjectModule.h"
 #include "MuCOE/EdGraphSchema_CustomizableObject.h"
@@ -14,8 +15,10 @@ class UCustomizableObjectNodeRemapPins;
 
 void UCustomizableObjectNodeExposePin::PostEditChangeProperty(FPropertyChangedEvent& PropertyChangedEvent)
 {
+	Super::PostEditChangeProperty(PropertyChangedEvent);
+	
 	FProperty* PropertyThatChanged = PropertyChangedEvent.Property;
-	if (PropertyThatChanged && PropertyThatChanged->GetName() == TEXT("Name"))
+	if (PropertyThatChanged && PropertyThatChanged->GetName() == GET_MEMBER_NAME_STRING_CHECKED(UCustomizableObjectNodeExposePin, Name))
 	{
 		OnNameChangedDelegate.Broadcast();
 	}
@@ -74,7 +77,7 @@ void UCustomizableObjectNodeExposePin::OnRenameNode(const FString& NewName)
 {
 	if (!NewName.IsEmpty())
 	{
-		Name = NewName;
+		FObjectEditorUtils::SetPropertyValue(this, GET_MEMBER_NAME_STRING_CHECKED(UCustomizableObjectNodeExposePin, Name), NewName);
 	}
 }
 
diff --git a/Engine/Plugins/Mutable/Source/CustomizableObjectEditor/Private/MuCOE/Nodes/CustomizableObjectNodeExposePin.h b/Engine/Plugins/Mutable/Source/CustomizableObjectEditor/Private/MuCOE/Nodes/CustomizableObjectNodeExposePin.h
--- a/Engine/Plugins/Mutable/Source/CustomizableObjectEditor/Private/MuCOE/Nodes/CustomizableObjectNodeExposePin.h
+++ b/Engine/Plugins/Mutable/Source/CustomizableObjectEditor/Private/MuCOE/Nodes/CustomizableObjectNodeExposePin.h
@@ -26,19 +26,19 @@ public:
 	GENERATED_BODY()
 
 	// UObject interface
-	void PostEditChangeProperty(FPropertyChangedEvent& PropertyChangedEvent) override;
+	virtual void PostEditChangeProperty(FPropertyChangedEvent& PropertyChangedEvent) override;
 
 	// EdGraphNode interface
-	FText GetNodeTitle(ENodeTitleType::Type TitleType) const override;
-	FLinearColor GetNodeTitleColor() const override;
-	FText GetTooltipText() const override;
+	virtual FText GetNodeTitle(ENodeTitleType::Type TitleType) const override;
+	virtual FLinearColor GetNodeTitleColor() const override;
+	virtual FText GetTooltipText() const override;
 	virtual bool GetCanRenameNode() const override { return true; }
 	virtual void OnRenameNode(const FString& NewName) override;
 	virtual void BackwardsCompatibleFixup(int32 CustomizableObjectCustomVersion) override;
 
 	// UCustomizableObjectNode interface
 	virtual bool CanConnect(const UEdGraphPin* InOwnedInputPin, const UEdGraphPin* InOutputPin, bool& bOutIsOtherNodeBlocklisted, bool& bOutArePinsCompatible) const override;
-	void AllocateDefaultPins(UCustomizableObjectNodeRemapPins* RemapPins) override;
+	virtual void AllocateDefaultPins(UCustomizableObjectNodeRemapPins* RemapPins) override;
 
 	/** Return the Expose Pin Node expose pin name. */
 	FString GetNodeName() const;
diff --git a/Engine/Plugins/Mutable/Source/CustomizableObjectEditor/Private/MuCOE/Nodes/CustomizableObjectNodeExternalPin.cpp b/Engine/Plugins/Mutable/Source/CustomizableObjectEditor/Private/MuCOE/Nodes/CustomizableObjectNodeExternalPin.cpp
--- a/Engine/Plugins/Mutable/Source/CustomizableObjectEditor/Private/MuCOE/Nodes/CustomizableObjectNodeExternalPin.cpp
+++ b/Engine/Plugins/Mutable/Source/CustomizableObjectEditor/Private/MuCOE/Nodes/CustomizableObjectNodeExternalPin.cpp
@@ -205,8 +205,7 @@ void UCustomizableObjectNodeExternalPin::SetExternalObjectNodeId(FGuid Guid)
 
 UEdGraphPin* UCustomizableObjectNodeExternalPin::GetExternalPin() const
 {
-	const FName TargetPinName = GeneratePinName();
-	return FindPin(TargetPinName);
+	return Pins[0];
 }
 
 
@@ -232,15 +231,10 @@ void UCustomizableObjectNodeExternalPin::PropagateConnectionChanged()
 	NodePinConnectionListChanged(FollowOutputPinArray(*GetExternalPin()));
 }
 
-FName UCustomizableObjectNodeExternalPin::GeneratePinName() const
-{
-	return GetNodeExposePin() ? FName(GetNodeExposePin()->GetNodeName()) : UEdGraphSchema_CustomizableObject::GetPinCategoryName(PinType);
-}
-
 
 void UCustomizableObjectNodeExternalPin::AllocateDefaultPins(UCustomizableObjectNodeRemapPins* RemapPins)
 {
-	const FName PinName = GeneratePinName();
+	const FName PinName = UEdGraphSchema_CustomizableObject::GetPinCategoryName(PinType);
 	const FText PinFriendlyName = UEdGraphSchema_CustomizableObject::GetPinCategoryFriendlyName(PinType);
 
 	const bool bIsArrayPinCategory = PinType == UEdGraphSchema_CustomizableObject::PC_GroupProjector;
@@ -259,7 +253,14 @@ FText UCustomizableObjectNodeExternalPin::GetNodeTitle(ENodeTitleType::Type Titl
 	}
 	else
 	{
-		return FText::Format(LOCTEXT("External_Pin_Title_WithName", "{0}\nImport {1} Pin"), FText::FromString(ExternalObject->GetName()), PinTypeName);
+		if (const UCustomizableObjectNodeExposePin* ExportNode = GetNodeExposePin())
+		{
+			return FText::Format(LOCTEXT("External_Pin_ExportNode_Title", "{0}\nImport {1} Pin"), FText::FromString(*ExportNode->GetNodeName()), PinTypeName);
+		}
+		else
+		{
+			return FText::Format(LOCTEXT("External_Pin_ExternalObject_Title", "{0}\nImport {1} Pin"), FText::FromString(ExternalObject->GetName()), PinTypeName);
+		}
 	}
 }
 
diff --git a/Engine/Plugins/Mutable/Source/CustomizableObjectEditor/Private/MuCOE/Nodes/CustomizableObjectNodeExternalPin.h b/Engine/Plugins/Mutable/Source/CustomizableObjectEditor/Private/MuCOE/Nodes/CustomizableObjectNodeExternalPin.h
--- a/Engine/Plugins/Mutable/Source/CustomizableObjectEditor/Private/MuCOE/Nodes/CustomizableObjectNodeExternalPin.h
+++ b/Engine/Plugins/Mutable/Source/CustomizableObjectEditor/Private/MuCOE/Nodes/CustomizableObjectNodeExternalPin.h
@@ -70,12 +70,6 @@ public:
 
 private:
 	
-	/**
-	 * Computes the name to be used by the pin. 
-	 * @return the name that the pin should use
-	 */
-	FName GeneratePinName() const;
-	
 	/** Linked Node Expose Pin node guid. */
 	UPROPERTY()
 	FGuid ExternalObjectNodeId;
diff --git a/Engine/Plugins/Mutable/Source/CustomizableObjectEditor/Public/MuCOE/Nodes/CustomizableObjectNode.h b/Engine/Plugins/Mutable/Source/CustomizableObjectEditor/Public/MuCOE/Nodes/CustomizableObjectNode.h
--- a/Engine/Plugins/Mutable/Source/CustomizableObjectEditor/Public/MuCOE/Nodes/CustomizableObjectNode.h
+++ b/Engine/Plugins/Mutable/Source/CustomizableObjectEditor/Public/MuCOE/Nodes/CustomizableObjectNode.h
@@ -224,7 +224,7 @@ public:
 	 */
 	virtual bool ShouldBreakExistingConnections(const UEdGraphPin* InputPin, const UEdGraphPin* OutputPin) const;
 
-	void PostEditChangeProperty(FPropertyChangedEvent& PropertyChangedEvent) override;
+	virtual void PostEditChangeProperty(FPropertyChangedEvent& PropertyChangedEvent) override;
 
 	/** Custom post duplicate function. Called at the beginning of duplication, before the nodes have their Guid updated. */
 	virtual void BeginPostDuplicate(bool bDuplicateForPIE) {};
