[Mutable] Fix Mesh Streaming streaming incorrect LODs.

diff --git a/Engine/Plugins/Mutable/Source/CustomizableObject/Private/MuCO/CustomizableObjectInstance.cpp b/Engine/Plugins/Mutable/Source/CustomizableObject/Private/MuCO/CustomizableObjectInstance.cpp
--- a/Engine/Plugins/Mutable/Source/CustomizableObject/Private/MuCO/CustomizableObjectInstance.cpp
+++ b/Engine/Plugins/Mutable/Source/CustomizableObject/Private/MuCO/CustomizableObjectInstance.cpp
@@ -3660,7 +3660,8 @@ void UCustomizableInstancePrivate::InitSkeletalMeshData(const TSharedRef<FUpdate
 		RenderData->NumInlinedLODs = NumLODsAvailablePerComponent - OperationData->FirstResidentLOD[ComponentName];
 		RenderData->NumNonOptionalLODs = NumLODsAvailablePerComponent - OperationData->FirstLODAvailable[ComponentName];
 		RenderData->CurrentFirstLODIdx = OperationData->FirstResidentLOD[ComponentName];
-		RenderData->LODBiasModifier = 0;
+		RenderData->PendingFirstLODIdx = RenderData->CurrentFirstLODIdx;
+		RenderData->LODBiasModifier = OperationData->FirstLODAvailable[ComponentName];		
 		
 		if (bEnableRayTracingFix)
 		{
diff --git a/Engine/Plugins/Mutable/Source/CustomizableObject/Private/MuCO/CustomizableObjectMeshUpdate.cpp b/Engine/Plugins/Mutable/Source/CustomizableObject/Private/MuCO/CustomizableObjectMeshUpdate.cpp
--- a/Engine/Plugins/Mutable/Source/CustomizableObject/Private/MuCO/CustomizableObjectMeshUpdate.cpp
+++ b/Engine/Plugins/Mutable/Source/CustomizableObject/Private/MuCO/CustomizableObjectMeshUpdate.cpp
@@ -169,7 +169,7 @@ namespace impl
 	{
 		MUTABLE_CPUPROFILER_SCOPE(Task_Mutable_UpdateMesh_Loop);
 
-		if (Task->IsCancelled() || LODIndex == OperationData->CurrentFirstLODIdx)
+		if (Task->IsCancelled() || LODIndex == OperationData->CurrentFirstLODIdx + OperationData->AssetLODBias)
 		{
 			Task_Mutable_UpdateMesh_End(OperationData, Task, InstanceID);
 			return;
@@ -222,7 +222,7 @@ namespace impl
 		TSharedPtr<const mu::FInstance> Instance = System->BeginUpdate(InstanceID, OperationData->Parameters, OperationData->State, LODMask);
 		check(Instance);
 
-		Task_Mutable_UpdateMesh_Loop(OperationData, Task, InstanceID, OperationData->PendingFirstLODIdx);
+		Task_Mutable_UpdateMesh_Loop(OperationData, Task, InstanceID, OperationData->PendingFirstLODIdx + OperationData->AssetLODBias);
 	}
 	
 } // namespace
@@ -252,6 +252,8 @@ void FCustomizableObjectMeshStreamIn::RequestMeshUpdate(const FContext& Context)
 	}
 #endif
 
+	OperationData->AssetLODBias = Context.AssetLODBias;
+
 	UCustomizableObjectSystemPrivate* CustomizableObjectSystem = UCustomizableObjectSystem::GetInstance()->GetPrivate();
 	
 	TaskSynchronization.Increment();
@@ -313,7 +315,7 @@ void FCustomizableObjectMeshStreamIn::ConvertMesh(const FContext& Context)
 
 	for (int32 LODIndex = PendingFirstLODIdx; LODIndex < CurrentFirstLODIdx; ++LODIndex)
 	{
-		TSharedPtr<const mu::FMesh> MutableMesh = OperationData->Meshes[LODIndex];
+		TSharedPtr<const mu::FMesh> MutableMesh = OperationData->Meshes[LODIndex + Context.AssetLODBias];
 
 		if (!MutableMesh)
 		{
diff --git a/Engine/Plugins/Mutable/Source/CustomizableObject/Private/MuCO/CustomizableObjectMeshUpdate.h b/Engine/Plugins/Mutable/Source/CustomizableObject/Private/MuCO/CustomizableObjectMeshUpdate.h
--- a/Engine/Plugins/Mutable/Source/CustomizableObject/Private/MuCO/CustomizableObjectMeshUpdate.h
+++ b/Engine/Plugins/Mutable/Source/CustomizableObject/Private/MuCO/CustomizableObjectMeshUpdate.h
@@ -39,6 +39,8 @@ struct FMutableMeshOperationData
 
 	/** The requested first LOD resource index. With domain = [0, ResourceState.NumLODs[. NOT THE ASSET LOD INDEX! */
 	int32 PendingFirstLODIdx = INDEX_NONE;
+
+	int32 AssetLODBias = 0;
 };
 
 extern template class TRenderAssetUpdate<FSkelMeshUpdateContext>;
