[Mutable] Avoid loading objects on node PostLoads.

Only load them if required due to backwards compatible code.

diff --git a/Engine/Plugins/Mutable/Source/CustomizableObjectEditor/Private/MuCOE/Nodes/CustomizableObjectNodeMaterial.cpp b/Engine/Plugins/Mutable/Source/CustomizableObjectEditor/Private/MuCOE/Nodes/CustomizableObjectNodeMaterial.cpp
--- a/Engine/Plugins/Mutable/Source/CustomizableObjectEditor/Private/MuCOE/Nodes/CustomizableObjectNodeMaterial.cpp
+++ b/Engine/Plugins/Mutable/Source/CustomizableObjectEditor/Private/MuCOE/Nodes/CustomizableObjectNodeMaterial.cpp
@@ -381,6 +381,11 @@ void UCustomizableObjectNodeMaterial::BackwardsCompatibleFixup(int32 Customizabl
 	
 	if (CustomizableObjectCustomVersion == FCustomizableObjectCustomVersion::AutomaticNodeMaterial)
 	{
+		if (Material)
+		{
+			ConditionalPostLoadReference(*Material); // Make sure the Material has been fully loaded.
+		}
+		
 		for (const FCustomizableObjectNodeMaterialImage& Image : Images_DEPRECATED)
 		{
 			const FString OldPinName = (Image.LayerIndex == -1 ? Image.Name : Image.PinName) + FString(TEXT("_Input_Image"));
@@ -569,11 +574,21 @@ void UCustomizableObjectNodeMaterial::BackwardsCompatibleFixup(int32 Customizabl
 	
 	if (CustomizableObjectCustomVersion == FCustomizableObjectCustomVersion::NodeMaterialAddTablePin)
 	{
+		if (Material)
+		{
+			ConditionalPostLoadReference(*Material); // Make sure the Material has been fully loaded.
+		}
+		
 		Super::ReconstructNode();
 	}
 
 	if (CustomizableObjectCustomVersion == FCustomizableObjectCustomVersion::AddedTableMaterialSwitch)
 	{
+		if (Material)
+		{
+			ConditionalPostLoadReference(*Material); // Make sure the Material has been fully loaded.
+		}
+		
 		UMaterialInstance* DefaultPinValue = nullptr;
 
 		if (const UEdGraphPin* MaterialAssetPin = GetMaterialAssetPin())
@@ -692,6 +707,8 @@ void UCustomizableObjectNodeMaterial::BackwardsCompatibleFixup(int32 Customizabl
 	{
 		if (Material && Material->GetCachedExpressionData().bHasMaterialLayers)
 		{
+			ConditionalPostLoadReference(*Material); // Make sure the Material has been fully loaded.
+			
 			// Needed since we can not get the layer index of repeated parameters
 			Super::ReconstructNode();
 		}
@@ -760,17 +777,6 @@ void UCustomizableObjectNodeMaterial::Serialize(FArchive& Ar)
 }
 
 
-void UCustomizableObjectNodeMaterial::PostLoad()
-{
-	Super::PostLoad();
-	
-	if (Material)
-	{
-		ConditionalPostLoadReference(*Material); // Make sure the Material has been fully loaded.
-	}
-}
-
-
 void UCustomizableObjectNodeMaterial::PostEditChangeProperty(FPropertyChangedEvent& PropertyChangedEvent)
 {
 	Super::PostEditChangeProperty(PropertyChangedEvent);
@@ -806,9 +812,6 @@ FText UCustomizableObjectNodeMaterial::GetNodeTitle(ENodeTitleType::Type TitleTy
 }
 
 
-
-
-
 UEdGraphPin* UCustomizableObjectNodeMaterial::OutputPin() const
 {
 	FString PinFriendlyName = TEXT("Mesh Section");
diff --git a/Engine/Plugins/Mutable/Source/CustomizableObjectEditor/Private/MuCOE/Nodes/CustomizableObjectNodeMaterial.h b/Engine/Plugins/Mutable/Source/CustomizableObjectEditor/Private/MuCOE/Nodes/CustomizableObjectNodeMaterial.h
--- a/Engine/Plugins/Mutable/Source/CustomizableObjectEditor/Private/MuCOE/Nodes/CustomizableObjectNodeMaterial.h
+++ b/Engine/Plugins/Mutable/Source/CustomizableObjectEditor/Private/MuCOE/Nodes/CustomizableObjectNodeMaterial.h
@@ -107,7 +107,6 @@ public:
 	
 	// UObject interface.
 	virtual void Serialize(FArchive& Ar) override;
-	virtual void PostLoad() override;
 	virtual void PostEditChangeProperty(FPropertyChangedEvent& PropertyChangedEvent) override;
 
 	// EdGraphNode interface
diff --git a/Engine/Plugins/Mutable/Source/CustomizableObjectEditor/Private/MuCOE/Nodes/CustomizableObjectNodeSkeletalMesh.cpp b/Engine/Plugins/Mutable/Source/CustomizableObjectEditor/Private/MuCOE/Nodes/CustomizableObjectNodeSkeletalMesh.cpp
--- a/Engine/Plugins/Mutable/Source/CustomizableObjectEditor/Private/MuCOE/Nodes/CustomizableObjectNodeSkeletalMesh.cpp
+++ b/Engine/Plugins/Mutable/Source/CustomizableObjectEditor/Private/MuCOE/Nodes/CustomizableObjectNodeSkeletalMesh.cpp
@@ -30,10 +30,8 @@ struct FSlateBrush;
 static const TCHAR* SKELETAL_MESH_PIN_NAME = TEXT("Skeletal Mesh");
 
 
-void UCustomizableObjectNodeSkeletalMesh::PostLoad()
+void UCustomizableObjectNodeSkeletalMesh::LoadObjects()
 {
-	Super::PostLoad();
-
 	if (SkeletalMesh)
 	{
 		ConditionalPostLoadReference(*SkeletalMesh);
@@ -441,6 +439,8 @@ void UCustomizableObjectNodeSkeletalMesh::BackwardsCompatibleFixup(int32 Customi
 	
 	if (CustomizableObjectCustomVersion == FCustomizableObjectCustomVersion::PostLoadToCustomVersion)
 	{
+		LoadObjects();
+		
 		for (FCustomizableObjectNodeSkeletalMeshLOD& LOD : LODs_DEPRECATED)
 		{
 			for(FCustomizableObjectNodeSkeletalMeshMaterial& Material : LOD.Materials)
@@ -504,6 +504,8 @@ void UCustomizableObjectNodeSkeletalMesh::BackwardsCompatibleFixup(int32 Customi
 
 	if (CustomizableObjectCustomVersion == FCustomizableObjectCustomVersion::AutomaticNodeSkeletalMesh)
 	{
+		LoadObjects();
+		
 		for (int32 LODIndex = 0; LODIndex < LODs_DEPRECATED.Num(); ++LODIndex)
 		{
 			const FCustomizableObjectNodeSkeletalMeshLOD& LOD = LODs_DEPRECATED[LODIndex];
@@ -569,16 +571,22 @@ void UCustomizableObjectNodeSkeletalMesh::BackwardsCompatibleFixup(int32 Customi
 
 	if (CustomizableObjectCustomVersion == FCustomizableObjectCustomVersion::AutomaticNodeSkeletalMeshPinDataOuter)
 	{
+		LoadObjects();
+		
 		ReconstructNode(); // Pins did not have Pin Data. Reconstruct them.
 	}
 
 	if (CustomizableObjectCustomVersion == FCustomizableObjectCustomVersion::AutomaticNodeSkeletalMeshPinDataUProperty)
 	{
+		LoadObjects();
+		
 		ReconstructNode(CreateRemapPinsByName()); // Correct pins but incorrect Pin Data. Reconstruct and remap pins only by name, no Pin Data.
 	}
 
 	if (CustomizableObjectCustomVersion == FCustomizableObjectCustomVersion::IgnoreDisabledSections)
 	{
+		LoadObjects();
+		
 		ReconstructNode(); // Pins representing disabled sections could be present. 
 	}
 
@@ -595,6 +603,8 @@ void UCustomizableObjectNodeSkeletalMesh::BackwardsCompatibleFixup(int32 Customi
 
 	if (CustomizableObjectCustomVersion == FCustomizableObjectCustomVersion::MoveLayoutToNodeSkeletalMesh)
 	{
+		LoadObjects();
+		
 		if (const FSkeletalMeshModel* ImportedModel = SkeletalMesh ? SkeletalMesh->GetImportedModel() : nullptr)
 		{
 			TArray<UEdGraphPin*> PinsToDelete;
diff --git a/Engine/Plugins/Mutable/Source/CustomizableObjectEditor/Private/MuCOE/Nodes/CustomizableObjectNodeSkeletalMesh.h b/Engine/Plugins/Mutable/Source/CustomizableObjectEditor/Private/MuCOE/Nodes/CustomizableObjectNodeSkeletalMesh.h
--- a/Engine/Plugins/Mutable/Source/CustomizableObjectEditor/Private/MuCOE/Nodes/CustomizableObjectNodeSkeletalMesh.h
+++ b/Engine/Plugins/Mutable/Source/CustomizableObjectEditor/Private/MuCOE/Nodes/CustomizableObjectNodeSkeletalMesh.h
@@ -214,7 +214,6 @@ public:
 	FGameplayTagContainer AnimationGameplayTags;
 
 	// UObject interface.
-	virtual void PostLoad() override;
 	virtual void PostEditChangeProperty(FPropertyChangedEvent& PropertyChangedEvent) override;
 	
 	// UEdGraphNode interface
@@ -270,6 +269,8 @@ private:
 	
 	const FSkelMeshSection* GetSkeletalMeshSectionFor(const int32 LODIndex, const int32 SectionIndex) const;
 
+	void LoadObjects();
+	
 	// Deprecated
 	
 	UPROPERTY()
diff --git a/Engine/Plugins/Mutable/Source/CustomizableObjectEditor/Private/MuCOE/Nodes/CustomizableObjectNodeStaticMesh.cpp b/Engine/Plugins/Mutable/Source/CustomizableObjectEditor/Private/MuCOE/Nodes/CustomizableObjectNodeStaticMesh.cpp
--- a/Engine/Plugins/Mutable/Source/CustomizableObjectEditor/Private/MuCOE/Nodes/CustomizableObjectNodeStaticMesh.cpp
+++ b/Engine/Plugins/Mutable/Source/CustomizableObjectEditor/Private/MuCOE/Nodes/CustomizableObjectNodeStaticMesh.cpp
@@ -29,25 +29,6 @@ struct FSlateBrush;
 static const TCHAR* STATIC_MESH_PIN_NAME = TEXT("Static Mesh");
 
 
-void UCustomizableObjectNodeStaticMesh::PostLoad()
-{
-	Super::PostLoad();
-
-	if (StaticMesh)
-	{
-		ConditionalPostLoadReference(*StaticMesh);
-
-		for (const FStaticMaterial& StaticMaterial : StaticMesh->GetStaticMaterials())
-		{
-			if (UMaterialInterface* Material = StaticMaterial.MaterialInterface)
-			{
-				ConditionalPostLoadReference(*Material);
-			}
-		}
-	}
-}
-
-
 void UCustomizableObjectNodeStaticMesh::PostEditChangeProperty(FPropertyChangedEvent& PropertyChangedEvent)
 {
 	FProperty* PropertyThatChanged = PropertyChangedEvent.Property;
diff --git a/Engine/Plugins/Mutable/Source/CustomizableObjectEditor/Private/MuCOE/Nodes/CustomizableObjectNodeStaticMesh.h b/Engine/Plugins/Mutable/Source/CustomizableObjectEditor/Private/MuCOE/Nodes/CustomizableObjectNodeStaticMesh.h
--- a/Engine/Plugins/Mutable/Source/CustomizableObjectEditor/Private/MuCOE/Nodes/CustomizableObjectNodeStaticMesh.h
+++ b/Engine/Plugins/Mutable/Source/CustomizableObjectEditor/Private/MuCOE/Nodes/CustomizableObjectNodeStaticMesh.h
@@ -131,7 +131,6 @@ public:
 	FEdGraphPinReference DefaultPin;
 	
 	// UObject interface.
-	virtual void PostLoad() override;
 	virtual void PostEditChangeProperty(FPropertyChangedEvent& PropertyChangedEvent) override;
 	virtual void Serialize(FArchive& Ar) override;
 
diff --git a/Engine/Plugins/Mutable/Source/CustomizableObjectEditor/Private/MuCOE/Nodes/CustomizableObjectNodeTable.cpp b/Engine/Plugins/Mutable/Source/CustomizableObjectEditor/Private/MuCOE/Nodes/CustomizableObjectNodeTable.cpp
--- a/Engine/Plugins/Mutable/Source/CustomizableObjectEditor/Private/MuCOE/Nodes/CustomizableObjectNodeTable.cpp
+++ b/Engine/Plugins/Mutable/Source/CustomizableObjectEditor/Private/MuCOE/Nodes/CustomizableObjectNodeTable.cpp
@@ -140,7 +140,7 @@ void UCustomizableObjectNodeTable::BackwardsCompatibleFixup(int32 CustomizableOb
 	}
 	
 	// Adding StructColumnID
-	if(CustomizableObjectCustomVersion < FCustomizableObjectCustomVersion::AddedColumnIdDataToTableNodePins)
+	if (CustomizableObjectCustomVersion < FCustomizableObjectCustomVersion::AddedColumnIdDataToTableNodePins)
 	{
 		const UScriptStruct* TableStruct = GetTableNodeStruct();
 
@@ -216,6 +216,16 @@ void UCustomizableObjectNodeTable::BackwardsCompatibleFixup(int32 CustomizableOb
 
 	if (CustomizableObjectCustomVersion == FCustomizableObjectCustomVersion::SetDisplayNamePropertyAsPinNameOfTableNodes)
 	{
+		if (Structure)
+		{
+			ConditionalPostLoadReference(*Structure);
+		}
+
+		if (Table)
+		{
+			ConditionalPostLoadReference(*Table);
+		}
+		
 		// We asume that until now we only supported UserDefinedStructs (Structure assets generated trhough the Editor)
 		const UUserDefinedStruct* UserStruct = Cast<UUserDefinedStruct>(GetTableNodeStruct());
 
@@ -265,59 +275,6 @@ void UCustomizableObjectNodeTable::BackwardsCompatibleFixup(int32 CustomizableOb
 }
 
 
-void UCustomizableObjectNodeTable::PostLoad()
-{
-	Super::PostLoad();
-	
-	if (Table)
-	{
-		ConditionalPostLoadReference(*Table);
-	}
-
-	if (Structure)
-	{
-		ConditionalPostLoadReference(*Structure);
-	}
-	
-	if (const UScriptStruct* TableStruct = GetTableNodeStruct())
-	{
-		TArray<int8> DefaultDataArray;
-		DefaultDataArray.SetNumZeroed(TableStruct->GetStructureSize());
-		TableStruct->InitializeStruct(DefaultDataArray.GetData());
-
-		FCookLoadScope CookLoadScope(ECookLoadType::EditorOnly);
-		for (TFieldIterator<FProperty> It(TableStruct); It; ++It)
-		{
-			FProperty* ColumnProperty = *It;
-			if (!ColumnProperty)
-			{
-				continue;
-			}
-		
-			if (const FSoftObjectProperty* SoftObjectProperty = CastField<FSoftObjectProperty>(ColumnProperty))
-			{
-				if (SoftObjectProperty->PropertyClass)
-				{
-					uint8* CellData = SoftObjectProperty->ContainerPtrToValuePtr<uint8>(DefaultDataArray.GetData());
-					if (!CellData)
-					{
-						continue;
-					}
-
-					UObject* Object = SoftObjectProperty->GetPropertyValue(CellData).LoadSynchronous();
-					if (!Object)
-					{
-						continue;
-					}
-
-					ConditionalPostLoadReference(*Object);
-				}
-			}
-		}
-	}
-}
-
-
 void UCustomizableObjectNodeTable::PreEditChange(FProperty* PropertyAboutToChange)
 {
 	Super::PreEditChange(PropertyAboutToChange);
diff --git a/Engine/Plugins/Mutable/Source/CustomizableObjectEditor/Private/MuCOE/Nodes/CustomizableObjectNodeTable.h b/Engine/Plugins/Mutable/Source/CustomizableObjectEditor/Private/MuCOE/Nodes/CustomizableObjectNodeTable.h
--- a/Engine/Plugins/Mutable/Source/CustomizableObjectEditor/Private/MuCOE/Nodes/CustomizableObjectNodeTable.h
+++ b/Engine/Plugins/Mutable/Source/CustomizableObjectEditor/Private/MuCOE/Nodes/CustomizableObjectNodeTable.h
@@ -253,7 +253,6 @@ public:
 public:
 
 	// UObject interface
-	virtual void PostLoad() override;
 	virtual void PreEditChange(FProperty* PropertyAboutToChange) override;
 	virtual void PostEditChangeProperty(FPropertyChangedEvent& PropertyChangedEvent) override;
 	
@@ -401,8 +400,8 @@ public:
 
 	/** Returns the SkeletalMesh section associated to the given skeletal mesh output pin. */
 	const FSkelMeshSection* GetDefaultSkeletalMeshSectionFor(const UEdGraphPin& MeshPin) const;
-private:
 
+private:
 	/** Number of properties to know when the node needs an update */
 	UPROPERTY()
 	int32 NumProperties;
