Fixed TMemoryTrackingAllocatorWrapper compilation error when UE_ENABLE_ARRAY_SLACK_TRACKING is defined

#rnx

diff --git a/Engine/Plugins/Mutable/Source/MutableRuntime/Public/MuR/MemoryTrackingAllocationPolicy.h b/Engine/Plugins/Mutable/Source/MutableRuntime/Public/MuR/MemoryTrackingAllocationPolicy.h
--- a/Engine/Plugins/Mutable/Source/MutableRuntime/Public/MuR/MemoryTrackingAllocationPolicy.h
+++ b/Engine/Plugins/Mutable/Source/MutableRuntime/Public/MuR/MemoryTrackingAllocationPolicy.h
@@ -127,13 +127,7 @@ namespace mu
 			using BaseAlloc::ForAnyElementType::GetInitialCapacity;
 
 #if UE_ENABLE_ARRAY_SLACK_TRACKING
-			FORCEINLINE void SlackTrackerLogNum(SizeType NewNumUsed)
-			{
-				if constexpr (TAllocatorTraits<BaseAlloc>::SupportsSlackTracking)
-				{
-					BaseAlloc::ForElementType::SlackTrackerLogNum(NewNumUsed);	
-				}
-			}
+			using BaseAlloc::ForAnyElementType::SlackTrackerLogNum;
 #endif
 
 		private:
