[mutable] Move mesh operations that were still using the ASTOpFixed representation to their own ASTOp subclasses.
#rnx

diff --git a/Engine/Plugins/Mutable/Source/CustomizableObjectEditor/Private/MuCOE/SMutableCodeViewer.cpp b/Engine/Plugins/Mutable/Source/CustomizableObjectEditor/Private/MuCOE/SMutableCodeViewer.cpp
--- a/Engine/Plugins/Mutable/Source/CustomizableObjectEditor/Private/MuCOE/SMutableCodeViewer.cpp
+++ b/Engine/Plugins/Mutable/Source/CustomizableObjectEditor/Private/MuCOE/SMutableCodeViewer.cpp
@@ -2420,7 +2420,7 @@ bool SMutableCodeViewer::IsConstantResourceUsedByOperation(const int32 IndexOnCo
 					const mu::OP::MeshClipMorphPlaneArgs Arguments =  InProgram.GetOpArgs<mu::OP::MeshClipMorphPlaneArgs>(OperationAddress);
 					
 					// Morph shape
-					bResourceLocated = IndexOnConstantsArray == Arguments.morphShape;
+					bResourceLocated = IndexOnConstantsArray == Arguments.MorphShape;
 					if (bResourceLocated)
 					{
 						break;
@@ -2429,7 +2429,7 @@ bool SMutableCodeViewer::IsConstantResourceUsedByOperation(const int32 IndexOnCo
 					if (Arguments.VertexSelectionType == EClipVertexSelectionType::Shape)
 					{
 						// Selection Shape
-						bResourceLocated = IndexOnConstantsArray == InProgram.GetOpArgs<mu::OP::MeshClipMorphPlaneArgs>(OperationAddress).vertexSelectionShapeOrBone;
+						bResourceLocated = IndexOnConstantsArray == InProgram.GetOpArgs<mu::OP::MeshClipMorphPlaneArgs>(OperationAddress).VertexSelectionShapeOrBone;
 					}
 				}
 				break;
diff --git a/Engine/Plugins/Mutable/Source/CustomizableObjectEditor/Private/MuCOE/SMutableGraphViewer.cpp b/Engine/Plugins/Mutable/Source/CustomizableObjectEditor/Private/MuCOE/SMutableGraphViewer.cpp
--- a/Engine/Plugins/Mutable/Source/CustomizableObjectEditor/Private/MuCOE/SMutableGraphViewer.cpp
+++ b/Engine/Plugins/Mutable/Source/CustomizableObjectEditor/Private/MuCOE/SMutableGraphViewer.cpp
@@ -108,7 +108,6 @@ static const char* MutableNodeNames[] =
 
 	"Mesh",
 	"MeshConstant",
-	"MeshInterpolate",
 	"MeshTable",
 	"MeshFormat",
 	"MeshTangents",
@@ -121,7 +120,6 @@ static const char* MutableNodeNames[] =
 	"MeshClipWithMesh",
 	"MeshApplyPose",
 	"MeshVariation",
-	"MeshGeometryOperation",
 	"MeshReshape",
 	"MeshClipDeform",
 
diff --git a/Engine/Plugins/Mutable/Source/MutableRuntime/Internal/MuR/Operations.h b/Engine/Plugins/Mutable/Source/MutableRuntime/Internal/MuR/Operations.h
--- a/Engine/Plugins/Mutable/Source/MutableRuntime/Internal/MuR/Operations.h
+++ b/Engine/Plugins/Mutable/Source/MutableRuntime/Internal/MuR/Operations.h
@@ -837,26 +837,18 @@ namespace mu
         //-------------------------------------------------------------------------------------
         struct MeshApplyLayoutArgs
         {
-            ADDRESS mesh;
-            ADDRESS layout;
-            uint16 channel;
+            ADDRESS Mesh;
+            ADDRESS Layout;
+            uint16 Channel;
         };
 
         struct MeshMergeArgs
         {
-            ADDRESS base;
-            ADDRESS added;
-
-            // If 0, it merges the surfaces, otherwise, add a new surface
-            // for the added mesh.
-            uint32 newSurfaceID;
-        };
+            ADDRESS Base;
+            ADDRESS Added;
 
-        struct MeshInterpolateArgs
-        {
-            ADDRESS factor;
-            ADDRESS base;
-            ADDRESS targets[ MUTABLE_OP_MAX_INTERPOLATE_COUNT-1 ];
+            // If 0, it merges the surfaces, otherwise, add a new surface for the added mesh.
+            uint32 NewSurfaceID;
         };
 
 		struct MeshMaskClipMeshArgs
@@ -876,8 +868,8 @@ namespace mu
 
         struct MeshMaskDiffArgs
         {
-            ADDRESS source;
-            ADDRESS fragment;
+            ADDRESS Source;
+            ADDRESS Fragment;
         };
 
         struct MeshFormatArgs
@@ -914,12 +906,11 @@ namespace mu
 
 		struct MeshClipMorphPlaneArgs
 		{
-			ADDRESS source;
-
-			ADDRESS morphShape;
-			ADDRESS vertexSelectionShapeOrBone;
+			ADDRESS Source;
+			ADDRESS MorphShape;
+			ADDRESS VertexSelectionShapeOrBone;
 
-			float dist, factor, maxBoneRadius;
+			float Dist, Factor, MaxBoneRadius;
 
 			EClipVertexSelectionType VertexSelectionType;
 
@@ -928,20 +919,20 @@ namespace mu
 
         struct MeshClipWithMeshArgs
         {
-            ADDRESS source;
-            ADDRESS clipMesh;
+            ADDRESS Source;
+            ADDRESS ClipMesh;
         };
 
         struct MeshSetSkeletonArgs
         {
-            ADDRESS source;
-            ADDRESS skeleton;
+            ADDRESS Source;
+            ADDRESS Skeleton;
         };
 
         struct MeshProjectArgs
         {
-            ADDRESS mesh;
-            ADDRESS projector;
+            ADDRESS Mesh;
+            ADDRESS Projector;
         };
 
         struct MeshApplyPoseArgs
@@ -1110,7 +1101,6 @@ namespace mu
             ScalarConstantArgs ScalarConstant;
             ColourConstantArgs ColourConstant;
         	MatrixConstantArgs TransformConstant;
-            //ParameterArgs Parameter;
 
             //-------------------------------------------------------------------------------------
             BoolLessArgs BoolLess;
@@ -1147,16 +1137,6 @@ namespace mu
             ImageDisplaceArgs ImageDisplace;
 			ImageInvertArgs ImageInvert;
 
-            //-------------------------------------------------------------------------------------
-            MeshApplyLayoutArgs MeshApplyLayout;
-            MeshMergeArgs MeshMerge;
-            MeshInterpolateArgs MeshInterpolate;
-            MeshMaskDiffArgs MeshMaskDiff;
-			MeshClipMorphPlaneArgs MeshClipMorphPlane;
-            MeshClipWithMeshArgs MeshClipWithMesh;
-            MeshSetSkeletonArgs MeshSetSkeleton;
-			MeshProjectArgs MeshProject;
-
         } FArgs;
 
         FArgs args;
diff --git a/Engine/Plugins/Mutable/Source/MutableRuntime/Private/MuR/CodeRunner.cpp b/Engine/Plugins/Mutable/Source/MutableRuntime/Private/MuR/CodeRunner.cpp
--- a/Engine/Plugins/Mutable/Source/MutableRuntime/Private/MuR/CodeRunner.cpp
+++ b/Engine/Plugins/Mutable/Source/MutableRuntime/Private/MuR/CodeRunner.cpp
@@ -1079,22 +1079,22 @@ namespace mu
 			case 0:
 			{
 				AddOp(FScheduledOp(Item.At, Item, 1),
-					FScheduledOp(Args.mesh, Item),
-					FScheduledOp(Args.layout, Item));
+					FScheduledOp(Args.Mesh, Item),
+					FScheduledOp(Args.Layout, Item));
 				break;
 			}
             case 1:
             {
            		MUTABLE_CPUPROFILER_SCOPE(ME_APPLYLAYOUT)
             		
-                TSharedPtr<const FMesh> Base = LoadMesh(FCacheAddress(Args.mesh, Item));
+                TSharedPtr<const FMesh> Base = LoadMesh(FCacheAddress(Args.Mesh, Item));
 
                 if (Base)
                 {
 					TSharedPtr<FMesh> Result = CloneOrTakeOver(Base);
 
-                    TSharedPtr<const FLayout> Layout = LoadLayout(FCacheAddress(Args.layout, Item));
-                    int32 texCoordsSet = Args.channel;
+                    TSharedPtr<const FLayout> Layout = LoadLayout(FCacheAddress(Args.Layout, Item));
+                    int32 texCoordsSet = Args.Channel;
 
                     MeshApplyLayout(Result.Get(), Layout.Get(), texCoordsSet);
 					
@@ -1315,16 +1315,16 @@ namespace mu
             case 0:
 			{
 				AddOp(FScheduledOp(Item.At, Item, 1),
-					FScheduledOp(Args.base, Item),
-					FScheduledOp(Args.added, Item));
+					FScheduledOp(Args.Base, Item),
+					FScheduledOp(Args.Added, Item));
 				break;
 			}
             case 1:
             {
 				MUTABLE_CPUPROFILER_SCOPE(ME_MERGE_1)
 
-                TSharedPtr<const FMesh> pA = LoadMesh(FCacheAddress(Args.base, Item));
-                TSharedPtr<const FMesh> pB = LoadMesh(FCacheAddress(Args.added, Item));
+                TSharedPtr<const FMesh> pA = LoadMesh(FCacheAddress(Args.Base, Item));
+                TSharedPtr<const FMesh> pB = LoadMesh(FCacheAddress(Args.Added, Item));
 
                 if (pA && pB && pA->GetVertexCount() && pB->GetVertexCount())
                 {
@@ -1336,15 +1336,15 @@ namespace mu
 
 					TSharedPtr<FMesh> Result = CreateMesh(pA->GetDataSize() + pB->GetDataSize());
 
-					MeshMerge(Result.Get(), pA.Get(), pB.Get(), !Args.newSurfaceID, Scratch);
+					MeshMerge(Result.Get(), pA.Get(), pB.Get(), !Args.NewSurfaceID, Scratch);
 
 					Release(Scratch.FirstReformat);
 					Release(Scratch.SecondReformat);
 
-                    if (Args.newSurfaceID)
+                    if (Args.NewSurfaceID)
                     {
 						check(pB->GetSurfaceCount() == 1);
-						Result->Surfaces.Last().Id = Args.newSurfaceID;
+						Result->Surfaces.Last().Id = Args.NewSurfaceID;
                     }
 
 					Release(pA);
@@ -1362,9 +1362,9 @@ namespace mu
 
                     check(Result->IsReference() || (Result->GetSurfaceCount() == 1) );
 
-                    if (Result->GetSurfaceCount() > 0 && Args.newSurfaceID)
+                    if (Result->GetSurfaceCount() > 0 && Args.NewSurfaceID)
                     {
-                        Result->Surfaces.Last().Id = Args.newSurfaceID;
+                        Result->Surfaces.Last().Id = Args.NewSurfaceID;
                     }
 
 					Release(pA);
@@ -1533,16 +1533,16 @@ namespace mu
             case 0:
 			{
 				AddOp(FScheduledOp(Item.At, Item, 1),
-					FScheduledOp(Args.source, Item),
-					FScheduledOp(Args.fragment, Item));
+					FScheduledOp(Args.Source, Item),
+					FScheduledOp(Args.Fragment, Item));
 				break;
 			}
             case 1:
             {
            		MUTABLE_CPUPROFILER_SCOPE(ME_MASKDIFF_1)
             		
-                TSharedPtr<const FMesh> Source = LoadMesh(FCacheAddress(Args.source, Item));
-                TSharedPtr<const FMesh> pClip = LoadMesh(FCacheAddress(Args.fragment, Item));
+                TSharedPtr<const FMesh> Source = LoadMesh(FCacheAddress(Args.Source, Item));
+                TSharedPtr<const FMesh> pClip = LoadMesh(FCacheAddress(Args.Fragment, Item));
 
                 // Only if both are valid.
                 if (Source.Get() && pClip.Get())
@@ -1815,10 +1815,10 @@ namespace mu
             {
             case 0:
 			{
-				if (Args.source)
+				if (Args.Source)
 				{
 					AddOp(FScheduledOp(Item.At, Item, 1),
-						FScheduledOp(Args.source, Item));
+						FScheduledOp(Args.Source, Item));
 				}
 				else
 				{
@@ -1830,29 +1830,29 @@ namespace mu
             {
            		MUTABLE_CPUPROFILER_SCOPE(ME_CLIPMORPHPLANE_1)
             		
-                TSharedPtr<const FMesh> Source = LoadMesh(FCacheAddress(Args.source, Item));
+                TSharedPtr<const FMesh> Source = LoadMesh(FCacheAddress(Args.Source, Item));
 
-                check(Args.morphShape < (uint32)InModel->GetPrivate()->Program.ConstantShapes.Num());
+                check(Args.MorphShape < (uint32)InModel->GetPrivate()->Program.ConstantShapes.Num());
 
                 // Should be an ellipse
-                const FShape& morphShape = Program.ConstantShapes[Args.morphShape];
+                const FShape& MorphShape = Program.ConstantShapes[Args.MorphShape];
 
-                const FVector3f& origin = morphShape.position;
-                const FVector3f& normal = morphShape.up;
+                const FVector3f& Origin = MorphShape.position;
+                const FVector3f& Normal = MorphShape.up;
 
 				bool bRemoveFaceIfAllVerticesCulled = Args.FaceCullStrategy==EFaceCullStrategy::AllVerticesCulled;
 
                 if (Args.VertexSelectionType == EClipVertexSelectionType::Shape)
                 {
-                    check(Args.vertexSelectionShapeOrBone < (uint32)InModel->GetPrivate()->Program.ConstantShapes.Num());
+                    check(Args.VertexSelectionShapeOrBone < (uint32)InModel->GetPrivate()->Program.ConstantShapes.Num());
 
                     // Should be None or an axis aligned box
-                    const FShape& selectionShape = Program.ConstantShapes[Args.vertexSelectionShapeOrBone];
+                    const FShape& SelectionShape = Program.ConstantShapes[Args.VertexSelectionShapeOrBone];
 
 					TSharedPtr<FMesh> Result = CreateMesh(Source ? Source->GetDataSize() : 0);
 
 					bool bOutSuccess = false;
-					MeshClipMorphPlane(Result.Get(), Source.Get(), origin, normal, Args.dist, Args.factor, morphShape.size[0], morphShape.size[1], morphShape.size[2], selectionShape, bRemoveFaceIfAllVerticesCulled, bOutSuccess, nullptr, -1);
+					MeshClipMorphPlane(Result.Get(), Source.Get(), Origin, Normal, Args.Dist, Args.Factor, MorphShape.size[0], MorphShape.size[1], MorphShape.size[2], SelectionShape, bRemoveFaceIfAllVerticesCulled, bOutSuccess, nullptr, -1);
 					
 					if (!bOutSuccess)
 					{
@@ -1868,16 +1868,16 @@ namespace mu
 
 				else if (Args.VertexSelectionType == EClipVertexSelectionType::BoneHierarchy)
 				{
-					FShape selectionShape;
-					selectionShape.type = (uint8)FShape::Type::None;
+					FShape SelectionShape;
+					SelectionShape.type = (uint8)FShape::Type::None;
 
 					TSharedPtr<FMesh> Result = CreateMesh(Source->GetDataSize());
 
-					check(Args.vertexSelectionShapeOrBone <= MAX_uint32);
-					const FBoneName Bone(Args.vertexSelectionShapeOrBone);
+					check(Args.VertexSelectionShapeOrBone <= MAX_uint32);
+					const FBoneName Bone(Args.VertexSelectionShapeOrBone);
 
 					bool bOutSuccess = false;
-					MeshClipMorphPlane(Result.Get(), Source.Get(), origin, normal, Args.dist, Args.factor, morphShape.size[0], morphShape.size[1], morphShape.size[2], selectionShape, bRemoveFaceIfAllVerticesCulled, bOutSuccess, &Bone, Args.maxBoneRadius);
+					MeshClipMorphPlane(Result.Get(), Source.Get(), Origin, Normal, Args.Dist, Args.Factor, MorphShape.size[0], MorphShape.size[1], MorphShape.size[2], SelectionShape, bRemoveFaceIfAllVerticesCulled, bOutSuccess, &Bone, Args.MaxBoneRadius);
 
 					if (!bOutSuccess)
 					{
@@ -1893,13 +1893,13 @@ namespace mu
                 else
                 {
                     // No vertex selection
-                    FShape selectionShape;
-                    selectionShape.type = (uint8)FShape::Type::None;
+                    FShape SelectionShape;
+                    SelectionShape.type = (uint8)FShape::Type::None;
 
 					TSharedPtr<FMesh> Result = CreateMesh(Source ? Source->GetDataSize() : 0);
 
 					bool bOutSuccess = false;
-					MeshClipMorphPlane(Result.Get(), Source.Get(), origin, normal, Args.dist, Args.factor, morphShape.size[0], morphShape.size[1], morphShape.size[2], selectionShape, bRemoveFaceIfAllVerticesCulled, bOutSuccess, nullptr, -1.0f);
+					MeshClipMorphPlane(Result.Get(), Source.Get(), Origin, Normal, Args.Dist, Args.Factor, MorphShape.size[0], MorphShape.size[1], MorphShape.size[2], SelectionShape, bRemoveFaceIfAllVerticesCulled, bOutSuccess, nullptr, -1.0f);
 
 					if (!bOutSuccess)
 					{
@@ -1931,11 +1931,11 @@ namespace mu
             {
             case 0:
 			{
-				if (Args.source)
+				if (Args.Source)
 				{
 					AddOp(FScheduledOp(Item.At, Item, 1),
-						FScheduledOp(Args.source, Item),
-						FScheduledOp(Args.clipMesh, Item));
+						FScheduledOp(Args.Source, Item),
+						FScheduledOp(Args.ClipMesh, Item));
 				}
 				else
 				{
@@ -1948,8 +1948,8 @@ namespace mu
             {
 				MUTABLE_CPUPROFILER_SCOPE(ME_CLIPWITHMESH_1)
 
-                TSharedPtr<const FMesh> Source = LoadMesh(FCacheAddress(Args.source, Item));
-                TSharedPtr<const FMesh> Clip = LoadMesh(FCacheAddress(Args.clipMesh, Item));
+                TSharedPtr<const FMesh> Source = LoadMesh(FCacheAddress(Args.Source, Item));
+                TSharedPtr<const FMesh> Clip = LoadMesh(FCacheAddress(Args.ClipMesh, Item));
 
                 // Only if both are valid.
                 if (Source && Clip)
@@ -2378,11 +2378,11 @@ namespace mu
             {
             case 0:
 			{
-				if (Args.source)
+				if (Args.Source)
 				{
 					AddOp(FScheduledOp(Item.At, Item, 1),
-						FScheduledOp(Args.source, Item),
-						FScheduledOp(Args.skeleton, Item));
+						FScheduledOp(Args.Source, Item),
+						FScheduledOp(Args.Skeleton, Item));
 				}
 				else
 				{
@@ -2394,8 +2394,8 @@ namespace mu
             {
             	MUTABLE_CPUPROFILER_SCOPE(ME_SETSKELETON_1)
             		
-                TSharedPtr<const FMesh> Source = LoadMesh(FCacheAddress(Args.source, Item));
-                TSharedPtr<const FMesh> Skeleton = LoadMesh(FCacheAddress(Args.skeleton, Item));
+                TSharedPtr<const FMesh> Source = LoadMesh(FCacheAddress(Args.Source, Item));
+                TSharedPtr<const FMesh> Skeleton = LoadMesh(FCacheAddress(Args.Skeleton, Item));
 
                 // Only if both are valid.
                 if (Source && Skeleton)
@@ -2689,11 +2689,11 @@ namespace mu
             {
             case 0:
 			{
-				if (Args.mesh)
+				if (Args.Mesh)
 				{
 					AddOp(FScheduledOp(Item.At, Item, 1),
-						FScheduledOp(Args.mesh, Item),
-						FScheduledOp(Args.projector, Item));
+						FScheduledOp(Args.Mesh, Item),
+						FScheduledOp(Args.Projector, Item));
 				}
 				else
 				{
@@ -2705,8 +2705,8 @@ namespace mu
             {
 				MUTABLE_CPUPROFILER_SCOPE(ME_PROJECT_1)
 
-                TSharedPtr<const FMesh> Mesh = LoadMesh(FCacheAddress(Args.mesh,Item));
-                const FProjector Projector = LoadProjector(FCacheAddress(Args.projector, Item));
+                TSharedPtr<const FMesh> Mesh = LoadMesh(FCacheAddress(Args.Mesh,Item));
+                const FProjector Projector = LoadProjector(FCacheAddress(Args.Projector, Item));
 
                 // Only if both are valid.
                 if (Mesh && Mesh->GetVertexBuffers().GetBufferCount() > 0)
diff --git a/Engine/Plugins/Mutable/Source/MutableRuntime/Private/MuR/Operations.cpp b/Engine/Plugins/Mutable/Source/MutableRuntime/Private/MuR/Operations.cpp
--- a/Engine/Plugins/Mutable/Source/MutableRuntime/Private/MuR/Operations.cpp
+++ b/Engine/Plugins/Mutable/Source/MutableRuntime/Private/MuR/Operations.cpp
@@ -301,41 +301,6 @@ namespace mu
 			f(&op.args.ImageInvert.base);
 			break;
 
-        //-------------------------------------------------------------------------------------
-        case EOpType::ME_APPLYLAYOUT:
-            f(&op.args.MeshApplyLayout.layout );
-            f(&op.args.MeshApplyLayout.mesh );
-            break;
-
-        case EOpType::ME_MERGE:
-            f(&op.args.MeshMerge.base );
-            f(&op.args.MeshMerge.added );
-            break;
-
-        case EOpType::ME_MASKDIFF:
-            f(&op.args.MeshMaskDiff.source );
-            f(&op.args.MeshMaskDiff.fragment );
-            break;
-
-        case EOpType::ME_CLIPMORPHPLANE:
-            f(&op.args.MeshClipMorphPlane.source);
-            break;
-
-        case EOpType::ME_CLIPWITHMESH :
-            f(&op.args.MeshClipWithMesh.source);
-            f(&op.args.MeshClipWithMesh.clipMesh);
-            break;
-
-        case EOpType::ME_SETSKELETON :
-            f(&op.args.MeshSetSkeleton.source);
-            f(&op.args.MeshSetSkeleton.skeleton);
-            break;
-
-		case EOpType::ME_PROJECT:
-			f(&op.args.MeshProject.mesh);
-			f(&op.args.MeshProject.projector);
-			break;
-
         default:
 			check( false );
             break;
@@ -767,8 +732,8 @@ namespace mu
         case EOpType::ME_APPLYLAYOUT:
         {
 			OP::MeshApplyLayoutArgs args = program.GetOpArgs<OP::MeshApplyLayoutArgs>(at);
-            f(args.layout );
-            f(args.mesh );
+            f(args.Layout );
+            f(args.Mesh );
             break;
         }
 
@@ -807,8 +772,8 @@ namespace mu
         case EOpType::ME_MERGE:
         {
 			OP::MeshMergeArgs args = program.GetOpArgs<OP::MeshMergeArgs>(at);
-            f(args.base );
-            f(args.added );
+            f(args.Base );
+            f(args.Added );
             break;
         }
 
@@ -832,8 +797,8 @@ namespace mu
         case EOpType::ME_MASKDIFF:
         {
 			OP::MeshMaskDiffArgs args = program.GetOpArgs<OP::MeshMaskDiffArgs>(at);
-            f(args.source );
-            f(args.fragment );
+            f(args.Source );
+            f(args.Fragment );
             break;
         }
 
@@ -899,15 +864,15 @@ namespace mu
         case EOpType::ME_CLIPMORPHPLANE:
         {
 			OP::MeshClipMorphPlaneArgs args = program.GetOpArgs<OP::MeshClipMorphPlaneArgs>(at);
-            f(args.source);
+            f(args.Source);
             break;
         }
 
         case EOpType::ME_CLIPWITHMESH :
         {
-			OP::MeshClipWithMeshArgs args = program.GetOpArgs<OP::MeshClipWithMeshArgs>(at);
-            f(args.source);
-            f(args.clipMesh);
+			OP::MeshClipWithMeshArgs Args = program.GetOpArgs<OP::MeshClipWithMeshArgs>(at);
+            f(Args.Source);
+            f(Args.ClipMesh);
             break;
         }
 
@@ -930,25 +895,25 @@ namespace mu
  
         case EOpType::ME_SETSKELETON :
         {
-			OP::MeshSetSkeletonArgs args = program.GetOpArgs<OP::MeshSetSkeletonArgs>(at);
-            f(args.source);
-            f(args.skeleton);
+			OP::MeshSetSkeletonArgs Args = program.GetOpArgs<OP::MeshSetSkeletonArgs>(at);
+            f(Args.Source);
+            f(Args.Skeleton);
             break;
         }
 
         case EOpType::ME_PROJECT :
         {
-			OP::MeshProjectArgs args = program.GetOpArgs<OP::MeshProjectArgs>(at);
-            f(args.mesh);
-            f(args.projector);
+			OP::MeshProjectArgs Args = program.GetOpArgs<OP::MeshProjectArgs>(at);
+            f(Args.Mesh);
+            f(Args.Projector);
             break;
         }
 
 		case EOpType::ME_APPLYPOSE:
 		{
-			OP::MeshApplyPoseArgs args = program.GetOpArgs<OP::MeshApplyPoseArgs>(at);
-			f(args.base);
-			f(args.pose);
+			OP::MeshApplyPoseArgs Args = program.GetOpArgs<OP::MeshApplyPoseArgs>(at);
+			f(Args.base);
+			f(Args.pose);
 			break;
 		}
 
diff --git a/Engine/Plugins/Mutable/Source/MutableTools/Private/MuT/AST.cpp b/Engine/Plugins/Mutable/Source/MutableTools/Private/MuT/AST.cpp
--- a/Engine/Plugins/Mutable/Source/MutableTools/Private/MuT/AST.cpp
+++ b/Engine/Plugins/Mutable/Source/MutableTools/Private/MuT/AST.cpp
@@ -1019,15 +1019,7 @@ void ASTOpFixed::Link( FProgram& program, FLinkerOptions* )
         program.OpAddress.Add((uint32)program.ByteCode.Num());
 
         AppendCode(program.ByteCode,lop.type);
-        if (lop.type==EOpType::ME_MERGE)
-        {
-            AppendCode(program.ByteCode,lop.args.MeshMerge);
-        }
-        else
-        {
-            // Generic encoding
-            AppendCode(program.ByteCode,lop.args);
-        }
+        AppendCode(program.ByteCode,lop.args);
     }
 }
 
@@ -1711,27 +1703,7 @@ FSourceDataDescriptor ASTOpFixed::GetSourceDataDescriptor(FGetSourceDataDescript
 	case EOpType::IM_COLOURMAP:
 		Result = GetSourceDataDescriptor(op.args.ImageColourMap.base, Context);
 		break;
-	case EOpType::ME_APPLYLAYOUT:
-		Result = GetSourceDataDescriptor(op.args.MeshApplyLayout.mesh, Context);
-		break;	
-	case EOpType::ME_MERGE:
-		Result = GetSourceDataDescriptor(op.args.MeshMerge.base, Context);
-		break;	
-	case EOpType::ME_MASKDIFF:
-		Result = GetSourceDataDescriptor(op.args.MeshMaskDiff.source, Context);
-		break;	
-	case EOpType::ME_CLIPMORPHPLANE:
-		Result = GetSourceDataDescriptor(op.args.MeshClipMorphPlane.source, Context);
-		break;	
-	case EOpType::ME_CLIPWITHMESH:
-		Result = GetSourceDataDescriptor(op.args.MeshClipWithMesh.source, Context);
-		break;	
-	case EOpType::ME_SETSKELETON:
-		Result = GetSourceDataDescriptor(op.args.MeshSetSkeleton.source, Context);
-		break;	
-	case EOpType::ME_PROJECT:
-		Result = GetSourceDataDescriptor(op.args.MeshProject.mesh, Context);
-		break;
+
 	default:
 		// This should probably be implemented
 		ensure(false);
diff --git a/Engine/Plugins/Mutable/Source/MutableTools/Private/MuT/ASTOpImageRasterMesh.cpp b/Engine/Plugins/Mutable/Source/MutableTools/Private/MuT/ASTOpImageRasterMesh.cpp
--- a/Engine/Plugins/Mutable/Source/MutableTools/Private/MuT/ASTOpImageRasterMesh.cpp
+++ b/Engine/Plugins/Mutable/Source/MutableTools/Private/MuT/ASTOpImageRasterMesh.cpp
@@ -8,6 +8,7 @@
 #include "MuT/ASTOpConditional.h"
 #include "MuT/ASTOpSwitch.h"
 #include "MuT/ASTOpMeshAddTags.h"
+#include "MuT/ASTOpMeshProject.h"
 #include "MuR/ModelPrivate.h"
 #include "MuR/RefCounted.h"
 #include "MuR/Types.h"
@@ -252,10 +253,10 @@ namespace mu
 			if (!imageAt)
 			{
 				// We remove the project from the raster children
-				const ASTOpFixed* MeshProjectOp = static_cast<const ASTOpFixed*>(sourceAt.get());
-				Ptr<ASTOpImageRasterMesh> nop = mu::Clone<ASTOpImageRasterMesh>(this);
-				nop->mesh = MeshProjectOp->children[MeshProjectOp->op.args.MeshProject.mesh].child();
-				at = nop;
+				const ASTOpMeshProject* MeshProjectOp = static_cast<const ASTOpMeshProject*>(sourceAt.get());
+				Ptr<ASTOpImageRasterMesh> NewOp = mu::Clone<ASTOpImageRasterMesh>(this);
+				NewOp->mesh = MeshProjectOp->Mesh.child();
+				at = NewOp;
 			}
 			break;
 		}
diff --git a/Engine/Plugins/Mutable/Source/MutableTools/Private/MuT/ASTOpLayoutFromMesh.cpp b/Engine/Plugins/Mutable/Source/MutableTools/Private/MuT/ASTOpLayoutFromMesh.cpp
--- a/Engine/Plugins/Mutable/Source/MutableTools/Private/MuT/ASTOpLayoutFromMesh.cpp
+++ b/Engine/Plugins/Mutable/Source/MutableTools/Private/MuT/ASTOpLayoutFromMesh.cpp
@@ -20,6 +20,7 @@
 #include "MuT/ASTOpMeshApplyShape.h"
 #include "MuT/ASTOpMeshBindShape.h"
 #include "MuT/ASTOpMeshAddTags.h"
+#include "MuT/ASTOpMeshMerge.h"
 #include "MuT/ASTOpSwitch.h"
 #include "MuT/ASTOpLayoutMerge.h"
 
@@ -219,11 +220,11 @@ namespace mu
 
 				case EOpType::ME_MERGE:
 				{
-					const ASTOpFixed* Typed = static_cast<const ASTOpFixed*>(at.get());
+					const ASTOpMeshMerge* Typed = static_cast<const ASTOpMeshMerge*>(at.get());
 
 					Ptr<ASTOpLayoutMerge> NewMerge = new ASTOpLayoutMerge;
-					NewMerge->Base = Visit(Typed->children[Typed->op.args.MeshMerge.base].child(), CurrentSinkingOp);
-					NewMerge->Added = Visit(Typed->children[Typed->op.args.MeshMerge.added].child(), CurrentSinkingOp);
+					NewMerge->Base = Visit(Typed->Base.child(), CurrentSinkingOp);
+					NewMerge->Added = Visit(Typed->Added.child(), CurrentSinkingOp);
 					newAt = NewMerge;
 					break;
 				}
diff --git a/Engine/Plugins/Mutable/Source/MutableTools/Private/MuT/ASTOpMeshApplyLayout.cpp b/Engine/Plugins/Mutable/Source/MutableTools/Private/MuT/ASTOpMeshApplyLayout.cpp
new file mode 100644
--- /dev/null
+++ b/Engine/Plugins/Mutable/Source/MutableTools/Private/MuT/ASTOpMeshApplyLayout.cpp
@@ -0,0 +1,133 @@
+// Copyright Epic Games, Inc. All Rights Reserved.
+
+#include "MuT/ASTOpMeshApplyLayout.h"
+#include "MuT/ASTOpMeshAddTags.h"
+
+#include "HAL/PlatformMath.h"
+#include "MuR/ModelPrivate.h"
+#include "MuR/RefCounted.h"
+#include "MuR/Types.h"
+
+namespace mu
+{
+
+	ASTOpMeshApplyLayout::ASTOpMeshApplyLayout()
+		: Mesh(this)
+		, Layout(this)
+	{
+	}
+
+
+	ASTOpMeshApplyLayout::~ASTOpMeshApplyLayout()
+	{
+		// Explicit call needed to avoid recursive destruction
+		ASTOp::RemoveChildren();
+	}
+
+
+	bool ASTOpMeshApplyLayout::IsEqual(const ASTOp& OtherUntyped) const
+	{
+		if (OtherUntyped.GetOpType()==GetOpType())
+		{
+			const ASTOpMeshApplyLayout* Other = static_cast<const ASTOpMeshApplyLayout*>(&OtherUntyped);
+			return Mesh == Other->Mesh &&
+				Layout == Other->Layout &&
+				Channel == Other->Channel;
+		}
+		return false;
+	}
+
+
+	uint64 ASTOpMeshApplyLayout::Hash() const
+	{
+		uint64 Result = std::hash<void*>()(Mesh.child().get());
+		hash_combine(Result, Layout.child().get());
+		hash_combine(Result, Channel);
+		return Result;
+	}
+
+
+	Ptr<ASTOp> ASTOpMeshApplyLayout::Clone(MapChildFuncRef MapChild) const
+	{
+		Ptr<ASTOpMeshApplyLayout> New = new ASTOpMeshApplyLayout();
+		New->Mesh = MapChild(Mesh.child());
+		New->Layout = MapChild(Layout.child());
+		New->Channel = Channel;
+		return New;
+	}
+
+
+	void ASTOpMeshApplyLayout::ForEachChild(const TFunctionRef<void(ASTChild&)> Func)
+	{
+		Func(Mesh);
+		Func(Layout);
+	}
+
+
+	void ASTOpMeshApplyLayout::Link(FProgram& Program, FLinkerOptions*)
+	{
+		// Already linked?
+		if (!linkedAddress)
+		{
+			OP::MeshApplyLayoutArgs Args = {};
+
+			if (Mesh) Args.Mesh = Mesh->linkedAddress;
+			if (Layout) Args.Layout = Layout->linkedAddress;
+			Args.Channel = Channel;
+
+			linkedAddress = (OP::ADDRESS)Program.OpAddress.Num();
+			Program.OpAddress.Add(Program.ByteCode.Num());
+			AppendCode(Program.ByteCode, GetOpType());
+			AppendCode(Program.ByteCode, Args);
+		}
+
+	}
+
+
+	Ptr<ASTOp> ASTOpMeshApplyLayout::OptimiseSink(const FModelOptimizationOptions&, FOptimizeSinkContext&) const
+	{
+		Ptr<ASTOp> NewOp;
+
+		Ptr<ASTOp> MeshAt = Mesh.child();
+
+		if (!MeshAt)
+		{
+			return nullptr;
+		}
+
+		EOpType MeshOpType = MeshAt->GetOpType();
+
+		switch (MeshOpType)
+		{
+			case EOpType::ME_ADDTAGS:
+			{
+				Ptr<ASTOpMeshAddTags> New = mu::Clone<ASTOpMeshAddTags>(MeshAt);
+				if (New->Source)
+				{
+					Ptr<ASTOpMeshApplyLayout> NewApplyPose = mu::Clone<ASTOpMeshApplyLayout>(this);
+					NewApplyPose->Mesh = New->Source.child();
+					New->Source = NewApplyPose;
+				}
+
+				NewOp = New;
+				break;
+			}
+
+			default:
+				break;
+		}
+
+		return NewOp;
+	}
+
+
+	FSourceDataDescriptor ASTOpMeshApplyLayout::GetSourceDataDescriptor(FGetSourceDataDescriptorContext* Context) const
+	{
+		if (Mesh)
+		{
+			return Mesh->GetSourceDataDescriptor(Context);
+		}
+
+		return {};
+	}
+}
diff --git a/Engine/Plugins/Mutable/Source/MutableTools/Private/MuT/ASTOpMeshApplyLayout.h b/Engine/Plugins/Mutable/Source/MutableTools/Private/MuT/ASTOpMeshApplyLayout.h
new file mode 100644
--- /dev/null
+++ b/Engine/Plugins/Mutable/Source/MutableTools/Private/MuT/ASTOpMeshApplyLayout.h
@@ -0,0 +1,43 @@
+// Copyright Epic Games, Inc. All Rights Reserved.
+
+#pragma once
+
+#include "HAL/Platform.h"
+#include "MuR/Operations.h"
+#include "MuR/Ptr.h"
+#include "MuT/AST.h"
+
+
+namespace mu
+{
+	struct FProgram;
+
+	/** */
+	class ASTOpMeshApplyLayout final : public ASTOp
+	{
+	public:
+
+		ASTChild Mesh;
+		ASTChild Layout;
+		uint16 Channel=0;
+
+	public:
+
+		ASTOpMeshApplyLayout();
+		ASTOpMeshApplyLayout(const ASTOpMeshApplyLayout&) = delete;
+		~ASTOpMeshApplyLayout();
+
+		// ASTOp interface
+		virtual EOpType GetOpType() const override { return EOpType::ME_APPLYLAYOUT; }
+		virtual uint64 Hash() const override;
+		virtual bool IsEqual(const ASTOp&) const override;
+		virtual Ptr<ASTOp> Clone(MapChildFuncRef) const override;
+		virtual void ForEachChild(const TFunctionRef<void(ASTChild&)>) override;
+		virtual void Link(FProgram&, FLinkerOptions*) override;
+		virtual Ptr<ASTOp> OptimiseSink(const FModelOptimizationOptions&, FOptimizeSinkContext&) const override;
+		virtual FSourceDataDescriptor GetSourceDataDescriptor(FGetSourceDataDescriptorContext*) const override;
+	};
+
+
+}
+
diff --git a/Engine/Plugins/Mutable/Source/MutableTools/Private/MuT/ASTOpMeshApplyPose.cpp b/Engine/Plugins/Mutable/Source/MutableTools/Private/MuT/ASTOpMeshApplyPose.cpp
--- a/Engine/Plugins/Mutable/Source/MutableTools/Private/MuT/ASTOpMeshApplyPose.cpp
+++ b/Engine/Plugins/Mutable/Source/MutableTools/Private/MuT/ASTOpMeshApplyPose.cpp
@@ -11,11 +11,9 @@
 namespace mu
 {
 
-
-	//-------------------------------------------------------------------------------------------------
 	ASTOpMeshApplyPose::ASTOpMeshApplyPose()
-		: base(this)
-		, pose(this)
+		: Base(this)
+		, Pose(this)
 	{
 	}
 
@@ -27,13 +25,13 @@ namespace mu
 	}
 
 
-	bool ASTOpMeshApplyPose::IsEqual(const ASTOp& otherUntyped) const
+	bool ASTOpMeshApplyPose::IsEqual(const ASTOp& OtherUntyped) const
 	{
-		if (otherUntyped.GetOpType()==GetOpType())
+		if (OtherUntyped.GetOpType()==GetOpType())
 		{
-			const ASTOpMeshApplyPose* other = static_cast<const ASTOpMeshApplyPose*>(&otherUntyped);
-			return base == other->base &&
-				pose == other->pose;
+			const ASTOpMeshApplyPose* Other = static_cast<const ASTOpMeshApplyPose*>(&OtherUntyped);
+			return Base == Other->Base &&
+				Pose == Other->Pose;
 		}
 		return false;
 	}
@@ -41,53 +39,52 @@ namespace mu
 
 	uint64 ASTOpMeshApplyPose::Hash() const
 	{
-		uint64 res = std::hash<void*>()(base.child().get());
-		hash_combine(res, pose.child().get());
-		return res;
+		uint64 Result = std::hash<void*>()(Base.child().get());
+		hash_combine(Result, Pose.child().get());
+		return Result;
 	}
 
 
-	mu::Ptr<ASTOp> ASTOpMeshApplyPose::Clone(MapChildFuncRef mapChild) const
+	Ptr<ASTOp> ASTOpMeshApplyPose::Clone(MapChildFuncRef MapChild) const
 	{
-		Ptr<ASTOpMeshApplyPose> n = new ASTOpMeshApplyPose();
-		n->base = mapChild(base.child());
-		n->pose = mapChild(pose.child());
-		return n;
+		Ptr<ASTOpMeshApplyPose> New = new ASTOpMeshApplyPose();
+		New->Base = MapChild(Base.child());
+		New->Pose = MapChild(Pose.child());
+		return New;
 	}
 
 
-	void ASTOpMeshApplyPose::ForEachChild(const TFunctionRef<void(ASTChild&)> f)
+	void ASTOpMeshApplyPose::ForEachChild(const TFunctionRef<void(ASTChild&)> Func)
 	{
-		f(base);
-		f(pose);
+		Func(Base);
+		Func(Pose);
 	}
 
 
-	void ASTOpMeshApplyPose::Link(FProgram& program, FLinkerOptions*)
+	void ASTOpMeshApplyPose::Link(FProgram& Program, FLinkerOptions*)
 	{
 		// Already linked?
 		if (!linkedAddress)
 		{
-			OP::MeshApplyPoseArgs args;
-			memset(&args, 0, sizeof(args));
+			OP::MeshApplyPoseArgs Args = {};
 
-			if (base) args.base = base->linkedAddress;
-			if (pose) args.pose = pose->linkedAddress;
+			if (Base) Args.base = Base->linkedAddress;
+			if (Pose) Args.pose = Pose->linkedAddress;
 
-			linkedAddress = (OP::ADDRESS)program.OpAddress.Num();
-			//program.m_code.push_back(op);
-			program.OpAddress.Add((uint32_t)program.ByteCode.Num());
-			AppendCode(program.ByteCode, EOpType::ME_APPLYPOSE);
-			AppendCode(program.ByteCode, args);
+			linkedAddress = (OP::ADDRESS)Program.OpAddress.Num();
+			Program.OpAddress.Add(Program.ByteCode.Num());
+			AppendCode(Program.ByteCode, GetOpType());
+			AppendCode(Program.ByteCode, Args);
 		}
 
 	}
 
+
 	Ptr<ASTOp> ASTOpMeshApplyPose::OptimiseSink(const FModelOptimizationOptions&, FOptimizeSinkContext&) const
 	{
 		Ptr<ASTOp> NewOp;
 
-		Ptr<ASTOp> MeshAt = base.child();
+		Ptr<ASTOp> MeshAt = Base.child();
 
 		if (!MeshAt)
 		{
@@ -104,7 +101,7 @@ namespace mu
 				if (New->Source)
 				{
 					Ptr<ASTOpMeshApplyPose> NewApplyPose = mu::Clone<ASTOpMeshApplyPose>(this);
-					NewApplyPose->base = New->Source.child();
+					NewApplyPose->Base = New->Source.child();
 					New->Source = NewApplyPose;
 				}
 
@@ -121,9 +118,9 @@ namespace mu
 
 	FSourceDataDescriptor ASTOpMeshApplyPose::GetSourceDataDescriptor(FGetSourceDataDescriptorContext* Context) const
 	{
-		if (base)
+		if (Base)
 		{
-			return base->GetSourceDataDescriptor(Context);
+			return Base->GetSourceDataDescriptor(Context);
 		}
 
 		return {};
diff --git a/Engine/Plugins/Mutable/Source/MutableTools/Private/MuT/ASTOpMeshApplyPose.h b/Engine/Plugins/Mutable/Source/MutableTools/Private/MuT/ASTOpMeshApplyPose.h
--- a/Engine/Plugins/Mutable/Source/MutableTools/Private/MuT/ASTOpMeshApplyPose.h
+++ b/Engine/Plugins/Mutable/Source/MutableTools/Private/MuT/ASTOpMeshApplyPose.h
@@ -10,18 +10,15 @@
 
 namespace mu
 {
-struct FProgram;
+	struct FProgram;
 
-
-	//---------------------------------------------------------------------------------------------
-	//!
-	//---------------------------------------------------------------------------------------------
+	/** */
 	class ASTOpMeshApplyPose final : public ASTOp
 	{
 	public:
 
-		ASTChild base;
-		ASTChild pose;
+		ASTChild Base;
+		ASTChild Pose;
 
 	public:
 
@@ -29,13 +26,14 @@ struct FProgram;
 		ASTOpMeshApplyPose(const ASTOpMeshApplyPose&) = delete;
 		~ASTOpMeshApplyPose();
 
-		EOpType GetOpType() const override { return EOpType::ME_APPLYPOSE; }
-		uint64 Hash() const override;
-		bool IsEqual(const ASTOp& otherUntyped) const override;
-		Ptr<ASTOp> Clone(MapChildFuncRef mapChild) const override;
-		void ForEachChild(const TFunctionRef<void(ASTChild&)>) override;
-		void Link(FProgram& program, FLinkerOptions* Options) override;
-		Ptr<ASTOp> OptimiseSink(const FModelOptimizationOptions&, FOptimizeSinkContext&) const override;
+		// ASTOp interface
+		virtual EOpType GetOpType() const override { return EOpType::ME_APPLYPOSE; }
+		virtual uint64 Hash() const override;
+		virtual bool IsEqual(const ASTOp&) const override;
+		virtual Ptr<ASTOp> Clone(MapChildFuncRef) const override;
+		virtual void ForEachChild(const TFunctionRef<void(ASTChild&)>) override;
+		virtual void Link(FProgram&, FLinkerOptions*) override;
+		virtual Ptr<ASTOp> OptimiseSink(const FModelOptimizationOptions&, FOptimizeSinkContext&) const override;
 		virtual FSourceDataDescriptor GetSourceDataDescriptor(FGetSourceDataDescriptorContext*) const override;
 	};
 
diff --git a/Engine/Plugins/Mutable/Source/MutableTools/Private/MuT/ASTOpMeshClipMorphPlane.cpp b/Engine/Plugins/Mutable/Source/MutableTools/Private/MuT/ASTOpMeshClipMorphPlane.cpp
--- a/Engine/Plugins/Mutable/Source/MutableTools/Private/MuT/ASTOpMeshClipMorphPlane.cpp
+++ b/Engine/Plugins/Mutable/Source/MutableTools/Private/MuT/ASTOpMeshClipMorphPlane.cpp
@@ -17,7 +17,7 @@ namespace mu
 	//-------------------------------------------------------------------------------------------------
 	//-------------------------------------------------------------------------------------------------
 	ASTOpMeshClipMorphPlane::ASTOpMeshClipMorphPlane()
-		: source(this)
+		: Source(this)
 	{
 	}
 
@@ -35,14 +35,14 @@ namespace mu
 		{
 			const ASTOpMeshClipMorphPlane* other = static_cast<const ASTOpMeshClipMorphPlane*>(&otherUntyped);
 
-			return source == other->source &&
-				morphShape == other->morphShape &&
-				selectionShape == other->selectionShape &&
-				vertexSelectionBone == other->vertexSelectionBone &&
+			return Source == other->Source &&
+				MorphShape == other->MorphShape &&
+				SelectionShape == other->SelectionShape &&
+				VertexSelectionBone == other->VertexSelectionBone &&
 				VertexSelectionType == other->VertexSelectionType &&
 				FaceCullStrategy == other->FaceCullStrategy &&
-				dist == other->dist &&
-				factor == other->factor;
+				Dist == other->Dist &&
+				Factor == other->Factor;
 		}
 		return false;
 	}
@@ -50,9 +50,9 @@ namespace mu
 
 	uint64 ASTOpMeshClipMorphPlane::Hash() const
 	{
-		uint64 res = std::hash<void*>()(source.child().get());
-		hash_combine(res, factor);
-		hash_combine(res, dist);
+		uint64 res = std::hash<void*>()(Source.child().get());
+		hash_combine(res, Factor);
+		hash_combine(res, Dist);
 		return res;
 	}
 
@@ -60,21 +60,21 @@ namespace mu
 	mu::Ptr<ASTOp> ASTOpMeshClipMorphPlane::Clone(MapChildFuncRef mapChild) const
 	{
 		Ptr<ASTOpMeshClipMorphPlane> n = new ASTOpMeshClipMorphPlane();
-		n->source = mapChild(source.child());
-		n->morphShape = morphShape;
-		n->selectionShape = selectionShape;
-		n->vertexSelectionBone = vertexSelectionBone;
+		n->Source = mapChild(Source.child());
+		n->MorphShape = MorphShape;
+		n->SelectionShape = SelectionShape;
+		n->VertexSelectionBone = VertexSelectionBone;
 		n->VertexSelectionType = VertexSelectionType;
 		n->FaceCullStrategy = FaceCullStrategy;
-		n->dist = dist;
-		n->factor = factor;
+		n->Dist = Dist;
+		n->Factor = Factor;
 		return n;
 	}
 
 
 	void ASTOpMeshClipMorphPlane::ForEachChild(const TFunctionRef<void(ASTChild&)> f)
 	{
-		f(source);
+		f(Source);
 	}
 
 
@@ -86,27 +86,27 @@ namespace mu
 			OP::MeshClipMorphPlaneArgs args;
 			FMemory::Memzero(&args, sizeof(args));
 
-			if (source) args.source = source->linkedAddress;
+			if (Source) args.Source = Source->linkedAddress;
 
-			args.morphShape = (OP::ADDRESS)program.ConstantShapes.Num();
-			program.ConstantShapes.Add(morphShape);
+			args.MorphShape = (OP::ADDRESS)program.ConstantShapes.Num();
+			program.ConstantShapes.Add(MorphShape);
 
 			args.FaceCullStrategy = FaceCullStrategy;
 
 			args.VertexSelectionType = VertexSelectionType;
 			if (VertexSelectionType == EClipVertexSelectionType::BoneHierarchy)
 			{
-				check(vertexSelectionBone.Id < MAX_uint32);
-				args.vertexSelectionShapeOrBone = vertexSelectionBone.Id;
+				check(VertexSelectionBone.Id < MAX_uint32);
+				args.VertexSelectionShapeOrBone = VertexSelectionBone.Id;
 			}
 			else if (VertexSelectionType == EClipVertexSelectionType::Shape)
 			{
-				args.vertexSelectionShapeOrBone = program.AddConstant(selectionShape);
+				args.VertexSelectionShapeOrBone = program.AddConstant(SelectionShape);
 			}
 
-			args.dist = dist;
-			args.factor = factor;
-			args.maxBoneRadius = vertexSelectionBoneMaxRadius;
+			args.Dist = Dist;
+			args.Factor = Factor;
+			args.MaxBoneRadius = VertexSelectionBoneMaxRadius;
 
 			linkedAddress = (OP::ADDRESS)program.OpAddress.Num();
 			//program.m_code.push_back(op);
@@ -122,12 +122,12 @@ namespace mu
 	{
 		Ptr<ASTOp> NewOp;
 
-		if (!source.child())
+		if (!Source.child())
 		{
 			return nullptr;
 		}
 
-		EOpType SourceType = source.child()->GetOpType();
+		EOpType SourceType = Source.child()->GetOpType();
 
 		// Optimize only the mesh parameter
 		switch (SourceType)
@@ -135,12 +135,12 @@ namespace mu
 
 		case EOpType::ME_ADDTAGS:
 		{
-			Ptr<ASTOpMeshAddTags> NewAddTags = mu::Clone<ASTOpMeshAddTags>(source.child());
+			Ptr<ASTOpMeshAddTags> NewAddTags = mu::Clone<ASTOpMeshAddTags>(Source.child());
 
 			if (NewAddTags->Source)
 			{
 				Ptr<ASTOpMeshClipMorphPlane> New = mu::Clone<ASTOpMeshClipMorphPlane>(this);
-				New->source = NewAddTags->Source.child();
+				New->Source = NewAddTags->Source.child();
 				NewAddTags->Source = New;
 			}
 
@@ -159,9 +159,9 @@ namespace mu
 
 	FSourceDataDescriptor ASTOpMeshClipMorphPlane::GetSourceDataDescriptor(FGetSourceDataDescriptorContext* Context) const
 	{
-		if (source)
+		if (Source)
 		{
-			return source->GetSourceDataDescriptor(Context);
+			return Source->GetSourceDataDescriptor(Context);
 		}
 
 		return {};
diff --git a/Engine/Plugins/Mutable/Source/MutableTools/Private/MuT/ASTOpMeshClipMorphPlane.h b/Engine/Plugins/Mutable/Source/MutableTools/Private/MuT/ASTOpMeshClipMorphPlane.h
--- a/Engine/Plugins/Mutable/Source/MutableTools/Private/MuT/ASTOpMeshClipMorphPlane.h
+++ b/Engine/Plugins/Mutable/Source/MutableTools/Private/MuT/ASTOpMeshClipMorphPlane.h
@@ -8,28 +8,25 @@
 #include "MuR/Ptr.h"
 #include "MuT/AST.h"
 
-
 namespace mu
 {
-struct FProgram;
+	struct FProgram;
 
-	//---------------------------------------------------------------------------------------------
-	//!
-	//---------------------------------------------------------------------------------------------
+	/** */
 	class ASTOpMeshClipMorphPlane final : public ASTOp
 	{
 	public:
 
-		ASTChild source;
+		ASTChild Source;
 
-		FShape morphShape;
-		FShape selectionShape;
-		FBoneName vertexSelectionBone;
+		FShape MorphShape;
+		FShape SelectionShape;
+		FBoneName VertexSelectionBone;
 
 		EClipVertexSelectionType VertexSelectionType = EClipVertexSelectionType::None;
 		EFaceCullStrategy FaceCullStrategy = EFaceCullStrategy::AllVerticesCulled;
 
-		float dist = 0.f, factor = 0.f, vertexSelectionBoneMaxRadius = -1.f;
+		float Dist = 0.f, Factor = 0.f, VertexSelectionBoneMaxRadius = -1.f;
 
 	public:
 
diff --git a/Engine/Plugins/Mutable/Source/MutableTools/Private/MuT/ASTOpMeshClipWithMesh.cpp b/Engine/Plugins/Mutable/Source/MutableTools/Private/MuT/ASTOpMeshClipWithMesh.cpp
new file mode 100644
--- /dev/null
+++ b/Engine/Plugins/Mutable/Source/MutableTools/Private/MuT/ASTOpMeshClipWithMesh.cpp
@@ -0,0 +1,129 @@
+// Copyright Epic Games, Inc. All Rights Reserved.
+
+#include "MuT/ASTOpMeshClipWithMesh.h"
+#include "MuT/ASTOpMeshAddTags.h"
+
+#include "HAL/PlatformMath.h"
+#include "MuR/ModelPrivate.h"
+#include "MuR/RefCounted.h"
+#include "MuR/Types.h"
+
+namespace mu
+{
+
+	ASTOpMeshClipWithMesh::ASTOpMeshClipWithMesh()
+		: Source(this)
+		, ClipMesh(this)
+	{
+	}
+
+
+	ASTOpMeshClipWithMesh::~ASTOpMeshClipWithMesh()
+	{
+		// Explicit call needed to avoid recursive destruction
+		ASTOp::RemoveChildren();
+	}
+
+
+	bool ASTOpMeshClipWithMesh::IsEqual(const ASTOp& OtherUntyped) const
+	{
+		if (OtherUntyped.GetOpType()==GetOpType())
+		{
+			const ASTOpMeshClipWithMesh* Other = static_cast<const ASTOpMeshClipWithMesh*>(&OtherUntyped);
+			return Source == Other->Source &&
+				ClipMesh == Other->ClipMesh;
+		}
+		return false;
+	}
+
+
+	uint64 ASTOpMeshClipWithMesh::Hash() const
+	{
+		uint64 Result = std::hash<void*>()(Source.child().get());
+		hash_combine(Result, ClipMesh.child().get());
+		return Result;
+	}
+
+
+	Ptr<ASTOp> ASTOpMeshClipWithMesh::Clone(MapChildFuncRef MapChild) const
+	{
+		Ptr<ASTOpMeshClipWithMesh> New = new ASTOpMeshClipWithMesh();
+		New->Source = MapChild(Source.child());
+		New->ClipMesh = MapChild(ClipMesh.child());
+		return New;
+	}
+
+
+	void ASTOpMeshClipWithMesh::ForEachChild(const TFunctionRef<void(ASTChild&)> Func)
+	{
+		Func(Source);
+		Func(ClipMesh);
+	}
+
+
+	void ASTOpMeshClipWithMesh::Link(FProgram& Program, FLinkerOptions*)
+	{
+		// Already linked?
+		if (!linkedAddress)
+		{
+			OP::MeshApplyPoseArgs Args = {};
+
+			if (Source) Args.base = Source->linkedAddress;
+			if (ClipMesh) Args.pose = ClipMesh->linkedAddress;
+
+			linkedAddress = (OP::ADDRESS)Program.OpAddress.Num();
+			Program.OpAddress.Add(Program.ByteCode.Num());
+			AppendCode(Program.ByteCode, GetOpType());
+			AppendCode(Program.ByteCode, Args);
+		}
+
+	}
+
+
+	Ptr<ASTOp> ASTOpMeshClipWithMesh::OptimiseSink(const FModelOptimizationOptions&, FOptimizeSinkContext&) const
+	{
+		Ptr<ASTOp> NewOp;
+
+		Ptr<ASTOp> MeshAt = Source.child();
+
+		if (!MeshAt)
+		{
+			return nullptr;
+		}
+
+		EOpType MeshType = MeshAt->GetOpType();
+
+		switch (MeshType)
+		{
+			case EOpType::ME_ADDTAGS:
+			{
+				Ptr<ASTOpMeshAddTags> New = mu::Clone<ASTOpMeshAddTags>(MeshAt);
+				if (New->Source)
+				{
+					Ptr<ASTOpMeshClipWithMesh> NewApplyPose = mu::Clone<ASTOpMeshClipWithMesh>(this);
+					NewApplyPose->Source = New->Source.child();
+					New->Source = NewApplyPose;
+				}
+
+				NewOp = New;
+				break;
+			}
+
+			default:
+				break;
+		}
+
+		return NewOp;
+	}
+
+
+	FSourceDataDescriptor ASTOpMeshClipWithMesh::GetSourceDataDescriptor(FGetSourceDataDescriptorContext* Context) const
+	{
+		if (Source)
+		{
+			return Source->GetSourceDataDescriptor(Context);
+		}
+
+		return {};
+	}
+}
diff --git a/Engine/Plugins/Mutable/Source/MutableTools/Private/MuT/ASTOpMeshClipWithMesh.h b/Engine/Plugins/Mutable/Source/MutableTools/Private/MuT/ASTOpMeshClipWithMesh.h
new file mode 100644
--- /dev/null
+++ b/Engine/Plugins/Mutable/Source/MutableTools/Private/MuT/ASTOpMeshClipWithMesh.h
@@ -0,0 +1,42 @@
+// Copyright Epic Games, Inc. All Rights Reserved.
+
+#pragma once
+
+#include "HAL/Platform.h"
+#include "MuR/Operations.h"
+#include "MuR/Ptr.h"
+#include "MuT/AST.h"
+
+
+namespace mu
+{
+	struct FProgram;
+
+	/** */
+	class ASTOpMeshClipWithMesh final : public ASTOp
+	{
+	public:
+
+		ASTChild Source;
+		ASTChild ClipMesh;
+
+	public:
+
+		ASTOpMeshClipWithMesh();
+		ASTOpMeshClipWithMesh(const ASTOpMeshClipWithMesh&) = delete;
+		~ASTOpMeshClipWithMesh();
+
+		// ASTOp interface
+		virtual EOpType GetOpType() const override { return EOpType::ME_CLIPWITHMESH; }
+		virtual uint64 Hash() const override;
+		virtual bool IsEqual(const ASTOp&) const override;
+		virtual Ptr<ASTOp> Clone(MapChildFuncRef) const override;
+		virtual void ForEachChild(const TFunctionRef<void(ASTChild&)>) override;
+		virtual void Link(FProgram&, FLinkerOptions*) override;
+		virtual Ptr<ASTOp> OptimiseSink(const FModelOptimizationOptions&, FOptimizeSinkContext&) const override;
+		virtual FSourceDataDescriptor GetSourceDataDescriptor(FGetSourceDataDescriptorContext*) const override;
+	};
+
+
+}
+
diff --git a/Engine/Plugins/Mutable/Source/MutableTools/Private/MuT/ASTOpMeshExtractLayoutBlocks.cpp b/Engine/Plugins/Mutable/Source/MutableTools/Private/MuT/ASTOpMeshExtractLayoutBlocks.cpp
--- a/Engine/Plugins/Mutable/Source/MutableTools/Private/MuT/ASTOpMeshExtractLayoutBlocks.cpp
+++ b/Engine/Plugins/Mutable/Source/MutableTools/Private/MuT/ASTOpMeshExtractLayoutBlocks.cpp
@@ -9,6 +9,9 @@
 #include "MuT/ASTOpMeshRemoveMask.h"
 #include "MuT/ASTOpMeshClipMorphPlane.h"
 #include "MuT/ASTOpMeshApplyPose.h"
+#include "MuT/ASTOpMeshApplyLayout.h"
+#include "MuT/ASTOpMeshMerge.h"
+#include "MuT/ASTOpMeshSetSkeleton.h"
 #include "MuR/ModelPrivate.h"
 #include "MuR/RefCounted.h"
 #include "MuR/Types.h"
@@ -150,33 +153,33 @@ namespace mu
 
 		case EOpType::ME_APPLYLAYOUT:
 		{
-			Ptr<ASTOpFixed> newOp = mu::Clone<ASTOpFixed>(at);
-			newOp->SetChild(newOp->op.args.MeshApplyLayout.mesh, Visit(newOp->children[newOp->op.args.MeshApplyLayout.mesh].child(), currentSinkOp));
-			newAt = newOp;
+			Ptr<ASTOpMeshApplyLayout> NewOp = mu::Clone<ASTOpMeshApplyLayout>(at);
+			NewOp->Mesh = Visit(NewOp->Mesh.child(), currentSinkOp);
+			newAt = NewOp;
 			break;
 		}
 
 		case EOpType::ME_SETSKELETON:
 		{
-			Ptr<ASTOpFixed> newOp = mu::Clone<ASTOpFixed>(at);
-			newOp->SetChild(newOp->op.args.MeshSetSkeleton.source, Visit(newOp->children[newOp->op.args.MeshSetSkeleton.source].child(), currentSinkOp));
-			newAt = newOp;
+			Ptr<ASTOpMeshSetSkeleton> NewOp = mu::Clone<ASTOpMeshSetSkeleton>(at);
+			NewOp->Source = Visit( NewOp->Source.child(), currentSinkOp);
+			newAt = NewOp;
 			break;
 		}
 
 		case EOpType::ME_ADDTAGS:
 		{
-			Ptr<ASTOpMeshAddTags> newOp = mu::Clone<ASTOpMeshAddTags>(at);
-			newOp->Source = Visit(newOp->Source.child(), currentSinkOp);
-			newAt = newOp;
+			Ptr<ASTOpMeshAddTags> NewOp = mu::Clone<ASTOpMeshAddTags>(at);
+			NewOp->Source = Visit(NewOp->Source.child(), currentSinkOp);
+			newAt = NewOp;
 			break;
 		}
 
 		case EOpType::ME_CLIPMORPHPLANE:
 		{
-			Ptr<ASTOpMeshClipMorphPlane> newOp = mu::Clone<ASTOpMeshClipMorphPlane>(at);
-			newOp->source = Visit(newOp->source.child(), currentSinkOp);
-			newAt = newOp;
+			Ptr<ASTOpMeshClipMorphPlane> NewOp = mu::Clone<ASTOpMeshClipMorphPlane>(at);
+			NewOp->Source = Visit(NewOp->Source.child(), currentSinkOp);
+			newAt = NewOp;
 			break;
 		}
 
@@ -191,17 +194,17 @@ namespace mu
 
 		case EOpType::ME_MERGE:
 		{
-			Ptr<ASTOpFixed> newOp = mu::Clone<ASTOpFixed>(at);
-			newOp->SetChild(newOp->op.args.MeshMerge.base, Visit(newOp->children[newOp->op.args.MeshMerge.base].child(), currentSinkOp));
-			newOp->SetChild(newOp->op.args.MeshMerge.added, Visit(newOp->children[newOp->op.args.MeshMerge.added].child(), currentSinkOp));
-			newAt = newOp;
+			Ptr<ASTOpMeshMerge> NewOp = mu::Clone<ASTOpMeshMerge>(at);
+			NewOp->Base = Visit(NewOp->Base.child(), currentSinkOp);
+			NewOp->Added = Visit(NewOp->Added.child(), currentSinkOp);
+			newAt = NewOp;
 			break;
 		}
 
 		case EOpType::ME_APPLYPOSE:
 		{
 			Ptr<ASTOpMeshApplyPose> NewOp = mu::Clone<ASTOpMeshApplyPose>(at);
-			NewOp->base = Visit(NewOp->base.child(), currentSinkOp);
+			NewOp->Base = Visit(NewOp->Base.child(), currentSinkOp);
 			newAt = NewOp;
 			break;
 		}
diff --git a/Engine/Plugins/Mutable/Source/MutableTools/Private/MuT/ASTOpMeshFormat.cpp b/Engine/Plugins/Mutable/Source/MutableTools/Private/MuT/ASTOpMeshFormat.cpp
--- a/Engine/Plugins/Mutable/Source/MutableTools/Private/MuT/ASTOpMeshFormat.cpp
+++ b/Engine/Plugins/Mutable/Source/MutableTools/Private/MuT/ASTOpMeshFormat.cpp
@@ -19,6 +19,9 @@
 #include "MuT/ASTOpMeshMorph.h"
 #include "MuT/ASTOpMeshAddTags.h"
 #include "MuT/ASTOpMeshApplyPose.h"
+#include "MuT/ASTOpMeshApplyLayout.h"
+#include "MuT/ASTOpMeshMerge.h"
+#include "MuT/ASTOpMeshSetSkeleton.h"
 #include "MuT/ASTOpSwitch.h"
 
 #include "GPUSkinPublicDefs.h"
@@ -274,7 +277,10 @@ namespace
 
 mu::Ptr<ASTOp> Sink_MeshFormatAST::Visit(const mu::Ptr<ASTOp>& at, const ASTOpMeshFormat* currentFormatOp)
 {
-	if (!at) return nullptr;
+	if (!at)
+	{
+		return nullptr;
+	}
 
 	// Already visited?
 	const Ptr<ASTOp>* Cached = OldToNew.Find({ at,currentFormatOp });
@@ -289,17 +295,17 @@ mu::Ptr<ASTOp> Sink_MeshFormatAST::Visit(const mu::Ptr<ASTOp>& at, const ASTOpMe
 
 	case EOpType::ME_APPLYLAYOUT:
 	{
-		auto newOp = mu::Clone<ASTOpFixed>(at);
-		newOp->SetChild(newOp->op.args.MeshApplyLayout.mesh, Visit(newOp->children[newOp->op.args.MeshApplyLayout.mesh].child(), currentFormatOp));
-		newAt = newOp;
+		Ptr<ASTOpMeshApplyLayout> NewOp = mu::Clone<ASTOpMeshApplyLayout>(at);
+		NewOp->Mesh = Visit(NewOp->Mesh.child(), currentFormatOp);
+		newAt = NewOp;
 		break;
 	}
 
 	case EOpType::ME_SETSKELETON:
 	{
-		auto newOp = mu::Clone<ASTOpFixed>(at);
-		newOp->SetChild(newOp->op.args.MeshSetSkeleton.source, Visit(newOp->children[newOp->op.args.MeshSetSkeleton.source].child(), currentFormatOp));
-		newAt = newOp;
+		Ptr<ASTOpMeshSetSkeleton> NewOp = mu::Clone<ASTOpMeshSetSkeleton>(at);
+		NewOp->Source = Visit(NewOp->Source.child(), currentFormatOp);
+		newAt = NewOp;
 		break;
 	}
 
@@ -314,7 +320,7 @@ mu::Ptr<ASTOp> Sink_MeshFormatAST::Visit(const mu::Ptr<ASTOp>& at, const ASTOpMe
 	case EOpType::ME_CLIPMORPHPLANE:
 	{
 		Ptr<ASTOpMeshClipMorphPlane> newOp = mu::Clone<ASTOpMeshClipMorphPlane>(at);
-		newOp->source = Visit(newOp->source.child(), currentFormatOp);
+		newOp->Source = Visit(newOp->Source.child(), currentFormatOp);
 		newAt = newOp;
 		break;
 	}
@@ -361,10 +367,10 @@ mu::Ptr<ASTOp> Sink_MeshFormatAST::Visit(const mu::Ptr<ASTOp>& at, const ASTOpMe
 
 	case EOpType::ME_MERGE:
 	{
-		Ptr<ASTOpFixed> newOp = mu::Clone<ASTOpFixed>(at);
-		newOp->SetChild(newOp->op.args.MeshMerge.base, Visit(newOp->children[newOp->op.args.MeshMerge.base].child(), currentFormatOp));
-		newOp->SetChild(newOp->op.args.MeshMerge.added, Visit(newOp->children[newOp->op.args.MeshMerge.added].child(), currentFormatOp));
-		newAt = newOp;
+		Ptr<ASTOpMeshMerge> NewOp = mu::Clone<ASTOpMeshMerge>(at);
+		NewOp->Base = Visit(NewOp->Base.child(), currentFormatOp);
+		NewOp->Added = Visit(NewOp->Added.child(), currentFormatOp);
+		newAt = NewOp;
 		break;
 	}
 
@@ -388,7 +394,7 @@ mu::Ptr<ASTOp> Sink_MeshFormatAST::Visit(const mu::Ptr<ASTOp>& at, const ASTOpMe
         // to indicate it should copy the skinning from the base mesh.
         NewFormat->Format = NewFormatConstant;
 
-        NewOp->base = Visit(NewOp->base.child(), NewFormat.get());
+        NewOp->Base = Visit(NewOp->Base.child(), NewFormat.get());
 		
 		newAt = NewOp;
 		break;
diff --git a/Engine/Plugins/Mutable/Source/MutableTools/Private/MuT/ASTOpMeshMaskDiff.cpp b/Engine/Plugins/Mutable/Source/MutableTools/Private/MuT/ASTOpMeshMaskDiff.cpp
new file mode 100644
--- /dev/null
+++ b/Engine/Plugins/Mutable/Source/MutableTools/Private/MuT/ASTOpMeshMaskDiff.cpp
@@ -0,0 +1,144 @@
+// Copyright Epic Games, Inc. All Rights Reserved.
+
+#include "MuT/ASTOpMeshMaskDiff.h"
+#include "MuT/ASTOpMeshAddTags.h"
+#include "MuT/ASTOpSwitch.h"
+
+#include "HAL/PlatformMath.h"
+#include "MuR/ModelPrivate.h"
+#include "MuR/RefCounted.h"
+#include "MuR/Types.h"
+
+namespace mu
+{
+
+	ASTOpMeshMaskDiff::ASTOpMeshMaskDiff()
+		: Source(this)
+		, Fragment(this)
+	{
+	}
+
+
+	ASTOpMeshMaskDiff::~ASTOpMeshMaskDiff()
+	{
+		// Explicit call needed to avoid recursive destruction
+		ASTOp::RemoveChildren();
+	}
+
+
+	bool ASTOpMeshMaskDiff::IsEqual(const ASTOp& OtherUntyped) const
+	{
+		if (OtherUntyped.GetOpType()==GetOpType())
+		{
+			const ASTOpMeshMaskDiff* Other = static_cast<const ASTOpMeshMaskDiff*>(&OtherUntyped);
+			return Source == Other->Source &&
+				Fragment == Other->Fragment;
+		}
+		return false;
+	}
+
+
+	uint64 ASTOpMeshMaskDiff::Hash() const
+	{
+		uint64 Result = std::hash<void*>()(Source.child().get());
+		hash_combine(Result, Fragment.child().get());
+		return Result;
+	}
+
+
+	Ptr<ASTOp> ASTOpMeshMaskDiff::Clone(MapChildFuncRef MapChild) const
+	{
+		Ptr<ASTOpMeshMaskDiff> New = new ASTOpMeshMaskDiff();
+		New->Source = MapChild(Source.child());
+		New->Fragment = MapChild(Fragment.child());
+		return New;
+	}
+
+
+	void ASTOpMeshMaskDiff::ForEachChild(const TFunctionRef<void(ASTChild&)> Func)
+	{
+		Func(Source);
+		Func(Fragment);
+	}
+
+
+	void ASTOpMeshMaskDiff::Link(FProgram& Program, FLinkerOptions*)
+	{
+		// Already linked?
+		if (!linkedAddress)
+		{
+			OP::MeshMaskDiffArgs Args = {};
+
+			if (Source) Args.Source = Source->linkedAddress;
+			if (Fragment) Args.Fragment = Fragment->linkedAddress;
+
+			linkedAddress = (OP::ADDRESS)Program.OpAddress.Num();
+			Program.OpAddress.Add(Program.ByteCode.Num());
+			AppendCode(Program.ByteCode, GetOpType());
+			AppendCode(Program.ByteCode, Args);
+		}
+	}
+
+
+	mu::Ptr<ASTOp> ASTOpMeshMaskDiff::OptimiseSemantic(const FModelOptimizationOptions&, int32 Pass) const
+	{
+		mu::Ptr<ASTOp> ResultOp;
+
+		EOpType FragmentType = Fragment->GetOpType();
+		switch (FragmentType)
+		{
+
+		case EOpType::ME_ADDTAGS:
+		{
+			// Tags in the fragment can be ignored.
+			const ASTOpMeshAddTags* Add = static_cast<const ASTOpMeshAddTags*>(Fragment.child().get());
+
+			Ptr<ASTOpMeshMaskDiff> NewAt = mu::Clone<ASTOpMeshMaskDiff>(this);
+			NewAt->Fragment = Add->Source.child();
+			ResultOp = NewAt;
+			break;
+		}
+		case EOpType::ME_SWITCH:
+		{
+			// Move the mask diff down all the paths
+			Ptr<ASTOpSwitch> NewOp = mu::Clone<ASTOpSwitch>(Fragment.child().get());
+
+			if (NewOp->def)
+			{
+				Ptr<ASTOpMeshMaskDiff> defOp = mu::Clone<ASTOpMeshMaskDiff>(this);
+				defOp->Fragment = NewOp->def.child();
+				NewOp->def = defOp;
+			}
+
+			// We need to copy the options because we change them
+			for (int32 v = 0; v < NewOp->cases.Num(); ++v)
+			{
+				if (NewOp->cases[v].branch)
+				{
+					Ptr<ASTOpMeshMaskDiff> bOp = mu::Clone<ASTOpMeshMaskDiff>(this);
+					bOp->Fragment = NewOp->cases[v].branch.child();
+					NewOp->cases[v].branch = bOp;
+				}
+			}
+
+			ResultOp = NewOp;
+			break;
+		}
+		default:
+			break;
+		}
+
+		return ResultOp;
+	}
+
+
+	FSourceDataDescriptor ASTOpMeshMaskDiff::GetSourceDataDescriptor(FGetSourceDataDescriptorContext* Context) const
+	{
+		if (Source)
+		{
+			return Source->GetSourceDataDescriptor(Context);
+		}
+
+		return {};
+	}
+}
diff --git a/Engine/Plugins/Mutable/Source/MutableTools/Private/MuT/ASTOpMeshMaskDiff.h b/Engine/Plugins/Mutable/Source/MutableTools/Private/MuT/ASTOpMeshMaskDiff.h
new file mode 100644
--- /dev/null
+++ b/Engine/Plugins/Mutable/Source/MutableTools/Private/MuT/ASTOpMeshMaskDiff.h
@@ -0,0 +1,39 @@
+// Copyright Epic Games, Inc. All Rights Reserved.
+
+#pragma once
+
+#include "HAL/Platform.h"
+#include "MuR/Operations.h"
+#include "MuR/Ptr.h"
+#include "MuT/AST.h"
+
+namespace mu
+{
+	struct FProgram;
+
+	/** */
+	class ASTOpMeshMaskDiff final : public ASTOp
+	{
+	public:
+
+		ASTChild Source;
+		ASTChild Fragment;
+
+	public:
+
+		ASTOpMeshMaskDiff();
+		ASTOpMeshMaskDiff(const ASTOpMeshMaskDiff&) = delete;
+		~ASTOpMeshMaskDiff();
+
+		// ASTOp interface
+		virtual EOpType GetOpType() const override { return EOpType::ME_MASKDIFF; }
+		virtual uint64 Hash() const override;
+		virtual bool IsEqual(const ASTOp&) const override;
+		virtual Ptr<ASTOp> Clone(MapChildFuncRef) const override;
+		virtual void ForEachChild(const TFunctionRef<void(ASTChild&)>) override;
+		virtual void Link(FProgram&, FLinkerOptions*) override;
+		virtual Ptr<ASTOp> OptimiseSemantic(const FModelOptimizationOptions&, int32 Pass) const override;
+		virtual FSourceDataDescriptor GetSourceDataDescriptor(FGetSourceDataDescriptorContext*) const override;
+	};
+
+}
diff --git a/Engine/Plugins/Mutable/Source/MutableTools/Private/MuT/ASTOpMeshMerge.cpp b/Engine/Plugins/Mutable/Source/MutableTools/Private/MuT/ASTOpMeshMerge.cpp
new file mode 100644
--- /dev/null
+++ b/Engine/Plugins/Mutable/Source/MutableTools/Private/MuT/ASTOpMeshMerge.cpp
@@ -0,0 +1,95 @@
+// Copyright Epic Games, Inc. All Rights Reserved.
+
+#include "MuT/ASTOpMeshMerge.h"
+#include "MuT/ASTOpMeshAddTags.h"
+
+#include "HAL/PlatformMath.h"
+#include "MuR/ModelPrivate.h"
+#include "MuR/RefCounted.h"
+#include "MuR/Types.h"
+
+namespace mu
+{
+
+	ASTOpMeshMerge::ASTOpMeshMerge()
+		: Base(this)
+		, Added(this)
+	{
+	}
+
+
+	ASTOpMeshMerge::~ASTOpMeshMerge()
+	{
+		// Explicit call needed to avoid recursive destruction
+		ASTOp::RemoveChildren();
+	}
+
+
+	bool ASTOpMeshMerge::IsEqual(const ASTOp& OtherUntyped) const
+	{
+		if (OtherUntyped.GetOpType()==GetOpType())
+		{
+			const ASTOpMeshMerge* Other = static_cast<const ASTOpMeshMerge*>(&OtherUntyped);
+			return Base == Other->Base &&
+				Added == Other->Added &&
+				NewSurfaceID == Other->NewSurfaceID;
+		}
+		return false;
+	}
+
+
+	uint64 ASTOpMeshMerge::Hash() const
+	{
+		uint64 Result = std::hash<void*>()(Base.child().get());
+		hash_combine(Result, Added.child().get());
+		hash_combine(Result, NewSurfaceID);
+		return Result;
+	}
+
+
+	Ptr<ASTOp> ASTOpMeshMerge::Clone(MapChildFuncRef MapChild) const
+	{
+		Ptr<ASTOpMeshMerge> New = new ASTOpMeshMerge();
+		New->Base = MapChild(Base.child());
+		New->Added = MapChild(Added.child());
+		New->NewSurfaceID = NewSurfaceID;
+		return New;
+	}
+
+
+	void ASTOpMeshMerge::ForEachChild(const TFunctionRef<void(ASTChild&)> Func)
+	{
+		Func(Base);
+		Func(Added);
+	}
+
+
+	void ASTOpMeshMerge::Link(FProgram& Program, FLinkerOptions*)
+	{
+		// Already linked?
+		if (!linkedAddress)
+		{
+			OP::MeshMergeArgs Args = {};
+
+			if (Base) Args.Base = Base->linkedAddress;
+			if (Added) Args.Added = Added->linkedAddress;
+			Args.NewSurfaceID = NewSurfaceID;
+
+			linkedAddress = (OP::ADDRESS)Program.OpAddress.Num();
+			Program.OpAddress.Add(Program.ByteCode.Num());
+			AppendCode(Program.ByteCode, GetOpType());
+			AppendCode(Program.ByteCode, Args);
+		}
+	}
+
+
+	FSourceDataDescriptor ASTOpMeshMerge::GetSourceDataDescriptor(FGetSourceDataDescriptorContext* Context) const
+	{
+		if (Base)
+		{
+			return Base->GetSourceDataDescriptor(Context);
+		}
+
+		return {};
+	}
+}
diff --git a/Engine/Plugins/Mutable/Source/MutableTools/Private/MuT/ASTOpMeshMerge.h b/Engine/Plugins/Mutable/Source/MutableTools/Private/MuT/ASTOpMeshMerge.h
new file mode 100644
--- /dev/null
+++ b/Engine/Plugins/Mutable/Source/MutableTools/Private/MuT/ASTOpMeshMerge.h
@@ -0,0 +1,41 @@
+// Copyright Epic Games, Inc. All Rights Reserved.
+
+#pragma once
+
+#include "HAL/Platform.h"
+#include "MuR/Operations.h"
+#include "MuR/Ptr.h"
+#include "MuT/AST.h"
+
+
+namespace mu
+{
+	struct FProgram;
+
+	/** */
+	class ASTOpMeshMerge final : public ASTOp
+	{
+	public:
+
+		ASTChild Base;
+		ASTChild Added;
+		uint32 NewSurfaceID = 0;
+
+	public:
+
+		ASTOpMeshMerge();
+		ASTOpMeshMerge(const ASTOpMeshMerge&) = delete;
+		~ASTOpMeshMerge();
+
+		// ASTOp interface
+		virtual EOpType GetOpType() const override { return EOpType::ME_MERGE; }
+		virtual uint64 Hash() const override;
+		virtual bool IsEqual(const ASTOp&) const override;
+		virtual Ptr<ASTOp> Clone(MapChildFuncRef) const override;
+		virtual void ForEachChild(const TFunctionRef<void(ASTChild&)>) override;
+		virtual void Link(FProgram&, FLinkerOptions*) override;
+		virtual FSourceDataDescriptor GetSourceDataDescriptor(FGetSourceDataDescriptorContext*) const override;
+	};
+
+}
+
diff --git a/Engine/Plugins/Mutable/Source/MutableTools/Private/MuT/ASTOpMeshOptimizeSkinning.cpp b/Engine/Plugins/Mutable/Source/MutableTools/Private/MuT/ASTOpMeshOptimizeSkinning.cpp
--- a/Engine/Plugins/Mutable/Source/MutableTools/Private/MuT/ASTOpMeshOptimizeSkinning.cpp
+++ b/Engine/Plugins/Mutable/Source/MutableTools/Private/MuT/ASTOpMeshOptimizeSkinning.cpp
@@ -5,6 +5,7 @@
 #include "MuT/ASTOpConditional.h"
 #include "MuT/ASTOpMeshMorph.h"
 #include "MuT/ASTOpMeshAddTags.h"
+#include "MuT/ASTOpMeshApplyLayout.h"
 #include "MuT/ASTOpSwitch.h"
 #include "MuR/ModelPrivate.h"
 #include "MuR/RefCounted.h"
@@ -144,8 +145,8 @@ namespace mu
 
 		case EOpType::ME_APPLYLAYOUT:
 		{
-			Ptr<ASTOpFixed> NewOp = mu::Clone<ASTOpFixed>(at);
-			NewOp->SetChild(NewOp->op.args.MeshApplyLayout.mesh, Visit(NewOp->children[NewOp->op.args.MeshApplyLayout.mesh].child(), CurrentSinkOp));
+			Ptr<ASTOpMeshApplyLayout> NewOp = mu::Clone<ASTOpMeshApplyLayout>(at);
+			NewOp->Mesh = Visit(NewOp->Mesh.child(), CurrentSinkOp);
 			NewAt = NewOp;
 			break;
 		}
diff --git a/Engine/Plugins/Mutable/Source/MutableTools/Private/MuT/ASTOpMeshProject.cpp b/Engine/Plugins/Mutable/Source/MutableTools/Private/MuT/ASTOpMeshProject.cpp
new file mode 100644
--- /dev/null
+++ b/Engine/Plugins/Mutable/Source/MutableTools/Private/MuT/ASTOpMeshProject.cpp
@@ -0,0 +1,184 @@
+// Copyright Epic Games, Inc. All Rights Reserved.
+
+#include "MuT/ASTOpMeshProject.h"
+#include "MuT/ASTOpMeshAddTags.h"
+#include "MuT/ASTOpConditional.h"
+#include "MuT/ASTOpSwitch.h"
+
+#include "HAL/PlatformMath.h"
+#include "MuR/ModelPrivate.h"
+#include "MuR/RefCounted.h"
+#include "MuR/Types.h"
+
+namespace mu
+{
+
+	ASTOpMeshProject::ASTOpMeshProject()
+		: Mesh(this)
+		, Projector(this)
+	{
+	}
+
+
+	ASTOpMeshProject::~ASTOpMeshProject()
+	{
+		// Explicit call needed to avoid recursive destruction
+		ASTOp::RemoveChildren();
+	}
+
+
+	bool ASTOpMeshProject::IsEqual(const ASTOp& OtherUntyped) const
+	{
+		if (OtherUntyped.GetOpType()==GetOpType())
+		{
+			const ASTOpMeshProject* Other = static_cast<const ASTOpMeshProject*>(&OtherUntyped);
+			return Mesh == Other->Mesh &&
+				Projector == Other->Projector;
+		}
+		return false;
+	}
+
+
+	uint64 ASTOpMeshProject::Hash() const
+	{
+		uint64 Result = std::hash<void*>()(Mesh.child().get());
+		hash_combine(Result, Projector.child().get());
+		return Result;
+	}
+
+
+	Ptr<ASTOp> ASTOpMeshProject::Clone(MapChildFuncRef MapChild) const
+	{
+		Ptr<ASTOpMeshProject> New = new ASTOpMeshProject();
+		New->Mesh = MapChild(Mesh.child());
+		New->Projector = MapChild(Projector.child());
+		return New;
+	}
+
+
+	void ASTOpMeshProject::ForEachChild(const TFunctionRef<void(ASTChild&)> Func)
+	{
+		Func(Mesh);
+		Func(Projector);
+	}
+
+
+	void ASTOpMeshProject::Link(FProgram& Program, FLinkerOptions*)
+	{
+		// Already linked?
+		if (!linkedAddress)
+		{
+			OP::MeshProjectArgs Args = {};
+
+			if (Mesh) Args.Mesh = Mesh->linkedAddress;
+			if (Projector) Args.Projector = Projector->linkedAddress;
+
+			linkedAddress = (OP::ADDRESS)Program.OpAddress.Num();
+			Program.OpAddress.Add(Program.ByteCode.Num());
+			AppendCode(Program.ByteCode, GetOpType());
+			AppendCode(Program.ByteCode, Args);
+		}
+
+	}
+
+
+	Ptr<ASTOp> ASTOpMeshProject::OptimiseSink(const FModelOptimizationOptions&, FOptimizeSinkContext&) const
+	{
+		Ptr<ASTOp> NewOp;
+
+		Ptr<ASTOp> SourceAt = Mesh.child();
+		Ptr<ASTOp> ProjectorAt = Projector.child();
+
+		if (!SourceAt || !ProjectorAt)
+		{
+			return NewOp;
+		}
+
+		EOpType SourceType = SourceAt->GetOpType();
+		switch (SourceType)
+		{
+
+		case EOpType::ME_CONDITIONAL:
+		{
+			if (ProjectorAt->GetOpType() == EOpType::PR_CONSTANT)
+			{
+				// We move the project down the two paths
+				Ptr<ASTOpConditional> nop = mu::Clone<ASTOpConditional>(SourceAt);
+
+				Ptr<ASTOpMeshProject> aOp = mu::Clone<ASTOpMeshProject>(this);
+				aOp->Mesh = nop->yes.child();
+				nop->yes = aOp;
+
+				Ptr<ASTOpMeshProject> bOp = mu::Clone<ASTOpMeshProject>(this);
+				bOp->Mesh = nop->no.child();
+				nop->no = bOp;
+
+				NewOp = nop;
+			}
+			break;
+		}
+
+		case EOpType::ME_SWITCH:
+		{
+			if (ProjectorAt->GetOpType() == EOpType::PR_CONSTANT)
+			{
+				// Move the format down all the paths
+				Ptr<ASTOpSwitch> nop = mu::Clone<ASTOpSwitch>(SourceAt);
+
+				if (nop->def)
+				{
+					Ptr<ASTOpMeshProject> defOp = mu::Clone<ASTOpMeshProject>(this);
+					defOp->Mesh = nop->def.child();
+					nop->def = defOp;
+				}
+
+				// We need to copy the options because we change them
+				for (int32 CaseIndex = 0; CaseIndex < nop->cases.Num(); ++CaseIndex)
+				{
+					if (nop->cases[CaseIndex].branch)
+					{
+						Ptr<ASTOpMeshProject> bOp = mu::Clone<ASTOpMeshProject>(this);
+						bOp->Mesh = nop->cases[CaseIndex].branch.child();
+						nop->cases[CaseIndex].branch = bOp;
+					}
+				}
+
+				NewOp = nop;
+			}
+			break;
+		}
+
+		case EOpType::ME_ADDTAGS:
+		{
+			// Add the tags after layout
+			Ptr<ASTOpMeshAddTags> NewAddTags = mu::Clone<ASTOpMeshAddTags>(SourceAt);
+
+			if (NewAddTags->Source)
+			{
+				Ptr<ASTOpMeshProject> NewAt = mu::Clone<ASTOpMeshProject>(this);
+				NewAt->Mesh = NewAddTags->Source.child();
+				NewAddTags->Source = NewAt;
+			}
+
+			NewOp = NewAddTags;
+			break;
+		}
+
+		default:
+			break;
+		}
+
+		return NewOp;
+	}
+
+
+	FSourceDataDescriptor ASTOpMeshProject::GetSourceDataDescriptor(FGetSourceDataDescriptorContext* Context) const
+	{
+		if (Mesh)
+		{
+			return Mesh->GetSourceDataDescriptor(Context);
+		}
+
+		return {};
+	}
+}
diff --git a/Engine/Plugins/Mutable/Source/MutableTools/Private/MuT/ASTOpMeshProject.h b/Engine/Plugins/Mutable/Source/MutableTools/Private/MuT/ASTOpMeshProject.h
new file mode 100644
--- /dev/null
+++ b/Engine/Plugins/Mutable/Source/MutableTools/Private/MuT/ASTOpMeshProject.h
@@ -0,0 +1,42 @@
+// Copyright Epic Games, Inc. All Rights Reserved.
+
+#pragma once
+
+#include "HAL/Platform.h"
+#include "MuR/Operations.h"
+#include "MuR/Ptr.h"
+#include "MuT/AST.h"
+
+
+namespace mu
+{
+	struct FProgram;
+
+	/** */
+	class ASTOpMeshProject final : public ASTOp
+	{
+	public:
+
+		ASTChild Mesh;
+		ASTChild Projector;
+
+	public:
+
+		ASTOpMeshProject();
+		ASTOpMeshProject(const ASTOpMeshProject&) = delete;
+		~ASTOpMeshProject();
+
+		// ASTOp interface
+		virtual EOpType GetOpType() const override { return EOpType::ME_PROJECT; }
+		virtual uint64 Hash() const override;
+		virtual bool IsEqual(const ASTOp&) const override;
+		virtual Ptr<ASTOp> Clone(MapChildFuncRef) const override;
+		virtual void ForEachChild(const TFunctionRef<void(ASTChild&)>) override;
+		virtual void Link(FProgram&, FLinkerOptions*) override;
+		virtual Ptr<ASTOp> OptimiseSink(const FModelOptimizationOptions&, FOptimizeSinkContext&) const override;
+		virtual FSourceDataDescriptor GetSourceDataDescriptor(FGetSourceDataDescriptorContext*) const override;
+	};
+
+
+}
+
diff --git a/Engine/Plugins/Mutable/Source/MutableTools/Private/MuT/ASTOpMeshSetSkeleton.cpp b/Engine/Plugins/Mutable/Source/MutableTools/Private/MuT/ASTOpMeshSetSkeleton.cpp
new file mode 100644
--- /dev/null
+++ b/Engine/Plugins/Mutable/Source/MutableTools/Private/MuT/ASTOpMeshSetSkeleton.cpp
@@ -0,0 +1,92 @@
+// Copyright Epic Games, Inc. All Rights Reserved.
+
+#include "MuT/ASTOpMeshSetSkeleton.h"
+#include "MuT/ASTOpMeshAddTags.h"
+
+#include "HAL/PlatformMath.h"
+#include "MuR/ModelPrivate.h"
+#include "MuR/RefCounted.h"
+#include "MuR/Types.h"
+
+namespace mu
+{
+
+	ASTOpMeshSetSkeleton::ASTOpMeshSetSkeleton()
+		: Source(this)
+		, Skeleton(this)
+	{
+	}
+
+
+	ASTOpMeshSetSkeleton::~ASTOpMeshSetSkeleton()
+	{
+		// Explicit call needed to avoid recursive destruction
+		ASTOp::RemoveChildren();
+	}
+
+
+	bool ASTOpMeshSetSkeleton::IsEqual(const ASTOp& OtherUntyped) const
+	{
+		if (OtherUntyped.GetOpType()==GetOpType())
+		{
+			const ASTOpMeshSetSkeleton* Other = static_cast<const ASTOpMeshSetSkeleton*>(&OtherUntyped);
+			return Source == Other->Source &&
+				Skeleton == Other->Skeleton;
+		}
+		return false;
+	}
+
+
+	uint64 ASTOpMeshSetSkeleton::Hash() const
+	{
+		uint64 Result = std::hash<void*>()(Source.child().get());
+		hash_combine(Result, Skeleton.child().get());
+		return Result;
+	}
+
+
+	Ptr<ASTOp> ASTOpMeshSetSkeleton::Clone(MapChildFuncRef MapChild) const
+	{
+		Ptr<ASTOpMeshSetSkeleton> New = new ASTOpMeshSetSkeleton();
+		New->Source = MapChild(Source.child());
+		New->Skeleton = MapChild(Skeleton.child());
+		return New;
+	}
+
+
+	void ASTOpMeshSetSkeleton::ForEachChild(const TFunctionRef<void(ASTChild&)> Func)
+	{
+		Func(Source);
+		Func(Skeleton);
+	}
+
+
+	void ASTOpMeshSetSkeleton::Link(FProgram& Program, FLinkerOptions*)
+	{
+		// Already linked?
+		if (!linkedAddress)
+		{
+			OP::MeshApplyPoseArgs Args = {};
+
+			if (Source) Args.base = Source->linkedAddress;
+			if (Skeleton) Args.pose = Skeleton->linkedAddress;
+
+			linkedAddress = (OP::ADDRESS)Program.OpAddress.Num();
+			Program.OpAddress.Add(Program.ByteCode.Num());
+			AppendCode(Program.ByteCode, GetOpType());
+			AppendCode(Program.ByteCode, Args);
+		}
+
+	}
+
+
+	FSourceDataDescriptor ASTOpMeshSetSkeleton::GetSourceDataDescriptor(FGetSourceDataDescriptorContext* Context) const
+	{
+		if (Source)
+		{
+			return Source->GetSourceDataDescriptor(Context);
+		}
+
+		return {};
+	}
+}
diff --git a/Engine/Plugins/Mutable/Source/MutableTools/Private/MuT/ASTOpMeshSetSkeleton.h b/Engine/Plugins/Mutable/Source/MutableTools/Private/MuT/ASTOpMeshSetSkeleton.h
new file mode 100644
--- /dev/null
+++ b/Engine/Plugins/Mutable/Source/MutableTools/Private/MuT/ASTOpMeshSetSkeleton.h
@@ -0,0 +1,39 @@
+// Copyright Epic Games, Inc. All Rights Reserved.
+
+#pragma once
+
+#include "HAL/Platform.h"
+#include "MuR/Operations.h"
+#include "MuR/Ptr.h"
+#include "MuT/AST.h"
+
+namespace mu
+{
+	struct FProgram;
+
+	/** */
+	class ASTOpMeshSetSkeleton final : public ASTOp
+	{
+	public:
+
+		ASTChild Source;
+		ASTChild Skeleton;
+
+	public:
+
+		ASTOpMeshSetSkeleton();
+		ASTOpMeshSetSkeleton(const ASTOpMeshSetSkeleton&) = delete;
+		~ASTOpMeshSetSkeleton();
+
+		// ASTOp interface
+		virtual EOpType GetOpType() const override { return EOpType::ME_SETSKELETON; }
+		virtual uint64 Hash() const override;
+		virtual bool IsEqual(const ASTOp&) const override;
+		virtual Ptr<ASTOp> Clone(MapChildFuncRef) const override;
+		virtual void ForEachChild(const TFunctionRef<void(ASTChild&)>) override;
+		virtual void Link(FProgram&, FLinkerOptions*) override;
+		virtual FSourceDataDescriptor GetSourceDataDescriptor(FGetSourceDataDescriptorContext*) const override;
+	};
+
+}
+
diff --git a/Engine/Plugins/Mutable/Source/MutableTools/Private/MuT/CodeGenerator.cpp b/Engine/Plugins/Mutable/Source/MutableTools/Private/MuT/CodeGenerator.cpp
--- a/Engine/Plugins/Mutable/Source/MutableTools/Private/MuT/CodeGenerator.cpp
+++ b/Engine/Plugins/Mutable/Source/MutableTools/Private/MuT/CodeGenerator.cpp
@@ -38,6 +38,9 @@
 #include "MuT/ASTOpMeshMorph.h"
 #include "MuT/ASTOpMeshOptimizeSkinning.h"
 #include "MuT/ASTOpMeshExtractLayoutBlocks.h"
+#include "MuT/ASTOpMeshApplyLayout.h"
+#include "MuT/ASTOpMeshMerge.h"
+#include "MuT/ASTOpMeshMaskDiff.h"
 #include "MuT/ASTOpParameter.h"
 #include "MuT/ASTOpLayoutRemoveBlocks.h"
 #include "MuT/ASTOpLayoutFromMesh.h"
@@ -1088,11 +1091,10 @@ namespace mu
 
 				// Create the expression to apply the layout to the mesh
 				{
-					Ptr<ASTOpFixed> ApplyLayoutOp = new ASTOpFixed();
-					ApplyLayoutOp->op.type = EOpType::ME_APPLYLAYOUT;
-					ApplyLayoutOp->SetChild(ApplyLayoutOp->op.args.MeshApplyLayout.mesh, LastMeshOp);
-					ApplyLayoutOp->SetChild(ApplyLayoutOp->op.args.MeshApplyLayout.layout, SurfaceLayoutOps[LayoutIndex]);
-					ApplyLayoutOp->op.args.MeshApplyLayout.channel = (uint16)LayoutIndex;
+					Ptr<ASTOpMeshApplyLayout> ApplyLayoutOp = new ASTOpMeshApplyLayout();
+					ApplyLayoutOp->Mesh = LastMeshOp;
+					ApplyLayoutOp->Layout = SurfaceLayoutOps[LayoutIndex];
+					ApplyLayoutOp->Channel = (uint16)LayoutIndex;
 					
 					LastMeshOp = ApplyLayoutOp;
 				}
@@ -1626,14 +1628,13 @@ namespace mu
                 // We add the merge op even for the first mesh, so that we set the surface id.
 				Ptr<ASTOp> mergeAd;
 				{
-                    Ptr<ASTOp> added = its.ResultMeshOp;
+                    Ptr<ASTOp> Added = its.ResultMeshOp;
  
-                    Ptr<ASTOpFixed> mop = new ASTOpFixed();
-                    mop->op.type = EOpType::ME_MERGE;
-                    mop->SetChild(mop->op.args.MeshMerge.base, lastMeshOp );
-                    mop->SetChild(mop->op.args.MeshMerge.added, added );
-                    mop->op.args.MeshMerge.newSurfaceID = surfaceID;
-                    mergeAd = mop;
+                    Ptr<ASTOpMeshMerge> MergeOp = new ASTOpMeshMerge();
+					MergeOp->Base = lastMeshOp;
+					MergeOp->Added = Added;
+					MergeOp->NewSurfaceID = surfaceID;
+                    mergeAd = MergeOp;
                 }
 
                 if (SurfaceConditionOp)
@@ -2115,12 +2116,11 @@ namespace mu
 					data.MeshFragment = AddedMeshOp;
 					BaseMeshResult.ExtraMeshLayouts[EditIndex] = data;
 
-					Ptr<ASTOpFixed> mop = new ASTOpFixed();
-					mop->op.type = EOpType::ME_MERGE;
-					mop->SetChild(mop->op.args.MeshMerge.base, LastMeshOp);
-					mop->SetChild(mop->op.args.MeshMerge.added, AddedMeshOp);
+					Ptr<ASTOpMeshMerge> MergeOp = new ASTOpMeshMerge();
+					MergeOp->Base = LastMeshOp;
+					MergeOp->Added = AddedMeshOp;
 					// will merge the meshes under the same surface
-					mop->op.args.MeshMerge.newSurfaceID = 0;
+					MergeOp->NewSurfaceID = 0;
 
 					// Condition to apply
 					if (m.FinalCondition)
@@ -2128,13 +2128,13 @@ namespace mu
 						Ptr<ASTOpConditional> conditionalAd = new ASTOpConditional();
 						conditionalAd->type = EOpType::ME_CONDITIONAL;
 						conditionalAd->no = LastMeshOp;
-						conditionalAd->yes = mop;
+						conditionalAd->yes = MergeOp;
 						conditionalAd->condition = m.FinalCondition;
 						LastMeshOp = conditionalAd;
 					}
 					else
 					{
-						LastMeshOp = mop;
+						LastMeshOp = MergeOp;
 					}
 				}
 
@@ -2171,13 +2171,12 @@ namespace mu
 
 					GenerateMesh(RemoveMeshOptions, removeResults, pRemove);
 
-					Ptr<ASTOpFixed> maskOp = new ASTOpFixed();
-					maskOp->op.type = EOpType::ME_MASKDIFF;
+					Ptr<ASTOpMeshMaskDiff> MaskOp = new ASTOpMeshMaskDiff();
 
 					// By default, remove from the base
-					Ptr<ASTOp> removeFrom = BaseMeshResult.BaseMeshOp;
-					maskOp->SetChild(maskOp->op.args.MeshMaskDiff.source, removeFrom);
-					maskOp->SetChild(maskOp->op.args.MeshMaskDiff.fragment, removeResults.MeshOp);
+					Ptr<ASTOp> RemoveFrom = BaseMeshResult.BaseMeshOp;
+					MaskOp->Source = RemoveFrom;
+					MaskOp->Fragment = removeResults.MeshOp;
 
 					if (!RemoveOp)
 					{
@@ -2187,7 +2186,7 @@ namespace mu
 						LastMeshOp = RemoveOp;
 					}
 
-					RemoveOp->AddRemove(m.FinalCondition, maskOp);
+					RemoveOp->AddRemove(m.FinalCondition, MaskOp);
 				}
 			}
 		}
@@ -2412,7 +2411,7 @@ namespace mu
 			{
 				const NodeModifierMeshClipMorphPlane* TypedNode = static_cast<const NodeModifierMeshClipMorphPlane*>(m.Node);
 				Ptr<ASTOpMeshClipMorphPlane> op = new ASTOpMeshClipMorphPlane();
-				op->source = LastMeshOp;
+				op->Source = LastMeshOp;
 				op->FaceCullStrategy = TypedNode->Parameters.FaceCullStrategy;
 
 				// Morph to an ellipse
@@ -2437,7 +2436,7 @@ namespace mu
 
 						morphShape.side = FVector3f::CrossProduct(TypedNode->Parameters.Normal, aux_base);
 					}
-					op->morphShape = morphShape;
+					op->MorphShape = morphShape;
 				}
 
 				// Selection box
@@ -2448,17 +2447,17 @@ namespace mu
 					selectionShape.type = (uint8)FShape::Type::AABox;
 					selectionShape.position = TypedNode->Parameters.SelectionBoxOrigin;
 					selectionShape.size = TypedNode->Parameters.SelectionBoxRadius;
-					op->selectionShape = selectionShape;
+					op->SelectionShape = selectionShape;
 				}
 				else if (op->VertexSelectionType == EClipVertexSelectionType::BoneHierarchy)
 				{
-					op->vertexSelectionBone = TypedNode->Parameters.VertexSelectionBone;
-					op->vertexSelectionBoneMaxRadius = TypedNode->Parameters.MaxEffectRadius;
+					op->VertexSelectionBone = TypedNode->Parameters.VertexSelectionBone;
+					op->VertexSelectionBoneMaxRadius = TypedNode->Parameters.MaxEffectRadius;
 				}
 
 				// FParameters
-				op->dist = TypedNode->Parameters.DistanceToPlane;
-				op->factor = TypedNode->Parameters.LinearityFactor;
+				op->Dist = TypedNode->Parameters.DistanceToPlane;
+				op->Factor = TypedNode->Parameters.LinearityFactor;
 
 				modifiedMeshOp = op;
 
diff --git a/Engine/Plugins/Mutable/Source/MutableTools/Private/MuT/CodeGenerator_Image.cpp b/Engine/Plugins/Mutable/Source/MutableTools/Private/MuT/CodeGenerator_Image.cpp
--- a/Engine/Plugins/Mutable/Source/MutableTools/Private/MuT/CodeGenerator_Image.cpp
+++ b/Engine/Plugins/Mutable/Source/MutableTools/Private/MuT/CodeGenerator_Image.cpp
@@ -33,6 +33,7 @@
 #include "MuT/ASTOpImageCrop.h"
 #include "MuT/ASTOpMeshExtractLayoutBlocks.h"
 #include "MuT/ASTOpMeshFormat.h"
+#include "MuT/ASTOpMeshProject.h"
 #include "MuT/ASTOpParameter.h"
 #include "MuT/ASTOpSwitch.h"
 #include "MuT/CodeGenerator.h"
@@ -1387,8 +1388,7 @@ namespace mu
 
         // Mesh project operation
         //------------------------------
-        Ptr<ASTOpFixed> ProjectOp = new ASTOpFixed();
-		ProjectOp->op.type = EOpType::ME_PROJECT;
+        Ptr<ASTOpMeshProject> ProjectOp = new ASTOpMeshProject();
 
         Ptr<ASTOp> LastMeshOp = ProjectOp;
 
@@ -1405,7 +1405,7 @@ namespace mu
             GenerateMissingProjectorCode( projectorResult, InNode->GetMessageContext() );
         }
 
-		ProjectOp->SetChild(ProjectOp->op.args.MeshProject.projector, projectorResult.op );
+		ProjectOp->Projector = projectorResult.op;
 
 		int32 LayoutBlockIndex = -1;
 		if (Options.LayoutToApply)
@@ -1505,7 +1505,7 @@ namespace mu
 				CurrentMeshToProjectOp = FormatOp;
             }
 
-			ProjectOp->SetChild(ProjectOp->op.args.MeshProject.mesh, CurrentMeshToProjectOp);
+			ProjectOp->Mesh = CurrentMeshToProjectOp;
 
         }
         else
@@ -1515,7 +1515,7 @@ namespace mu
             Ptr<ASTOpConstantResource> cop = new ASTOpConstantResource();
             cop->Type = EOpType::ME_CONSTANT;
             cop->SetValue(TempMesh, CompilerOptions->OptimisationOptions.DiskCacheContext);
-			ProjectOp->SetChild(ProjectOp->op.args.MeshProject.mesh, cop );
+			ProjectOp->Mesh = cop;
             ErrorLog->Add( "Projector mesh not set.", ELMT_ERROR, InNode->GetMessageContext() );
         }
 
diff --git a/Engine/Plugins/Mutable/Source/MutableTools/Private/MuT/CodeGenerator_Mesh.cpp b/Engine/Plugins/Mutable/Source/MutableTools/Private/MuT/CodeGenerator_Mesh.cpp
--- a/Engine/Plugins/Mutable/Source/MutableTools/Private/MuT/CodeGenerator_Mesh.cpp
+++ b/Engine/Plugins/Mutable/Source/MutableTools/Private/MuT/CodeGenerator_Mesh.cpp
@@ -23,6 +23,7 @@
 #include "MuT/ASTOpMeshBindShape.h"
 #include "MuT/ASTOpMeshClipDeform.h"
 #include "MuT/ASTOpMeshClipMorphPlane.h"
+#include "MuT/ASTOpMeshClipWithMesh.h"
 #include "MuT/ASTOpMeshExtractLayoutBlocks.h"
 #include "MuT/ASTOpMeshFormat.h"
 #include "MuT/ASTOpMeshMorphReshape.h"
@@ -1324,7 +1325,7 @@ namespace mu
         {
 			FMeshGenerationOptions BaseOptions = InOptions;
             GenerateMesh(BaseOptions, OutResult, ClipNode->Source);
-            op->source = OutResult.MeshOp;
+            op->Source = OutResult.MeshOp;
         }
         else
         {
@@ -1334,10 +1335,10 @@ namespace mu
 
         // Morph to an ellipse
         {
-            op->morphShape.type = (uint8_t)FShape::Type::Ellipse;
-            op->morphShape.position = ClipNode->Parameters.Origin;
-            op->morphShape.up = ClipNode->Parameters.Normal;
-            op->morphShape.size = FVector3f(ClipNode->Parameters.Radius1, ClipNode->Parameters.Radius2, ClipNode->Parameters.Rotation); // TODO: Move rotation to ellipse rotation reference base instead of passing it directly
+            op->MorphShape.type = (uint8_t)FShape::Type::Ellipse;
+            op->MorphShape.position = ClipNode->Parameters.Origin;
+            op->MorphShape.up = ClipNode->Parameters.Normal;
+            op->MorphShape.size = FVector3f(ClipNode->Parameters.Radius1, ClipNode->Parameters.Radius2, ClipNode->Parameters.Rotation); // TODO: Move rotation to ellipse rotation reference base instead of passing it directly
 
             // Generate a "side" vector.
             // \todo: make generic and move to the vector class
@@ -1350,7 +1351,7 @@ namespace mu
                     aux_base = FVector3f(0.f, 0.f, 1.f);
                 }
 
-                op->morphShape.side = FVector3f::CrossProduct(ClipNode->Parameters.Normal, aux_base);
+                op->MorphShape.side = FVector3f::CrossProduct(ClipNode->Parameters.Normal, aux_base);
             }
         }
 
@@ -1358,20 +1359,20 @@ namespace mu
 		op->VertexSelectionType = ClipNode->Parameters.VertexSelectionType;
         if (op->VertexSelectionType == EClipVertexSelectionType::Shape)
         {
-            op->selectionShape.type = (uint8_t)FShape::Type::AABox;
-            op->selectionShape.position = ClipNode->Parameters.SelectionBoxOrigin;
-            op->selectionShape.size = ClipNode->Parameters.SelectionBoxRadius;
+            op->SelectionShape.type = (uint8_t)FShape::Type::AABox;
+            op->SelectionShape.position = ClipNode->Parameters.SelectionBoxOrigin;
+            op->SelectionShape.size = ClipNode->Parameters.SelectionBoxRadius;
         }
         else if (op->VertexSelectionType == EClipVertexSelectionType::BoneHierarchy)
         {
             // Selection by bone hierarchy?
-            op->vertexSelectionBone = ClipNode->Parameters.VertexSelectionBone;
-			op->vertexSelectionBoneMaxRadius = ClipNode->Parameters.MaxEffectRadius;
+            op->VertexSelectionBone = ClipNode->Parameters.VertexSelectionBone;
+			op->VertexSelectionBoneMaxRadius = ClipNode->Parameters.MaxEffectRadius;
         }
  
         // FParameters
-        op->dist = ClipNode->Parameters.DistanceToPlane;
-        op->factor = ClipNode->Parameters.LinearityFactor;
+        op->Dist = ClipNode->Parameters.DistanceToPlane;
+        op->Factor = ClipNode->Parameters.LinearityFactor;
 
         OutResult.MeshOp = op;
     }
@@ -1381,14 +1382,13 @@ namespace mu
     void CodeGenerator::GenerateMesh_ClipWithMesh(const FMeshGenerationOptions& InOptions, FMeshGenerationResult& OutResult,
                                                    const NodeMeshClipWithMesh* ClipNode)
     {
-        Ptr<ASTOpFixed> op = new ASTOpFixed();
-        op->op.type = EOpType::ME_CLIPWITHMESH;
+        Ptr<ASTOpMeshClipWithMesh> Op = new ASTOpMeshClipWithMesh();
 
         // Base
         if (ClipNode->Source)
         {
             GenerateMesh(InOptions, OutResult, ClipNode->Source );
-            op->SetChild( op->op.args.MeshClipWithMesh.source, OutResult.MeshOp );
+			Op->Source = OutResult.MeshOp;
         }
         else
         {
@@ -1406,7 +1406,7 @@ namespace mu
 
             FMeshGenerationResult clipResult;
             GenerateMesh(ClipOptions, clipResult, ClipNode->ClipMesh);
-            op->SetChild( op->op.args.MeshClipWithMesh.clipMesh, clipResult.MeshOp );
+            Op->ClipMesh = clipResult.MeshOp;
 		}
         else
         {
@@ -1414,7 +1414,7 @@ namespace mu
             ErrorLog->Add("Mesh clip-with-mesh clipping mesh node is not set.", ELMT_ERROR, ClipNode->GetMessageContext());
         }
 
-        OutResult.MeshOp = op;
+        OutResult.MeshOp = Op;
     }
 
 	//---------------------------------------------------------------------------------------------
@@ -1466,7 +1466,7 @@ namespace mu
         if (node.m_pBase)
         {
             GenerateMesh(InOptions, OutResult, node.m_pBase );
-            op->base = OutResult.MeshOp;
+            op->Base = OutResult.MeshOp;
         }
         else
         {
@@ -1482,9 +1482,9 @@ namespace mu
 			PoseOptions.OverrideLayouts.Empty();
 			PoseOptions.ActiveTags.Empty();
 
-            FMeshGenerationResult poseResult;
-            GenerateMesh(PoseOptions, poseResult, node.m_pPose );
-            op->pose = poseResult.MeshOp;
+            FMeshGenerationResult PoseResult;
+            GenerateMesh(PoseOptions, PoseResult, node.m_pPose );
+            op->Pose = PoseResult.MeshOp;
 		}
         else
         {
diff --git a/Engine/Plugins/Mutable/Source/MutableTools/Private/MuT/CodeOptimiser.cpp b/Engine/Plugins/Mutable/Source/MutableTools/Private/MuT/CodeOptimiser.cpp
--- a/Engine/Plugins/Mutable/Source/MutableTools/Private/MuT/CodeOptimiser.cpp
+++ b/Engine/Plugins/Mutable/Source/MutableTools/Private/MuT/CodeOptimiser.cpp
@@ -1234,7 +1234,7 @@ namespace mu
 				if (typedOp->VertexSelectionType == EClipVertexSelectionType::BoneHierarchy)
 				{
 					// We need the skeleton for the source mesh
-					RecurseWithState( typedOp->source.child(), true );
+					RecurseWithState( typedOp->Source.child(), true );
 					return false;
 				}
 
@@ -1246,8 +1246,8 @@ namespace mu
 				ASTOpMeshApplyPose* typedOp = static_cast<ASTOpMeshApplyPose*>(node.get());
 
 				// We need the skeleton for both meshes
-				RecurseWithState(typedOp->base.child(), true);
-				RecurseWithState(typedOp->pose.child(), true);
+				RecurseWithState(typedOp->Base.child(), true);
+				RecurseWithState(typedOp->Pose.child(), true);
 				return false;
 
 				break;
@@ -1401,31 +1401,6 @@ namespace mu
 			}
 			return true;
 		});
-
-	//    for( auto& rep: replacementsFound )
-	//    {
-	//        check(rep.AddMeshOp->GetOpType()==EOpType::IN_ADDMESH);
-	//        auto newAddMeshRootTyped = dynamic_cast<ASTOpInstanceAdd*>(rep.AddMeshOp.get());
-	//        auto meshRoot = newAddMeshRootTyped->value.child();
-
-	//        // Add new skeleton constant instruction
-	//        Ptr<ASTOpConstantResource> cop = new ASTOpConstantResource();
-	//        cop->type = EOpType::ME_CONSTANT;
-	//        cop->SetValue( rep.FinalSkeleton, options.bUseDiskCache );
-
-	//        // Add new "apply skeleton" instruction
-	//        Ptr<ASTOpFixed> skop = new ASTOpFixed();
-	//        skop->op.type = EOpType::ME_SETSKELETON;
-	//        skop->SetChild( skop->op.args.MeshSetSkeleton.source, meshRoot);
-	//        skop->SetChild( skop->op.args.MeshSetSkeleton.skeleton, cop);
-
-	//        // We need a new mesh add instruction, it is not safe to edit the current one.
-	//        auto newOp = mu::Clone<ASTOpInstanceAdd>(newAddMeshRootTyped);
-	//        newOp->value = skop;
-
-	//        // TODO BLEH UNSAFE?
-	//        ASTOp::Replace(rep.AddMeshOp,newOp);
-	//    }
 	}
 
 
diff --git a/Engine/Plugins/Mutable/Source/MutableTools/Private/MuT/CodeOptimiser_Logic.cpp b/Engine/Plugins/Mutable/Source/MutableTools/Private/MuT/CodeOptimiser_Logic.cpp
--- a/Engine/Plugins/Mutable/Source/MutableTools/Private/MuT/CodeOptimiser_Logic.cpp
+++ b/Engine/Plugins/Mutable/Source/MutableTools/Private/MuT/CodeOptimiser_Logic.cpp
@@ -9,6 +9,7 @@
 #include "MuT/ASTOpConditional.h"
 #include "MuT/ASTOpInstanceAdd.h"
 #include "MuT/ASTOpSwitch.h"
+#include "MuT/ASTOpMeshMerge.h"
 #include "MuT/CodeOptimiser.h"
 
 
@@ -1080,13 +1081,13 @@ namespace mu
                     Ptr<ASTOp> yes = topConditional->yes.child();
                     if ( yes && yes->GetOpType()==EOpType::ME_MERGE )
                     {
-                        Ptr<ASTOpFixed> merge = static_cast<ASTOpFixed*>(yes.get());
+                        Ptr<ASTOpMeshMerge> Merge = static_cast<ASTOpMeshMerge*>(yes.get());
 
                         bool ended = false;
                         while (!ended)
                         {
                             ended = true;
-							Ptr<ASTOp> base = merge->children[merge->op.args.MeshMerge.base].child();
+							Ptr<ASTOp> base = Merge->Base.child();
                             if ( base && base->GetOpType()==EOpType::ME_CONDITIONAL )
                             {
 								ASTOpConditional* bottomConditional = static_cast<ASTOpConditional*>(base.get());
@@ -1107,19 +1108,19 @@ namespace mu
 
                                 if (conditionaAreExclusive)
                                 {
-                                    if (merge->GetParentCount()==1)
+                                    if (Merge->GetParentCount()==1)
                                     {
                                         // Directly modify the instruction to skip the impossible child option
-                                        merge->SetChild(merge->op.args.MeshMerge.base, bottomConditional->no.child() );
+										Merge->Base = bottomConditional->no.child();
                                     }
                                     else
                                     {
                                         // Other parents may not impose the same condition that allows
                                         // the optimisation.
-                                        Ptr<ASTOpFixed> newMerge = mu::Clone<ASTOpFixed>(merge);
-                                        newMerge->SetChild(newMerge->op.args.MeshMerge.base, bottomConditional->no.child());
-                                        topConditional->yes = newMerge;
-                                        merge = newMerge;
+                                        Ptr<ASTOpMeshMerge> NewMerge = mu::Clone<ASTOpMeshMerge>(Merge);
+										NewMerge->Base = bottomConditional->no.child();
+                                        topConditional->yes = NewMerge;
+                                        Merge = NewMerge;
                                     }
 
                                     modified = true;
diff --git a/Engine/Plugins/Mutable/Source/MutableTools/Private/MuT/CodeOptimiser_Semantic.cpp b/Engine/Plugins/Mutable/Source/MutableTools/Private/MuT/CodeOptimiser_Semantic.cpp
--- a/Engine/Plugins/Mutable/Source/MutableTools/Private/MuT/CodeOptimiser_Semantic.cpp
+++ b/Engine/Plugins/Mutable/Source/MutableTools/Private/MuT/CodeOptimiser_Semantic.cpp
@@ -398,85 +398,6 @@ namespace mu
 			break;
 		}
 
-		//-------------------------------------------------------------------------------------
-		case EOpType::ME_MASKDIFF:
-		{
-			Ptr<ASTOp> Fragment = children[op.args.MeshMaskDiff.fragment].child();
-			EOpType FragmentType = Fragment->GetOpType();
-			switch (FragmentType)
-			{
-			case EOpType::ME_ADDTAGS:
-			{
-				// Tags in the fragment can be ignored.
-				const ASTOpMeshAddTags* Add = static_cast<const ASTOpMeshAddTags*>(Fragment.get());
-
-				Ptr<ASTOpFixed> NewAt = mu::Clone<ASTOpFixed>(this);
-				NewAt->SetChild(NewAt->op.args.MeshMaskDiff.fragment, Add->Source);
-				at = NewAt;
-				break;
-			}
-			case EOpType::ME_SWITCH:
-			{
-				// Move the mask diff down all the paths
-				Ptr<ASTOpSwitch> nop = mu::Clone<ASTOpSwitch>(Fragment.get());
-
-				if (nop->def)
-				{
-					Ptr<ASTOpFixed> defOp = mu::Clone<ASTOpFixed>(this);
-					defOp->SetChild(defOp->op.args.MeshMaskDiff.fragment, nop->def);
-					nop->def = defOp;
-				}
-
-				// We need to copy the options because we change them
-				for (int32 v = 0; v < nop->cases.Num(); ++v)
-				{
-					if (nop->cases[v].branch)
-					{
-						Ptr<ASTOpFixed> bOp = mu::Clone<ASTOpFixed>(this);
-						bOp->SetChild(bOp->op.args.MeshMaskDiff.fragment, nop->cases[v].branch);
-						nop->cases[v].branch = bOp;
-					}
-				}
-
-				at = nop;
-				break;
-			}
-			default:
-				break;
-			}
-
-			break;
-		}
-
-		//-------------------------------------------------------------------------------------
-		case EOpType::ME_APPLYLAYOUT:
-		{
-			Ptr<ASTOp> Base = children[op.args.MeshApplyLayout.mesh].child();
-			EOpType BaseType = Base->GetOpType();
-			switch (BaseType)
-			{
-			case EOpType::ME_ADDTAGS:
-			{
-				// Add the tags after layout
-				Ptr<ASTOpMeshAddTags> NewAddTags = mu::Clone<ASTOpMeshAddTags>(Base);
-
-				if (NewAddTags->Source)
-				{
-					Ptr<ASTOpFixed> NewAt = mu::Clone<ASTOpFixed>(this);
-					NewAt->SetChild(NewAt->op.args.MeshApplyLayout.mesh, NewAddTags->Source);
-					NewAddTags->Source = NewAt;
-				}
-
-				at = NewAddTags;
-				break;
-			}
-
-			default:
-				break;
-			}
-
-			break;
-		}
 
         default:
             break;
@@ -528,92 +449,6 @@ namespace mu
         switch ( type )
         {
 
-        case EOpType::ME_PROJECT:
-        {
-            auto sourceAt = children[op.args.MeshProject.mesh].child();
-            if (!sourceAt) { break; }
-
-            auto projectorAt = children[op.args.MeshProject.projector].child();
-            if (!projectorAt) { break; }
-
-            EOpType sourceType = sourceAt->GetOpType();
-            switch ( sourceType )
-            {
-
-            case EOpType::ME_CONDITIONAL:
-            {
-                if (projectorAt->GetOpType()==EOpType::PR_CONSTANT)
-                {
-                    // We move the project down the two paths
-                    auto nop = mu::Clone<ASTOpConditional>(sourceAt);
-
-                    auto aOp = mu::Clone<ASTOpFixed>(this);
-                    aOp->SetChild( aOp->op.args.MeshProject.mesh, nop->yes );
-                    nop->yes = aOp;
-
-                    auto bOp = mu::Clone<ASTOpFixed>(this);
-                    bOp->SetChild( bOp->op.args.MeshProject.mesh, nop->no);
-                    nop->no = bOp;
-
-                    at = nop;
-                }
-                break;
-            }
-
-            case EOpType::ME_SWITCH:
-            {
-                if (projectorAt->GetOpType()==EOpType::PR_CONSTANT)
-                {
-                    // Move the format down all the paths
-                    auto nop = mu::Clone<ASTOpSwitch>(sourceAt);
-
-                    if (nop->def)
-                    {
-                        auto defOp = mu::Clone<ASTOpFixed>(this);
-                        defOp->SetChild( defOp->op.args.MeshProject.mesh, nop->def );
-                        nop->def = defOp;
-                    }
-
-                    // We need to copy the options because we change them
-                    for ( int32 v=0; v<nop->cases.Num(); ++v )
-                    {
-                        if ( nop->cases[v].branch )
-                        {
-                            auto bOp = mu::Clone<ASTOpFixed>(this);
-                            bOp->SetChild( bOp->op.args.MeshProject.mesh, nop->cases[v].branch );
-                            nop->cases[v].branch = bOp;
-                        }
-                    }
-
-                    at = nop;
-                }
-                break;
-            }
-
-			case EOpType::ME_ADDTAGS:
-			{
-				// Add the tags after layout
-				Ptr<ASTOpMeshAddTags> NewAddTags = mu::Clone<ASTOpMeshAddTags>(sourceAt);
-
-				if (NewAddTags->Source)
-				{
-					Ptr<ASTOpFixed> NewAt = mu::Clone<ASTOpFixed>(this);
-					NewAt->SetChild(NewAt->op.args.MeshProject.mesh, NewAddTags->Source);
-					NewAddTags->Source = NewAt;
-				}
-
-				at = NewAddTags;
-				break;
-			}
-
-            default:
-                break;
-            }
-
-            break;
-        }
-
-
         //-----------------------------------------------------------------------------------------
         //-----------------------------------------------------------------------------------------
         //-----------------------------------------------------------------------------------------
diff --git a/Engine/Plugins/Mutable/Source/MutableTools/Private/MuT/DataPacker.cpp b/Engine/Plugins/Mutable/Source/MutableTools/Private/MuT/DataPacker.cpp
--- a/Engine/Plugins/Mutable/Source/MutableTools/Private/MuT/DataPacker.cpp
+++ b/Engine/Plugins/Mutable/Source/MutableTools/Private/MuT/DataPacker.cpp
@@ -25,6 +25,8 @@
 #include "MuT/ASTOpMeshRemoveMask.h"
 #include "MuT/ASTOpMeshDifference.h"
 #include "MuT/ASTOpMeshMorph.h"
+#include "MuT/ASTOpMeshProject.h"
+#include "MuT/ASTOpMeshApplyLayout.h"
 #include "MuT/ASTOpLayoutFromMesh.h"
 #include "MuT/ASTOpImageRasterMesh.h"
 #include "MuT/CompilerPrivate.h"
@@ -593,13 +595,13 @@ namespace mu
             {
                 recurse = false;
 
-				const ASTOpFixed* op = static_cast<const ASTOpFixed*>(node.get());
+				const ASTOpMeshApplyLayout* Op = static_cast<const ASTOpMeshApplyLayout*>(node.get());
 
-                uint64 newState = currentSemantics;
-                newState |= (UINT64_C(1)<< uint32(EMeshBufferSemantic::LayoutBlock));
-                RecurseWithState( op->children[op->op.args.MeshApplyLayout.mesh].child(), newState );
+                uint64 NewState = currentSemantics;
+				NewState |= (UINT64_C(1)<< uint32(EMeshBufferSemantic::LayoutBlock));
+                RecurseWithState( Op->Mesh.child(), NewState);
 
-                RecurseWithState( op->children[op->op.args.MeshApplyLayout.layout].child(), currentSemantics );
+                RecurseWithState( Op->Layout.child(), currentSemantics );
                 break;
             }
 
@@ -607,13 +609,13 @@ namespace mu
 			{
 				recurse = false;
 
-				const ASTOpFixed* op = static_cast<const ASTOpFixed*>(node.get());
+				const ASTOpMeshProject* op = static_cast<const ASTOpMeshProject*>(node.get());
 
 				uint64 newState = currentSemantics;
 				newState |= (UINT64_C(1) << uint32(EMeshBufferSemantic::LayoutBlock));
-				RecurseWithState(op->children[op->op.args.MeshProject.mesh].child(), newState);
+				RecurseWithState(op->Mesh.child(), newState);
 
-				RecurseWithState(op->children[op->op.args.MeshProject.projector].child(), currentSemantics);
+				RecurseWithState(op->Projector.child(), currentSemantics);
 				break;
 			}
 
diff --git a/Engine/Plugins/Mutable/Source/MutableTools/Private/MuT/ErrorLog.cpp b/Engine/Plugins/Mutable/Source/MutableTools/Private/MuT/ErrorLog.cpp
--- a/Engine/Plugins/Mutable/Source/MutableTools/Private/MuT/ErrorLog.cpp
+++ b/Engine/Plugins/Mutable/Source/MutableTools/Private/MuT/ErrorLog.cpp
@@ -325,7 +325,7 @@ namespace mu
 		TEXT("LA_FROMMESH	   "),
 	};
 
-	static_assert(sizeof(s_opNames) / sizeof(void*) == int32(EOpType::COUNT));
+	static_assert(UE_ARRAY_COUNT(s_opNames) == int32(EOpType::COUNT));
 
     // clang-format on
 
diff --git a/Engine/Plugins/Mutable/Source/MutableTools/Public/MuT/Node.h b/Engine/Plugins/Mutable/Source/MutableTools/Public/MuT/Node.h
--- a/Engine/Plugins/Mutable/Source/MutableTools/Public/MuT/Node.h
+++ b/Engine/Plugins/Mutable/Source/MutableTools/Public/MuT/Node.h
@@ -32,7 +32,6 @@ namespace mu
 
 			Mesh,
 			MeshConstant,
-			MeshInterpolate,
 			MeshTable,
 			MeshFormat,
 			MeshTangents,
@@ -45,7 +44,6 @@ namespace mu
 			MeshClipWithMesh,
 			MeshApplyPose,
 			MeshVariation,
-			MeshGeometryOperation,
 			MeshReshape,
 			MeshClipDeform,
 
