[Mutable] Remove BlueprintSpawnableComponent from UCustomizableObjectInstanceUsage.

It is not a component.

diff --git a/Engine/Plugins/Mutable/Source/CustomizableObject/Public/MuCO/CustomizableObjectInstanceUsage.h b/Engine/Plugins/Mutable/Source/CustomizableObject/Public/MuCO/CustomizableObjectInstanceUsage.h
--- a/Engine/Plugins/Mutable/Source/CustomizableObject/Public/MuCO/CustomizableObjectInstanceUsage.h
+++ b/Engine/Plugins/Mutable/Source/CustomizableObject/Public/MuCO/CustomizableObjectInstanceUsage.h
@@ -23,7 +23,7 @@ DECLARE_DELEGATE(FCustomizableObjectInstanceUsageUpdatedDelegate);
 // UCustomizableObjectInstance and a USkeletalComponent so that the CustomizableObjectSystem takes care of updating it and its LODs, 
 // streaming, etc. It's a UObject, so it will be much cheaper than a UCustomizableComponent as it won't have to refresh its transforms
 // every time it's moved.
-UCLASS(Blueprintable, BlueprintType, ClassGroup = (CustomizableObject), meta = (BlueprintSpawnableComponent))
+UCLASS(Blueprintable, BlueprintType, ClassGroup = (CustomizableObject))
 class CUSTOMIZABLEOBJECT_API UCustomizableObjectInstanceUsage : public UObject
 {
 	friend UCustomizableObjectInstanceUsagePrivate;
