[Mutable] Fix clothing physics assets ignored in some case when compiling the CO.

#rnx

diff --git a/Engine/Plugins/Mutable/Source/CustomizableObjectEditor/Private/MuCOE/GenerateMutableSource/GenerateMutableSourceMesh.cpp b/Engine/Plugins/Mutable/Source/CustomizableObjectEditor/Private/MuCOE/GenerateMutableSource/GenerateMutableSourceMesh.cpp
--- a/Engine/Plugins/Mutable/Source/CustomizableObjectEditor/Private/MuCOE/GenerateMutableSource/GenerateMutableSourceMesh.cpp
+++ b/Engine/Plugins/Mutable/Source/CustomizableObjectEditor/Private/MuCOE/GenerateMutableSource/GenerateMutableSourceMesh.cpp
@@ -1505,6 +1505,8 @@ TSharedPtr<mu::FMesh> ConvertSkeletalMeshToMutable(
 					return;
 				}
 
+				OutPhysicsAssetIndex = Context.PhysicsAssets.AddUnique(Asset->PhysicsAsset);
+
 				int32 FoundIndex = Context.ClothingAssetsData.IndexOfByPredicate(
 					[SourceAsset = Asset](const FCustomizableObjectClothingAssetData& Asset)
 					{
@@ -1582,12 +1584,6 @@ TSharedPtr<mu::FMesh> ConvertSkeletalMeshToMutable(
 				if (FoundIndex != INDEX_NONE)
 				{
 					OutClothingAssetIndex = FoundIndex;
-					OutPhysicsAssetIndex = Context.PhysicsAssets.IndexOfByPredicate(
-						[PhysicsAsset = Asset->PhysicsAsset](const TSoftObjectPtr<UPhysicsAsset>& OtherPhysicsAsset)
-						{
-							return TSoftObjectPtr<UPhysicsAsset>(PhysicsAsset) == OtherPhysicsAsset;
-						});
-
 					return;
 				}
 
