[Mutable] Add missing cook scopes to remove editor-only assets from packaged builds.

#rnx

diff --git a/Engine/Plugins/Mutable/Source/CustomizableObjectEditor/Private/MuCOE/CustomizableObjectCompiler.cpp b/Engine/Plugins/Mutable/Source/CustomizableObjectEditor/Private/MuCOE/CustomizableObjectCompiler.cpp
--- a/Engine/Plugins/Mutable/Source/CustomizableObjectEditor/Private/MuCOE/CustomizableObjectCompiler.cpp
+++ b/Engine/Plugins/Mutable/Source/CustomizableObjectEditor/Private/MuCOE/CustomizableObjectCompiler.cpp
@@ -117,6 +117,8 @@ void FCustomizableObjectCompiler::PreloadingReferencerAssetsCallback(bool bAsync
 	
 	if (AsynchronousStreamableHandlePtr)
 	{
+		FCookLoadScope CookLoadScope(ECookLoadType::EditorOnly);
+
 		TArray<FSoftObjectPath> AssetsToStream;
 		AsynchronousStreamableHandlePtr->GetRequestedAssets(AssetsToStream);
 
@@ -369,6 +371,8 @@ void ProcessChildObjectsRecursively(const UCustomizableObject* ParentObject, FMu
 		return;
 	}
 	
+	FCookLoadScope CookLoadScope(ECookLoadType::EditorOnly);
+	
 	// Required to be deterministic.
 	ReferencedObjectNames.Sort([](const FName& A, const FName& B)
 	{
@@ -1038,7 +1042,7 @@ void FCustomizableObjectCompiler::CompileInternal(bool bAsync)
 		for (const TPair<TSoftObjectPtr<const UTexture>, FMutableGraphGenerationContext::FGeneratedReferencedTexture>& Pair : GenerationContext.CompileTimeTextureMap)
 		{
 			check(Pair.Value.ID == NewCompileTimeReferencedTextures.Num());
-
+			FCookLoadScope CookLoadScope(ECookLoadType::EditorOnly);
 			FMutableSourceTextureData Tex(*Pair.Key.LoadSynchronous());
 			NewCompileTimeReferencedTextures.Add(Tex);
 		}
diff --git a/Engine/Plugins/Mutable/Source/CustomizableObjectEditor/Private/MuCOE/GenerateMutableSource/GenerateMutableSourceComponent.cpp b/Engine/Plugins/Mutable/Source/CustomizableObjectEditor/Private/MuCOE/GenerateMutableSource/GenerateMutableSourceComponent.cpp
--- a/Engine/Plugins/Mutable/Source/CustomizableObjectEditor/Private/MuCOE/GenerateMutableSource/GenerateMutableSourceComponent.cpp
+++ b/Engine/Plugins/Mutable/Source/CustomizableObjectEditor/Private/MuCOE/GenerateMutableSource/GenerateMutableSourceComponent.cpp
@@ -210,8 +210,8 @@ mu::Ptr<mu::NodeComponent> GenerateMutableSourceComponent(const UEdGraphPin* Pin
 			GenerationContext.Log(FText::FromString(Msg), TypedComponentPassthroughMesh, EMessageSeverity::Warning);
 			return nullptr;
 		}
-
-		USkeletalMesh* SkeletalMesh = Cast<USkeletalMesh>(TypedComponentPassthroughMesh->Mesh.TryLoad());
+		
+		USkeletalMesh* SkeletalMesh = Cast<USkeletalMesh>(GenerationContext.LoadObject(TSoftObjectPtr<USkeletalMesh>(TypedComponentPassthroughMesh->Mesh)));
 		if (!SkeletalMesh)
 		{
 			FString Msg = FString::Printf(TEXT("Only SkeletalMeshes are supported in this node, for now."));
diff --git a/Engine/Plugins/Mutable/Source/CustomizableObjectEditor/Private/MuCOE/GenerateMutableSource/GenerateMutableSourceTable.cpp b/Engine/Plugins/Mutable/Source/CustomizableObjectEditor/Private/MuCOE/GenerateMutableSource/GenerateMutableSourceTable.cpp
--- a/Engine/Plugins/Mutable/Source/CustomizableObjectEditor/Private/MuCOE/GenerateMutableSource/GenerateMutableSourceTable.cpp
+++ b/Engine/Plugins/Mutable/Source/CustomizableObjectEditor/Private/MuCOE/GenerateMutableSource/GenerateMutableSourceTable.cpp
@@ -28,6 +28,7 @@
 #include "Rendering/SkeletalMeshModel.h"
 #include "MuCOE/CustomizableObjectVersionBridge.h"
 #include "ClothingAsset.h"
+#include "UObject/ICookInfo.h"
 
 #define LOCTEXT_NAMESPACE "CustomizableObjectEditor"
 
@@ -937,6 +938,8 @@ bool GenerateTableColumn(const UCustomizableObjectNodeTable* TableNode, const UE
 	// Pre-pass to request async loading of all data. This seems to be slightly faster because it avoids the sync after every separate load.
 	if (!GenerationContext.bParticipatingObjectsPass)
 	{
+		FCookLoadScope CookLoadScope(ECookLoadType::EditorOnly);
+
 		TArray<int32> LoadRequests;
 		LoadRequests.Reserve(RowNames.Num());
 
@@ -1256,6 +1259,8 @@ UDataTable* GenerateDataTableFromStruct(const UCustomizableObjectNodeTable* Tabl
 		return nullptr;
 	}
 
+	FCookLoadScope CookLoadScope(ECookLoadType::EditorOnly);
+
 	TArray<FAssetData> DataTableAssets = TableNode->GetParentTables();
 
 	UCompositeDataTable* CompositeDataTable = NewObject<UCompositeDataTable>();
diff --git a/Engine/Plugins/Mutable/Source/CustomizableObjectEditor/Private/MuCOE/GraphTraversal.cpp b/Engine/Plugins/Mutable/Source/CustomizableObjectEditor/Private/MuCOE/GraphTraversal.cpp
--- a/Engine/Plugins/Mutable/Source/CustomizableObjectEditor/Private/MuCOE/GraphTraversal.cpp
+++ b/Engine/Plugins/Mutable/Source/CustomizableObjectEditor/Private/MuCOE/GraphTraversal.cpp
@@ -25,6 +25,7 @@
 #include "MuCOE/Nodes/CustomizableObjectNodeObjectGroup.h"
 #include "MuCOE/Nodes/CustomizableObjectNodeReroute.h"
 #include "UObject/UObjectIterator.h"
+#include "UObject/ICookInfo.h"
 
 
 TArray<UEdGraphPin*> FollowPinArray(const UEdGraphPin& Pin, bool bIgnoreOrphan, bool* bOutCycleDetected)
@@ -653,6 +654,8 @@ void GetNodeGroupObjectNodeMappingImmersive(UCustomizableObject* Object, FAssetR
 		}
 	}
 
+	FCookLoadScope CookLoadScope(ECookLoadType::EditorOnly);
+
 	Filter.bIncludeOnlyOnDiskAssets = false;
 	
 	TArray<FAssetData> ArrayAssetData;
diff --git a/Engine/Plugins/Mutable/Source/CustomizableObjectEditor/Private/MuCOE/Nodes/CustomizableObjectNodeComponentPassthroughMesh.cpp b/Engine/Plugins/Mutable/Source/CustomizableObjectEditor/Private/MuCOE/Nodes/CustomizableObjectNodeComponentPassthroughMesh.cpp
--- a/Engine/Plugins/Mutable/Source/CustomizableObjectEditor/Private/MuCOE/Nodes/CustomizableObjectNodeComponentPassthroughMesh.cpp
+++ b/Engine/Plugins/Mutable/Source/CustomizableObjectEditor/Private/MuCOE/Nodes/CustomizableObjectNodeComponentPassthroughMesh.cpp
@@ -14,6 +14,7 @@
 #include "MuCOE/Nodes/CustomizableObjectNodeLayoutBlocks.h"
 #include "MuCOE/RemapPins/CustomizableObjectNodeRemapPinsByNameDefaultPin.h"
 #include "MuCOE/UnrealEditorPortabilityHelpers.h"
+#include "UObject/ICookInfo.h"
 
 class UCustomizableObjectNodeRemapPinsByName;
 class UObject;
@@ -88,6 +89,8 @@ void UCustomizableObjectNodeComponentPassthroughMesh::AllocateDefaultPins(UCusto
 		return;
 	}
 
+	FCookLoadScope CookLoadScope(ECookLoadType::EditorOnly);
+
 	// Support SkeletalMeshes only, for now.
 	USkeletalMesh* SkeletalMesh = Cast<USkeletalMesh>( Mesh.TryLoad() );
 	if (!SkeletalMesh)
@@ -183,6 +186,8 @@ void UCustomizableObjectNodeComponentPassthroughMesh::GetPinSection(const UEdGra
 
 UMaterialInterface* UCustomizableObjectNodeComponentPassthroughMesh::GetMaterialFor(const UEdGraphPin* Pin) const
 {
+	FCookLoadScope CookLoadScope(ECookLoadType::EditorOnly);
+
 	USkeletalMesh* SkeletalMesh = Cast<USkeletalMesh>(Mesh.TryLoad());
 	if (SkeletalMesh)
 	{
@@ -234,6 +239,8 @@ UMaterialInterface* UCustomizableObjectNodeComponentPassthroughMesh::GetMaterial
 
 FSkeletalMaterial* UCustomizableObjectNodeComponentPassthroughMesh::GetSkeletalMaterialFor(const int32 LODIndex, const int32 SectionIndex, const FSkeletalMeshModel* ImportedModel) const
 {
+	FCookLoadScope CookLoadScope(ECookLoadType::EditorOnly);
+
 	USkeletalMesh* SkeletalMesh = Cast<USkeletalMesh>(Mesh.TryLoad());
 	if (!SkeletalMesh)
 	{
diff --git a/Engine/Plugins/Mutable/Source/CustomizableObjectEditor/Private/MuCOE/Nodes/CustomizableObjectNodeGroupProjectorParameter.cpp b/Engine/Plugins/Mutable/Source/CustomizableObjectEditor/Private/MuCOE/Nodes/CustomizableObjectNodeGroupProjectorParameter.cpp
--- a/Engine/Plugins/Mutable/Source/CustomizableObjectEditor/Private/MuCOE/Nodes/CustomizableObjectNodeGroupProjectorParameter.cpp
+++ b/Engine/Plugins/Mutable/Source/CustomizableObjectEditor/Private/MuCOE/Nodes/CustomizableObjectNodeGroupProjectorParameter.cpp
@@ -5,6 +5,8 @@
 #include "MuCO/CustomizableObjectCustomVersion.h"
 #include "MuCOE/EdGraphSchema_CustomizableObject.h"
 #include "MuCOE/UnrealEditorPortabilityHelpers.h"
+#include "UObject/ICookInfo.h"
+
 
 class UCustomizableObjectNodeRemapPins;
 
@@ -32,6 +34,7 @@ TArray<FGroupProjectorParameterImage> UCustomizableObjectNodeGroupProjectorParam
 
 	int32 NameIndex = 0;
 
+	FCookLoadScope CookLoadScope(ECookLoadType::EditorOnly);
 	for (TMap<FName, uint8*>::TConstIterator RowIt = OptionTexturesDataTable->GetRowMap().CreateConstIterator(); RowIt; ++RowIt)
 	{
 		uint8* RowData = RowIt.Value();
