[mutable] Fixed generation of textures of 1 pixel
- Switchnode descriptor request from the Transform node was producing an invalid size for a texture due to the switchnode having textures of various sizes.
- Removed unnecessary includes
- Fixed typos

diff --git a/Engine/Plugins/Mutable/Source/MutableTools/Private/MuT/CodeGenerator_Image.cpp b/Engine/Plugins/Mutable/Source/MutableTools/Private/MuT/CodeGenerator_Image.cpp
--- a/Engine/Plugins/Mutable/Source/MutableTools/Private/MuT/CodeGenerator_Image.cpp
+++ b/Engine/Plugins/Mutable/Source/MutableTools/Private/MuT/CodeGenerator_Image.cpp
@@ -11,11 +11,8 @@
 #include "MuR/Operations.h"
 #include "MuR/Parameters.h"
 #include "MuR/ParametersPrivate.h"
-#include "MuR/Platform.h"
 #include "MuR/Ptr.h"
-#include "MuR/RefCounted.h"
 #include "MuR/Serialisation.h"
-#include "MuR/System.h"
 #include "MuT/AST.h"
 #include "MuT/ASTOpConditional.h"
 #include "MuT/ASTOpConstantResource.h"
@@ -45,7 +42,6 @@
 #include "MuT/NodeBool.h"
 #include "MuT/NodeColour.h"
 #include "MuT/NodeColourFromScalars.h"
-#include "MuT/NodeColourSampleImage.h"
 #include "MuT/NodeColourConstant.h"
 #include "MuT/NodeImage.h"
 #include "MuT/NodeImageBinarise.h"
@@ -94,9 +90,6 @@
 #include "MuT/NodeImageTransformPrivate.h"
 #include "MuT/NodeImageVariation.h"
 #include "MuT/NodeImageVariationPrivate.h"
-#include "MuT/NodeMeshTable.h"
-#include "MuT/NodeProjector.h"
-#include "MuT/NodeRange.h"
 #include "MuT/NodeScalar.h"
 #include "MuT/NodeScalarConstant.h"
 #include "MuT/Table.h"
@@ -648,7 +641,7 @@ namespace mu
 			Rotation = ScalarResult.op;
 		}
 
-		// If one of the inputs (offset or scale) is missig assume unifrom translation/scaling 
+		// If one of the inputs (offset or scale) is missing assume uniform translation/scaling 
 		Op->OffsetX = OffsetX ? OffsetX : OffsetY;
 		Op->OffsetY = OffsetY ? OffsetY : OffsetX;
  		Op->ScaleX = ScaleX ? ScaleX : ScaleY;
@@ -679,7 +672,7 @@ namespace mu
             Base = GenerateMissingImageCode(TEXT("Image Transform Base"), EImageFormat::RGB_UByte, InNode->GetMessageContext(), NewOptions);
         }
 		
-		FImageDesc BaseDesc = Base->GetImageDesc();
+		FImageDesc BaseDesc = Base->GetImageDesc(true);
 		
         Op->Base = Base;
 		Op->SourceSizeX = BaseDesc.m_size.X;
