[Mutable] Convert ModelResources to UObject.

This conversion allows us to fix some long-standing issues with cooked data.
- Removed custom serialization for ModelResources.
- Fixed async compilation for cook.
- Fixed missing StreamedResources and StreamedExtension assets on packaged builds.
- Add an API method to UCustomizableObjectExtension to allow private references to be moved to the respective container at cook time.

Keep StreamedResources and StreamedExtension data loaded in the editor. UCustomizableObjectResourceDataContainer will only be used on cooked builds.

#rnx

diff --git a/Engine/Plugins/Experimental/HairStrandsMutable/Source/HairStrandsMutableEditor/Private/CustomizableObjectNodeGroomConstant.cpp b/Engine/Plugins/Experimental/HairStrandsMutable/Source/HairStrandsMutableEditor/Private/CustomizableObjectNodeGroomConstant.cpp
--- a/Engine/Plugins/Experimental/HairStrandsMutable/Source/HairStrandsMutableEditor/Private/CustomizableObjectNodeGroomConstant.cpp
+++ b/Engine/Plugins/Experimental/HairStrandsMutable/Source/HairStrandsMutableEditor/Private/CustomizableObjectNodeGroomConstant.cpp
@@ -54,20 +54,16 @@ mu::Ptr<mu::NodeExtensionData> UCustomizableObjectNodeGroomConstant::GenerateMut
 
 	// Create node and extension data container
 	mu::Ptr<mu::NodeExtensionDataConstant> Result = new mu::NodeExtensionDataConstant();
-	
+
+	FInstancedStruct Struct;
+	Struct.InitializeAs<FGroomPinData>(GroomData);
+
 	// Grooms are usually quite large, so set it up as a streaming constant to allow it to be
 	// loaded on demand.
 	//
 	// If needed we could expose an editable UPROPERTY to give the user the option of making this
 	// an always-loaded constant.
-	UCustomizableObjectResourceDataContainer* Container = nullptr;
-	Result->SetValue(CompilerInterface.MakeStreamedExtensionData(Container));
-
-	if (Container)
-	{
-		// Populate instanced struct
-		Container->Data.Data.InitializeAs<FGroomPinData>(GroomData);
-	}
+	Result->SetValue(CompilerInterface.MakeStreamedExtensionData(MoveTemp(Struct)));
 
 	return Result;
 }
diff --git a/Engine/Plugins/Experimental/RigLogicMutable/Source/RigLogicMutable/Private/RigLogicMutableExtension.cpp b/Engine/Plugins/Experimental/RigLogicMutable/Source/RigLogicMutable/Private/RigLogicMutableExtension.cpp
--- a/Engine/Plugins/Experimental/RigLogicMutable/Source/RigLogicMutable/Private/RigLogicMutableExtension.cpp
+++ b/Engine/Plugins/Experimental/RigLogicMutable/Source/RigLogicMutable/Private/RigLogicMutableExtension.cpp
@@ -29,17 +29,10 @@ FDNAPinData& FDNAPinData::operator=(FDNAPinData&& Source)
 	return *this;
 }
 
-void FDNAPinData::CopyFromDNAAsset(const UDNAAsset* SourceAsset, UObject* OuterForOwnedObjects)
+void FDNAPinData::SetDNAAsset(UDNAAsset* SourceAsset)
 {
-	if (SourceAsset)
-	{
-		DNAAsset = URigLogicMutableExtension::CopyDNAAsset(SourceAsset, OuterForOwnedObjects);
-		DNAAsset->bKeepDNAAfterInitialization = true;
-	}
-	else
-	{
-		DNAAsset = nullptr;
-	}
+	DNAAsset = SourceAsset;
+	DNAAsset->bKeepDNAAfterInitialization = SourceAsset != nullptr;
 }
 
 TArray<FCustomizableObjectPinType> URigLogicMutableExtension::GetPinTypes() const
@@ -94,6 +87,18 @@ void URigLogicMutableExtension::OnSkeletalMeshCreated(
 	}
 }
 
+#if WITH_EDITOR
+void URigLogicMutableExtension::MovePrivateReferencesToContainer(FInstancedStruct& Struct, UObject* Container) const
+{
+	FDNAPinData* Data = Struct.GetMutablePtr<FDNAPinData>();
+	if (Data
+		&& Data->GetDNAAsset() != nullptr)
+	{
+		Data->SetDNAAsset(CopyDNAAsset(Data->GetDNAAsset(), Container));
+	}
+}
+#endif
+
 UDNAAsset* URigLogicMutableExtension::CopyDNAAsset(const UDNAAsset* Source, UObject* OuterForCopy)
 {
 	check(IsInGameThread());
diff --git a/Engine/Plugins/Experimental/RigLogicMutable/Source/RigLogicMutable/Public/RigLogicMutableExtension.h b/Engine/Plugins/Experimental/RigLogicMutable/Source/RigLogicMutable/Public/RigLogicMutableExtension.h
--- a/Engine/Plugins/Experimental/RigLogicMutable/Source/RigLogicMutable/Public/RigLogicMutableExtension.h
+++ b/Engine/Plugins/Experimental/RigLogicMutable/Source/RigLogicMutable/Public/RigLogicMutableExtension.h
@@ -32,11 +32,11 @@ public:
 	FDNAPinData& operator=(FDNAPinData&& Source);
 
 	/**
-	 * Makes a copy of the given asset and assigns the DNAAsset member variable to the copy.
+	 * Assign the DNAAsset member variable to the PinData.
 	 *
 	 * If SourceAsset is nullptr, DNAAsset will be set to nullptr.
 	 */
-	void CopyFromDNAAsset(const UDNAAsset* SourceAsset, UObject* OuterForOwnedObjects);
+	void SetDNAAsset(UDNAAsset* SourceAsset);
 
 	const UDNAAsset* GetDNAAsset() const { return DNAAsset; }
 
@@ -75,6 +75,10 @@ public:
 		FName ComponentName,
 		USkeletalMesh* SkeletalMesh) const override;
 
+#if WITH_EDITOR
+	virtual void MovePrivateReferencesToContainer(FInstancedStruct& Struct, UObject* Container) const;
+#endif
+
 	/**
 	 * Makes a copy of the Source asset and returns it.
 	 *
diff --git a/Engine/Plugins/Experimental/RigLogicMutable/Source/RigLogicMutableEditor/Private/CustomizableObjectNodeDNAConstant.cpp b/Engine/Plugins/Experimental/RigLogicMutable/Source/RigLogicMutableEditor/Private/CustomizableObjectNodeDNAConstant.cpp
--- a/Engine/Plugins/Experimental/RigLogicMutable/Source/RigLogicMutableEditor/Private/CustomizableObjectNodeDNAConstant.cpp
+++ b/Engine/Plugins/Experimental/RigLogicMutable/Source/RigLogicMutableEditor/Private/CustomizableObjectNodeDNAConstant.cpp
@@ -65,6 +65,16 @@ mu::Ptr<mu::NodeExtensionData> UCustomizableObjectNodeDNAConstant::GenerateMutab
 	FDNAPinData PinData;
 	PinData.ComponentName = ComponentName;
 	
+	if (SkeletalMesh)
+	{
+		// Note that this may be nullptr if the mesh doesn't have a DNA asset
+		PinData.SetDNAAsset(Cast<UDNAAsset>(SkeletalMesh->GetAssetUserDataOfClass(UDNAAsset::StaticClass())));
+	}
+
+	// Populate instanced struct
+	FInstancedStruct Struct;
+	Struct.InitializeAs<FDNAPinData>(MoveTemp(PinData));
+
 	// DNA is usually quite large, so ideally it would be made a streaming constant to allow it to 
 	// be loaded on demand.
 	//
@@ -74,34 +84,10 @@ mu::Ptr<mu::NodeExtensionData> UCustomizableObjectNodeDNAConstant::GenerateMutab
 
 	if (bShouldStreamDNA)
 	{
-		UCustomizableObjectResourceDataContainer* Container = nullptr;
-		Result->SetValue(CompilerInterface.MakeStreamedExtensionData(Container));
-
-		// Populate instanced struct
-		if (Container)
-		{
-			if (SkeletalMesh)
-			{
-				// Note that this may be nullptr if the mesh doesn't have a DNA asset
-				PinData.CopyFromDNAAsset(Cast<UDNAAsset>(SkeletalMesh->GetAssetUserDataOfClass(UDNAAsset::StaticClass())), Container);
-			}
-
-			Container->Data.Data.InitializeAs<FDNAPinData>(MoveTemp(PinData));
-		}
+		Result->SetValue(CompilerInterface.MakeStreamedExtensionData(MoveTemp(Struct)));
 	}
 	else
 	{
-		if (SkeletalMesh)
-		{
-			UObject* OuterForOwnedObjects = const_cast<UObject*>(CompilerInterface.GetOuterForAlwaysLoadedObjects());
-
-			// Note that this may be nullptr if the mesh doesn't have a DNA asset
-			PinData.CopyFromDNAAsset(Cast<UDNAAsset>(SkeletalMesh->GetAssetUserDataOfClass(UDNAAsset::StaticClass())), OuterForOwnedObjects);
-		}
-
-		FInstancedStruct Struct;
-		Struct.InitializeAs<FDNAPinData>(MoveTemp(PinData));
-
 		Result->SetValue(CompilerInterface.MakeAlwaysLoadedExtensionData(MoveTemp(Struct)));
 	}
 
diff --git a/Engine/Plugins/Mutable/Source/CustomizableObject/Internal/MuCO/CustomizableObjectInstancePrivate.h b/Engine/Plugins/Mutable/Source/CustomizableObject/Internal/MuCO/CustomizableObjectInstancePrivate.h
--- a/Engine/Plugins/Mutable/Source/CustomizableObject/Internal/MuCO/CustomizableObjectInstancePrivate.h
+++ b/Engine/Plugins/Mutable/Source/CustomizableObject/Internal/MuCO/CustomizableObjectInstancePrivate.h
@@ -24,7 +24,6 @@ namespace mu
 	typedef uint64 FResourceID;
 }
 
-struct FModelResources;
 struct FMutableModelImageProperties;
 struct FMutableRefSkeletalMeshData;
 struct FMutableImageCacheKey;
@@ -35,6 +34,7 @@ struct FGeneratedTexture;
 class UPhysicsAsset;
 class USkeleton;
 class UCustomizableObjectExtension;
+class UModelResources;
 
 
 // Log texts
@@ -301,7 +301,7 @@ private:
 	void InitSkeletalMeshData(const TSharedRef<FUpdateContextPrivate>& OperationData, USkeletalMesh* SkeletalMesh, const FMutableRefSkeletalMeshData& RefSkeletalMeshData, const UCustomizableObject& CustomizableObject, FCustomizableObjectComponentIndex ObjectComponentIndex);
 
 	bool BuildSkeletonData(const TSharedRef<FUpdateContextPrivate>& OperationData, USkeletalMesh& SkeletalMesh, const FMutableRefSkeletalMeshData& RefSkeletalMeshData, UCustomizableObject& CustomizableObject, FCustomizableObjectInstanceComponentIndex InstanceComponentIndex);
-	void BuildMeshSockets(const TSharedRef<FUpdateContextPrivate>& OperationData, USkeletalMesh* SkeletalMesh, const FModelResources& ModelResources, const FMutableRefSkeletalMeshData& RefSkeletalMeshData, TSharedPtr<const mu::FMesh> MutableMesh);
+	void BuildMeshSockets(const TSharedRef<FUpdateContextPrivate>& OperationData, USkeletalMesh* SkeletalMesh, const UModelResources& ModelResources, const FMutableRefSkeletalMeshData& RefSkeletalMeshData, TSharedPtr<const mu::FMesh> MutableMesh);
 	void BuildOrCopyElementData(const TSharedRef<FUpdateContextPrivate>& OperationData, USkeletalMesh* SkeletalMesh, UCustomizableObject& CustomizableObjectInstance, FCustomizableObjectInstanceComponentIndex InstanceComponentIndex);
 	void BuildOrCopyMorphTargetsData(const TSharedRef<FUpdateContextPrivate>& OperationData, USkeletalMesh* SkeletalMesh, const USkeletalMesh* SrcSkeletalMesh, UCustomizableObject& CustomizableObject, FCustomizableObjectInstanceComponentIndex InstanceComponentIndex);
 	bool BuildOrCopyRenderData(const TSharedRef<FUpdateContextPrivate>& OperationData, USkeletalMesh* SkeletalMesh, const USkeletalMesh* SrcSkeletalMesh, UCustomizableObjectInstance* CustomizableObjectInstance, FCustomizableObjectInstanceComponentIndex InstanceComponentIndex);
diff --git a/Engine/Plugins/Mutable/Source/CustomizableObject/Internal/MuCO/CustomizableObjectPrivate.h b/Engine/Plugins/Mutable/Source/CustomizableObject/Internal/MuCO/CustomizableObjectPrivate.h
--- a/Engine/Plugins/Mutable/Source/CustomizableObject/Internal/MuCO/CustomizableObjectPrivate.h
+++ b/Engine/Plugins/Mutable/Source/CustomizableObject/Internal/MuCO/CustomizableObjectPrivate.h
@@ -50,7 +50,7 @@ class UTexture;
 class UAnimInstance;
 class UAssetUserData;
 class USkeletalMeshLODSettings;
-struct FModelResources;
+class UModelResources;
 struct FModelStreamableBulkData;
 struct FObjectAndNameAsStringProxyArchive;
 struct FCustomizableObjectInstanceDescriptor;
@@ -58,6 +58,7 @@ struct FCustomizableObjectInstanceDescriptor;
 
 FGuid CUSTOMIZABLEOBJECT_API GenerateIdentifier(const UCustomizableObject& CustomizableObject);
 
+FString CUSTOMIZABLEOBJECT_API GetModelResourcesNameForPlatform(const UCustomizableObject& CustomizableObject, const ITargetPlatform& Platform);
 
 // A USTRUCT version of FMeshToMeshVertData in SkeletalMeshTypes.h
 // We are taking advantage of the padding data to store from which asset this data comes from
@@ -216,10 +217,6 @@ struct FMutableRemappedBone
 	{
 		return Name == InName;
 	}
-
-#if WITH_EDITORONLY_DATA
-	friend FArchive& operator<<(FArchive& Ar, FMutableRemappedBone& RemappedBone);
-#endif
 };
 
 
@@ -369,10 +366,6 @@ struct CUSTOMIZABLEOBJECT_API FMutableModelImageProperties
 	TEnumAsByte<enum TextureAddress> AddressY;
 
 	bool operator!=(const FMutableModelImageProperties& rhs) const;
-
-#if WITH_EDITORONLY_DATA
-	friend FArchive& operator<<(FArchive& Ar, FMutableModelImageProperties& ImageProps);
-#endif
 };
 
 
@@ -482,22 +475,6 @@ struct FMutableRefSkeletalMeshSettings
 };
 
 
-USTRUCT()
-struct FMutableRefAssetUserData
-{
-	GENERATED_BODY()
-
-	UPROPERTY()
-	TObjectPtr<UCustomizableObjectResourceDataContainer> AssetUserData;
-
-#if WITH_EDITORONLY_DATA
-	UPROPERTY()
-	int32 AssetUserDataIndex = INDEX_NONE;
-#endif
-
-};
-
-
 USTRUCT()
 struct CUSTOMIZABLEOBJECT_API FMutableRefSkeletalMeshData
 {
@@ -549,14 +526,11 @@ struct CUSTOMIZABLEOBJECT_API FMutableRefSkeletalMeshData
 
 	// Asset user data
 	UPROPERTY()
-	TArray<FMutableRefAssetUserData> AssetUserData;
-	
+	TArray<int32> AssetUserDataIndices;
+
 #if WITH_EDITORONLY_DATA
 	friend FArchive& operator<<(FArchive& Ar, FMutableRefSkeletalMeshData& Data);
-
-	void InitResources(UCustomizableObject* InOuter, FModelResources& InModelResources, const ITargetPlatform* InTargetPlatform);
 #endif
-
 };
 
 
@@ -788,6 +762,35 @@ struct FMutableParameterIndex
 	int32 TypedIndex = INDEX_NONE;
 };
 
+
+USTRUCT()
+struct CUSTOMIZABLEOBJECT_API FIntegerParameterOptionKey
+{
+	GENERATED_BODY()
+
+	UPROPERTY()
+	FString ParameterName;
+
+	UPROPERTY()
+	FString ParameterOption;
+
+#if WITH_EDITOR
+	friend CUSTOMIZABLEOBJECT_API uint32 GetTypeHash(const FIntegerParameterOptionKey& Key);
+	bool operator==(const FIntegerParameterOptionKey& Other) const = default;
+#endif
+};
+
+
+USTRUCT()
+struct CUSTOMIZABLEOBJECT_API FIntegerParameterOptionDataTable
+{
+	GENERATED_BODY()
+
+	UPROPERTY()
+	TSet<TSoftObjectPtr<UDataTable>> DataTables;
+};
+
+
 USTRUCT()
 struct CUSTOMIZABLEOBJECT_API FIntegerParameterUIData
 {
@@ -1020,8 +1023,8 @@ namespace MutablePrivate
 		/** Serialized mu::FModel */
 		TArray64<uint8> ModelData;
 
-		/** Serialized FModelResources */
-		TArray64<uint8> ModelResourcesData;
+		/** Serialized UModelResources */
+		TStrongObjectPtr<UModelResources> ModelResources;
 
 		/** Streamable resources info such as files and offsets. */
 		TSharedPtr<FModelStreamableBulkData> ModelStreamables;
@@ -1129,39 +1132,18 @@ struct FMutableParamNameSet
 	GENERATED_BODY()
 
 	TSet<FString> ParamNames;
-
-	friend FArchive& operator<<(FArchive& Ar, FMutableParamNameSet& MutableParamNameSet);
 };
 
 
-/** Wrapper to Mutable Extensions data.
-  * Only used when serializing extensions data. Required to use Unreal's serialization system, which allows serializing optional data. */
+/** Class containing all UE resources derived from a CO compilation. These resources will be embedded in the CO at cook time but not in the editor.
+  * Editor compilations will serialize this class to disk using the Serialize methods. Ensure new fields are serialized, too.
+  * Variables and settings that should not change until the CO is re-compiled should be stored here. */
 UCLASS()
-class UExtensionData : public UObject
+class CUSTOMIZABLEOBJECT_API UModelResources : public UObject
 {
 	GENERATED_BODY()
 
 public:
-	void CopyFrom(const FModelResources& ModelResources);
-	
-	void MoveTo(FModelResources& ModelResources);
-	
-	UPROPERTY()
-	TArray<FCustomizableObjectResourceData> AlwaysLoadedExtensionData;
-
-	UPROPERTY()
-	TArray<FCustomizableObjectStreamedResourceData> StreamedExtensionData;
-};
-
-
-/** Struct containing all UE resources derived from a CO compilation. These resources will be embedded in the CO at cook time but not in the editor.
-  * Editor compilations will serialize this struct to disk using the Serialize methods. Ensure new fields are serialized, too.
-  * Variables and settings that should not change until the CO is re-compiled should be stored here. */
-USTRUCT()
-struct FModelResources
-{
-	GENERATED_BODY()
-
 	/** 
 	 * All the SkeletalMeshes generated for this CustomizableObject instances will use the Reference Skeletal Mesh
 	 * properties for everything that Mutable doesn't create or modify. This struct stores the information used from
@@ -1244,7 +1226,8 @@ struct FModelResources
 
 #if WITH_EDITORONLY_DATA
 	/** DataTable used by an int parameter and its value. */
-	TMap<TTuple<FString, FString>, TSet<TSoftObjectPtr<UDataTable>>> IntParameterOptionDataTable;
+	UPROPERTY()
+	TMap<FIntegerParameterOptionKey, FIntegerParameterOptionDataTable> IntParameterOptionDataTable;
 #endif
 	
 	UPROPERTY()
@@ -1271,26 +1254,37 @@ struct FModelResources
 
 	/** Map to identify what CustomizableObject owns a parameter. Used to display a tooltip when hovering a parameter
 	 * in the Prev. instance panel */
+	UPROPERTY()
 	TMap<FString, FString> CustomizableObjectPathMap;
 
+	UPROPERTY()
 	TMap<FString, FCustomizableObjectIdPair> GroupNodeMap;
 
 	/** If the object is compiled with maximum optimizations. */
+	UPROPERTY()
 	bool bIsCompiledWithOptimization = true;
 
 	/** This is a non-user-controlled flag to disable streaming (set at object compilation time, depending on optimization). */
+	UPROPERTY()
 	bool bIsTextureStreamingDisabled = false;
 
 	/** List of external packages that if changed, a compilation is required.
 	  * Key is the package name. Value is the the UPackage::Guid, which is regenerated each time the packages is saved.
 	  *
 	  * Updated each time the CO is compiled and saved in the Derived Data. */
+	UPROPERTY()
 	TMap<FName, FGuid> ParticipatingObjects;
 
 	// Used to know if roms and other resources must be streamed from the DDC.
 	bool bIsStoredInDDC = false;
 	UE::DerivedData::FCacheKey DDCKey = UE::DerivedData::FCacheKey::Empty;
 	UE::DerivedData::ECachePolicy DDCDefaultPolicy = UE::DerivedData::ECachePolicy::Default;
+
+	UPROPERTY()
+	TArray<FCustomizableObjectResourceData> StreamedResourceDataEditor;
+
+	UPROPERTY()
+	TArray<FCustomizableObjectResourceData> StreamedExtensionDataEditor;
 #endif
 
 	// Constant Resources streamed in on demand when generating meshes
@@ -1324,17 +1318,17 @@ struct FModelResources
 	UPROPERTY()
 	FString ReleaseVersion;
 
+	UPROPERTY()
+	int32 CodeVersion = 0;
+
 #if WITH_EDITORONLY_DATA
 	/** Value of the variable TextureCompression in the last compilation of this CO.
 		* this is needed since we can compile a CO through blueprints with a different 
 		* compilation setting than the stored in the COE.*/
 	UPROPERTY()
 	bool bCompiledWithHDTextureCompression = false;
-#endif
-
-#if WITH_EDITORONLY_DATA
-	void CUSTOMIZABLEOBJECT_API Serialize(FObjectAndNameAsStringProxyArchive& Ar, bool bIsCooking);
-	bool CUSTOMIZABLEOBJECT_API Unserialize(FObjectAndNameAsStringProxyArchive& Ar, UCustomizableObject& Outer, const ITargetPlatform* InTargetPlatform, bool bIsCooking);
+	
+	void InitCookData();
 #endif
 };
 
@@ -1442,7 +1436,7 @@ class CUSTOMIZABLEOBJECT_API UCustomizableObjectPrivate : public UObject
 
 	/** Stores resources to be used by MutableModel In-Game. Cooked resources. */
 	UPROPERTY()
-	FModelResources ModelResources;
+	TObjectPtr<UModelResources> ModelResources;
 
 #if WITH_EDITORONLY_DATA
 	/** 
@@ -1450,7 +1444,7 @@ class CUSTOMIZABLEOBJECT_API UCustomizableObjectPrivate : public UObject
 	 * Editor-Only to avoid packaging assets referenced by editor compilations. 
 	 */
 	UPROPERTY(Transient)
-	FModelResources ModelResourcesEditor;
+	TObjectPtr<UModelResources> ModelResourcesEditor;
 
 	/**
 	 * Stores streamable data info to be used by MutableModel in the Editor. Editor streaming.
@@ -1464,13 +1458,16 @@ public:
 	const TSharedPtr<mu::FModel, ESPMode::ThreadSafe>& GetModel();
 	const TSharedPtr<const mu::FModel, ESPMode::ThreadSafe> GetModel() const;
 
-	FModelResources& GetModelResources();
-	const FModelResources& GetModelResources() const;
+	UModelResources* GetModelResources();
+	const UModelResources* GetModelResources() const;
+	const UModelResources& GetModelResourcesChecked() const;
+
 
 #if WITH_EDITORONLY_DATA
-	FModelResources& GetModelResources(bool bIsCooking);
-	const FModelResources& GetModelResources(bool bIsCooking) const;
+	UModelResources* GetModelResources(bool bIsCooking);
+	const UModelResources* GetModelResources(bool bIsCooking) const;
 
+	void SetModelResources(UModelResources* InModelResources, bool bIsCooking);
 	void SetModelStreamableBulkData(const TSharedPtr<FModelStreamableBulkData>& StreamableData, bool bIsCooking);
 #endif
 	
@@ -1662,7 +1659,7 @@ public:
 
 	/** Reference to all UObject used in game. Only updated during the compilation if the user explicitly wants to save all references. */
 	UPROPERTY()
-	FModelResources References;
+	TObjectPtr<UModelResources> ReferencedObjects;
 	
 	// Map of name to index of ParameterProperties.
 	// use this to lookup fast by Name
@@ -1821,6 +1818,7 @@ public:
 		RemoveObsoletMeshInterpolateAndGeometryOp,
 
 		RemoveObsoleteDataTypesFromEnum,
+		ConvertModelResourcesToUObject,
 
 		// -----<new versions can be added above this line>--------
 		LastCustomizableObjectVersion
@@ -1828,6 +1826,6 @@ public:
 	
 	static constexpr int32 CurrentSupportedVersion = ECustomizableObjectVersions::LastCustomizableObjectVersion;
 	
-	static constexpr int32 DerivedDataVersion = 2;
+	static constexpr int32 DerivedDataVersion = 3;
 };
 
diff --git a/Engine/Plugins/Mutable/Source/CustomizableObject/Internal/MuCO/CustomizableObjectResourceDataTypes.h b/Engine/Plugins/Mutable/Source/CustomizableObject/Internal/MuCO/CustomizableObjectResourceDataTypes.h
--- a/Engine/Plugins/Mutable/Source/CustomizableObject/Internal/MuCO/CustomizableObjectResourceDataTypes.h
+++ b/Engine/Plugins/Mutable/Source/CustomizableObject/Internal/MuCO/CustomizableObjectResourceDataTypes.h
@@ -21,10 +21,4 @@ public:
 	// Valid in Cooked builds
 	UPROPERTY()
 	TObjectPtr<UAssetUserData> AssetUserData;
-
-#if WITH_EDITORONLY_DATA
-	// Valid in the Editor
-	UPROPERTY(Transient)
-	TObjectPtr<UAssetUserData> AssetUserDataEditor;
-#endif
 };
\ No newline at end of file
diff --git a/Engine/Plugins/Mutable/Source/CustomizableObject/Private/MuCO/CustomizableInstanceLODManagement.cpp b/Engine/Plugins/Mutable/Source/CustomizableObject/Private/MuCO/CustomizableInstanceLODManagement.cpp
--- a/Engine/Plugins/Mutable/Source/CustomizableObject/Private/MuCO/CustomizableInstanceLODManagement.cpp
+++ b/Engine/Plugins/Mutable/Source/CustomizableObject/Private/MuCO/CustomizableInstanceLODManagement.cpp
@@ -406,7 +406,6 @@ void UCustomizableInstanceLODManagement::UpdateInstanceDistsAndLODs(FMutableInst
 					COI->SetIsBeingUsedByComponentInPlay(true);
 
 					const FName& ComponentName = CustomizableObjectInstanceUsage->GetComponentName();
-					const int32 ObjectComponentIndex = COI->GetCustomizableObject()->GetPrivate()->GetModelResources().ComponentNamesPerObjectComponent.IndexOfByKey(ComponentName);
 
 #if WITH_EDITOR
 					// If the instance is generated but the component doesn't have a mesh, set it.
diff --git a/Engine/Plugins/Mutable/Source/CustomizableObject/Private/MuCO/CustomizableObject.cpp b/Engine/Plugins/Mutable/Source/CustomizableObject/Private/MuCO/CustomizableObject.cpp
--- a/Engine/Plugins/Mutable/Source/CustomizableObject/Private/MuCO/CustomizableObject.cpp
+++ b/Engine/Plugins/Mutable/Source/CustomizableObject/Private/MuCO/CustomizableObject.cpp
@@ -83,7 +83,7 @@ TAutoConsoleVariable<bool> CVarMutableUseBulkData(
 
 TAutoConsoleVariable<int32> CVarMutableDerivedDataCacheUsage(
 	TEXT("mutable.DerivedDataCacheUsage"),
-	2,
+	0, // TODO Pere: Set back to 2 before submit
 	TEXT("Derived data cache access for cooked data.")
 	TEXT("0 - None. Disables access to the cache.")
 	TEXT("1 - Local. Allow cache requests to query and store records and values in local caches.")
@@ -93,7 +93,7 @@ TAutoConsoleVariable<int32> CVarMutableDerivedDataCacheUsage(
 
 TAutoConsoleVariable<bool> CVarMutableAsyncCook(
 	TEXT("Mutable.CookAsync"),
-	false,
+	true,
 	TEXT("True - Customizable Objects will be compiled asynchronously during cook.\n")
 	TEXT("False - Sync compilation.\n"));
 
@@ -248,28 +248,13 @@ void UCustomizableObject::PreSave(FObjectPreSaveContext ObjectSaveContext)
 			
 			// Clear model resources
 			GetPrivate()->SetModel(nullptr, FGuid());
-			GetPrivate()->GetModelResources(true /* bIsCooking */) = FModelResources();
+			GetPrivate()->SetModelResources(nullptr, true /* bIsCooking */);
 			GetPrivate()->GetModelStreamableBulkData(true).Reset();
 		}
 	}
 }
 
 
-void UCustomizableObject::PostSaveRoot(FObjectPostSaveRootContext ObjectSaveContext)
-{
-	MUTABLE_CPUPROFILER_SCOPE(UCustomizableObject::PostSaveRoot);
-
-	Super::PostSaveRoot(ObjectSaveContext);
-
-	if (ObjectSaveContext.IsCooking())
-	{
-		// Free cached data after saving;
-		const ITargetPlatform* TargetPlatform = ObjectSaveContext.GetTargetPlatform();
-		GetPrivate()->CachedPlatformsData.Remove(TargetPlatform->PlatformName());
-	}
-}
-
-
 bool UCustomizableObjectPrivate::TryUpdateIsChildObject()
 {
 	if (const ICustomizableObjectEditorModule* Module = ICustomizableObjectEditorModule::Get())
@@ -304,17 +289,13 @@ bool UCustomizableObjectPrivate::TryLoadCompiledCookDataForPlatform(const ITarge
 		return false;
 	}
 
-	FMemoryReaderView ModelResourcesReader(PlatformData->ModelResourcesData);
-	if (LoadModelResources(ModelResourcesReader, TargetPlatform, true))
-	{
-		SetModelStreamableBulkData(PlatformData->ModelStreamables, true);
+	SetModelResources(PlatformData->ModelResources.Get(), true);
+	SetModelStreamableBulkData(PlatformData->ModelStreamables, true);
 
-		FMemoryReaderView ModelReader(PlatformData->ModelData);
-		LoadModel(ModelReader);
-		return GetModel() != nullptr;
-	}
+	FMemoryReaderView ModelReader(PlatformData->ModelData);
+	LoadModel(ModelReader);
 
-	return false;
+	return GetModel() != nullptr;
 }
 
 #endif // End WITH_EDITOR
@@ -528,17 +509,26 @@ FGuid GenerateIdentifier(const UCustomizableObject& CustomizableObject)
 }
 
 
+FString GetModelResourcesNameForPlatform(const UCustomizableObject& CustomizableObject, const ITargetPlatform& InTargetPlatform)
+{
+	return GenerateIdentifier(CustomizableObject).ToString() + InTargetPlatform.PlatformName();
+}
+
+
 bool UCustomizableObjectPrivate::LoadModelResources(FArchive& MemoryReader, const ITargetPlatform* InTargetPlatform, bool bIsCooking)
 {
 	// Make sure mutable has been initialised.
 	UCustomizableObjectSystem::GetInstance();
-
-	FModelResources LocalModelResources;
+	
+	const FString ModelResourcesName = GetModelResourcesNameForPlatform(*GetPublic(), *InTargetPlatform);
+	UObject* Outer = bIsCooking ? GetPublic() : (UObject*)GetTransientPackage();
+	TObjectPtr<UModelResources> LocalModelResources = NewObject<UModelResources>(Outer, FName(*ModelResourcesName), RF_Public);
 	
 	FObjectAndNameAsStringProxyArchive ObjectReader(MemoryReader, true);
-	const bool bLoadedSuccessfully = LocalModelResources.Unserialize(ObjectReader, *GetPublic(), InTargetPlatform, bIsCooking);
+	LocalModelResources->Serialize(ObjectReader);
 
-	GetModelResources(bIsCooking) = MoveTemp(LocalModelResources);
+	const bool bLoadedSuccessfully = LocalModelResources->CodeVersion == UCustomizableObjectPrivate::CurrentSupportedVersion;
+	SetModelResources(bLoadedSuccessfully ? LocalModelResources : nullptr, bIsCooking);
 
 	return bLoadedSuccessfully;
 }
@@ -566,415 +556,75 @@ void UCustomizableObjectPrivate::LoadModel(FArchive& MemoryReader)
 }
 
 
-void SerializeStreamedResources(FArchive& Ar, TArray<FCustomizableObjectStreamedResourceData>& StreamedResources)
+void UModelResources::InitCookData()
 {
-	check(Ar.IsSaving());
-
-	int32 NumStreamedResources = StreamedResources.Num();
-	Ar << NumStreamedResources;
+	UCustomizableObject* CustomizableObject = Cast<UCustomizableObject>(GetOuter());
+	check(CustomizableObject);
 
-	for (const FCustomizableObjectStreamedResourceData& ResourceData : StreamedResources)
+	const FString ObjectName = GetNameSafe(CustomizableObject);
+	
+	const int32 NumStreamedResources = StreamedResourceDataEditor.Num();
+	for (int32 Index = 0; Index < NumStreamedResources; ++Index)
 	{
-		const FCustomizableObjectResourceData& Data = ResourceData.GetPath().LoadSynchronous()->Data;
-		uint32 ResourceDataType = (uint32)Data.Type;
-		Ar << ResourceDataType;
+		FCustomizableObjectResourceData& ResourceData = StreamedResourceDataEditor[Index];
 
-		switch (Data.Type)
-		{
-		case ECOResourceDataType::AssetUserData:
-		{
-			const FCustomizableObjectAssetUserData* AssetUserData = Data.Data.GetPtr<FCustomizableObjectAssetUserData>();
-			
-			FString AssetUserDataPath;
-			
-			if (AssetUserData && AssetUserData->AssetUserDataEditor)
-			{
-				AssetUserDataPath = TSoftObjectPtr<UAssetUserData>(AssetUserData->AssetUserDataEditor).ToString();
-			}
-			else
-			{
-				UE_LOG(LogMutable, Warning, TEXT("Failed to serialize streamed resource of type AssetUserData."));
-			}
-
-			Ar << AssetUserDataPath;
-			break;
-		}
-		default:
-			check(false);
-			break;
-		}
-	}
-}
-
-
-void UnserializeStreamedResources(FArchive& Ar, UObject* Object, TArray<FCustomizableObjectStreamedResourceData>& StreamedResources, bool bIsCooking)
-{
-	check(Ar.IsLoading());
-
-	const FString CustomizableObjectName = GetNameSafe(Object) + TEXT("_");
-
-	int32 NumStreamedResources = 0;
-	Ar << NumStreamedResources;
-
-	StreamedResources.SetNum(NumStreamedResources);
-
-	for (int32 ResourceIndex = 0; ResourceIndex < NumStreamedResources; ++ResourceIndex)
-	{
-		// Override existing containers
-		UCustomizableObjectResourceDataContainer* Container = StreamedResources[ResourceIndex].GetPath().Get();
+		const FString& ContainerName = ObjectName + FString::Printf(TEXT("_SR_%d"), Index);
 
-		// Create a new container if none.
+		UCustomizableObjectResourceDataContainer* Container = FindObject<UCustomizableObjectResourceDataContainer>(CustomizableObject, *ContainerName);
 		if (!Container)
 		{
-			// Generate a deterministic name to help with deterministic cooking
-			const FString ContainerName = CustomizableObjectName + FString::Printf(TEXT("SR_%d"), ResourceIndex);
-
-			UCustomizableObjectResourceDataContainer* ExistingContainer = FindObject<UCustomizableObjectResourceDataContainer>(Object, *ContainerName);
-			Container = ExistingContainer ? ExistingContainer : NewObject<UCustomizableObjectResourceDataContainer>(
-				Object,
+			Container = NewObject<UCustomizableObjectResourceDataContainer>(
+				CustomizableObject,
 				FName(*ContainerName),
 				RF_Public);
-
-			StreamedResources[ResourceIndex] = { Container };
 		}
 
-		check(Container);
-		uint32 Type = 0;
-		Ar << Type;
+		Container->Data = ResourceData;
 
-		Container->Data.Type = (ECOResourceDataType)Type;
-		switch (Container->Data.Type)
-		{
-		case ECOResourceDataType::AssetUserData:
+		if (FCustomizableObjectAssetUserData* AUDResource = Container->Data.Data.GetMutablePtr<FCustomizableObjectAssetUserData>())
 		{
-			FString AssetUserDataPath;
-			Ar << AssetUserDataPath;
+			// Find or duplicate the AUD replacing the outer
+			const FString AssetName = GetNameSafe(AUDResource->AssetUserData);
+			check(AUDResource->AssetUserData != nullptr);
 
-			FCustomizableObjectAssetUserData ResourceData;
-
-			TSoftObjectPtr<UAssetUserData> SoftAssetUserData = TSoftObjectPtr<UAssetUserData>(FSoftObjectPath(AssetUserDataPath));
-			ResourceData.AssetUserDataEditor = !SoftAssetUserData.IsNull() ? SoftAssetUserData.LoadSynchronous() : nullptr;
-
-			if (!ResourceData.AssetUserDataEditor)
+			UAssetUserData* AssetUserData = FindObject<UAssetUserData>(Container, *AssetName);
+			if (!AssetUserData)
 			{
-				UE_LOG(LogMutable, Warning, TEXT("Failed to load streamed resource of type AssetUserData. Resource name: [%s]"), *AssetUserDataPath);
-			}
-
-			if (bIsCooking)
-			{
-				// Rename the asset user data for duplicate
-				const FString AssetName = CustomizableObjectName + GetNameSafe(ResourceData.AssetUserDataEditor);
-
-				// Find or duplicate the AUD replacing the outer
-				ResourceData.AssetUserData = FindObject<UAssetUserData>(Container, *AssetName);
-				if (!ResourceData.AssetUserData)
-				{
-					// AUD may be private objects within meshes. Duplicate changing the outer to avoid including meshes into the builds.
-					ResourceData.AssetUserData = DuplicateObject<UAssetUserData>(ResourceData.AssetUserDataEditor, Container, FName(*AssetName));
-				}
+				// AUD may be private objects within meshes. Duplicate changing the outer to avoid including meshes into the builds.
+				AssetUserData = DuplicateObject<UAssetUserData>(AUDResource->AssetUserData, Container, FName(*AssetName));
 			}
-
-			Container->Data.Data = FInstancedStruct::Make(ResourceData);
-			break;
-		}
-		default:
-			check(false);
-			break;
+			
+			AUDResource->AssetUserData = AssetUserData;
 		}
-	}
-}
-
-
-void FModelResources::Serialize(FObjectAndNameAsStringProxyArchive& MemoryWriter, bool bIsCooking)
-{
-	MUTABLE_CPUPROFILER_SCOPE(FModelResources::Serialize);
-	check(IsInGameThread());
-
-	int32 SupportedVersion = UCustomizableObjectPrivate::CurrentSupportedVersion;
-	MemoryWriter << SupportedVersion;
-
-	MemoryWriter << ReferenceSkeletalMeshesData;
-
-	SerializeStreamedResources(MemoryWriter, StreamedResourceData);
-
-	UExtensionData* ExtensionData = NewObject<UExtensionData>();
-	ExtensionData->CopyFrom(*this);
-	ExtensionData->Serialize(MemoryWriter);
-	
-	int32 NumReferencedMaterials = Materials.Num();
-	MemoryWriter << NumReferencedMaterials;
-
-	for (const TSoftObjectPtr<UMaterialInterface>& Material : Materials)
-	{
-		FString StringRef = Material.ToString();
-		MemoryWriter << StringRef;
-	}
-
-	int32 NumReferencedSkeletons = Skeletons.Num();
-	MemoryWriter << NumReferencedSkeletons;
-
-	for (const TSoftObjectPtr<USkeleton>& Skeleton : Skeletons)
-	{
-		FString StringRef = Skeleton.ToString();
-		MemoryWriter << StringRef;
-	}
-
-	int32 NumPassthroughTextures = PassThroughTextures.Num();
-	MemoryWriter << NumPassthroughTextures;
-
-	for (const TSoftObjectPtr<UTexture>& PassthroughTexture : PassThroughTextures)
-	{
-		FString StringRef = PassthroughTexture.ToString();
-		MemoryWriter << StringRef;
-	}
-
-	int32 NumPassthroughMeshes = PassThroughMeshes.Num();
-	MemoryWriter << NumPassthroughMeshes;
-
-	for (const TSoftObjectPtr<UStreamableRenderAsset>& PassthroughMesh : PassThroughMeshes)
-	{
-		FString StringRef = PassthroughMesh.ToString();
-		MemoryWriter << StringRef;
-	}
-
-#if WITH_EDITORONLY_DATA
-	int32 NumRuntimeReferencedTextures = RuntimeReferencedTextures.Num();
-	MemoryWriter << NumRuntimeReferencedTextures;
-
-	for (const TSoftObjectPtr<const UTexture>& RuntimeReferencedTexture : RuntimeReferencedTextures)
-	{
-		FString StringRef = RuntimeReferencedTexture.ToString();
-		MemoryWriter << StringRef;
-	}
-#endif
-
-	int32 NumPhysicsAssets = PhysicsAssets.Num();
-	MemoryWriter << NumPhysicsAssets;
-
-	for (const TSoftObjectPtr<UPhysicsAsset>& PhysicsAsset : PhysicsAssets)
-	{
-		FString StringRef = PhysicsAsset.ToString();
-		MemoryWriter << StringRef;
-	}
-
-	int32 NumAnimBps = AnimBPs.Num();
-	MemoryWriter << NumAnimBps;
-
-	for (const TSoftClassPtr<UAnimInstance>& AnimBp : AnimBPs)
-	{
-		FString StringRef = AnimBp.ToString();
-		MemoryWriter << StringRef;
-	}
-
-	MemoryWriter << AnimBpOverridePhysiscAssetsInfo;
-
-	MemoryWriter << MaterialSlotNames;
-	MemoryWriter << BoneNamesMap;
-	MemoryWriter << SocketArray;
-
-	MemoryWriter << SkinWeightProfilesInfo;
-
-	MemoryWriter << ImageProperties;
-	MemoryWriter << MeshMetadata;
-	MemoryWriter << SurfaceMetadata;
-	MemoryWriter << ParameterUIDataMap;
-	MemoryWriter << StateUIDataMap;
-
-#if WITH_EDITORONLY_DATA
-	MemoryWriter << IntParameterOptionDataTable;
-#endif
-
-	MemoryWriter << ClothingAssetsData;
-	MemoryWriter << ClothSharedConfigsData;
-
-	MemoryWriter << ComponentLODs;
-	MemoryWriter << ComponentNumLODsToStream;
-	MemoryWriter << ComponentFirstLODAvailable;
-
-	MemoryWriter << ComponentNamesPerObjectComponent;
-	MemoryWriter << ReleaseVersion;
-
-	// Editor Only data
-	if (!bIsCooking)
-	{
-		MemoryWriter << bIsTextureStreamingDisabled;
-		MemoryWriter << bIsCompiledWithOptimization;
-		MemoryWriter << CustomizableObjectPathMap;
-		MemoryWriter << GroupNodeMap;
-		MemoryWriter << ParticipatingObjects;
-		MemoryWriter << TableToParamNames;
-
-		MemoryWriter << EditorOnlyMorphTargetReconstructionData;
-		MemoryWriter << EditorOnlyClothingMeshToMeshVertData;
-	}
-
-#if WITH_EDITORONLY_DATA
-	MemoryWriter << bCompiledWithHDTextureCompression;
-#endif
-}
-
-
-bool FModelResources::Unserialize(FObjectAndNameAsStringProxyArchive& MemoryReader, UCustomizableObject& Outer, const ITargetPlatform* InTargetPlatform, bool bIsCooking)
-{
-	MUTABLE_CPUPROFILER_SCOPE(FModelResources::Unserialize);
-	check(IsInGameThread());
-
-	int32 SupportedVersion = 0;
-	MemoryReader << SupportedVersion;
-
-	if (SupportedVersion != UCustomizableObjectPrivate::CurrentSupportedVersion)
-	{
-		return false;
-	}
-
-	MemoryReader << ReferenceSkeletalMeshesData;
-
-	UnserializeStreamedResources(MemoryReader, &Outer, StreamedResourceData, bIsCooking);
-
-	UExtensionData* ExtensionData = NewObject<UExtensionData>();
-	ExtensionData->Serialize(MemoryReader);
-	ExtensionData->MoveTo(*this);
-
-	// Initialize resources. 
-	for (FMutableRefSkeletalMeshData& ReferenceSkeletalMeshData : ReferenceSkeletalMeshesData)
-	{
-		ReferenceSkeletalMeshData.InitResources(&Outer, *this, InTargetPlatform);
-	}
-
-	int32 NumReferencedMaterials = 0;
-	MemoryReader << NumReferencedMaterials;
-	Materials.Reset(NumReferencedMaterials);
-
-	for (int32 i = 0; i < NumReferencedMaterials; ++i)
-	{
-		FString StringRef;
-		MemoryReader << StringRef;
-
-		Materials.Add(TSoftObjectPtr<UMaterialInterface>(FSoftObjectPath(StringRef)));
-	}
-
-	int32 NumReferencedSkeletons = 0;
-	MemoryReader << NumReferencedSkeletons;
-	Skeletons.Reset(NumReferencedMaterials);
-
-	for (int32 SkeletonIndex = 0; SkeletonIndex < NumReferencedSkeletons; ++SkeletonIndex)
-	{
-		FString StringRef;
-		MemoryReader << StringRef;
-
-		Skeletons.Add(TSoftObjectPtr<USkeleton>(FSoftObjectPath(StringRef)));
-	}
 
-	int32 NumPassthroughTextures = 0;
-	MemoryReader << NumPassthroughTextures;
-	PassThroughTextures.Reset(NumPassthroughTextures);
-
-	for (int32 Index = 0; Index < NumPassthroughTextures; ++Index)
-	{
-		FString StringRef;
-		MemoryReader << StringRef;
-
-		PassThroughTextures.Add(TSoftObjectPtr<UTexture>(FSoftObjectPath(StringRef)));
+		StreamedResourceData.Emplace(Container);
 	}
 
-	int32 NumPassthroughMeshes = 0;
-	MemoryReader << NumPassthroughMeshes;
-	PassThroughMeshes.Reset(NumPassthroughMeshes);
+	const TArrayView<const UCustomizableObjectExtension* const> AllExtensions = ICustomizableObjectModule::Get().GetRegisteredExtensions();
 
-	for (int32 Index = 0; Index < NumPassthroughMeshes; ++Index)
+	const int32 NumStreamedExtensionResources = StreamedExtensionDataEditor.Num();
+	for (int32 Index = 0; Index < NumStreamedExtensionResources; ++Index)
 	{
-		FString StringRef;
-		MemoryReader << StringRef;
-
-		PassThroughMeshes.Add(TSoftObjectPtr<UStreamableRenderAsset>(FSoftObjectPath(StringRef)));
-	}
-
-#if WITH_EDITORONLY_DATA
-	int32 NumRuntimeReferencedTextures = 0;
-	MemoryReader << NumRuntimeReferencedTextures;
-	RuntimeReferencedTextures.Reset(NumRuntimeReferencedTextures);
-
-	for (int32 Index = 0; Index < NumRuntimeReferencedTextures; ++Index)
-	{
-		FString StringRef;
-		MemoryReader << StringRef;
-
-		RuntimeReferencedTextures.Add(TSoftObjectPtr<const UTexture>(FSoftObjectPath(StringRef)));
-	}
-#endif
-
-	int32 NumPhysicsAssets = 0;
-	MemoryReader << NumPhysicsAssets;
-	PhysicsAssets.Reset(NumPhysicsAssets);
-
-	for (int32 i = 0; i < NumPhysicsAssets; ++i)
-	{
-		FString StringRef;
-		MemoryReader << StringRef;
-
-		PhysicsAssets.Add(TSoftObjectPtr<UPhysicsAsset>(FSoftObjectPath(StringRef)));
-	}
-
-
-	int32 NumAnimBps = 0;
-	MemoryReader << NumAnimBps;
-	AnimBPs.Reset(NumAnimBps);
-
-	for (int32 Index = 0; Index < NumAnimBps; ++Index)
-	{
-		FString StringRef;
-		MemoryReader << StringRef;
-
-		AnimBPs.Add(TSoftClassPtr<UAnimInstance>(StringRef));
-	}
-
-	MemoryReader << AnimBpOverridePhysiscAssetsInfo;
-
-	MemoryReader << MaterialSlotNames;
-	MemoryReader << BoneNamesMap;
-	MemoryReader << SocketArray;
-
-	MemoryReader << SkinWeightProfilesInfo;
-
-	MemoryReader << ImageProperties;
-	MemoryReader << MeshMetadata;
-	MemoryReader << SurfaceMetadata;
-	MemoryReader << ParameterUIDataMap;
-	MemoryReader << StateUIDataMap;
-
-#if WITH_EDITORONLY_DATA
-	MemoryReader << IntParameterOptionDataTable;
-#endif
-
-	MemoryReader << ClothingAssetsData;
-	MemoryReader << ClothSharedConfigsData;
-
-	MemoryReader << ComponentLODs;
-	MemoryReader << ComponentNumLODsToStream;
-	MemoryReader << ComponentFirstLODAvailable;
+		const FString& ContainerName = ObjectName + FString::Printf(TEXT("_SE_%d"), Index);
+		FCustomizableObjectResourceData& ResourceData = StreamedExtensionDataEditor[Index];
 
-	MemoryReader << ComponentNamesPerObjectComponent;
-	MemoryReader << ReleaseVersion;
+		UCustomizableObjectResourceDataContainer* Container = FindObject<UCustomizableObjectResourceDataContainer>(CustomizableObject, *ContainerName);
+		if (!Container)
+		{
+			Container = NewObject<UCustomizableObjectResourceDataContainer>(
+				CustomizableObject,
+				FName(*ContainerName),
+				RF_Public);
+		}
 
-	// Editor Only data
-	if (!bIsCooking)
-	{
-		MemoryReader << bIsTextureStreamingDisabled;
-		MemoryReader << bIsCompiledWithOptimization;
-		MemoryReader << CustomizableObjectPathMap;
-		MemoryReader << GroupNodeMap;
-		MemoryReader << ParticipatingObjects;
-		MemoryReader << TableToParamNames;
+		for (const UCustomizableObjectExtension* Extension : AllExtensions)
+		{
+			Extension->MovePrivateReferencesToContainer(ResourceData.Data, Container);
+		}
 
-		MemoryReader << EditorOnlyMorphTargetReconstructionData;
-		MemoryReader << EditorOnlyClothingMeshToMeshVertData;
+		Container->Data = MoveTemp(ResourceData);
+		StreamedExtensionData.Emplace(Container);
 	}
-
-#if WITH_EDITORONLY_DATA
-	MemoryReader << bCompiledWithHDTextureCompression;
-#endif
-
-	return true;
 }
 
 
@@ -1320,12 +970,14 @@ USkeletalMesh* UCustomizableObject::GetComponentMeshReferenceSkeletalMesh(const
 	}
 #endif
 
-	const FModelResources& ModelResources = Private->GetModelResources();
-	int32 ObjectComponentIndex = ModelResources.ComponentNamesPerObjectComponent.IndexOfByKey(ComponentName);
-	if (ModelResources.ReferenceSkeletalMeshesData.IsValidIndex(ObjectComponentIndex))
+	if (const UModelResources* ModelResources = Private->GetModelResources())
 	{
-		// Can be nullptr if RefSkeletalMeshes are not loaded yet.
-		return ModelResources.ReferenceSkeletalMeshesData[ObjectComponentIndex].SkeletalMesh;
+		int32 ObjectComponentIndex = ModelResources->ComponentNamesPerObjectComponent.IndexOfByKey(ComponentName);
+		if (ModelResources->ReferenceSkeletalMeshesData.IsValidIndex(ObjectComponentIndex))
+		{
+			// Can be nullptr if RefSkeletalMeshes are not loaded yet.
+			return ModelResources->ReferenceSkeletalMeshesData[ObjectComponentIndex].SkeletalMesh;
+		}
 	}
 	
 	return nullptr;
@@ -1468,9 +1120,9 @@ UCustomizableObjectInstance* UCustomizableObject::CreateInstance()
 
 int32 UCustomizableObject::GetComponentCount() const
 {
-	if (IsCompiled())
+	if (const UModelResources* ModelResources = GetPrivate()->GetModelResources())
 	{
-		return GetPrivate()->GetModelResources().ComponentNamesPerObjectComponent.Num();
+		return ModelResources->ComponentNamesPerObjectComponent.Num();
 	}
 
 	return 0;
@@ -1479,9 +1131,9 @@ int32 UCustomizableObject::GetComponentCount() const
 
 FName UCustomizableObjectPrivate::GetComponentName(FCustomizableObjectComponentIndex ObjectComponentIndex) const
 {
-	if (GetPublic()->IsCompiled())
+	if (const UModelResources* LocalModelResources = GetModelResources())
 	{
-		const TArray<FName>& ComponentNames = GetModelResources().ComponentNamesPerObjectComponent;
+		const TArray<FName>& ComponentNames = LocalModelResources->ComponentNamesPerObjectComponent;
 		if (ComponentNames.IsValidIndex(ObjectComponentIndex.GetValue()))
 		{
 			return ComponentNames[ObjectComponentIndex.GetValue()];
@@ -1958,27 +1610,39 @@ USkeletalMesh* UCustomizableObject::GetRefSkeletalMesh(int32 ObjectComponentInde
 
 FMutableParamUIMetadata UCustomizableObject::GetParameterUIMetadata(const FString& ParamName) const
 {
-	const FMutableParameterData* ParameterData = Private->GetModelResources().ParameterUIDataMap.Find(ParamName);
-	return ParameterData ? ParameterData->ParamUIMetadata : FMutableParamUIMetadata();
+	if (const UModelResources* ModelResources = Private->GetModelResources())
+	{
+		const FMutableParameterData* ParameterData = ModelResources->ParameterUIDataMap.Find(ParamName);
+		return ParameterData ? ParameterData->ParamUIMetadata : FMutableParamUIMetadata();
+	}
+
+	return FMutableParamUIMetadata();
 }
 
 
 FMutableParamUIMetadata UCustomizableObject::GetEnumParameterValueUIMetadata(const FString& ParamName, const FString& OptionName) const
 {
+	const UModelResources* ModelResources = Private->GetModelResources();
+	if (!ModelResources)
+	{
+		return FMutableParamUIMetadata();
+	}
+
 	const int32 ParameterIndex = GetPrivate()->FindParameter(ParamName);
 	if (ParameterIndex == INDEX_NONE)
 	{
-		return {};
+		return FMutableParamUIMetadata();
 	}
-	
-	const FMutableParameterData* ParameterData = Private->GetModelResources().ParameterUIDataMap.Find(ParamName);
-	if (!ParameterData)
+
+	if (const FMutableParameterData* ParameterData = ModelResources->ParameterUIDataMap.Find(ParamName))
 	{
-		return {};
+		if (const FIntegerParameterUIData* IntegerParameterUIData = ParameterData->ArrayIntegerParameterOption.Find(OptionName))
+		{
+			return IntegerParameterUIData->ParamUIMetadata;
+		}
 	}
 
-	const FIntegerParameterUIData* IntegerParameterUIData = ParameterData->ArrayIntegerParameterOption.Find(OptionName);
-	return IntegerParameterUIData ? IntegerParameterUIData->ParamUIMetadata : FMutableParamUIMetadata();
+	return FMutableParamUIMetadata();
 }
 
 
@@ -1990,19 +1654,24 @@ FMutableParamUIMetadata UCustomizableObject::GetIntParameterOptionUIMetadata(con
 
 ECustomizableObjectGroupType UCustomizableObject::GetEnumParameterGroupType(const FString& ParamName) const
 {
+	const UModelResources* ModelResources = Private->GetModelResources();
+	if (!ModelResources)
+	{
+		return  ECustomizableObjectGroupType::COGT_TOGGLE;
+	}
+
 	const int32 ParameterIndex = GetPrivate()->FindParameter(ParamName);
 	if (ParameterIndex == INDEX_NONE)
 	{
 		return ECustomizableObjectGroupType::COGT_TOGGLE;
-	}		
-	
-	const FMutableParameterData* ParameterData = Private->GetModelResources().ParameterUIDataMap.Find(ParamName);
-	if (!ParameterData)
+	}
+
+	if (const FMutableParameterData* ParameterData = ModelResources->ParameterUIDataMap.Find(ParamName))
 	{
-		return ECustomizableObjectGroupType::COGT_TOGGLE;
+		return ParameterData->IntegerParameterGroupType;
 	}
 
-	return ParameterData->IntegerParameterGroupType;
+	return ECustomizableObjectGroupType::COGT_TOGGLE;
 }
 
 
@@ -2014,18 +1683,35 @@ ECustomizableObjectGroupType UCustomizableObject::GetIntParameterGroupType(const
 
 FMutableStateUIMetadata UCustomizableObject::GetStateUIMetadata(const FString& StateName) const
 {
-	const FMutableStateData* StateData = Private->GetModelResources().StateUIDataMap.Find(StateName);
-	return StateData ? StateData->StateUIMetadata : FMutableStateUIMetadata();
+	if (const UModelResources* ModelResources = Private->GetModelResources())
+	{
+		if (const FMutableStateData* StateData = ModelResources->StateUIDataMap.Find(StateName))
+		{
+			return StateData->StateUIMetadata;
+		}
+	}
+
+	return FMutableStateUIMetadata();
 }
 
 
 #if WITH_EDITOR
+
+uint32 GetTypeHash(const FIntegerParameterOptionKey& Key)
+{
+	uint32 Hash = GetTypeHash(Key.ParameterName);
+	Hash = HashCombine(Hash, GetTypeHash(Key.ParameterOption));
+	return Hash;
+}
+
 TArray<TSoftObjectPtr<UDataTable>> UCustomizableObject::GetEnumParameterValueDataTable(const FString& ParamName, const FString& ValueName)
 {
-	const FModelResources& ModelResources = GetPrivate()->GetModelResources();
-	if (const TSet<TSoftObjectPtr<UDataTable>>* Result = ModelResources.IntParameterOptionDataTable.Find(MakeTuple(ParamName, ValueName)))
+	if(const UModelResources* LocalModelResources = GetPrivate()->GetModelResources())
 	{
-		return Result->Array();
+		if (const FIntegerParameterOptionDataTable* Result = LocalModelResources->IntParameterOptionDataTable.Find({ ParamName, ValueName }))
+		{
+			return Result->DataTables.Array();
+		}
 	}
 	
 	return {};
@@ -2244,8 +1930,14 @@ bool UCustomizableObjectPrivate::IsParameterMultidimensional(const int32& InPara
 
 void UCustomizableObjectPrivate::ApplyStateForcedValuesToParameters(FCustomizableObjectInstanceDescriptor& Descriptor)
 {
+	const UModelResources* LocalModelResources = GetModelResources();
+	if (!LocalModelResources)
+	{
+		return;
+	}
+
 	const FString& StateName = Descriptor.GetCurrentState();
-	const FMutableStateData* StateData = GetModelResources().StateUIDataMap.Find(StateName);
+	const FMutableStateData* StateData = LocalModelResources->StateUIDataMap.Find(StateName);
 	if (!StateData)
 	{
 		return;
@@ -2278,13 +1970,13 @@ void UCustomizableObjectPrivate::GetLowPriorityTextureNames(TArray<FString>& Out
 {
 	OutTextureNames.Reset(GetPublic()->LowPriorityTextures.Num());
 
-	if (!GetPublic()->LowPriorityTextures.IsEmpty())
+	const UModelResources* LocalModelResources = GetModelResources();
+	if (LocalModelResources && !GetPublic()->LowPriorityTextures.IsEmpty())
 	{
-		const FModelResources& LocalModelResources = GetModelResources();
-		const int32 ImageCount = LocalModelResources.ImageProperties.Num();
+		const int32 ImageCount = LocalModelResources->ImageProperties.Num();
 		for (int32 ImageIndex = 0; ImageIndex < ImageCount; ++ImageIndex)
 		{
-			if (GetPublic()->LowPriorityTextures.Find(FName(LocalModelResources.ImageProperties[ImageIndex].TextureParameterName)) != INDEX_NONE)
+			if (GetPublic()->LowPriorityTextures.Find(FName(LocalModelResources->ImageProperties[ImageIndex].TextureParameterName)) != INDEX_NONE)
 			{
 				OutTextureNames.Add(FString::FromInt(ImageIndex));
 			}
@@ -2317,9 +2009,14 @@ uint8 UCustomizableObjectPrivate::GetMinLODIndex(const FName& ComponentName) con
 
 	// Get the first lod the current platform can get generated. It represents the Min LOD mutable can generate for the current platform.
 	uint8 FirstLODAvailable = 0;
-	if (const uint8* Result = GetModelResources().ComponentFirstLODAvailable.Find(ComponentName))
+
+	const UModelResources* LocalModelResources = GetModelResources();
+	if (LocalModelResources)
 	{
-		FirstLODAvailable = *Result;
+		if (const uint8* Result = LocalModelResources->ComponentFirstLODAvailable.Find(ComponentName))
+		{
+			FirstLODAvailable = *Result;
+		}
 	}
 
 	return FMath::Max(MinLODIdx, FirstLODAvailable);
@@ -2459,20 +2156,6 @@ void UModelStreamableData::Serialize(FArchive& Ar)
 }
 
 
-void UExtensionData::CopyFrom(const FModelResources& ModelResources)
-{
-	AlwaysLoadedExtensionData = ModelResources.AlwaysLoadedExtensionData;
-	StreamedExtensionData = ModelResources.StreamedExtensionData;
-}
-
-
-void UExtensionData::MoveTo(FModelResources& ModelResources)
-{
-	ModelResources.AlwaysLoadedExtensionData = MoveTemp(AlwaysLoadedExtensionData);
-	ModelResources.StreamedExtensionData = MoveTemp(StreamedExtensionData);
-}
-
-
 void UCustomizableObjectPrivate::SetModel(const TSharedPtr<mu::FModel, ESPMode::ThreadSafe>& Model, const FGuid Id)
 {
 	if (MutableModel == Model
@@ -2548,14 +2231,14 @@ TSharedPtr<FModelStreamableBulkData> UCustomizableObjectPrivate::GetModelStreama
 }
 
 
-FModelResources& UCustomizableObjectPrivate::GetModelResources()
+UModelResources* UCustomizableObjectPrivate::GetModelResources()
 {
 	const UCustomizableObjectPrivate* ConstThis = this;
-	return *const_cast<FModelResources*>(&ConstThis->GetModelResources());
+	return const_cast<UModelResources*>(ConstThis->GetModelResources());
 }
 
 
-const FModelResources& UCustomizableObjectPrivate::GetModelResources() const
+const UModelResources* UCustomizableObjectPrivate::GetModelResources() const
 {
 #if WITH_EDITORONLY_DATA
 	return ModelResourcesEditor;
@@ -2564,19 +2247,40 @@ const FModelResources& UCustomizableObjectPrivate::GetModelResources() const
 #endif
 }
 
+const UModelResources& UCustomizableObjectPrivate::GetModelResourcesChecked() const
+{
+	const UModelResources* LocalModelResources = GetModelResources();
+	check(LocalModelResources);
+
+	return *LocalModelResources;
+}
+
 
 #if WITH_EDITORONLY_DATA
-FModelResources& UCustomizableObjectPrivate::GetModelResources(bool bIsCooking)
+UModelResources* UCustomizableObjectPrivate::GetModelResources(bool bIsCooking)
 {
 	const UCustomizableObjectPrivate* ConstThis = this;
-	return *const_cast<FModelResources*>(&ConstThis->GetModelResources(bIsCooking));
+	return const_cast<UModelResources*>(ConstThis->GetModelResources(bIsCooking));
 }
 
 
-const FModelResources& UCustomizableObjectPrivate::GetModelResources(bool bIsCooking) const
+const UModelResources* UCustomizableObjectPrivate::GetModelResources(bool bIsCooking) const
 {
 	return bIsCooking ? ModelResources : ModelResourcesEditor;
 }
+
+
+void UCustomizableObjectPrivate::SetModelResources(UModelResources* InModelResources, bool bIsCooking)
+{
+	if (bIsCooking)
+	{
+		ModelResources = InModelResources;
+	}
+	else
+	{
+		ModelResourcesEditor = InModelResources;
+	}
+}
 #endif
 
 
@@ -3425,7 +3129,7 @@ void UCustomizableObjectPrivate::LoadCompiledDataFromDDC(FCompilationOptions Opt
 		FMemoryReaderView ModelResourcesReader(ModelResourcesBytesDDC.GetView());
 		if (LoadModelResources(ModelResourcesReader, Options.TargetPlatform, Options.bIsCooking))
 		{
-			FModelResources& LocalModelResources = GetModelResources(Options.bIsCooking);
+			UModelResources& LocalModelResources = *GetModelResources(Options.bIsCooking);
 			LocalModelResources.bIsStoredInDDC = true;
 			LocalModelResources.DDCKey = CacheKey;
 			LocalModelResources.DDCDefaultPolicy = ECachePolicy::Default;
@@ -3439,9 +3143,10 @@ void UCustomizableObjectPrivate::LoadCompiledDataFromDDC(FCompilationOptions Opt
 
 		TSharedPtr<mu::FModel> Model = GetModel();
 		TSharedPtr<FModelStreamableBulkData> ModelStreamables = GetModelStreamableBulkData(Options.bIsCooking);
+		UModelResources* LocalModelResources = GetModelResources(Options.bIsCooking);
 
 		// Cache CookedPlatfomData. 
-		if (Options.bIsCooking && Model && ModelStreamables)
+		if (Options.bIsCooking && Model && LocalModelResources && ModelStreamables)
 		{
 			// Sync cache cooked platform data
 			// TODO UE-220138: Sync -> Async
@@ -3451,7 +3156,10 @@ void UCustomizableObjectPrivate::LoadCompiledDataFromDDC(FCompilationOptions Opt
 			
 			// Cache Model, ModelResources and ModelStreamables
 			CachedData.ModelData.Append(reinterpret_cast<const uint8*>(ModelBytesDDC.GetData()), ModelBytesDDC.GetSize());
-			CachedData.ModelResourcesData.Append(reinterpret_cast<const uint8*>(ModelResourcesBytesDDC.GetData()), ModelResourcesBytesDDC.GetSize());
+			
+			CachedData.ModelResources = TStrongObjectPtr<UModelResources>(LocalModelResources);
+			CachedData.ModelResources->InitCookData();
+
 			CachedData.ModelStreamables = ModelStreamables;
 
 			// Value Id to file mapping to reconstruct the cached data
@@ -3680,36 +3388,38 @@ void UCustomizableObjectPrivate::LogMemory()
 
 		int32 ModeleResData = 0;
 		{
-			ModeleResData += sizeof(FModelResources);
+			ModeleResData += sizeof(UModelResources);
 			
 			FMutableMemoryCounterArchive Arch;
-			Arch << ModelResources.ReferenceSkeletalMeshesData;
-			Arch << ModelResources.Skeletons;
-			Arch << ModelResources.Materials;
-			Arch << ModelResources.PassThroughTextures;
-			Arch << ModelResources.PassThroughMeshes;
-			Arch << ModelResources.PhysicsAssets;
-			Arch << ModelResources.AnimBPs;
-			Arch << ModelResources.AnimBpOverridePhysiscAssetsInfo;
-			Arch << ModelResources.MaterialSlotNames;
-			Arch << ModelResources.BoneNamesMap;
-			Arch << ModelResources.SocketArray;
-			Arch << ModelResources.SkinWeightProfilesInfo;
-			Arch << ModelResources.MeshMetadata;
-			Arch << ModelResources.SurfaceMetadata;
-			Arch << ModelResources.ParameterUIDataMap;
-			Arch << ModelResources.StateUIDataMap;
-			Arch << ModelResources.ClothSharedConfigsData;
-			Arch << ModelResources.ClothingAssetsData;
-			Arch << ModelResources.ComponentNamesPerObjectComponent;
-
-			ModeleResData += ModelResources.StreamedResourceData.GetAllocatedSize();
-			for (FCustomizableObjectStreamedResourceData& Data : ModelResources.StreamedResourceData)
+			if (ModelResources)
 			{
-				TSoftObjectPtr<UCustomizableObjectResourceDataContainer> PathCopy = Data.GetPath();
-				Arch << PathCopy;
+				Arch << ModelResources->ReferenceSkeletalMeshesData;
+				Arch << ModelResources->Skeletons;
+				Arch << ModelResources->Materials;
+				Arch << ModelResources->PassThroughTextures;
+				Arch << ModelResources->PassThroughMeshes;
+				Arch << ModelResources->PhysicsAssets;
+				Arch << ModelResources->AnimBPs;
+				Arch << ModelResources->AnimBpOverridePhysiscAssetsInfo;
+				Arch << ModelResources->MaterialSlotNames;
+				Arch << ModelResources->BoneNamesMap;
+				Arch << ModelResources->SocketArray;
+				Arch << ModelResources->SkinWeightProfilesInfo;
+				Arch << ModelResources->MeshMetadata;
+				Arch << ModelResources->SurfaceMetadata;
+				Arch << ModelResources->ParameterUIDataMap;
+				Arch << ModelResources->StateUIDataMap;
+				Arch << ModelResources->ClothSharedConfigsData;
+				Arch << ModelResources->ClothingAssetsData;
+				Arch << ModelResources->ComponentNamesPerObjectComponent;
+
+				ModeleResData += ModelResources->StreamedResourceData.GetAllocatedSize();
+				for (FCustomizableObjectStreamedResourceData& Data : ModelResources->StreamedResourceData)
+				{
+					TSoftObjectPtr<UCustomizableObjectResourceDataContainer> PathCopy = Data.GetPath();
+					Arch << PathCopy;
+				}
 			}
-
 			ModeleResData += Arch.TotalSize();
 		}
 
@@ -3867,6 +3577,15 @@ bool FMutableSkinWeightProfileInfo::operator==(const FMutableSkinWeightProfileIn
 	return Name == Other.Name;
 }
 
+FArchive& operator<<(FArchive& Ar, FMutableSkinWeightProfileInfo& Info)
+{
+	Ar << Info.Name;
+	Ar << Info.NameId;
+	Ar << Info.DefaultProfile;
+	Ar << Info.DefaultProfileFromLODIndex;
+
+	return Ar;
+}
 
 FIntegerParameterUIData::FIntegerParameterUIData(const FMutableParamUIMetadata& InParamUIMetadata)
 {
@@ -3880,58 +3599,7 @@ FMutableParameterData::FMutableParameterData(const FMutableParamUIMetadata& InPa
 	Type = InType;
 }
 
-
 #if WITH_EDITORONLY_DATA
-FArchive& operator<<(FArchive& Ar, FMutableRemappedBone& RemappedBone)
-{
-	Ar << RemappedBone.Name;
-	Ar << RemappedBone.Hash;
-	return Ar;
-}
-
-FArchive& operator<<(FArchive& Ar, FMutableModelImageProperties& ImageProps)
-{
-	Ar << ImageProps.TextureParameterName;
-	Ar << ImageProps.Filter;
-
-	// Bitfields don't serialize automatically with FArchive
-	if (Ar.IsLoading())
-	{
-		int32 Aux = 0;
-		Ar << Aux;
-		ImageProps.SRGB = Aux;
-
-		Aux = 0;
-		Ar << Aux;
-		ImageProps.FlipGreenChannel = Aux;
-
-		Aux = 0;
-		Ar << Aux;
-		ImageProps.IsPassThrough = Aux;
-	}
-	else
-	{
-		int32 Aux = ImageProps.SRGB;
-		Ar << Aux;
-
-		Aux = ImageProps.FlipGreenChannel;
-		Ar << Aux;
-
-		Aux = ImageProps.IsPassThrough;
-		Ar << Aux;
-	}
-
-	Ar << ImageProps.LODBias;
-	Ar << ImageProps.MipGenSettings;
-	Ar << ImageProps.LODGroup;
-
-	Ar << ImageProps.AddressX;
-	Ar << ImageProps.AddressY;
-
-	return Ar;
-}
-
-
 FArchive& operator<<(FArchive& Ar, FAnimBpOverridePhysicsAssetsInfo& Info)
 {
 	FString AnimInstanceClassPathString;
@@ -4013,114 +3681,21 @@ FArchive& operator<<(FArchive& Ar, FMutableRefSkeletalMeshSettings& Data)
 }
 
 
-FArchive& operator<<(FArchive& Ar, FMutableRefAssetUserData& Data)
-{
-	Ar << Data.AssetUserDataIndex;
-
-	return Ar;
-}
-
-
-FArchive& operator<<(FArchive& Ar, FMutableSkinWeightProfileInfo& Info)
-{
-	Ar << Info.Name;
-	Ar << Info.NameId;
-	Ar << Info.DefaultProfile;
-	Ar << Info.DefaultProfileFromLODIndex;
-
-	return Ar;
-}
-
-
-//-------------------------------------------------------------------------------------------------
-//-------------------------------------------------------------------------------------------------
-//-------------------------------------------------------------------------------------------------
-void FMutableRefSkeletalMeshData::InitResources(UCustomizableObject* InOuter, FModelResources& InModelResources, const ITargetPlatform* InTargetPlatform)
-{
-	check(InOuter);
-
-	const bool bHasServer = InTargetPlatform ? !InTargetPlatform->IsClientOnly() : false;
-	if (InOuter->bEnableUseRefSkeletalMeshAsPlaceholder || bHasServer)
-	{
-		SkeletalMesh = TSoftObjectPtr<USkeletalMesh>(SoftSkeletalMesh).LoadSynchronous();
-	}
-
-	// Initialize AssetUserData
-	for (FMutableRefAssetUserData& Data : AssetUserData)
-	{
-		if (!InModelResources.StreamedResourceData.IsValidIndex(Data.AssetUserDataIndex))
-		{
-			check(false);
-			continue;
-		}
-
-		Data.AssetUserData = InModelResources.StreamedResourceData[Data.AssetUserDataIndex].GetPath().LoadSynchronous();
-		check(Data.AssetUserData);
-	}
-}
-
-
 FArchive& operator<<(FArchive& Ar, FMutableRefSkeletalMeshData& Data)
 {
 	Ar << Data.LODData;
 	Ar << Data.Sockets;
 	Ar << Data.Bounds;
 	Ar << Data.Settings;
-
-	if (Ar.IsSaving())
-	{
-		FString AssetPath = Data.SoftSkeletalMesh.ToString();
-		Ar << AssetPath;
-
-		AssetPath = TSoftObjectPtr<USkeletalMeshLODSettings>(Data.SkeletalMeshLODSettings).ToString();
-		Ar << AssetPath;
-
-		AssetPath = TSoftObjectPtr<USkeleton>(Data.Skeleton).ToString();
-		Ar << AssetPath;
-
-		AssetPath = TSoftObjectPtr<UPhysicsAsset>(Data.PhysicsAsset).ToString();
-		Ar << AssetPath;
-
-		AssetPath = Data.PostProcessAnimInst.ToString();
-		Ar << AssetPath;
-
-		AssetPath = TSoftObjectPtr<UPhysicsAsset>(Data.ShadowPhysicsAsset).ToString();
-		Ar << AssetPath;
-
-	}
-	else
-	{
-		FString SkeletalMeshAssetPath;
-		Ar << SkeletalMeshAssetPath;
-		Data.SoftSkeletalMesh = SkeletalMeshAssetPath;
-
-		FString SkeletalMeshLODSettingsAssetPath;
-		Ar << SkeletalMeshLODSettingsAssetPath;
-		Data.SkeletalMeshLODSettings = TSoftObjectPtr<USkeletalMeshLODSettings>(FSoftObjectPath(SkeletalMeshLODSettingsAssetPath)).LoadSynchronous();
-
-		FString SkeletonAssetPath;
-		Ar << SkeletonAssetPath;
-		Data.Skeleton = TSoftObjectPtr<USkeleton>(FSoftObjectPath(SkeletonAssetPath)).LoadSynchronous();
-
-		FString PhysicsAssetPath;
-		Ar << PhysicsAssetPath;
-		Data.PhysicsAsset = TSoftObjectPtr<UPhysicsAsset>(FSoftObjectPath(PhysicsAssetPath)).LoadSynchronous();
-
-		FString PostProcessAnimInstAssetPath;
-		Ar << PostProcessAnimInstAssetPath;
-		Data.PostProcessAnimInst = TSoftClassPtr<UAnimInstance>(FSoftObjectPath(PostProcessAnimInstAssetPath)).LoadSynchronous();
-
-		FString ShadowPhysicsAssetPath;
-		Ar << ShadowPhysicsAssetPath;
-		Data.ShadowPhysicsAsset = TSoftObjectPtr<UPhysicsAsset>(FSoftObjectPath(ShadowPhysicsAssetPath)).LoadSynchronous();
-	}
-
-	Ar << Data.AssetUserData;
+	Ar << Data.SkeletalMesh;
+	Ar << Data.SkeletalMeshLODSettings;
+	Ar << Data.Skeleton;
+	Ar << Data.ShadowPhysicsAsset;
+	Ar << Data.AssetUserDataIndices;
 
 	return Ar;
 }
-#endif // WITH_EDITORONLY_DATA
-
+#endif 
 
 #if WITH_EDITOR
 FCompilationRequest::FCompilationRequest(UCustomizableObject& InCustomizableObject, bool bAsyncCompile)
@@ -4215,13 +3790,4 @@ bool FCompilationRequest::operator==(const FCompilationRequest& Other) const
 }
 #endif
 
-
-FArchive& operator<<(FArchive& Ar, FMutableParamNameSet& MutableParamNameSet)
-{
-	Ar << MutableParamNameSet.ParamNames;
-
-	return Ar;
-}
-
-
 #undef LOCTEXT_NAMESPACE
diff --git a/Engine/Plugins/Mutable/Source/CustomizableObject/Private/MuCO/CustomizableObjectInstance.cpp b/Engine/Plugins/Mutable/Source/CustomizableObject/Private/MuCO/CustomizableObjectInstance.cpp
--- a/Engine/Plugins/Mutable/Source/CustomizableObject/Private/MuCO/CustomizableObjectInstance.cpp
+++ b/Engine/Plugins/Mutable/Source/CustomizableObject/Private/MuCO/CustomizableObjectInstance.cpp
@@ -256,12 +256,12 @@ void UCustomizableInstancePrivate::InitCustomizableObjectData(const UCustomizabl
 FCustomizableInstanceComponentData* UCustomizableInstancePrivate::GetComponentData(const FName& ComponentName)
 {
 	UCustomizableObject* Object = GetPublic()->GetCustomizableObject();
-	if (!Object)
+	if (!Object || !Object->IsCompiled())
 	{
 		return nullptr;
 	}
 
-	const int32 ObjectComponentIndex = Object->GetPrivate()->GetModelResources().ComponentNamesPerObjectComponent.IndexOfByKey(ComponentName);
+	const int32 ObjectComponentIndex = Object->GetPrivate()->GetModelResourcesChecked().ComponentNamesPerObjectComponent.IndexOfByKey(ComponentName);
 	if (ObjectComponentIndex == INDEX_NONE)
 	{
 		return nullptr;
@@ -1920,6 +1920,8 @@ bool UCustomizableInstancePrivate::DoComponentsNeedUpdate(UCustomizableObjectIns
 		return false;
 	}
 
+	const UModelResources& ModelResources = *CustomizableObject->GetPrivate()->GetModelResources();
+
 	const int32 NumInstanceComponents = OperationData->InstanceUpdateData.Components.Num();
 
 	// To be indexed with instance component index
@@ -1976,7 +1978,7 @@ bool UCustomizableInstancePrivate::DoComponentsNeedUpdate(UCustomizableObjectIns
 			continue;
 		}
 
-		const FName ComponentName = CustomizableObject->GetPrivate()->GetModelResources().ComponentNamesPerObjectComponent[ObjectComponentIndex.GetValue()];
+		const FName ComponentName = ModelResources.ComponentNamesPerObjectComponent[ObjectComponentIndex.GetValue()];
 
 		if (OperationData->bUseMeshCache)
 		{
@@ -2067,7 +2069,7 @@ bool UCustomizableInstancePrivate::DoComponentsNeedUpdate(UCustomizableObjectIns
 			{
 				const FInstanceUpdateData::FComponent& Component = OperationData->InstanceUpdateData.Components[InstanceComponentIndex];
 				const FCustomizableObjectComponentIndex ObjectComponentIndex = Component.Id;
-				const FName ComponentName = OperationData->Instance->GetCustomizableObject()->GetPrivate()->GetModelResources().ComponentNamesPerObjectComponent[ObjectComponentIndex.GetValue()];
+				const FName ComponentName = ModelResources.ComponentNamesPerObjectComponent[ObjectComponentIndex.GetValue()];
 
 				if (ComponentName == OldMesh.Key)
 				{
@@ -2123,7 +2125,7 @@ bool UCustomizableInstancePrivate::UpdateSkeletalMesh_PostBeginUpdate0(UCustomiz
 
 	TMap<FName, TObjectPtr<USkeletalMesh>> OldSkeletalMeshes = SkeletalMeshes;
 
-	const FModelResources& ModelResources = CustomizableObject->GetPrivate()->GetModelResources();
+	const UModelResources& ModelResources = CustomizableObject->GetPrivate()->GetModelResourcesChecked();
 
 	// Collate the Extension Data on the instance into groups based on the extension that produced
 	// it, so that we only need to call extension functions such as OnSkeletalMeshCreated once for
@@ -2197,8 +2199,11 @@ bool UCustomizableInstancePrivate::UpdateSkeletalMesh_PostBeginUpdate0(UCustomiz
 
 				case mu::FExtensionData::EOrigin::ConstantStreamed:
 				{
+#if WITH_EDITOR
+					check(ModelResources.StreamedExtensionDataEditor.IsValidIndex(ExtensionOutput.Data->Index));
+					ReferencedExtensionData = &ModelResources.StreamedExtensionDataEditor[ExtensionOutput.Data->Index];
+#else
 					check(ModelResources.StreamedExtensionData.IsValidIndex(ExtensionOutput.Data->Index));
-
 					const FCustomizableObjectStreamedResourceData& StreamedData = ModelResources.StreamedExtensionData[ExtensionOutput.Data->Index];
 					if (!StreamedData.IsLoaded())
 					{
@@ -2212,6 +2217,7 @@ bool UCustomizableInstancePrivate::UpdateSkeletalMesh_PostBeginUpdate0(UCustomiz
 					}
 
 					ReferencedExtensionData = &StreamedData.GetLoadedData();
+#endif
 				}
 				break;
 
@@ -2255,11 +2261,11 @@ bool UCustomizableInstancePrivate::UpdateSkeletalMesh_PostBeginUpdate0(UCustomiz
 			continue;
 		}
 		const FCustomizableObjectComponentIndex ObjectComponentIndex = Component->Id;
-		if (!CustomizableObject->GetPrivate()->GetModelResources().ComponentNamesPerObjectComponent.IsValidIndex(ObjectComponentIndex.GetValue()))
+		if (!ModelResources.ComponentNamesPerObjectComponent.IsValidIndex(ObjectComponentIndex.GetValue()))
 		{
 			continue;
 		}
-		const FName ComponentName = CustomizableObject->GetPrivate()->GetModelResources().ComponentNamesPerObjectComponent[ObjectComponentIndex.GetValue()];
+		const FName ComponentName = ModelResources.ComponentNamesPerObjectComponent[ObjectComponentIndex.GetValue()];
 
 		FCustomizableInstanceComponentData* ComponentData = GetComponentData(ComponentName);
 		if (!ComponentData)
@@ -3094,7 +3100,11 @@ void UCustomizableInstancePrivate::SetReferenceSkeletalMesh() const
 		return;
 	}
 	
-	const FModelResources& ModelResources = CustomizableObject->GetPrivate()->GetModelResources();
+	const UModelResources* ModelResources = CustomizableObject->GetPrivate()->GetModelResources();
+	if (!ModelResources)
+	{
+		return;
+	}
 	
 	for (TObjectIterator<UCustomizableObjectInstanceUsage> It; It; ++It)
 	{
@@ -3112,8 +3122,8 @@ void UCustomizableInstancePrivate::SetReferenceSkeletalMesh() const
 #endif
 
 		const FName& ComponentName = CustomizableObjectInstanceUsage->GetComponentName();
-		const int32 ObjectComponentIndex = Instance->GetCustomizableObject()->GetPrivate()->GetModelResources().ComponentNamesPerObjectComponent.IndexOfByKey(ComponentName);
-		if (!ModelResources.ReferenceSkeletalMeshesData.IsValidIndex(ObjectComponentIndex))
+		const int32 ObjectComponentIndex = ModelResources->ComponentNamesPerObjectComponent.IndexOfByKey(ComponentName);
+		if (!ModelResources->ReferenceSkeletalMeshesData.IsValidIndex(ObjectComponentIndex))
 		{
 			continue;
 		}
@@ -3122,7 +3132,7 @@ void UCustomizableInstancePrivate::SetReferenceSkeletalMesh() const
 		{
 			Parent->EmptyOverrideMaterials();
 			
-			USkeletalMesh* SkeletalMesh = ModelResources.ReferenceSkeletalMeshesData[ObjectComponentIndex].SoftSkeletalMesh.LoadSynchronous();
+			USkeletalMesh* SkeletalMesh = ModelResources->ReferenceSkeletalMeshesData[ObjectComponentIndex].SoftSkeletalMesh.LoadSynchronous();
 			Parent->SetSkeletalMesh(SkeletalMesh);
 		}
 	}
@@ -3611,18 +3621,6 @@ void UCustomizableInstancePrivate::InitSkeletalMeshData(const TSharedRef<FUpdate
 
 	// Asset User Data
 	{
-		for (const FMutableRefAssetUserData& MutAssetUserData : RefSkeletalMeshData.AssetUserData)
-		{
-			if (const FCustomizableObjectAssetUserData* DataPtr = MutAssetUserData.AssetUserData->Data.Data.GetPtr<FCustomizableObjectAssetUserData>())
-			{
-#if WITH_EDITORONLY_DATA
-				SkeletalMesh->AddAssetUserData(DataPtr->AssetUserDataEditor);
-#else
-				SkeletalMesh->AddAssetUserData(DataPtr->AssetUserData);
-#endif
-			}
-		}
-
 		const FCustomizableInstanceComponentData* ComponentData = GetComponentData(ObjectComponentIndex);
 		check(ComponentData);
 		for (TObjectPtr<UAssetUserData> AssetUserData : ComponentData->AssetUserDataArray)
@@ -3778,7 +3776,7 @@ bool UCustomizableInstancePrivate::BuildSkeletonData(const TSharedRef<FUpdateCon
 
 		BoneInfoMap.Reserve(RawRefBoneCount);
 
-		const FModelResources& ModelResources = CustomizableObject.GetPrivate()->GetModelResources();
+		const UModelResources& ModelResources = CustomizableObject.GetPrivate()->GetModelResourcesChecked();
 		for (int32 Index = 0; Index < RawRefBoneCount; ++Index)
 		{
 			const FName BoneName = RawRefBoneInfo[Index].Name;
@@ -3899,7 +3897,7 @@ bool UCustomizableInstancePrivate::BuildSkeletonData(const TSharedRef<FUpdateCon
 }
 
 
-void UCustomizableInstancePrivate::BuildMeshSockets(const TSharedRef<FUpdateContextPrivate>& OperationData, USkeletalMesh* SkeletalMesh, const FModelResources& ModelResources, const FMutableRefSkeletalMeshData& RefSkeletalMeshData, TSharedPtr<const mu::FMesh> MutableMesh)
+void UCustomizableInstancePrivate::BuildMeshSockets(const TSharedRef<FUpdateContextPrivate>& OperationData, USkeletalMesh* SkeletalMesh, const UModelResources& ModelResources, const FMutableRefSkeletalMeshData& RefSkeletalMeshData, TSharedPtr<const mu::FMesh> MutableMesh)
 {
 	// Build mesh sockets.
 	MUTABLE_CPUPROFILER_SCOPE(UCustomizableInstancePrivate::BuildMeshSockets);
@@ -4113,7 +4111,7 @@ void UCustomizableInstancePrivate::BuildOrCopyClothingData(const TSharedRef<FUpd
 	// It must be not null as it's checked in the calling function
 	check(CustomizableObject);
 
-	const FModelResources& ModelResources = CustomizableObject->GetPrivate()->GetModelResources();
+	const UModelResources& ModelResources = CustomizableObject->GetPrivate()->GetModelResourcesChecked();
 	const TArray<FCustomizableObjectClothingAssetData>& ClothingAssetsData = ModelResources.ClothingAssetsData;
 	const TArray<FCustomizableObjectClothConfigData>& ClothSharedConfigsData = ModelResources.ClothSharedConfigsData;
 
@@ -5109,7 +5107,7 @@ bool UCustomizableInstancePrivate::BuildOrCopyRenderData(const TSharedRef<FUpdat
 		return false;
 	}
 
-	const FModelResources& ModelResources = CustomizableObject->GetPrivate()->GetModelResources();
+	const UModelResources& ModelResources = CustomizableObject->GetPrivate()->GetModelResourcesChecked();
 	const FName ComponentName = ModelResources.ComponentNamesPerObjectComponent[Component->Id.GetValue()];
 
 	const int32 FirstGeneratedLOD = FMath::Max(OperationData->GetRequestedLODs()[ComponentName], OperationData->GetMinLOD()[ComponentName]);
@@ -5344,7 +5342,7 @@ UE::Tasks::FTask UCustomizableInstancePrivate::LoadAdditionalAssetsAndData(const
 
 	UCustomizableObject* CustomizableObject = GetPublic()->GetCustomizableObject();
 
-	const FModelResources& ModelResources = CustomizableObject->GetPrivate()->GetModelResources();
+	const UModelResources& ModelResources = CustomizableObject->GetPrivate()->GetModelResourcesChecked();
 	const TSharedPtr<FModelStreamableBulkData>& ModelStreamableBulkData = CustomizableObject->GetPrivate()->GetModelStreamableBulkData();
 
 	TArray<FSoftObjectPath> AssetsToStream;
@@ -5443,9 +5441,21 @@ UE::Tasks::FTask UCustomizableInstancePrivate::LoadAdditionalAssetsAndData(const
 		for (int32 InstanceComponentIndex = 0; InstanceComponentIndex < OperationData->InstanceUpdateData.Components.Num(); ++InstanceComponentIndex)
 		{
 			const FInstanceUpdateData::FComponent& Component = Components[InstanceComponentIndex];
+			FCustomizableInstanceComponentData* ComponentData = GetComponentData(Component.Id);
 
-			const FName ComponentName = OperationData->ComponentNames[Component.Id.GetValue()];
 
+			const FCustomizableObjectComponentIndex ObjectComponentIndex = Component.Id;
+			const FMutableRefSkeletalMeshData& RefSkeletalMeshData = ModelResources.ReferenceSkeletalMeshesData[ObjectComponentIndex.GetValue()];
+
+			for (int32 AssetUserDataIndex : RefSkeletalMeshData.AssetUserDataIndices)
+			{
+#if !WITH_EDITOR
+				OperationData->StreamedResourceIndex.AddUnique(AssetUserDataIndex); // Used to hold/release streamed resources in non-editor builds.
+#endif
+				ComponentData->StreamedResourceIndex.AddUnique(AssetUserDataIndex);
+			}
+
+			const FName ComponentName = OperationData->ComponentNames[Component.Id.GetValue()];
 			for (int32 LODIndex = OperationData->FirstLODAvailable[ComponentName]; LODIndex < Component.LODCount; ++LODIndex)
 			{
 				const FInstanceUpdateData::FLOD& LOD = OperationData->InstanceUpdateData.LODs[Component.FirstLOD + LODIndex];
@@ -5456,31 +5466,19 @@ UE::Tasks::FTask UCustomizableInstancePrivate::LoadAdditionalAssetsAndData(const
 					continue;
 				}
 
-				FCustomizableInstanceComponentData* ComponentData = GetComponentData(Component.Id);
-
 				const TArray<uint64>& StreamedResources = MutableMesh->GetStreamedResources();
 
 				for (uint64 ResourceId : StreamedResources)
 				{
-					FCustomizableObjectStreameableResourceId TypedResourceId = BitCast<FCustomizableObjectStreameableResourceId>(ResourceId);	
-		
+					FCustomizableObjectStreameableResourceId TypedResourceId = BitCast<FCustomizableObjectStreameableResourceId>(ResourceId);
+
 					if (TypedResourceId.Type == (uint8)FCustomizableObjectStreameableResourceId::EType::AssetUserData)
 					{
 						const uint32 ResourceIndex = TypedResourceId.Id;
-						if (!ModelResources.StreamedResourceData.IsValidIndex(ResourceIndex))
-						{
-							UE_LOG(LogMutable, Error, TEXT("Invalid streamed resource index. Max Index [%d]. Resource Index [%d]."), ModelResources.StreamedResourceData.Num(), ResourceIndex);
-							continue; 
-						}
-
-						const FCustomizableObjectStreamedResourceData& StreamedResource = ModelResources.StreamedResourceData[ResourceIndex];
-						if (!StreamedResource.IsLoaded())
-						{
-							OperationData->StreamedResourceIndex.Add(ResourceIndex);
-							AssetsToStream.AddUnique(StreamedResource.GetPath().ToSoftObjectPath());
-						}
-
-						ComponentData->StreamedResourceIndex.Add(ResourceIndex);
+#if !WITH_EDITOR
+						OperationData->StreamedResourceIndex.AddUnique(ResourceIndex); // Used to hold/release streamed resources in non-editor builds.
+#endif
+						ComponentData->StreamedResourceIndex.AddUnique(ResourceIndex);
 					}
 					else if (TypedResourceId.Type == (uint8)FCustomizableObjectStreameableResourceId::EType::RealTimeMorphTarget)
 					{
@@ -5613,10 +5611,10 @@ UE::Tasks::FTask UCustomizableInstancePrivate::LoadAdditionalAssetsAndData(const
 
 									FString WarningMessage = FString::Printf(TEXT("Two submeshes have the same anim slot index [%s] in a Mutable Instance."), *SlotIndex.ToString());
 									UE_LOG(LogMutable, Warning, TEXT("%s"), *WarningMessage);
-	#if WITH_EDITOR
+#if WITH_EDITOR
 									FMessageLog MessageLog("Mutable");
 									MessageLog.Notify(FText::FromString(WarningMessage), EMessageSeverity::Warning, true);
-	#endif
+#endif
 								}
 							}
 						}
@@ -5625,12 +5623,12 @@ UE::Tasks::FTask UCustomizableInstancePrivate::LoadAdditionalAssetsAndData(const
 					{
 						AnimBPGameplayTags.AddTag(FGameplayTag::RequestGameplayTag(*Tag));
 					}
-	#if WITH_EDITORONLY_DATA
+#if WITH_EDITORONLY_DATA
 					else if (Tag.RemoveFromStart("__MeshPath:"))
 					{
 						ComponentData->MeshPartPaths.Add(Tag);
 					}
-	#endif
+#endif
 				}
 
 				const int32 AdditionalPhysicsNum = MutableMesh->AdditionalPhysicsBodies.Num();
@@ -5642,6 +5640,29 @@ UE::Tasks::FTask UCustomizableInstancePrivate::LoadAdditionalAssetsAndData(const
 					AssetsToStream.Add(ModelResources.AnimBpOverridePhysiscAssetsInfo[ExternalId].SourceAsset.ToSoftObjectPath());
 				}
 			}
+
+
+			for (int32 ResourceIndex : ComponentData->StreamedResourceIndex)
+			{
+#if WITH_EDITOR
+				if (!ModelResources.StreamedResourceDataEditor.IsValidIndex(ResourceIndex))
+				{
+					UE_LOG(LogMutable, Error, TEXT("Invalid streamed resource index. Max Index [%d]. Resource Index [%d]."), ModelResources.StreamedResourceDataEditor.Num(), ResourceIndex);
+				}
+#else
+				if (!ModelResources.StreamedResourceData.IsValidIndex(ResourceIndex))
+				{
+					UE_LOG(LogMutable, Error, TEXT("Invalid streamed resource index. Max Index [%d]. Resource Index [%d]."), ModelResources.StreamedResourceData.Num(), ResourceIndex);
+					continue;
+				}
+
+				const FCustomizableObjectStreamedResourceData& StreamedResource = ModelResources.StreamedResourceData[ResourceIndex];
+				if (!StreamedResource.IsLoaded())
+				{
+					AssetsToStream.AddUnique(StreamedResource.GetPath().ToSoftObjectPath());
+				}
+#endif
+			}
 		}
 	}
 
@@ -5669,6 +5690,15 @@ UE::Tasks::FTask UCustomizableInstancePrivate::LoadAdditionalAssetsAndData(const
 		NewEntry.Name = Name;
 		check(NewEntry.Name != NAME_None);
 
+#if WITH_EDITOR
+		if (!ModelResources.StreamedExtensionDataEditor.IsValidIndex(ExtensionData->Index))
+		{
+			// The compiled data appears to be out of sync with the CO's properties
+
+			UE_LOG(LogMutable, Error, TEXT("Couldn't find streamed Extension Data with index %d in %s. Compiled data may be stale."),
+				ExtensionData->Index, *CustomizableObject->GetFullName());
+		}
+#else
 		if (!ModelResources.StreamedExtensionData.IsValidIndex(ExtensionData->Index))
 		{
 			// The compiled data appears to be out of sync with the CO's properties
@@ -5690,6 +5720,7 @@ UE::Tasks::FTask UCustomizableInstancePrivate::LoadAdditionalAssetsAndData(const
 
 		OperationData->ExtensionStreamedResourceIndex.Add(ExtensionData->Index);
 		AssetsToStream.Add(StreamedData.GetPath().ToSoftObjectPath());
+#endif
 	}
 	
 	TArray<UE::Tasks::FTaskEvent> StreamingCompletionEvents;
@@ -6039,7 +6070,13 @@ const TArray<TObjectPtr<UMaterialInterface>>* UCustomizableObjectInstance::GetOv
 		return nullptr;
 	}
 
-	FName ComponentName = Object->GetPrivate()->GetModelResources().ComponentNamesPerObjectComponent[ComponentIndex.GetValue()];
+	const UModelResources* ModelResources = Object->GetPrivate()->GetModelResources();
+	if (!ModelResources)
+	{
+		return nullptr;
+	}
+
+	FName ComponentName = ModelResources->ComponentNamesPerObjectComponent[ComponentIndex.GetValue()];
 	FCustomizableInstanceComponentData* ComponentData = PrivateData->GetComponentData(ComponentName);
 	return ComponentData ? &ComponentData->OverrideMaterials : nullptr;
 }
@@ -6055,7 +6092,7 @@ void UCustomizableInstancePrivate::AdditionalAssetsAsyncLoaded(const TSharedRef<
 
 	UCustomizableObjectPrivate* CustomizableObjectPrivate = GetPublic()->GetCustomizableObject()->GetPrivate();
 
-	FModelResources& ModelResources = CustomizableObjectPrivate->GetModelResources();
+	UModelResources& ModelResources = *CustomizableObjectPrivate->GetModelResources();
 
 	for (int32 ResourceIndex : OperationData->StreamedResourceIndex)
 	{
@@ -6098,16 +6135,25 @@ void UCustomizableInstancePrivate::AdditionalAssetsAsyncLoaded(const TSharedRef<
 	{
 		for (int32 ResourceIndex : ComponentData.StreamedResourceIndex)
 		{
-			const FCustomizableObjectResourceData& ResourceData = ModelResources.StreamedResourceData[ResourceIndex].GetLoadedData();
-
-			if (const FCustomizableObjectAssetUserData* AUDResource = ResourceData.Data.GetPtr<FCustomizableObjectAssetUserData>())
+#if WITH_EDITOR
+			if (ModelResources.StreamedResourceDataEditor.IsValidIndex(ResourceIndex))
 			{
-#if WITH_EDITORONLY_DATA
-				ComponentData.AssetUserDataArray.Add(AUDResource->AssetUserDataEditor);
+				if (const FCustomizableObjectAssetUserData* AUDResource = ModelResources.StreamedResourceDataEditor[ResourceIndex].Data.GetPtr<FCustomizableObjectAssetUserData>())
+				{
+					ComponentData.AssetUserDataArray.Add(AUDResource->AssetUserData);
+				}
+			}
 #else
-				ComponentData.AssetUserDataArray.Add(AUDResource->AssetUserData);
-#endif
+			if (ModelResources.StreamedResourceData.IsValidIndex(ResourceIndex))
+			{
+				const FCustomizableObjectResourceData& ResourceData = ModelResources.StreamedResourceData[ResourceIndex].GetLoadedData();
+
+				if (const FCustomizableObjectAssetUserData* AUDResource = ResourceData.Data.GetPtr<FCustomizableObjectAssetUserData>())
+				{
+					ComponentData.AssetUserDataArray.Add(AUDResource->AssetUserData);
+				}
 			}
+#endif
 		}
 
 		// Loaded Skeletons
@@ -6318,7 +6364,7 @@ void UCustomizableInstancePrivate::BuildMaterials(const TSharedRef<FUpdateContex
 
 	UCustomizableObject* CustomizableObject = Public->GetCustomizableObject();
 
-	const FModelResources& ModelResources = CustomizableObject->GetPrivate()->GetModelResources();
+	const UModelResources& ModelResources = *CustomizableObject->GetPrivate()->GetModelResources();
 
 	TArray<FGeneratedTexture> NewGeneratedTextures;
 
@@ -6361,11 +6407,11 @@ void UCustomizableInstancePrivate::BuildMaterials(const TSharedRef<FUpdateContex
 
 		const FCustomizableObjectComponentIndex ObjectComponentIndex = Component.Id;
 
-		if (!CustomizableObject->GetPrivate()->GetModelResources().ComponentNamesPerObjectComponent.IsValidIndex(ObjectComponentIndex.GetValue()))
+		if (!ModelResources.ComponentNamesPerObjectComponent.IsValidIndex(ObjectComponentIndex.GetValue()))
 		{
 			continue;
 		}
-		const FName& ComponentName = CustomizableObject->GetPrivate()->GetModelResources().ComponentNamesPerObjectComponent[ObjectComponentIndex.GetValue()];
+		const FName& ComponentName = ModelResources.ComponentNamesPerObjectComponent[ObjectComponentIndex.GetValue()];
 
 		TObjectPtr<USkeletalMesh>* Result = SkeletalMeshes.Find(ComponentName);
 		TObjectPtr<USkeletalMesh> SkeletalMesh = Result ? *Result : nullptr;
@@ -7050,7 +7096,7 @@ void UCustomizableInstancePrivate::BuildMaterials(const TSharedRef<FUpdateContex
 #endif
 
 			const FName& ComponentName = CustomizableObjectInstanceUsage->GetComponentName();
-			const int32 ObjectComponentIndex = OperationData->Instance->GetCustomizableObject()->GetPrivate()->GetModelResources().ComponentNamesPerObjectComponent.IndexOfByKey(ComponentName);
+			const int32 ObjectComponentIndex = ModelResources.ComponentNamesPerObjectComponent.IndexOfByKey(ComponentName);
 			
 			int32 InstanceComponentIndex = -1;
 			for (int32 CurrentInstanceIndex=0; CurrentInstanceIndex<OperationData->InstanceUpdateData.Components.Num(); ++CurrentInstanceIndex)
@@ -7220,8 +7266,12 @@ void UCustomizableObjectInstance::SetRequestedLODs(const TMap<FName, uint8>& InM
 	FMutableUpdateCandidate MutableUpdateCandidate(this);
 
 	// Clamp Min LOD
-	const FModelResources& ModelResources = CustomizableObject->GetPrivate()->GetModelResources();
-	
+	const UModelResources* ModelResources = CustomizableObject->GetPrivate()->GetModelResources();
+	if (!ModelResources)
+	{
+		return;
+	}
+
 	bool bMinLODChanged = false;
 	
 	// Save the new LODs
@@ -7232,7 +7282,7 @@ void UCustomizableObjectInstance::SetRequestedLODs(const TMap<FName, uint8>& InM
 	
 	const TMap<FName, uint8>& GeneratedLODsPerComponent = GetPrivate()->CommittedDescriptorHash.RequestedLODsPerComponent;
 	
-	for (const FName& ComponentName : CustomizableObject->GetPrivate()->GetModelResources().ComponentNamesPerObjectComponent)
+	for (const FName& ComponentName : ModelResources->ComponentNamesPerObjectComponent)
 	{
 		uint8& InMinLOD = MutableUpdateCandidate.MinLOD.FindOrAdd(ComponentName);
 		if (const uint8* Result = InMinLODs.Find(ComponentName))
@@ -7242,8 +7292,7 @@ void UCustomizableObjectInstance::SetRequestedLODs(const TMap<FName, uint8>& InM
 
 		const uint8 MinLODIdx = CustomizableObject->GetPrivate()->GetMinLODIndex(ComponentName);
 		MutableUpdateCandidate.FirstLODAvailable.Add(ComponentName, MinLODIdx);
-
-		const int32 MaxLODIdx = ModelResources.ComponentLODs[ComponentName] - 1;
+		const int32 MaxLODIdx = ModelResources->ComponentLODs[ComponentName] - 1;
 		InMinLOD = FMath::Clamp(InMinLOD, MinLODIdx, MaxLODIdx);
 
 		uint8 DescriptorMinLOD = 0;
@@ -7251,9 +7300,9 @@ void UCustomizableObjectInstance::SetRequestedLODs(const TMap<FName, uint8>& InM
 		{
 			DescriptorMinLOD = *Result;
 		}
-		
+
 		bMinLODChanged |= DescriptorMinLOD != InMinLOD;
-		
+
 		if (UCustomizableObjectSystem::GetInstance()->IsOnlyGenerateRequestedLODsEnabled())
 		{
 			uint8 CurrentMinLOD = 0;
@@ -7261,10 +7310,10 @@ void UCustomizableObjectInstance::SetRequestedLODs(const TMap<FName, uint8>& InM
 			{
 				CurrentMinLOD = *Result;
 			}
-			
+
 			PrivateData->SetCOInstanceFlags(InMinLOD > CurrentMinLOD ? PendingLODsDowngrade : ECONone);
 
-			const uint8 FirstNonStreamedLODIndex = FMath::Clamp(CustomizableObject->GetPrivate()->GetModelResources().ComponentNumLODsToStream[ComponentName], 0, MaxLODIdx);
+			const uint8 FirstNonStreamedLODIndex = FMath::Clamp(ModelResources->ComponentNumLODsToStream[ComponentName], 0, MaxLODIdx);
 
 			MutableUpdateCandidate.RequestedLODLevels.Add(ComponentName, FirstNonStreamedLODIndex);
 
@@ -7278,12 +7327,12 @@ void UCustomizableObjectInstance::SetRequestedLODs(const TMap<FName, uint8>& InM
 			{
 				PredictedLOD = FMath::Min(*Result, PredictedLOD);
 			}
-				
+
 			// Save new RequestedLODs
 			MutableUpdateCandidate.RequestedLODLevels[ComponentName] = PredictedLOD;
 		}
 	}
-
+	
 	if (bMinLODChanged || GeneratedLODsPerComponent != MutableUpdateCandidate.RequestedLODLevels)
 	{
 		// TODO: Remove this flag as it will become redundant with the new InOutRequestedUpdates system
diff --git a/Engine/Plugins/Mutable/Source/CustomizableObject/Private/MuCO/CustomizableObjectSkeletalMesh.cpp b/Engine/Plugins/Mutable/Source/CustomizableObject/Private/MuCO/CustomizableObjectSkeletalMesh.cpp
--- a/Engine/Plugins/Mutable/Source/CustomizableObject/Private/MuCO/CustomizableObjectSkeletalMesh.cpp
+++ b/Engine/Plugins/Mutable/Source/CustomizableObject/Private/MuCO/CustomizableObjectSkeletalMesh.cpp
@@ -52,7 +52,7 @@ void UCustomizableObjectSkeletalMesh::InitMutableStreamingData(const TSharedRef<
 		MeshIDs[LODIndex] = LOD.MeshID;
 	}
 
-	const FModelResources& ModelResources = CustomizableObject.GetPrivate()->GetModelResources();
+	const UModelResources& ModelResources = CustomizableObject.GetPrivate()->GetModelResourcesChecked();
 
 	for (const FSkinWeightProfileInfo& SkinWeightProfile : GetSkinWeightProfiles())
 	{
diff --git a/Engine/Plugins/Mutable/Source/CustomizableObject/Private/MuCO/CustomizableObjectSystem.cpp b/Engine/Plugins/Mutable/Source/CustomizableObject/Private/MuCO/CustomizableObjectSystem.cpp
--- a/Engine/Plugins/Mutable/Source/CustomizableObject/Private/MuCO/CustomizableObjectSystem.cpp
+++ b/Engine/Plugins/Mutable/Source/CustomizableObject/Private/MuCO/CustomizableObjectSystem.cpp
@@ -191,17 +191,19 @@ FUpdateContextPrivate::FUpdateContextPrivate(UCustomizableObjectInstance& InInst
 
 	Object = InObject;
 
-	const FModelResources& ModelResources = InObject->GetPrivate()->GetModelResources();
-
 	CapturedDescriptor = Descriptor;
 	InObject->GetPrivate()->ApplyStateForcedValuesToParameters(CapturedDescriptor);
 
 	CapturedDescriptorHash = FDescriptorHash(Descriptor);
 	NumObjectComponents = InObject->GetComponentCount();
-	FirstLODAvailable = ModelResources.ComponentFirstLODAvailable;
-	FirstResidentLOD = ModelResources.ComponentNumLODsToStream;
-	ComponentNames = ModelResources.ComponentNamesPerObjectComponent;
 	
+	if (const UModelResources* ModelResources = InObject->GetPrivate()->GetModelResources())
+	{
+		FirstLODAvailable = ModelResources->ComponentFirstLODAvailable;
+		FirstResidentLOD = ModelResources->ComponentNumLODsToStream;
+		ComponentNames = ModelResources->ComponentNamesPerObjectComponent;
+	}
+
 	MutableSystem = UCustomizableObjectSystem::GetInstance()->GetPrivate()->MutableSystem;	
 	check(MutableSystem);
 
@@ -837,6 +839,7 @@ void FinishUpdateGlobal(const TSharedRef<FUpdateContextPrivate>& Context)
 	UCustomizableObjectInstance* Instance = Context->Instance.Get();
 
 	UCustomizableObject* Object = Instance ? Instance->GetCustomizableObject() : nullptr;
+	UModelResources* ModelResources = Object ? Object->GetPrivate()->GetModelResources() : nullptr;
 
 	UCustomizableObjectSystem* System = UCustomizableObjectSystem::GetInstance();
 	UCustomizableObjectSystemPrivate* SystemPrivate = System ? System->GetPrivate() : nullptr;
@@ -923,16 +926,16 @@ void FinishUpdateGlobal(const TSharedRef<FUpdateContextPrivate>& Context)
 	Context->UpdateCallback.ExecuteIfBound(ContextPublic);
 	Context->UpdateNativeCallback.Broadcast(ContextPublic);
 
-	if (Object)
+	if (ModelResources)
 	{
 		for (const int32 ResourceIndex : Context->StreamedResourceIndex)
 		{
-			Object->GetPrivate()->GetModelResources().StreamedResourceData[ResourceIndex].Release();
+			ModelResources->StreamedResourceData[ResourceIndex].Release();
 		}
 
 		for (const int32 ResourceIndex : Context->ExtensionStreamedResourceIndex)
 		{
-			Object->GetPrivate()->GetModelResources().StreamedExtensionData[ResourceIndex].Release();
+			ModelResources->StreamedExtensionData[ResourceIndex].Release();
 		}
 	}
 	
@@ -956,7 +959,7 @@ void FinishUpdateGlobal(const TSharedRef<FUpdateContextPrivate>& Context)
 		Context->UpdateTime = FPlatformTime::Seconds() - Context->StartUpdateTime;		
 	}
 	
-	const FName ObjectName = Object ? Object->GetFName() : NAME_None;
+	const FName ObjectName = GetFNameSafe(Object);
 	const FName InstanceName = GetFNameSafe(Instance);
 	UE_LOG(LogMutable, Verbose, TEXT("Finished Update Skeletal Mesh Async. CustomizableObject=%s Instance=%s, Frame=%d  QueueTime=%f, UpdateTime=%f"), *ObjectName.ToString(), *InstanceName.ToString(), GFrameNumber, Context->QueueTime, Context->UpdateTime);
 	
@@ -1144,18 +1147,22 @@ void UCustomizableObjectSystemPrivate::GetMipStreamingConfig(const UCustomizable
 
 	// From user-controlled per-state flag?
 	const FString CurrentState = Instance.GetCurrentState();
-	if (const FMutableStateData* State = Instance.GetCustomizableObject()->GetPrivate()->GetModelResources().StateUIDataMap.Find(CurrentState))
+
+	if (const UModelResources* ModelResources = Instance.GetCustomizableObject()->GetPrivate()->GetModelResources())
 	{
-		bOutNeverStream = State->bDisableTextureStreaming;
-	}
+		if (const FMutableStateData* State = ModelResources->StateUIDataMap.Find(CurrentState))
+		{
+			bOutNeverStream = State->bDisableTextureStreaming;
+		}
 
 #if WITH_EDITORONLY_DATA
-	// Was streaming disabled at object-compilation time? 
-	if (Instance.GetCustomizableObject()->GetPrivate()->GetModelResources().bIsTextureStreamingDisabled)
-	{
-		bOutNeverStream = true;
-	}
+		// Was streaming disabled at object-compilation time? 
+		if (ModelResources->bIsTextureStreamingDisabled)
+		{
+			bOutNeverStream = true;
+		}
 #endif
+	}
 
 	OutMipsToSkip = 0; // 0 means generate all mips
 
@@ -2274,7 +2281,7 @@ namespace impl
 
 		const UCustomizableObjectInstance* Instance = OperationData->Instance.Get();
 		const UCustomizableObject* CustomizableObject = Instance->GetCustomizableObject();
-		const FModelResources& ModelResources = CustomizableObject->GetPrivate()->GetModelResources();
+		const UModelResources& ModelResources = *CustomizableObject->GetPrivate()->GetModelResources();
 
 		const int32 ImageIndex = (*GetImagesData)[GetImageIndex].ImageIndex;
 
@@ -2337,7 +2344,7 @@ namespace impl
 		
 		const UCustomizableObjectInstance* Instance = OperationData->Instance.Get();
 		const UCustomizableObject* CustomizableObject = Instance->GetCustomizableObject();
-		const FModelResources& ModelResources = CustomizableObject->GetPrivate()->GetModelResources();
+		const UModelResources& ModelResources = *CustomizableObject->GetPrivate()->GetModelResources();
 
 		const mu::FInstance* MutableInstance = OperationData->MutableInstance.Get();
 
@@ -2545,7 +2552,7 @@ namespace impl
 		{
 			const UCustomizableObjectInstance* Instance = OperationData->Instance.Get();
 			const UCustomizableObject* CustomizableObject = Instance->GetCustomizableObject();
-			const FModelResources& ModelResources = CustomizableObject->GetPrivate()->GetModelResources();
+			const UModelResources& ModelResources = *CustomizableObject->GetPrivate()->GetModelResources();
 
 			uint32 ReferenceID = LOD.Mesh->GetReferencedMesh();
 
@@ -3522,9 +3529,11 @@ if (CVarMutableLogObjectMemoryOnUpdate.GetValueOnAnyThread())
 		UCustomizableObjectSystemPrivate* SystemPrivateData = System->GetPrivate();
 
 		SystemPrivateData->CurrentInstanceBeingUpdated = CandidateInstance;
-		
+
+		const UModelResources& ModelResources = CandidateInstance->GetCustomizableObject()->GetPrivate()->GetModelResourcesChecked();
+
 		FString StateName = CandidateInstance->GetCustomizableObject()->GetPrivate()->GetStateName(CandidateInstance->GetPrivate()->GetState());
-		const FMutableStateData* StateData = CandidateInstance->GetCustomizableObject()->GetPrivate()->GetModelResources().StateUIDataMap.Find(StateName);
+		const FMutableStateData* StateData = ModelResources.StateUIDataMap.Find(StateName);
 
 		Operation->bLiveUpdateMode = false;
 
@@ -3573,7 +3582,7 @@ if (CVarMutableLogObjectMemoryOnUpdate.GetValueOnAnyThread())
 
 #if WITH_EDITOR
 		// Async load all Runtime Referenced Textures.
-		const TArray<TSoftObjectPtr<const UTexture>>& RuntimeReferencedTextures = CustomizableObject->GetPrivate()->GetModelResources().RuntimeReferencedTextures;
+		const TArray<TSoftObjectPtr<const UTexture>>& RuntimeReferencedTextures = ModelResources.RuntimeReferencedTextures;
 		if (!RuntimeReferencedTextures.IsEmpty())
 		{
 			UE::Tasks::FTaskEvent Event = UE::Tasks::FTaskEvent(TEXT("Texture"));
@@ -3602,7 +3611,8 @@ if (CVarMutableLogObjectMemoryOnUpdate.GetValueOnAnyThread())
 					return;
 				}
 			
-				System->GetPrivate()->GetResourceProviderChecked()->CacheRuntimeReferencedImages(Operation->Model.ToSharedRef(), CustomizableObject->GetPrivate()->GetModelResources().RuntimeReferencedTextures);
+				const UModelResources& ModelResources = CustomizableObject->GetPrivate()->GetModelResourcesChecked();
+				System->GetPrivate()->GetResourceProviderChecked()->CacheRuntimeReferencedImages(Operation->Model.ToSharedRef(), ModelResources.RuntimeReferencedTextures);
 				Event.Trigger();
 			}));
 		}
@@ -3665,7 +3675,7 @@ if (CVarMutableLogObjectMemoryOnUpdate.GetValueOnAnyThread())
 		if (bRequestAllLODs)
 		{
 			TMap<FName, uint8> RequestedLODs = Operation->GetRequestedLODs();
-			for (const FName& ComponentName : Operation->Object->GetPrivate()->GetModelResources().ComponentNamesPerObjectComponent)
+			for (const FName& ComponentName : ModelResources.ComponentNamesPerObjectComponent)
 			{
 				RequestedLODs.Add(ComponentName, 0);
 			}
@@ -4096,7 +4106,7 @@ void UCustomizableObjectSystem::DiscardInstances()
 	{
 
 		UCustomizableObjectInstance* COI = Iterator->CustomizableObjectInstance.Get();
-		
+
 		const bool bUpdating = Private->CurrentMutableOperation && Private->CurrentMutableOperation->Instance == Iterator->CustomizableObjectInstance;
 		if (COI && COI->GetPrivate() && !bUpdating)
 		{
@@ -4110,54 +4120,51 @@ void UCustomizableObjectSystem::DiscardInstances()
 
 				if (UCustomizableObject* CustomizableObject = COI->GetCustomizableObject())
 				{
-					const FModelResources& ModelResources = CustomizableObject->GetPrivate()->GetModelResources();
-				
-					for (TObjectIterator<UCustomizableObjectInstanceUsage> It; It; ++It)
+					if (const UModelResources* ModelResources = CustomizableObject->GetPrivate()->GetModelResources())
 					{
-						UCustomizableObjectInstanceUsage* InstanceUsage = *It;
 
-						if (InstanceUsage->GetCustomizableObjectInstance() != COI)
+						for (TObjectIterator<UCustomizableObjectInstanceUsage> It; It; ++It)
 						{
-							continue;
-						}
+							UCustomizableObjectInstanceUsage* InstanceUsage = *It;
+							if (InstanceUsage->GetCustomizableObjectInstance() != COI)
+							{
+								continue;
+							}
 
 #if WITH_EDITOR
-						if (InstanceUsage->GetPrivate()->IsNetMode(NM_DedicatedServer))
-						{
-							continue;
-						}
+							if (InstanceUsage->GetPrivate()->IsNetMode(NM_DedicatedServer))
+							{
+								continue;
+							}
 #endif
-						
-						if (USkeletalMeshComponent* Parent = Cast<USkeletalMeshComponent>(InstanceUsage->GetAttachParent()))
-						{
-							const FName& ComponentName = InstanceUsage->GetComponentName();
-							const int32 ObjectComponentIndex = ModelResources.ComponentNamesPerObjectComponent.IndexOfByKey(ComponentName);
-							if (ModelResources.ReferenceSkeletalMeshesData.IsValidIndex(ObjectComponentIndex))
+
+							if (USkeletalMeshComponent* Parent = Cast<USkeletalMeshComponent>(InstanceUsage->GetAttachParent()))
 							{
-								USkeletalMesh* SkeletalMesh = nullptr;
-					
-								if (IsReplaceDiscardedWithReferenceMeshEnabled())
+								const FName& ComponentName = InstanceUsage->GetComponentName();
+								const int32 ObjectComponentIndex = ModelResources->ComponentNamesPerObjectComponent.IndexOfByKey(ComponentName);
+								if (ModelResources->ReferenceSkeletalMeshesData.IsValidIndex(ObjectComponentIndex))
 								{
-									// Force load the reference mesh if necessary. 
-									SkeletalMesh = ModelResources.ReferenceSkeletalMeshesData[ObjectComponentIndex].SoftSkeletalMesh.LoadSynchronous();
-								}
+									USkeletalMesh* SkeletalMesh = nullptr;
+
+									if (IsReplaceDiscardedWithReferenceMeshEnabled())
+									{
+										// Force load the reference mesh if necessary. 
+										SkeletalMesh = ModelResources->ReferenceSkeletalMeshesData[ObjectComponentIndex].SoftSkeletalMesh.LoadSynchronous();
+									}
 
-								Parent->EmptyOverrideMaterials();
-								Parent->SetSkeletalMesh(SkeletalMesh);
+									Parent->EmptyOverrideMaterials();
+									Parent->SetSkeletalMesh(SkeletalMesh);
+								}
 							}
-						}
-					
-						for (const UCustomizableObjectExtension* Extension : ICustomizableObjectModule::Get().GetRegisteredExtensions())
-						{
-							Extension->OnCustomizableObjectInstanceUsageDiscarded(*InstanceUsage);
+
 						}
 					}
 				}
+
+				Iterator.RemoveCurrent();
+				NumInstancesDiscarded++;
 			}
 		}
-
-		Iterator.RemoveCurrent();
-		NumInstancesDiscarded++;
 	}
 }
 
diff --git a/Engine/Plugins/Mutable/Source/CustomizableObject/Private/MuCO/UnrealMutableModelDiskStreamer.cpp b/Engine/Plugins/Mutable/Source/CustomizableObject/Private/MuCO/UnrealMutableModelDiskStreamer.cpp
--- a/Engine/Plugins/Mutable/Source/CustomizableObject/Private/MuCO/UnrealMutableModelDiskStreamer.cpp
+++ b/Engine/Plugins/Mutable/Source/CustomizableObject/Private/MuCO/UnrealMutableModelDiskStreamer.cpp
@@ -89,7 +89,7 @@ bool FUnrealMutableModelBulkReader::PrepareStreamingForObject(UCustomizableObjec
 		!=
 		nullptr;
 
-	const FModelResources& ModelResources = CustomizableObject->GetPrivate()->GetModelResources();
+	const UModelResources& ModelResources = CustomizableObject->GetPrivate()->GetModelResourcesChecked();
 
 	if (!bAlreadyStreaming)
 	{
diff --git a/Engine/Plugins/Mutable/Source/CustomizableObject/Public/MuCO/CustomizableObject.h b/Engine/Plugins/Mutable/Source/CustomizableObject/Public/MuCO/CustomizableObject.h
--- a/Engine/Plugins/Mutable/Source/CustomizableObject/Public/MuCO/CustomizableObject.h
+++ b/Engine/Plugins/Mutable/Source/CustomizableObject/Public/MuCO/CustomizableObject.h
@@ -403,7 +403,6 @@ public:
 	virtual FString GetDesc() override;
 	virtual bool IsEditorOnly() const override;
 	virtual void PreSave(FObjectPreSaveContext ObjectSaveContext) override;
-	virtual void PostSaveRoot(FObjectPostSaveRootContext ObjectSaveContext) override;
 	virtual void PostRename(UObject* OldOuter, const FName OldName) override;	
 	virtual void BeginCacheForCookedPlatformData(const ITargetPlatform* TargetPlatform) override;
 	virtual bool IsCachedCookedPlatformDataLoaded(const ITargetPlatform* TargetPlatform) override;
diff --git a/Engine/Plugins/Mutable/Source/CustomizableObject/Public/MuCO/CustomizableObjectExtension.h b/Engine/Plugins/Mutable/Source/CustomizableObject/Public/MuCO/CustomizableObjectExtension.h
--- a/Engine/Plugins/Mutable/Source/CustomizableObject/Public/MuCO/CustomizableObjectExtension.h
+++ b/Engine/Plugins/Mutable/Source/CustomizableObject/Public/MuCO/CustomizableObjectExtension.h
@@ -122,6 +122,23 @@ public:
 	 */
 	virtual FInstancedStruct GenerateExtensionInstanceData(const TArray<FInputPinDataContainer>& InputPinData) const { return FInstancedStruct(); }
 
+#if WITH_EDITOR
+	/**
+	 * Non-owned references to private UObjects must be converted to owned references at cook time.
+	 * Duplicate the private UObjects with the container as its new outer object, or the cook will 
+	 * fail to serialize them.
+	 * 
+	 * E.g. Any UObjects that are referenced by the Extension Data that aren't in their own asset
+	 * package should be copied with the Container as their Outer. This ensures that they get cooked
+	 * into the correct package for streaming.
+	 * 
+	 * @param Struct - Instanced struct of unknown type that could hold external private references.
+	 *
+	 * @param Container - Outer UObject for duplicated private references.
+	 */
+	virtual void MovePrivateReferencesToContainer(FInstancedStruct& Struct, UObject* Container) const {};
+#endif
+
 	/** Called when a Customizable Object Instance Usage is being updated.
 	 * 
 	 * @param Usage Customizable Object Instance Usage being updated. */
diff --git a/Engine/Plugins/Mutable/Source/CustomizableObjectEditor/Private/MuCOE/CustomizableInstanceDetails.cpp b/Engine/Plugins/Mutable/Source/CustomizableObjectEditor/Private/MuCOE/CustomizableInstanceDetails.cpp
--- a/Engine/Plugins/Mutable/Source/CustomizableObjectEditor/Private/MuCOE/CustomizableInstanceDetails.cpp
+++ b/Engine/Plugins/Mutable/Source/CustomizableObjectEditor/Private/MuCOE/CustomizableInstanceDetails.cpp
@@ -1009,11 +1009,14 @@ TSharedRef<SWidget> FCustomizableInstanceDetails::GenerateIntWidget(const UCusto
 
 	// Tooltip for the selected option
 	FString ToolTipText = FString("None");
-	if (const FString* Identifier = CustomizableObject.GetPrivate()->GetModelResources().GroupNodeMap.FindKey(FCustomizableObjectIdPair(ParamName, SelectedOption)))
+	if (const UModelResources* ModelResources = CustomizableObject.GetPrivate()->GetModelResources())
 	{
-		if (const FString* CustomizableObjectPath = CustomizableObject.GetPrivate()->GetModelResources().CustomizableObjectPathMap.Find(*Identifier))
+		if (const FString* Identifier = ModelResources->GroupNodeMap.FindKey(FCustomizableObjectIdPair(ParamName, SelectedOption)))
 		{
-			ToolTipText = *CustomizableObjectPath;
+			if (const FString* CustomizableObjectPath = ModelResources->CustomizableObjectPathMap.Find(*Identifier))
+			{
+				ToolTipText = *CustomizableObjectPath;
+			}
 		}
 	}
 
diff --git a/Engine/Plugins/Mutable/Source/CustomizableObjectEditor/Private/MuCOE/CustomizableObjectBenchmarkingUtils.cpp b/Engine/Plugins/Mutable/Source/CustomizableObjectEditor/Private/MuCOE/CustomizableObjectBenchmarkingUtils.cpp
--- a/Engine/Plugins/Mutable/Source/CustomizableObjectEditor/Private/MuCOE/CustomizableObjectBenchmarkingUtils.cpp
+++ b/Engine/Plugins/Mutable/Source/CustomizableObjectEditor/Private/MuCOE/CustomizableObjectBenchmarkingUtils.cpp
@@ -11,6 +11,12 @@ bool CustomizableObjectBenchmarkingUtils::GenerateDeterministicSetOfInstances(TO
 	{
 		return false;
 	}
+
+	const UModelResources* ModelResources = TargetCustomizableObject->GetPrivate()->GetModelResources();
+	if (!ModelResources)
+	{
+		return false;
+	}
 	
 	// Test this parameter configuration in all the states of the CO
 	const uint32 StateCount = TargetCustomizableObject->GetStateCount();
@@ -33,7 +39,7 @@ bool CustomizableObjectBenchmarkingUtils::GenerateDeterministicSetOfInstances(TO
 		{
 			// Force generation of all LODS
 			TMap<FName, uint8> RequestedLodLevels;
-			for (const FName ComponentName : TargetCustomizableObject->GetPrivate()->GetModelResources().ComponentNamesPerObjectComponent)
+			for (const FName ComponentName : ModelResources->ComponentNamesPerObjectComponent)
 			{
 				RequestedLodLevels.Add(ComponentName, 0);
 			}
diff --git a/Engine/Plugins/Mutable/Source/CustomizableObjectEditor/Private/MuCOE/CustomizableObjectCompileRunnable.cpp b/Engine/Plugins/Mutable/Source/CustomizableObjectEditor/Private/MuCOE/CustomizableObjectCompileRunnable.cpp
--- a/Engine/Plugins/Mutable/Source/CustomizableObjectEditor/Private/MuCOE/CustomizableObjectCompileRunnable.cpp
+++ b/Engine/Plugins/Mutable/Source/CustomizableObjectEditor/Private/MuCOE/CustomizableObjectCompileRunnable.cpp
@@ -341,13 +341,14 @@ void FCustomizableObjectCompileRunnable::Tick()
 
 FCustomizableObjectSaveDDRunnable::FCustomizableObjectSaveDDRunnable(const TSharedPtr<FCompilationRequest>& InRequest,
 	TSharedPtr<mu::FModel> InModel,
-	FModelResources& InModelResources,
+	UModelResources& InModelResources,
 	TSharedPtr<FModelStreamableBulkData> InModelStreamables)
 {
 	MUTABLE_CPUPROFILER_SCOPE(FCustomizableObjectSaveDDRunnable::FCustomizableObjectSaveDDRunnable)
 
 	Model = InModel;
 	ModelStreamables = InModelStreamables;
+	PlatformData.ModelResources = TStrongObjectPtr<UModelResources>(&InModelResources);
 
 	Options = InRequest->GetCompileOptions();
 
@@ -362,9 +363,9 @@ FCustomizableObjectSaveDDRunnable::FCustomizableObjectSaveDDRunnable(const TShar
 
 	// Cache ModelResources
 	{
-		FMemoryWriter64 MemoryWriter(PlatformData.ModelResourcesData);
+		FMemoryWriter64 MemoryWriter(ModelResourcesData);
 		FObjectAndNameAsStringProxyArchive ObjectWriter(MemoryWriter, true);
-		InModelResources.Serialize(ObjectWriter, Options.bIsCooking);
+		InModelResources.Serialize(ObjectWriter);
 	}
 
 	// Cache Morphs and Clothing. 
@@ -568,7 +569,7 @@ void FCustomizableObjectSaveDDRunnable::StoreCachedPlatformDataInDDC(bool& bStor
 	{
 		MUTABLE_CPUPROFILER_SCOPE(SerializeModelResourcesForDDC);
 
-		const FValue ModelResourcesValue = FValue::Compress(FSharedBuffer::MakeView(PlatformData.ModelResourcesData.GetData(), PlatformData.ModelResourcesData.Num()));
+		const FValue ModelResourcesValue = FValue::Compress(FSharedBuffer::MakeView(ModelResourcesData.GetData(), ModelResourcesData.Num()));
 		RecordBuilder.AddValue(MutablePrivate::GetDerivedDataModelResourcesId(), ModelResourcesValue);
 	}
 
@@ -691,7 +692,7 @@ void FCustomizableObjectSaveDDRunnable::StoreCachedPlatformDataToDisk(bool& bSto
 		// Serailize headers to validate data
 		*ModelMemoryWriter << CustomizableObjectHeader;
 
-		ModelMemoryWriter->Serialize(PlatformData.ModelResourcesData.GetData(), PlatformData.ModelResourcesData.Num() * sizeof(uint8));
+		ModelMemoryWriter->Serialize(ModelResourcesData.GetData(), ModelResourcesData.Num() * sizeof(uint8));
 
 		{
 			// ModelMemoryWriter (Writer to disk) doesn't handle FNames properly. Serialize them ModelStreamables in two steps.
diff --git a/Engine/Plugins/Mutable/Source/CustomizableObjectEditor/Private/MuCOE/CustomizableObjectCompileRunnable.h b/Engine/Plugins/Mutable/Source/CustomizableObjectEditor/Private/MuCOE/CustomizableObjectCompileRunnable.h
--- a/Engine/Plugins/Mutable/Source/CustomizableObjectEditor/Private/MuCOE/CustomizableObjectCompileRunnable.h
+++ b/Engine/Plugins/Mutable/Source/CustomizableObjectEditor/Private/MuCOE/CustomizableObjectCompileRunnable.h
@@ -95,7 +95,7 @@ class FCustomizableObjectSaveDDRunnable : public FRunnable
 {
 public:
 
-	FCustomizableObjectSaveDDRunnable(const TSharedPtr<FCompilationRequest>& InRequest, TSharedPtr<mu::FModel> InModel, FModelResources& ModelResources, TSharedPtr<FModelStreamableBulkData> ModelStreamables);
+	FCustomizableObjectSaveDDRunnable(const TSharedPtr<FCompilationRequest>& InRequest, TSharedPtr<mu::FModel> InModel, UModelResources& ModelResources, TSharedPtr<FModelStreamableBulkData> ModelStreamables);
 
 	// FRunnable interface
 	uint32 Run() override;
@@ -135,6 +135,7 @@ public:
 
 	TSharedPtr<mu::FModel, ESPMode::ThreadSafe> Model;
 	TSharedPtr<FModelStreamableBulkData> ModelStreamables;
+	TArray64<uint8> ModelResourcesData;
 
 	// Cached platform data
 	MutablePrivate::FMutableCachedPlatformData PlatformData;
diff --git a/Engine/Plugins/Mutable/Source/CustomizableObjectEditor/Private/MuCOE/CustomizableObjectCompiler.cpp b/Engine/Plugins/Mutable/Source/CustomizableObjectEditor/Private/MuCOE/CustomizableObjectCompiler.cpp
--- a/Engine/Plugins/Mutable/Source/CustomizableObjectEditor/Private/MuCOE/CustomizableObjectCompiler.cpp
+++ b/Engine/Plugins/Mutable/Source/CustomizableObjectEditor/Private/MuCOE/CustomizableObjectCompiler.cpp
@@ -351,11 +351,8 @@ bool FCustomizableObjectCompiler::IsRequestQueued(const TSharedRef<FCompilationR
 void FCustomizableObjectCompiler::AddReferencedObjects(FReferenceCollector& Collector)
 {
 	Collector.AddReferencedObjects(ArrayGCProtect);
-	
-	if (CurrentObject)
-	{
-		Collector.AddReferencedObject(CurrentObject);
-	}
+	Collector.AddReferencedObject(CurrentObject);
+	Collector.AddReferencedObject(ModelResources);
 }
 
 
@@ -612,7 +609,10 @@ mu::Ptr<mu::NodeObject> GenerateMutableRoot(
 
 	if (!GenerationContext.CompilationContext.Options.ParamNamesToSelectedOptions.IsEmpty())
 	{
-		GenerationContext.TableToParamNames = Object->GetPrivate()->GetModelResources().TableToParamNames;
+		if (const UModelResources* ModelResources = Object->GetPrivate()->GetModelResources())
+		{
+			GenerationContext.TableToParamNames = ModelResources->TableToParamNames;
+		}
 	}
 
 	GenerationContext.bPartialCompilation = LocalRootNodeObject->ParentObject != nullptr;
@@ -626,11 +626,6 @@ mu::Ptr<mu::NodeObject> GenerateMutableRoot(
 
 	// Generate ReferenceSkeletalMeshes data;
 	PopulateReferenceSkeletalMeshesData(GenerationContext);
-	
-	//for (const USkeletalMesh* RefSkeletalMesh : Object->ReferenceSkeletalMeshes)
-	//{
-	//	GenerationContext.CheckPhysicsAssetInSkeletalMesh(RefSkeletalMesh);
-	//}
 
 	// Display warnings for unnamed node objects
 	FText Message = LOCTEXT("Unnamed Node Object", "Unnamed Node Object");
@@ -752,70 +747,74 @@ void FCustomizableObjectCompiler::CompileInternal(bool bAsync)
 	}
 	else
 	{
-		FModelResources& ModelResources = CurrentObject->GetPrivate()->GetModelResources(CurrentOptions.bIsCooking);
-		ModelResources = FModelResources();
+		const FString ModelResourcesName = GetModelResourcesNameForPlatform(*CurrentObject, *CurrentOptions.TargetPlatform);
+		UObject* Outer = CurrentOptions.bIsCooking ? CurrentObject : StaticCast<UObject*>(GetTransientPackage());
+		ModelResources = NewObject<UModelResources>(Outer, FName(*ModelResourcesName), RF_Public);
+
 		ModelStreamableBulkData = MakeShared<FModelStreamableBulkData>();
 		
-		ModelResources.ReferenceSkeletalMeshesData = MoveTemp(GenerationContext.ReferenceSkeletalMeshesData);
+		ModelResources->CodeVersion = UCustomizableObjectPrivate::CurrentSupportedVersion;
+
+		ModelResources->ReferenceSkeletalMeshesData = MoveTemp(GenerationContext.ReferenceSkeletalMeshesData);
 
-		ModelResources.Skeletons.Reserve(GenerationContext.CompilationContext.ReferencedSkeletons.Num());
+		ModelResources->Skeletons.Reserve(GenerationContext.CompilationContext.ReferencedSkeletons.Num());
 		for (const USkeleton* Skeleton : GenerationContext.CompilationContext.ReferencedSkeletons)
 		{
-			ModelResources.Skeletons.Emplace(const_cast<USkeleton*>(Skeleton));
+			ModelResources->Skeletons.Emplace(const_cast<USkeleton*>(Skeleton));
 		}
 		
-		ModelResources.Materials.Reserve(GenerationContext.ReferencedMaterials.Num());
+		ModelResources->Materials.Reserve(GenerationContext.ReferencedMaterials.Num());
 		for (const UMaterialInterface* Material : GenerationContext.ReferencedMaterials)
 		{
-			ModelResources.Materials.Emplace(const_cast<UMaterialInterface*>(Material));
+			ModelResources->Materials.Emplace(const_cast<UMaterialInterface*>(Material));
 		}
 
 		for (const TPair<TSoftObjectPtr<UStreamableRenderAsset>, FMutableGraphGenerationContext::FGeneratedReferencedMesh>& Pair : GenerationContext.PassthroughMeshMap)
 		{
-			check(Pair.Value.ID == ModelResources.PassThroughMeshes.Num());
-			ModelResources.PassThroughMeshes.Add(Pair.Key);
+			check(Pair.Value.ID == ModelResources->PassThroughMeshes.Num());
+			ModelResources->PassThroughMeshes.Add(Pair.Key);
 		}
 
 		for (const TPair<TSoftObjectPtr<const UStreamableRenderAsset>, FMutableGraphGenerationContext::FGeneratedReferencedMesh>& Pair : GenerationContext.RuntimeReferencedMeshMap)
 		{
-			check(Pair.Value.ID == ModelResources.RuntimeReferencedMeshes.Num());
-			ModelResources.RuntimeReferencedMeshes.Add(Pair.Key);
+			check(Pair.Value.ID == ModelResources->RuntimeReferencedMeshes.Num());
+			ModelResources->RuntimeReferencedMeshes.Add(Pair.Key);
 		}
 
 		for (const TPair<TSoftObjectPtr<UTexture>, FMutableGraphGenerationContext::FGeneratedReferencedTexture>& Pair : GenerationContext.PassthroughTextureMap)
 		{
-			check(Pair.Value.ID == ModelResources.PassThroughTextures.Num());
-			ModelResources.PassThroughTextures.Add(Pair.Key);
+			check(Pair.Value.ID == ModelResources->PassThroughTextures.Num());
+			ModelResources->PassThroughTextures.Add(Pair.Key);
 		}
 
 		for (const TPair<TSoftObjectPtr<const UTexture>, FMutableGraphGenerationContext::FGeneratedReferencedTexture>& Pair : GenerationContext.RuntimeReferencedTextureMap)
 		{
-			check(Pair.Value.ID == ModelResources.RuntimeReferencedTextures.Num());
-			ModelResources.RuntimeReferencedTextures.Add(Pair.Key);
+			check(Pair.Value.ID == ModelResources->RuntimeReferencedTextures.Num());
+			ModelResources->RuntimeReferencedTextures.Add(Pair.Key);
 		}
 		
-		ModelResources.PhysicsAssets = MoveTemp(GenerationContext.CompilationContext.PhysicsAssets);
+		ModelResources->PhysicsAssets = MoveTemp(GenerationContext.CompilationContext.PhysicsAssets);
 
-		ModelResources.AnimBPs = MoveTemp(GenerationContext.AnimBPAssets);
-		ModelResources.AnimBpOverridePhysiscAssetsInfo = MoveTemp(GenerationContext.CompilationContext.AnimBpOverridePhysicsAssetsInfo);
+		ModelResources->AnimBPs = MoveTemp(GenerationContext.AnimBPAssets);
+		ModelResources->AnimBpOverridePhysiscAssetsInfo = MoveTemp(GenerationContext.CompilationContext.AnimBpOverridePhysicsAssetsInfo);
 
-		ModelResources.MaterialSlotNames = MoveTemp(GenerationContext.ReferencedMaterialSlotNames);
-		ModelResources.SocketArray = MoveTemp(GenerationContext.SocketArray);
+		ModelResources->MaterialSlotNames = MoveTemp(GenerationContext.ReferencedMaterialSlotNames);
+		ModelResources->SocketArray = MoveTemp(GenerationContext.SocketArray);
 
 		const int32 NumBones = GenerationContext.CompilationContext.UniqueBoneNames.Num() + GenerationContext.CompilationContext.RemappedBoneNames.Num();
-		ModelResources.BoneNamesMap.Reserve(NumBones);
+		ModelResources->BoneNamesMap.Reserve(NumBones);
 
 		for (auto& It : GenerationContext.CompilationContext.UniqueBoneNames)
 		{
-			ModelResources.BoneNamesMap.Add(It.Value, It.Key.Id);
+			ModelResources->BoneNamesMap.Add(It.Value, It.Key.Id);
 		}
 
 		for (auto& It : GenerationContext.CompilationContext.RemappedBoneNames)
 		{
-			ModelResources.BoneNamesMap.Add(It.Key, It.Value.Id);
+			ModelResources->BoneNamesMap.Add(It.Key, It.Value.Id);
 		}
 
-		ModelResources.SkinWeightProfilesInfo = MoveTemp(GenerationContext.CompilationContext.SkinWeightProfilesInfo);
+		ModelResources->SkinWeightProfilesInfo = MoveTemp(GenerationContext.CompilationContext.SkinWeightProfilesInfo);
 
 		TArray<FGeneratedImageProperties> ImageProperties;
 		GenerationContext.ImageProperties.GenerateValueArray(ImageProperties);
@@ -824,11 +823,11 @@ void FCustomizableObjectCompiler::CompileInternal(bool bAsync)
 		ImageProperties.Sort([](const FGeneratedImageProperties& PropsA, const FGeneratedImageProperties& PropsB)
 			{ return PropsA.ImagePropertiesIndex < PropsB.ImagePropertiesIndex;	});
 
-		ModelResources.ImageProperties.Empty(ImageProperties.Num());
+		ModelResources->ImageProperties.Empty(ImageProperties.Num());
 
 		for (const FGeneratedImageProperties& ImageProp : ImageProperties)
 		{
-			ModelResources.ImageProperties.Add({ ImageProp.TextureParameterName,
+			ModelResources->ImageProperties.Add({ ImageProp.TextureParameterName,
 										ImageProp.Filter,
 										ImageProp.SRGB,
 										ImageProp.bFlipGreenChannel,
@@ -839,9 +838,9 @@ void FCustomizableObjectCompiler::CompileInternal(bool bAsync)
 										ImageProp.AddressX, ImageProp.AddressY });
 		}
 
-		ModelResources.ParameterUIDataMap = MoveTemp(GenerationContext.ParameterUIDataMap);
-		ModelResources.StateUIDataMap = MoveTemp(GenerationContext.StateUIDataMap);
-		ModelResources.IntParameterOptionDataTable = MoveTemp(GenerationContext.IntParameterOptionDataTable);
+		ModelResources->ParameterUIDataMap = MoveTemp(GenerationContext.ParameterUIDataMap);
+		ModelResources->StateUIDataMap = MoveTemp(GenerationContext.StateUIDataMap);
+		ModelResources->IntParameterOptionDataTable = MoveTemp(GenerationContext.IntParameterOptionDataTable);
 
 		// Create the RealTimeMorphsTargets Blocks from the per mesh Morph data.
 		uint64 RealTimeMorphDataSize = 0;
@@ -851,7 +850,7 @@ void FCustomizableObjectCompiler::CompileInternal(bool bAsync)
 		}
 		
 		ModelStreamableBulkData->RealTimeMorphStreamables.Empty(32);
-		ModelResources.EditorOnlyMorphTargetReconstructionData.Empty(RealTimeMorphDataSize);
+		ModelResources->EditorOnlyMorphTargetReconstructionData.Empty(RealTimeMorphDataSize);
 
 		uint64 RealTimeMorphDataOffsetInBytes = 0;
 		for (const TPair<uint32, FRealTimeMorphMeshData>& MeshData : GenerationContext.CompilationContext.RealTimeMorphTargetPerMeshData)
@@ -869,7 +868,7 @@ void FCustomizableObjectCompiler::CompileInternal(bool bAsync)
 			ResourceMeshData.SourceId = MeshData.Value.SourceId;
 
 			RealTimeMorphDataOffsetInBytes += DataSizeInBytes;
-			ModelResources.EditorOnlyMorphTargetReconstructionData.Append(MeshData.Value.Data);
+			ModelResources->EditorOnlyMorphTargetReconstructionData.Append(MeshData.Value.Data);
 		}
 	
 		// Create the Clothing Blocks from the per mesh Morph data.
@@ -880,7 +879,7 @@ void FCustomizableObjectCompiler::CompileInternal(bool bAsync)
 		}
 		
 		ModelStreamableBulkData->ClothingStreamables.Empty(32);
-		ModelResources.EditorOnlyClothingMeshToMeshVertData.Empty(ClothingDataNum);
+		ModelResources->EditorOnlyClothingMeshToMeshVertData.Empty(ClothingDataNum);
 
 		uint64 ClothingDataOffsetInBytes = 0;
 		for (const TPair<uint32, FClothingMeshData>& MeshData : GenerationContext.CompilationContext.ClothingPerMeshData)
@@ -901,10 +900,10 @@ void FCustomizableObjectCompiler::CompileInternal(bool bAsync)
 			ResourceMeshData.SourceId = MeshData.Value.SourceId;
 
 			ClothingDataOffsetInBytes += DataSizeInBytes;
-			ModelResources.EditorOnlyClothingMeshToMeshVertData.Append(MeshData.Value.Data);
+			ModelResources->EditorOnlyClothingMeshToMeshVertData.Append(MeshData.Value.Data);
 		}
 
-		ModelResources.ClothingAssetsData = MoveTemp(GenerationContext.CompilationContext.ClothingAssetsData);
+		ModelResources->ClothingAssetsData = MoveTemp(GenerationContext.CompilationContext.ClothingAssetsData);
 
 		// A clothing backend, e.g. Chaos cloth, can use 2 config files, one owned by the asset, and another that is shared 
 		// among all assets in a SkeletalMesh. When merging different assets in a skeletalmesh we need to make sure only one of 
@@ -918,13 +917,13 @@ void FCustomizableObjectCompiler::CompileInternal(bool bAsync)
 		};
 		
 		// Find shared configs to be used (One of each type) 
-		for (FCustomizableObjectClothingAssetData& ClothingAssetData : ModelResources.ClothingAssetsData)
+		for (FCustomizableObjectClothingAssetData& ClothingAssetData : ModelResources->ClothingAssetsData)
 		{
 			 for (FCustomizableObjectClothConfigData& ClothConfigData : ClothingAssetData.ConfigsData)
 			 {
 				  if (IsSharedConfigData(ClothConfigData))
 				  {
-					  FCustomizableObjectClothConfigData* FoundConfig = ModelResources.ClothSharedConfigsData.FindByPredicate(
+					  FCustomizableObjectClothConfigData* FoundConfig = ModelResources->ClothSharedConfigsData.FindByPredicate(
 						   [Name = ClothConfigData.ConfigName](const FCustomizableObjectClothConfigData& Other)
 						   {
 							   return Name == Other.ConfigName;
@@ -932,110 +931,73 @@ void FCustomizableObjectCompiler::CompileInternal(bool bAsync)
 
 					  if (!FoundConfig)
 					  {
-						   ModelResources.ClothSharedConfigsData.AddDefaulted_GetRef() = ClothConfigData;
+						  ModelResources->ClothSharedConfigsData.AddDefaulted_GetRef() = ClothConfigData;
 					  }
 				  }
 			 }
 		}
 		
 		// Remove shared configs
-		for (FCustomizableObjectClothingAssetData& ClothingAssetData : ModelResources.ClothingAssetsData)
+		for (FCustomizableObjectClothingAssetData& ClothingAssetData : ModelResources->ClothingAssetsData)
 		{
 			 ClothingAssetData.ConfigsData.RemoveAllSwap(IsSharedConfigData);
 		}
 
-		ModelResources.MeshMetadata = MoveTemp(GenerationContext.CompilationContext.MeshMetadata);
-		ModelResources.SurfaceMetadata = MoveTemp(GenerationContext.SurfaceMetadata);
+		ModelResources->MeshMetadata = MoveTemp(GenerationContext.CompilationContext.MeshMetadata);
+		ModelResources->SurfaceMetadata = MoveTemp(GenerationContext.SurfaceMetadata);
 
-		ModelResources.GroupNodeMap = GenerationContext.GroupNodeMap;
+		ModelResources->GroupNodeMap = GenerationContext.GroupNodeMap;
 
 		// If the optimization level is "none" disable texture streaming, because textures are all referenced
 		// unreal assets and progressive generation is not supported.
-		ModelResources.bIsTextureStreamingDisabled = GenerationContext.CompilationContext.Options.OptimizationLevel == 0;
+		ModelResources->bIsTextureStreamingDisabled = GenerationContext.CompilationContext.Options.OptimizationLevel == 0;
 		
-		ModelResources.bIsCompiledWithOptimization = GenerationContext.CompilationContext.Options.OptimizationLevel == UE_MUTABLE_MAX_OPTIMIZATION;
+		ModelResources->bIsCompiledWithOptimization = GenerationContext.CompilationContext.Options.OptimizationLevel == UE_MUTABLE_MAX_OPTIMIZATION;
 
-		ModelResources.bCompiledWithHDTextureCompression = GenerationContext.CompilationContext.Options.TextureCompression == ECustomizableObjectTextureCompression::HighQuality;
+		ModelResources->bCompiledWithHDTextureCompression = GenerationContext.CompilationContext.Options.TextureCompression == ECustomizableObjectTextureCompression::HighQuality;
 
-		ModelResources.AlwaysLoadedExtensionData = MoveTemp(GenerationContext.AlwaysLoadedExtensionData);
-
-		ModelResources.StreamedExtensionData.Empty(GenerationContext.StreamedExtensionData.Num());
-		for (const TPair<FName, UCustomizableObjectResourceDataContainer*>& Pair : GenerationContext.StreamedExtensionData)
-		{
-			const FName& ContainerName = Pair.Key;
-			UCustomizableObjectResourceDataContainer* Container = Pair.Value;
-
-			UCustomizableObjectResourceDataContainer* NewContainer = FindObject<UCustomizableObjectResourceDataContainer>(CurrentObject, *ContainerName.ToString());
-			if (!NewContainer)
-			{
-				NewContainer = NewObject<UCustomizableObjectResourceDataContainer>(
-					CurrentObject,
-					FName(*ContainerName.ToString()),
-					RF_Public);
-			}
-			
-			NewContainer->Data = Container->Data;
-			
-			ModelResources.StreamedExtensionData.Emplace(NewContainer);
-		}
+		ModelResources->AlwaysLoadedExtensionData = MoveTemp(GenerationContext.AlwaysLoadedExtensionData);
+		ModelResources->StreamedExtensionDataEditor = MoveTemp(GenerationContext.StreamedExtensionData);
 
 #if WITH_EDITORONLY_DATA
 		// Cache the tables that are used by more than one param so that CompileOnlySelected can work properly
-		ModelResources.TableToParamNames = MoveTemp(GenerationContext.TableToParamNames);
-		ModelResources.CustomizableObjectPathMap = MoveTemp(GenerationContext.CustomizableObjectPathMap);
+		ModelResources->TableToParamNames = MoveTemp(GenerationContext.TableToParamNames);
+		ModelResources->CustomizableObjectPathMap = MoveTemp(GenerationContext.CustomizableObjectPathMap);
 #endif
 
-		ModelResources.ComponentNamesPerObjectComponent = MoveTemp(GenerationContext.ComponentNames);
+		ModelResources->ComponentNamesPerObjectComponent = MoveTemp(GenerationContext.ComponentNames);
 
 		if (ICustomizableObjectVersionBridgeInterface* VersionBridge = Cast<ICustomizableObjectVersionBridgeInterface>(GraphTraversal::GetRootObject(CurrentObject)->VersionBridge))
 		{
-			ModelResources.ReleaseVersion = VersionBridge->GetCurrentVersionAsString();
+			ModelResources->ReleaseVersion = VersionBridge->GetCurrentVersionAsString();
 		}
 
-		ModelResources.ComponentLODs = GenerationContext.NumLODs;
+		ModelResources->ComponentLODs = GenerationContext.NumLODs;
 
 		if (GenerationContext.bEnableLODStreaming)
 		{
-			ModelResources.ComponentNumLODsToStream = GenerationContext.NumMaxLODsToStream;
+			ModelResources->ComponentNumLODsToStream = GenerationContext.NumMaxLODsToStream;
 		}
 		else
 		{
-			for (TTuple<FName, uint8>& Pair : ModelResources.ComponentNumLODsToStream)
+			for (TTuple<FName, uint8>& Pair : ModelResources->ComponentNumLODsToStream)
 			{
 				Pair.Value = 0;
 			}
 		}
 
-		ModelResources.ComponentFirstLODAvailable = GenerationContext.FirstLODAvailable;
+		ModelResources->ComponentFirstLODAvailable = GenerationContext.FirstLODAvailable;
 		
-		ModelResources.ParticipatingObjects = MoveTemp(GenerationContext.ParticipatingObjects);
+		ModelResources->ParticipatingObjects = MoveTemp(GenerationContext.ParticipatingObjects);
 		
 		if (CurrentOptions.bGatherReferences)
 		{
-			CurrentObject->GetPrivate()->References = ModelResources;
-			CurrentObject->GetPrivate()->References.RuntimeReferencedTextures.Empty(); // Empty in case the of none optimization. In maximum optimization, they are Mutable textures. 
+			CurrentObject->GetPrivate()->ReferencedObjects = Cast<UModelResources>(StaticDuplicateObject(ModelResources, CurrentObject));
+			CurrentObject->GetPrivate()->ReferencedObjects->RuntimeReferencedTextures.Empty(); // Empty in case the of none optimization. In maximum optimization, they are Mutable textures. 
 			CurrentObject->Modify();
 		}
 
-		ModelResources.StreamedResourceData.Empty(GenerationContext.StreamedResourceData.Num());
-		for (TPair<FName, UCustomizableObjectResourceDataContainer*>& Pair : GenerationContext.StreamedResourceData)
-		{
-			const FName& ContainerName = Pair.Key;
-			UCustomizableObjectResourceDataContainer* Container = Pair.Value;
-
-			UCustomizableObjectResourceDataContainer* NewContainer = FindObject<UCustomizableObjectResourceDataContainer>(CurrentObject, *ContainerName.ToString());
-			if (!NewContainer)
-			{
-				NewContainer = NewObject<UCustomizableObjectResourceDataContainer>(
-					CurrentObject,
-					FName(*ContainerName.ToString()),
-					RF_Public);
-			}
-			
-			NewContainer->Data = Container->Data;
-			
-			ModelResources.StreamedResourceData.Emplace(NewContainer);
-		}
+		ModelResources->StreamedResourceDataEditor = MoveTemp(GenerationContext.StreamedResourceData);
 		
 		// Pass-through textures
 		TArray<FMutableSourceTextureData> NewCompileTimeReferencedTextures;
@@ -1128,16 +1090,23 @@ void FCustomizableObjectCompiler::CompleteRequest(ECompilationStatePrivate State
 			Model->GetPrivate()->UnloadRoms();
 		}
 
-		if (Result == ECompilationResultPrivate::Success || Result == ECompilationResultPrivate::Warnings)
-		{
-			CurrentObject->GetPrivate()->SetModel(Model, GenerateIdentifier(*CurrentObject));
-		}
-		else
+		if (!CurrentOptions.bIsCooking)
 		{
-			CurrentObject->GetPrivate()->SetModel(nullptr, {});
-		}
+			if (Result == ECompilationResultPrivate::Success || Result == ECompilationResultPrivate::Warnings)
+			{
+				CurrentObject->GetPrivate()->SetModel(Model, GenerateIdentifier(*CurrentObject));
+				CurrentObject->GetPrivate()->SetModelResources(ModelResources, CurrentOptions.bIsCooking);
+				CurrentObject->GetPrivate()->SetModelStreamableBulkData(ModelStreamableBulkData, CurrentOptions.bIsCooking);
+			}
+			else
+			{
+				CurrentObject->GetPrivate()->SetModel(nullptr, {});
+				CurrentObject->GetPrivate()->SetModelResources(nullptr, CurrentOptions.bIsCooking);
+				CurrentObject->GetPrivate()->SetModelStreamableBulkData(nullptr, CurrentOptions.bIsCooking);
+			}
 
-		CurrentObject->GetPrivate()->PostCompile();
+			CurrentObject->GetPrivate()->PostCompile();
+		}
 
 		UE_LOG(LogMutable, Display, TEXT("Finished compiling Customizable Object %s. Compilation took %5.3f seconds to complete."),
 			*CurrentObject->GetName(), FPlatformTime::Seconds() - CompilationStartTime);
@@ -1181,6 +1150,8 @@ void FCustomizableObjectCompiler::CompleteRequest(ECompilationStatePrivate State
 	CurrentObject = nullptr;
 	CurrentRequest.Reset();
 	Model.Reset();
+	ModelResources = nullptr;
+	ModelStreamableBulkData.Reset();
 
 	UE_LOG(LogMutable, Verbose, TEXT("PROFILE: [ %16.8f ] Completed compile request."), FPlatformTime::Seconds());
 	UE_LOG(LogMutable, Verbose, TEXT("PROFILE: -----------------------------------------------------------"));
@@ -1251,6 +1222,8 @@ bool FCustomizableObjectCompiler::TryLoadCompiledDataFromDDC(UCustomizableObject
 	CustomizableObject.GetPrivate()->LoadCompiledDataFromDDC(CurrentOptions, DefaultPolicy, &CacheKey);
 	
 	Model = CustomizableObject.GetPrivate()->GetModel();
+	ModelResources = CustomizableObject.GetPrivate()->GetModelResources(CurrentOptions.bIsCooking);
+	ModelStreamableBulkData = CustomizableObject.GetPrivate()->GetModelStreamableBulkData(CurrentOptions.bIsCooking);
 	return CustomizableObject.IsCompiled();
 }
 
@@ -1361,9 +1334,6 @@ void FCustomizableObjectCompiler::FinishCompilationTask()
 			ModelStreamables.Add(FileIndex, FMutableStreamableBlock{ 0, uint32(Flags), Offset });
 			Offset += ResourceSize;
 		}
-
-		// Always work with the ModelStreamableData (Editor) when compiling. They'll be copied to the cooked version during PreSave.
-		CurrentObject->GetPrivate()->SetModelStreamableBulkData(ModelStreamableBulkData, false);
 	}
 
 	// Order matters
@@ -1375,7 +1345,7 @@ void FCustomizableObjectCompiler::FinishCompilationTask()
 
 	// Create SaveDD task
 	TRACE_BEGIN_REGION(UE_MUTABLE_SAVEDD_REGION);
-	SaveDDTask = MakeShareable(new FCustomizableObjectSaveDDRunnable(CurrentRequest, Model, CurrentObject->GetPrivate()->GetModelResources(CurrentOptions.bIsCooking), ModelStreamableBulkData));
+	SaveDDTask = MakeShareable(new FCustomizableObjectSaveDDRunnable(CurrentRequest, Model, *ModelResources, ModelStreamableBulkData));
 }
 
 
@@ -1396,6 +1366,8 @@ void FCustomizableObjectCompiler::FinishSavingDerivedDataTask()
 
 		MutablePrivate::FMutableCachedPlatformData& Data = CurrentObject->GetPrivate()->CachedPlatformsData.Add(PlatformName);
 		Data = MoveTemp(SaveDDTask->PlatformData);
+
+		Data.ModelResources->InitCookData();
 	}
 
 	// Order matters
diff --git a/Engine/Plugins/Mutable/Source/CustomizableObjectEditor/Private/MuCOE/CustomizableObjectCompiler.h b/Engine/Plugins/Mutable/Source/CustomizableObjectEditor/Private/MuCOE/CustomizableObjectCompiler.h
--- a/Engine/Plugins/Mutable/Source/CustomizableObjectEditor/Private/MuCOE/CustomizableObjectCompiler.h
+++ b/Engine/Plugins/Mutable/Source/CustomizableObjectEditor/Private/MuCOE/CustomizableObjectCompiler.h
@@ -22,6 +22,7 @@ class FRunnableThread;
 class FText;
 class UCustomizableObject;
 class UCustomizableObjectNode;
+class UModelResources;
 struct FMutableGraphGenerationContext;
 
 namespace mu
@@ -130,6 +131,9 @@ private:
 	TSharedPtr<mu::FModel, ESPMode::ThreadSafe> Model;
 	TSharedPtr<FModelStreamableBulkData> ModelStreamableBulkData;
 
+	// Protected from GC with FCustomizableObjectCompiler::AddReferencedObjects
+	TObjectPtr<UModelResources> ModelResources;
+
 	// Protected from GC with FCustomizableObjectCompiler::AddReferencedObjects
 	TObjectPtr<UCustomizableObject> CurrentObject = nullptr;
 
diff --git a/Engine/Plugins/Mutable/Source/CustomizableObjectEditor/Private/MuCOE/CustomizableObjectCookPackageSplitter.cpp b/Engine/Plugins/Mutable/Source/CustomizableObjectEditor/Private/MuCOE/CustomizableObjectCookPackageSplitter.cpp
--- a/Engine/Plugins/Mutable/Source/CustomizableObjectEditor/Private/MuCOE/CustomizableObjectCookPackageSplitter.cpp
+++ b/Engine/Plugins/Mutable/Source/CustomizableObjectEditor/Private/MuCOE/CustomizableObjectCookPackageSplitter.cpp
@@ -12,6 +12,21 @@ REGISTER_COOKPACKAGE_SPLITTER(FCustomizableObjectCookPackageSplitter, UCustomiza
 
 namespace
 {
+UModelResources* FindModelResources(UCustomizableObject& Object)
+{
+	// All platforms should have the same resources
+	const TMap<FString, MutablePrivate::FMutableCachedPlatformData>& CachePlatforms = Object.GetPrivate()->CachedPlatformsData;
+	for (const TPair<FString, MutablePrivate::FMutableCachedPlatformData>& PlatformData : CachePlatforms)
+	{
+		if (PlatformData.Value.ModelResources)
+		{
+			return PlatformData.Value.ModelResources.Get();
+		}
+	}
+
+	return nullptr;
+}
+
 // Look up a streamed Resource Data constant by name on a Customizable Object.
 //
 // Returns nullptr if not found.
@@ -22,20 +37,20 @@ FCustomizableObjectStreamedResourceData* FindStreamedResourceData(
 {
 	return Algo::FindByPredicate(
 		StreamedResources,
-			[&ContainerName](const FCustomizableObjectStreamedResourceData& StreamedData)
+		[&ContainerName](const FCustomizableObjectStreamedResourceData& StreamedData)
+		{
+			const FSoftObjectPath& Path = StreamedData.GetPath().ToSoftObjectPath();
+
+			// ContainerName should match the last element of the path, which could be the
+			// sub-path string or the asset name.
+
+			if (Path.GetSubPathString().Len() > 0)
 			{
-				const FSoftObjectPath& Path = StreamedData.GetPath().ToSoftObjectPath();
-				
-				// ContainerName should match the last element of the path, which could be the
-				// sub-path string or the asset name.
-
-				if (Path.GetSubPathString().Len() > 0)
-				{
-					return Path.GetSubPathString() == ContainerName;
-				}
-
-				return Path.GetAssetName() == ContainerName;
-			});
+				return Path.GetSubPathString() == ContainerName;
+			}
+
+			return Path.GetAssetName() == ContainerName;
+		});
 }
 
 enum class EMoveContainerError
@@ -50,11 +65,11 @@ const TCHAR* LexToString(EMoveContainerError Error)
 {
 	switch (Error)
 	{
-		case EMoveContainerError::None: return TEXT("None");
-		case EMoveContainerError::FailedToLoadContainer: return TEXT("FailedToLoadContainer");
-		case EMoveContainerError::NameCollision: return TEXT("NameCollision");
-		case EMoveContainerError::RenameFailed: return TEXT("RenameFailed");
-		default: return TEXT("Unknown");
+	case EMoveContainerError::None: return TEXT("None");
+	case EMoveContainerError::FailedToLoadContainer: return TEXT("FailedToLoadContainer");
+	case EMoveContainerError::NameCollision: return TEXT("NameCollision");
+	case EMoveContainerError::RenameFailed: return TEXT("RenameFailed");
+	default: return TEXT("Unknown");
 	}
 }
 
@@ -78,7 +93,7 @@ EMoveContainerError MoveContainerToNewOuter(
 	if (Container->GetOuter() != NewOuter)
 	{
 		// Ensure the target object doesn't exist
-		if(FindObject<UObject>(NewOuter, *Container->GetName()))
+		if (FindObject<UObject>(NewOuter, *Container->GetName()))
 		{
 			return EMoveContainerError::NameCollision;
 		}
@@ -152,9 +167,16 @@ void GenerateNewPackage(const FCustomizableObjectStreamedResourceData& StreamedD
 bool FCustomizableObjectCookPackageSplitter::ShouldSplit(UObject* SplitData)
 {
 	UCustomizableObject* Object = CastChecked<UCustomizableObject>(SplitData);
-	const FModelResources& ModelResources = Object->GetPrivate()->GetModelResources(true);
+	
+	if (!Object->IsChildObject())
+	{
+		if(const UModelResources* ModelResources = FindModelResources(*Object))
+		{
+			return ModelResources->StreamedResourceData.Num() || ModelResources->StreamedExtensionData.Num();
+		}
+	}
 
-	return !Object->IsChildObject() && (ModelResources.StreamedResourceData.Num() > 0 || ModelResources.StreamedExtensionData.Num() > 0);
+	return false;
 }
 
 TArray<ICookPackageSplitter::FGeneratedPackage> FCustomizableObjectCookPackageSplitter::GetGenerateList(
@@ -164,19 +186,22 @@ TArray<ICookPackageSplitter::FGeneratedPackage> FCustomizableObjectCookPackageSp
 	// Keep a strong reference to the CO.
 	StrongObject.Reset(OwnerObject);
 	
-	const UCustomizableObject* Object = CastChecked<UCustomizableObject>(OwnerObject);
-	const FModelResources& ModelResources = Object->GetPrivate()->GetModelResources(true);
+	UCustomizableObject* Object = const_cast<UCustomizableObject*>(CastChecked<UCustomizableObject>(OwnerObject));
+
+	// All platforms should have the same resources
+	const UModelResources* ModelResources = FindModelResources(*Object);
+	check(ModelResources);
 
 	TArray<ICookPackageSplitter::FGeneratedPackage> Result;
 
 	// Generate a new package for each streamed Resource Data
-	for (const FCustomizableObjectStreamedResourceData& StreamedData : ModelResources.StreamedResourceData)
+	for (const FCustomizableObjectStreamedResourceData& StreamedData : ModelResources->StreamedResourceData)
 	{
 		GenerateNewPackage(StreamedData, OwnerPackage, OwnerObject, Result);
 	}
 
 	// Generate a new package for each streamed Extension Data
-	for (const FCustomizableObjectStreamedResourceData& StreamedData : ModelResources.StreamedExtensionData)
+	for (const FCustomizableObjectStreamedResourceData& StreamedData : ModelResources->StreamedExtensionData)
 	{
 		GenerateNewPackage(StreamedData, OwnerPackage, OwnerObject, Result);
 	}
@@ -229,11 +254,17 @@ bool FCustomizableObjectCookPackageSplitter::PreSaveGeneratorPackage(
 	};
 
 	UCustomizableObject* Object = CastChecked<UCustomizableObject>(OwnerObject);
-	FModelResources& ModelResources = Object->GetPrivate()->GetModelResources(true);
+
+	UModelResources* ModelResources = FindModelResources(*Object);
+	if (!ModelResources)
+	{
+		UE_LOG(LogMutable, Warning, TEXT("Couldn't find ModelResources. CO %s"), *GetNameSafe(Object));
+		return false;
+	}
 
 	// There should be one generated package per streamed Resource Data
-	const int32 NumStreamedData = ModelResources.StreamedResourceData.Num();
-	const int32 NumStreamedExtensionData = ModelResources.StreamedExtensionData.Num();
+	const int32 NumStreamedData = ModelResources->StreamedResourceData.Num();
+	const int32 NumStreamedExtensionData = ModelResources->StreamedExtensionData.Num();
 	
 	check(NumStreamedData + NumStreamedExtensionData == PlaceholderPackages.Num());
 
@@ -247,7 +278,7 @@ bool FCustomizableObjectCookPackageSplitter::PreSaveGeneratorPackage(
 	for (int32 Index = 0; Index < NumStreamedData; ++Index)
 	{
 		const ICookPackageSplitter::FGeneratedPackageForPreSave& GeneratedPackage = PlaceholderPackages[Index];
-		if (!PreSavePackage(GeneratedPackage, ModelResources.StreamedResourceData))
+		if (!PreSavePackage(GeneratedPackage, ModelResources->StreamedResourceData))
 		{
 			return false;
 		}
@@ -258,7 +289,7 @@ bool FCustomizableObjectCookPackageSplitter::PreSaveGeneratorPackage(
 	for (int32 Index = 0; Index < NumStreamedExtensionData; ++Index)
 	{
 		const ICookPackageSplitter::FGeneratedPackageForPreSave& GeneratedPackage = PlaceholderPackages[NumStreamedData + Index];
-		if (!PreSavePackage(GeneratedPackage, ModelResources.StreamedExtensionData))
+		if (!PreSavePackage(GeneratedPackage, ModelResources->StreamedExtensionData))
 		{
 			return false;
 		}
@@ -275,18 +306,23 @@ void FCustomizableObjectCookPackageSplitter::PostSaveGeneratorPackage(UPackage*
 	// array on the CO to how it was before PreSaveGeneratorPackage.
 
 	UCustomizableObject* Object = CastChecked<UCustomizableObject>(OwnerObject);
-	FModelResources& ModelResources = Object->GetPrivate()->GetModelResources(true);
+	UModelResources* ModelResources = FindModelResources(*Object);
+	if (!ModelResources)
+	{
+		UE_LOG(LogMutable, Warning, TEXT("Couldn't find ModelResources. CO %s"), *GetNameSafe(Object));
+		return;
+	}
 
 	TArray<FCustomizableObjectStreamedResourceData> NewArray;
 	NewArray.Reset(SavedContainerNames.Num());
 
 	for (const FString& ContainerName : SavedContainerNames)
 	{
-		FCustomizableObjectStreamedResourceData* ResourceData = FindStreamedResourceData(ModelResources.StreamedResourceData, ContainerName);
+		FCustomizableObjectStreamedResourceData* ResourceData = FindStreamedResourceData(ModelResources->StreamedResourceData, ContainerName);
 		if (!ResourceData)
 		{
 			UE_LOG(LogMutable, Error, TEXT("Couldn't find streamed Resource Data container with name %s in array of %d entries"),
-				*ContainerName, ModelResources.StreamedResourceData.Num());
+				*ContainerName, ModelResources->StreamedResourceData.Num());
 
 			continue;
 		}
@@ -298,17 +334,17 @@ void FCustomizableObjectCookPackageSplitter::PostSaveGeneratorPackage(UPackage*
 		NewArray.Emplace(Container);
 	}
 
-	ModelResources.StreamedResourceData = NewArray;
+	ModelResources->StreamedResourceData = NewArray;
 
 	NewArray.Reset(SavedExtensionContainerNames.Num());
 
 	for (const FString& ContainerName : SavedExtensionContainerNames)
 	{
-		FCustomizableObjectStreamedResourceData* ResourceData = FindStreamedResourceData(ModelResources.StreamedExtensionData, ContainerName);
+		FCustomizableObjectStreamedResourceData* ResourceData = FindStreamedResourceData(ModelResources->StreamedExtensionData, ContainerName);
 		if (!ResourceData)
 		{
 			UE_LOG(LogMutable, Error, TEXT("Couldn't find streamed Extension Data container with name %s in array of %d entries"),
-				*ContainerName, ModelResources.StreamedExtensionData.Num());
+				*ContainerName, ModelResources->StreamedExtensionData.Num());
 
 			continue;
 		}
@@ -320,7 +356,7 @@ void FCustomizableObjectCookPackageSplitter::PostSaveGeneratorPackage(UPackage*
 		NewArray.Emplace(Container);
 	}
 
-	ModelResources.StreamedExtensionData = NewArray;
+	ModelResources->StreamedExtensionData = NewArray;
 }
 
 bool FCustomizableObjectCookPackageSplitter::PopulateGeneratedPackage(
@@ -333,18 +369,18 @@ bool FCustomizableObjectCookPackageSplitter::PopulateGeneratedPackage(
 	// Move the container into its newly generated package
 
 	UCustomizableObject* Object = CastChecked<UCustomizableObject>(OwnerObject);
-	FModelResources& ModelResources = Object->GetPrivate()->GetModelResources(true);
+	UModelResources* ModelResources = FindModelResources(*Object);
 
-	FCustomizableObjectStreamedResourceData* ResourceData = FindStreamedResourceData(ModelResources.StreamedResourceData, GeneratedPackage.RelativePath);
+	FCustomizableObjectStreamedResourceData* ResourceData = FindStreamedResourceData(ModelResources->StreamedResourceData, GeneratedPackage.RelativePath);
 	if (!ResourceData)
 	{
-		ResourceData = FindStreamedResourceData(ModelResources.StreamedExtensionData, GeneratedPackage.RelativePath);
+		ResourceData = FindStreamedResourceData(ModelResources->StreamedExtensionData, GeneratedPackage.RelativePath);
 	}
 
 	if (!ResourceData)
 	{
 		UE_LOG(LogMutable, Error, TEXT("Couldn't find streamed resource Data container with name %s in arrays of %d and %d entries"),
-			*GeneratedPackage.RelativePath, ModelResources.StreamedResourceData.Num(), ModelResources.StreamedExtensionData.Num());
+			*GeneratedPackage.RelativePath, ModelResources->StreamedResourceData.Num(), ModelResources->StreamedExtensionData.Num());
 
 		return false;
 	}
@@ -381,18 +417,18 @@ void FCustomizableObjectCookPackageSplitter::PostSaveGeneratedPackage(
 	// that everything is the same as it was before cooking.
 
 	UCustomizableObject* Object = CastChecked<UCustomizableObject>(OwnerObject);
-	FModelResources& ModelResources = Object->GetPrivate()->GetModelResources(true);
+	UModelResources* ModelResources = FindModelResources(*Object);
 
-	FCustomizableObjectStreamedResourceData* ResourceData = FindStreamedResourceData(ModelResources.StreamedResourceData, GeneratedPackage.RelativePath);
+	FCustomizableObjectStreamedResourceData* ResourceData = FindStreamedResourceData(ModelResources->StreamedResourceData, GeneratedPackage.RelativePath);
 	if (!ResourceData)
 	{
-		ResourceData = FindStreamedResourceData(ModelResources.StreamedExtensionData, GeneratedPackage.RelativePath);
+		ResourceData = FindStreamedResourceData(ModelResources->StreamedExtensionData, GeneratedPackage.RelativePath);
 	}
 
 	if (!ResourceData)
 	{
 		UE_LOG(LogMutable, Error, TEXT("Couldn't find streamed resource Data container with name %s in arrays of %d and %d entries"),
-			*GeneratedPackage.RelativePath, ModelResources.StreamedResourceData.Num(), ModelResources.StreamedExtensionData.Num());
+			*GeneratedPackage.RelativePath, ModelResources->StreamedResourceData.Num(), ModelResources->StreamedExtensionData.Num());
 
 		return;
 	}
diff --git a/Engine/Plugins/Mutable/Source/CustomizableObjectEditor/Private/MuCOE/CustomizableObjectEditor.cpp b/Engine/Plugins/Mutable/Source/CustomizableObjectEditor/Private/MuCOE/CustomizableObjectEditor.cpp
--- a/Engine/Plugins/Mutable/Source/CustomizableObjectEditor/Private/MuCOE/CustomizableObjectEditor.cpp
+++ b/Engine/Plugins/Mutable/Source/CustomizableObjectEditor/Private/MuCOE/CustomizableObjectEditor.cpp
@@ -1487,7 +1487,7 @@ void FCustomizableObjectEditor::DebugObject() const
 
 void FCustomizableObjectEditor::ClearGatheredReferences()
 {
-	CustomizableObject->GetPrivate()->References = {};
+	CustomizableObject->GetPrivate()->ReferencedObjects = {};
 	CustomizableObject->Modify();
 }
 
diff --git a/Engine/Plugins/Mutable/Source/CustomizableObjectEditor/Private/MuCOE/CustomizableObjectEditorModule.cpp b/Engine/Plugins/Mutable/Source/CustomizableObjectEditor/Private/MuCOE/CustomizableObjectEditorModule.cpp
--- a/Engine/Plugins/Mutable/Source/CustomizableObjectEditor/Private/MuCOE/CustomizableObjectEditorModule.cpp
+++ b/Engine/Plugins/Mutable/Source/CustomizableObjectEditor/Private/MuCOE/CustomizableObjectEditorModule.cpp
@@ -157,9 +157,16 @@ void ShowOnScreenCompileWarnings()
 			continue;
 		}
 	
+		const UModelResources* ModelResources = Object->GetPrivate()->GetModelResources();
+		if (!ModelResources)
+		{
+			continue;
+		}
+
 		// Show a warning if the compilation was not done with optimizations.
 		const uint64 KeyCompiledWithOptimization = reinterpret_cast<uint64>(Object);
-		if (!Object->GetPrivate()->GetModelResources().bIsCompiledWithOptimization)
+
+		if (!ModelResources->bIsCompiledWithOptimization)
 		{
 			FString Msg = FString::Printf(TEXT("Customizable Object [%s] was compiled without optimization."), *Object->GetName());
 			GEngine->AddOnScreenDebugMessage(KeyCompiledWithOptimization, ShowOnScreenCompileWarningsTickerTime * 2.0f, FColor::Yellow, Msg);
@@ -582,7 +589,13 @@ bool FCustomizableObjectEditorModule::IsCompilationOutOfDate(const UCustomizable
 	
 	const FAssetRegistryModule& AssetRegistryModule = FModuleManager::LoadModuleChecked<FAssetRegistryModule>("AssetRegistry");
 
-	const TMap<FName, FGuid>& OldParticipatingObjects = Object.GetPrivate()->GetModelResources().ParticipatingObjects;
+	const UModelResources* ModelResources = Object.GetPrivate()->GetModelResources();
+	if (!ModelResources)
+	{
+		return true;
+	}
+
+	const TMap<FName, FGuid>& OldParticipatingObjects = ModelResources->ParticipatingObjects;
 	
 	for (const TTuple<FName, FGuid>& ParticipatingObject : OldParticipatingObjects)
 	{
@@ -653,7 +666,7 @@ bool FCustomizableObjectEditorModule::IsCompilationOutOfDate(const UCustomizable
 	bReleaseVersion = false;
 	if (ICustomizableObjectVersionBridgeInterface* VersionBridge = Cast<ICustomizableObjectVersionBridgeInterface>(Object.VersionBridge))
 	{
-		bReleaseVersion = Object.GetPrivate()->GetModelResources().ReleaseVersion != VersionBridge->GetCurrentVersionAsString();
+		bReleaseVersion = ModelResources->ReleaseVersion != VersionBridge->GetCurrentVersionAsString();
 	}
 
 	return bReleaseVersion || !OutOfDatePackages.IsEmpty() || !AddedPackages.IsEmpty() || !RemovedPackages.IsEmpty();
diff --git a/Engine/Plugins/Mutable/Source/CustomizableObjectEditor/Private/MuCOE/CustomizableObjectInstanceBakingUtils.cpp b/Engine/Plugins/Mutable/Source/CustomizableObjectEditor/Private/MuCOE/CustomizableObjectInstanceBakingUtils.cpp
--- a/Engine/Plugins/Mutable/Source/CustomizableObjectEditor/Private/MuCOE/CustomizableObjectInstanceBakingUtils.cpp
+++ b/Engine/Plugins/Mutable/Source/CustomizableObjectEditor/Private/MuCOE/CustomizableObjectInstanceBakingUtils.cpp
@@ -411,10 +411,18 @@ bool BakeCustomizableObjectInstance(
 		return false;
 	}
 
+	// Early out if the Instance ModelResources is not valid.
+	const UModelResources* ModelResources = InstanceCO->GetPrivate()->GetModelResources();
+	if (!ModelResources)
+	{
+		UE_LOG(LogMutable, Error, TEXT("The ModelResources from the Customizable Object is not valid."));
+		return false;
+	}
+
 	// COI Validation completed : Proceed with the baking operation ----------------------------------------------------------------------------------
 	
 	// Notify of better configuration -> Continue operation normally
-	if (InstanceCO->GetPrivate()->GetModelResources().bCompiledWithHDTextureCompression == false)
+	if (ModelResources->bCompiledWithHDTextureCompression == false)
 	{
 		FCustomizableObjectEditorLogger::CreateLog(
 		LOCTEXT("CustomizableObjecInstanceBakingUtils_LowQualityTextures", "The Customizable Object wasn't compiled with high quality textures. For the best baking results, change the Texture Compression setting and recompile it."))
@@ -842,9 +850,9 @@ bool BakeCustomizableObjectInstance(
 
 			// Copy LODSettings from the Reference Skeletal Mesh
 			{
-				if (InstanceCO && InstanceCO->GetPrivate()->GetModelResources().ReferenceSkeletalMeshesData.IsValidIndex(ObjectComponentIndex))
+				if (ModelResources->ReferenceSkeletalMeshesData.IsValidIndex(ObjectComponentIndex))
 				{
-					USkeletalMeshLODSettings* LODSettings = InstanceCO->GetPrivate()->GetModelResources().ReferenceSkeletalMeshesData[ObjectComponentIndex].SkeletalMeshLODSettings;
+					USkeletalMeshLODSettings* LODSettings = ModelResources->ReferenceSkeletalMeshesData[ObjectComponentIndex].SkeletalMeshLODSettings;
 					SkeletalMesh->SetLODSettings(LODSettings);
 				}
 			}
diff --git a/Engine/Plugins/Mutable/Source/CustomizableObjectEditor/Private/MuCOE/ExtensionDataCompilerInterface.cpp b/Engine/Plugins/Mutable/Source/CustomizableObjectEditor/Private/MuCOE/ExtensionDataCompilerInterface.cpp
--- a/Engine/Plugins/Mutable/Source/CustomizableObjectEditor/Private/MuCOE/ExtensionDataCompilerInterface.cpp
+++ b/Engine/Plugins/Mutable/Source/CustomizableObjectEditor/Private/MuCOE/ExtensionDataCompilerInterface.cpp
@@ -13,25 +13,15 @@ FExtensionDataCompilerInterface::FExtensionDataCompilerInterface(FMutableGraphGe
 {
 }
 
-TSharedPtr<const mu::FExtensionData> FExtensionDataCompilerInterface::MakeStreamedExtensionData(UCustomizableObjectResourceDataContainer*& OutContainer)
+TSharedPtr<const mu::FExtensionData> FExtensionDataCompilerInterface::MakeStreamedExtensionData(FInstancedStruct&& Data)
 {
 	TSharedPtr<mu::FExtensionData> Result = MakeShared<mu::FExtensionData>();
 	Result->Origin = mu::FExtensionData::EOrigin::ConstantStreamed;
 	Result->Index = GenerationContext.StreamedExtensionData.Num();
 	
-	if (!GenerationContext.bParticipatingObjectsPass)
-	{
-		// Generate a deterministic name to help with deterministic cooking
-		const FString ContainerName = FString::Printf(TEXT("Streamed_%d"), Result->Index);
-
-		OutContainer = NewObject<UCustomizableObjectResourceDataContainer>(
-			GetTransientPackage(),
-			FName(*ContainerName),
-			RF_Public);
+	FCustomizableObjectResourceData* CompileTimeExtensionData = &GenerationContext.StreamedExtensionData.AddDefaulted_GetRef();
+	CompileTimeExtensionData->Data = MoveTemp(Data);
 
-		GenerationContext.StreamedExtensionData.Emplace(ContainerName, OutContainer);
-	}
-	
 	return Result;
 }
 
diff --git a/Engine/Plugins/Mutable/Source/CustomizableObjectEditor/Private/MuCOE/GenerateMutableSource/GenerateMutableSource.cpp b/Engine/Plugins/Mutable/Source/CustomizableObjectEditor/Private/MuCOE/GenerateMutableSource/GenerateMutableSource.cpp
--- a/Engine/Plugins/Mutable/Source/CustomizableObjectEditor/Private/MuCOE/GenerateMutableSource/GenerateMutableSource.cpp
+++ b/Engine/Plugins/Mutable/Source/CustomizableObjectEditor/Private/MuCOE/GenerateMutableSource/GenerateMutableSource.cpp
@@ -353,52 +353,31 @@ TSharedPtr<mu::FMesh> FMutableGraphGenerationContext::FindGeneratedMesh( const F
 }
 
 
-int32 FMutableGraphGenerationContext::AddStreamedResource(uint32 InResourceHash, UCustomizableObjectResourceDataContainer*& OutNewResource)
+int32 FMutableGraphGenerationContext::AddStreamedResource(UAssetUserData& AssetUserData)
 {
 	if (bParticipatingObjectsPass)
 	{
-		return INDEX_NONE;	
+		return INDEX_NONE;
 	}
-	
-	OutNewResource = nullptr;
+
+	const uint32 AssetIdentifier = AssetUserData.GetUniqueID();
 
 	// Return resource index if found.
-	if (int32* ResourceIndex = StreamedResourceIndices.Find(InResourceHash))
+	if (int32* ResourceIndex = StreamedResourceIndices.Find(AssetIdentifier))
 	{
 		return *ResourceIndex;
 	}
-	
-	int32 NewResourceIndex = StreamedResourceData.Num();
-	const FString ContainerName = GetObjectName() + FString::Printf(TEXT("_SR_%d"), NewResourceIndex);
-
-	OutNewResource = NewObject<UCustomizableObjectResourceDataContainer>(
-		GetTransientPackage(),
-		FName(*ContainerName),
-		RF_Public);
-
-	StreamedResourceData.Emplace(ContainerName, OutNewResource);
-	StreamedResourceIndices.Add({ InResourceHash, NewResourceIndex });
 
-	return NewResourceIndex;
-}
-
-
-int32 FMutableGraphGenerationContext::AddAssetUserDataToStreamedResources(UAssetUserData* AssetUserData)
-{
-	check(AssetUserData);
-	const uint32 AssetIdentifier = AssetUserData->GetUniqueID();
+	FCustomizableObjectAssetUserData ResourceData;
+	ResourceData.AssetUserData = &AssetUserData;
 
-	UCustomizableObjectResourceDataContainer* NewResource = nullptr;
-	const int32 ResourceIndex = AddStreamedResource(AssetIdentifier, NewResource);
+	const int32 ResourceIndex = StreamedResourceData.Num();
 
-	if (NewResource) // Nullptr if not new
-	{
-		FCustomizableObjectAssetUserData ResourceData;
-		ResourceData.AssetUserDataEditor = AssetUserData;
+	FCustomizableObjectResourceData& Resource = StreamedResourceData.AddDefaulted_GetRef();
+	Resource.Type = ECOResourceDataType::AssetUserData;
+	Resource.Data.InitializeAs<FCustomizableObjectAssetUserData>(MoveTemp(ResourceData));
 
-		NewResource->Data.Type = ECOResourceDataType::AssetUserData;
-		NewResource->Data.Data = FInstancedStruct::Make(ResourceData);
-	}
+	StreamedResourceIndices.Add({ AssetIdentifier, ResourceIndex });
 
 	return ResourceIndex;
 }
@@ -468,22 +447,6 @@ void FMutableGraphGenerationContext::GenerateSharedSurfacesUniqueIds()
 }
 
 
-//void FMutableGraphGenerationContext::CheckPhysicsAssetInSkeletalMesh(const USkeletalMesh* SkeletalMesh)
-//{
-//	if (SkeletalMesh && SkeletalMesh->GetPhysicsAsset() && !DiscartedPhysicsAssetMap.Find(SkeletalMesh->GetPhysicsAsset()))
-//	{
-//		for (const TObjectPtr<USkeletalBodySetup>& Bodies : SkeletalMesh->GetPhysicsAsset()->SkeletalBodySetups)
-//		{
-//			if (Bodies->BoneName != NAME_None && SkeletalMesh->GetRefSkeleton().FindBoneIndex(Bodies->BoneName) == INDEX_NONE)
-//			{
-//				DiscartedPhysicsAssetMap.Add(SkeletalMesh->GetPhysicsAsset(), DiscartedPhysicsAssetMap.Num());
-//				break;
-//			}
-//		}
-//	}
-//}
-
-
 FMutableComponentInfo* FMutableGraphGenerationContext::GetCurrentComponentInfo()
 {
 	return CompilationContext.GetComponentInfo(CurrentMeshComponent);
@@ -1346,7 +1309,12 @@ void PopulateReferenceSkeletalMeshesData(FMutableGraphGenerationContext& Generat
 		FMutableRefSkeletalMeshData& Data = GenerationContext.ReferenceSkeletalMeshesData[ComponentIndex];
 
 		// Set the RefSkeletalMesh
-		Data.SkeletalMesh = TObjectPtr<USkeletalMesh>(RefSkeletalMesh);
+		if (!GenerationContext.CompilationContext.Options.TargetPlatform->IsClientOnly()
+			|| GenerationContext.CompilationContext.Object->bEnableUseRefSkeletalMeshAsPlaceholder)
+		{
+			Data.SkeletalMesh = RefSkeletalMesh;
+		}
+
 		Data.SoftSkeletalMesh = RefSkeletalMesh;
 
 		// Set the optional USkeletalMeshLODSettings that will be applied to the generated transient meshes or the baked meshes
@@ -1454,13 +1422,7 @@ void PopulateReferenceSkeletalMeshesData(FMutableGraphGenerationContext& Generat
 			{
 				if (AssetUserData)
 				{
-					FMutableRefAssetUserData& MutAssetUserData = Data.AssetUserData.AddDefaulted_GetRef();
-					MutAssetUserData.AssetUserDataIndex = GenerationContext.AddAssetUserDataToStreamedResources(AssetUserData);
-					if (MutAssetUserData.AssetUserDataIndex != INDEX_NONE)
-					{
-						MutAssetUserData.AssetUserData = GenerationContext.StreamedResourceData[MutAssetUserData.AssetUserDataIndex].Value;
-						check(MutAssetUserData.AssetUserData->Data.Type == ECOResourceDataType::AssetUserData);
-					}
+					Data.AssetUserDataIndices.AddUnique(GenerationContext.AddStreamedResource(*AssetUserData));
 				}
 			}
 		}
diff --git a/Engine/Plugins/Mutable/Source/CustomizableObjectEditor/Private/MuCOE/GenerateMutableSource/GenerateMutableSource.h b/Engine/Plugins/Mutable/Source/CustomizableObjectEditor/Private/MuCOE/GenerateMutableSource/GenerateMutableSource.h
--- a/Engine/Plugins/Mutable/Source/CustomizableObjectEditor/Private/MuCOE/GenerateMutableSource/GenerateMutableSource.h
+++ b/Engine/Plugins/Mutable/Source/CustomizableObjectEditor/Private/MuCOE/GenerateMutableSource/GenerateMutableSource.h
@@ -745,14 +745,9 @@ public:
 	/** Find a mesh if already generated for a given source and flags. */
 	TSharedPtr<mu::FMesh> FindGeneratedMesh(const FGeneratedMeshData::FKey& Key);
 
-	/** Add a resource to the streamed resources array.
-	  * OutStreamedResourceContainer - Container to store the streamed resource.
-	  * Returns resource index in the array of streamed resources. */
-	int32 AddStreamedResource(const uint32 InResourceHash, UCustomizableObjectResourceDataContainer*& OutStreamedResourceContainer);
-
 	/** Adds a streamed resource of type AssetUserData.
 	  * Returns resource index in the array of streamed resources. */
-	int32 AddAssetUserDataToStreamedResources(UAssetUserData* AssetUserData);
+	int32 AddStreamedResource(UAssetUserData& AssetUserData);
 
 	// Check if the Id of the node Node already exists, if it's new adds it to NodeIds array, otherwise, returns new Id
 	const FGuid GetNodeIdUnique(const UCustomizableObjectNode* Node);
@@ -761,10 +756,6 @@ public:
 	void GenerateSharedSurfacesUniqueIds();
 
 
-	/** Check if the PhysicsAsset of a given SkeletalMesh has any SkeletalBodySetup with BoneNames not present in the
-	* InSkeletalMesh's RefSkeleton, if so, adds the PhysicsAsset to the DiscartedPhysicsAssetMap to display a warning later on */
-	//void CheckPhysicsAssetInSkeletalMesh(const USkeletalMesh* InSkeletalMesh);
-
 	/** Get the reference skeletal mesh associated to the current mesh component being generated */
 	FMutableComponentInfo* GetCurrentComponentInfo();
 
@@ -803,8 +794,7 @@ public:
 	TMap<FString, FString> CustomizableObjectPathMap;
 	TMap<FString, FMutableParameterData> ParameterUIDataMap;
 	TMap<FString, FMutableStateData> StateUIDataMap;
-	TMap<TTuple<FString, FString>, TSet<TSoftObjectPtr<UDataTable>>> IntParameterOptionDataTable;
-	//TMap<UPhysicsAsset*, uint32> DiscartedPhysicsAssetMap;
+	TMap<FIntegerParameterOptionKey, FIntegerParameterOptionDataTable> IntParameterOptionDataTable;
 
 
 	// Used to aviod Nodes with duplicated ids
@@ -902,12 +892,12 @@ public:
 
 	/** Resource Data constants */
 	TMap<uint32, int32> StreamedResourceIndices;
-	TArray<TPair<FName, UCustomizableObjectResourceDataContainer*>> StreamedResourceData;
+	TArray<FCustomizableObjectResourceData> StreamedResourceData;
 
 	/** Extension Data constants are collected here */
 	FExtensionDataCompilerInterface ExtensionDataCompilerInterface;
 	TArray<FCustomizableObjectResourceData> AlwaysLoadedExtensionData;
-	TArray<TPair<FName, UCustomizableObjectResourceDataContainer*>> StreamedExtensionData;
+	TArray<FCustomizableObjectResourceData> StreamedExtensionData;
 
 	/** Map to relate a Composite Data Table Row and its original DataTable */
 	TMap<UDataTable*,TMap<FName, TArray<UDataTable*>>> CompositeDataTableRowToOriginalDataTableMap;
diff --git a/Engine/Plugins/Mutable/Source/CustomizableObjectEditor/Private/MuCOE/GenerateMutableSource/GenerateMutableSourceMesh.cpp b/Engine/Plugins/Mutable/Source/CustomizableObjectEditor/Private/MuCOE/GenerateMutableSource/GenerateMutableSourceMesh.cpp
--- a/Engine/Plugins/Mutable/Source/CustomizableObjectEditor/Private/MuCOE/GenerateMutableSource/GenerateMutableSourceMesh.cpp
+++ b/Engine/Plugins/Mutable/Source/CustomizableObjectEditor/Private/MuCOE/GenerateMutableSource/GenerateMutableSourceMesh.cpp
@@ -3494,7 +3494,7 @@ mu::Ptr<mu::NodeMesh> GenerateMutableSourceMesh(const UEdGraphPin* Pin,
 							continue;
 						}
 
-						const int32 ResourceIndex = GenerationContext.AddAssetUserDataToStreamedResources(AssetUserData);
+						const int32 ResourceIndex = GenerationContext.AddStreamedResource(*AssetUserData);
 						if (ResourceIndex >= 0)
 						{	
 							FCustomizableObjectStreameableResourceId ResourceId;
diff --git a/Engine/Plugins/Mutable/Source/CustomizableObjectEditor/Private/MuCOE/GenerateMutableSource/GenerateMutableSourceTable.cpp b/Engine/Plugins/Mutable/Source/CustomizableObjectEditor/Private/MuCOE/GenerateMutableSource/GenerateMutableSourceTable.cpp
--- a/Engine/Plugins/Mutable/Source/CustomizableObjectEditor/Private/MuCOE/GenerateMutableSource/GenerateMutableSourceTable.cpp
+++ b/Engine/Plugins/Mutable/Source/CustomizableObjectEditor/Private/MuCOE/GenerateMutableSource/GenerateMutableSourceTable.cpp
@@ -262,12 +262,12 @@ bool FillTableColumn(const UCustomizableObjectNodeTable* TableNode, mu::Ptr<mu::
 							continue;
 						}
 
-						const int32 ResourceIndex = GenerationContext.AddAssetUserDataToStreamedResources(AssetUserData);
+						const int32 ResourceIndex = GenerationContext.AddStreamedResource(*AssetUserData);
 						
 						if (ResourceIndex >= 0)
 						{
 							FCustomizableObjectStreameableResourceId ResourceId;
-							ResourceId.Id = GenerationContext.AddAssetUserDataToStreamedResources(AssetUserData);
+							ResourceId.Id = (uint32)ResourceIndex;
 							ResourceId.Type = (uint8)FCustomizableObjectStreameableResourceId::EType::AssetUserData;
 
 							StreamedResources.Add(ResourceId);
@@ -1008,8 +1008,8 @@ void GenerateTableParameterUIData(const UDataTable* DataTable, const UCustomizab
 
 	for (const FName& Name : RowNames)
 	{
-		TSet<TSoftObjectPtr<UDataTable>>& DataTables = GenerationContext.IntParameterOptionDataTable.FindOrAdd(MakeTuple(TableNode->ParameterName, Name.ToString()));
-		DataTables.Add(TSoftObjectPtr<UDataTable>(const_cast<UDataTable*>(DataTable)));
+		FIntegerParameterOptionDataTable& Data = GenerationContext.IntParameterOptionDataTable.FindOrAdd({ TableNode->ParameterName, Name.ToString() });
+		Data.DataTables.Add(TSoftObjectPtr<UDataTable>(const_cast<UDataTable*>(DataTable)));
 	}
 	
 	// Generating Parameter UI MetaData if not exists
diff --git a/Engine/Plugins/Mutable/Source/CustomizableObjectEditor/Private/MuCOE/SCustomizableObjectEditorViewport.cpp b/Engine/Plugins/Mutable/Source/CustomizableObjectEditor/Private/MuCOE/SCustomizableObjectEditorViewport.cpp
--- a/Engine/Plugins/Mutable/Source/CustomizableObjectEditor/Private/MuCOE/SCustomizableObjectEditorViewport.cpp
+++ b/Engine/Plugins/Mutable/Source/CustomizableObjectEditor/Private/MuCOE/SCustomizableObjectEditorViewport.cpp
@@ -221,9 +221,12 @@ FText SCustomizableObjectEditorViewport::GetWarningText() const
 	{
 		if (const UCustomizableObject* Object = Instance->GetCustomizableObject())
 		{
-			if (!Object->GetPrivate()->GetModelResources().bIsCompiledWithOptimization)
+			if (const UModelResources* ModelResources = Object->GetPrivate()->GetModelResources())
 			{
-				return LOCTEXT("CompiledWithoutOptimization", "Compiled without maximum optimization. Updates will be slower!");
+				if (!ModelResources->bIsCompiledWithOptimization)
+				{
+					return LOCTEXT("CompiledWithoutOptimization", "Compiled without maximum optimization. Updates will be slower!");
+				}
 			}
 		}
 	}
diff --git a/Engine/Plugins/Mutable/Source/CustomizableObjectEditor/Public/MuCOE/ExtensionDataCompilerInterface.h b/Engine/Plugins/Mutable/Source/CustomizableObjectEditor/Public/MuCOE/ExtensionDataCompilerInterface.h
--- a/Engine/Plugins/Mutable/Source/CustomizableObjectEditor/Public/MuCOE/ExtensionDataCompilerInterface.h
+++ b/Engine/Plugins/Mutable/Source/CustomizableObjectEditor/Public/MuCOE/ExtensionDataCompilerInterface.h
@@ -5,7 +5,6 @@
 #include "MuR/ExtensionData.h"
 
 struct FInstancedStruct;
-class UCustomizableObjectResourceDataContainer;
 struct FMutableGraphGenerationContext;
 
 /**
@@ -22,16 +21,6 @@ public:
 	 * 
 	 * This constant will be cooked out to its own package and will be loaded as needed.
 	 * 
-	 * Importantly, any UObjects that are referenced by the Extension Data that aren't in an asset
-	 * package should have the container returned from this function as their Outer. This ensures
-	 * that they get cooked into the correct package for streaming.
-	 *
-	 * For example, if the Extension Data has a dynamically generated USkeletalMesh that it
-	 * references, which doesn't belong to another asset, the USkeletalMesh should have the given
-	 * container as its Outer.
-	 * 
-	 * Assign the constant data to the FInstancedStruct within the container's Data member.
-	 * 
 	 * The provided FInstancedStruct will be visible to the garbage collector, so any object
 	 * references from UPROPERTYs in the instanced struct will be treated as hard references from
 	 * the Customizable Object.
@@ -40,7 +29,7 @@ public:
 	 * TSoftObjectPtr, because the container will be moved to a different package during cooking
 	 * and your paths won't be automatically updated.
 	 */
-	TSharedPtr<const mu::FExtensionData> MakeStreamedExtensionData(UCustomizableObjectResourceDataContainer*& OutContainer);
+	TSharedPtr<const mu::FExtensionData> MakeStreamedExtensionData(FInstancedStruct&& Data);
 
 	/**
 	 * Register a new Extension Data compile-time constant that will always be loaded.
