[Mutable] Re-enable DDC for cook and turn off Asynch cook.

#rnx

diff --git a/Engine/Plugins/Mutable/Source/CustomizableObject/Private/MuCO/CustomizableObject.cpp b/Engine/Plugins/Mutable/Source/CustomizableObject/Private/MuCO/CustomizableObject.cpp
--- a/Engine/Plugins/Mutable/Source/CustomizableObject/Private/MuCO/CustomizableObject.cpp
+++ b/Engine/Plugins/Mutable/Source/CustomizableObject/Private/MuCO/CustomizableObject.cpp
@@ -83,7 +83,7 @@ TAutoConsoleVariable<bool> CVarMutableUseBulkData(
 
 TAutoConsoleVariable<int32> CVarMutableDerivedDataCacheUsage(
 	TEXT("mutable.DerivedDataCacheUsage"),
-	0, // TODO Pere: Set back to 2 before submit
+	2,
 	TEXT("Derived data cache access for cooked data.")
 	TEXT("0 - None. Disables access to the cache.")
 	TEXT("1 - Local. Allow cache requests to query and store records and values in local caches.")
@@ -93,7 +93,7 @@ TAutoConsoleVariable<int32> CVarMutableDerivedDataCacheUsage(
 
 TAutoConsoleVariable<bool> CVarMutableAsyncCook(
 	TEXT("Mutable.CookAsync"),
-	true,
+	false,
 	TEXT("True - Customizable Objects will be compiled asynchronously during cook.\n")
 	TEXT("False - Sync compilation.\n"));
 
