[mutable] Fixed the CustomizableObject editor Instance Preview Automatic LODs option, so that the correct LOD is shown in the preview when the customized mesh is zoomed-in and out.
#tests Editor
#rnx

diff --git a/Engine/Plugins/Mutable/Source/CustomizableObjectEditor/Private/MuCOE/CustomizableObjectEditorViewportClient.cpp b/Engine/Plugins/Mutable/Source/CustomizableObjectEditor/Private/MuCOE/CustomizableObjectEditorViewportClient.cpp
--- a/Engine/Plugins/Mutable/Source/CustomizableObjectEditor/Private/MuCOE/CustomizableObjectEditorViewportClient.cpp
+++ b/Engine/Plugins/Mutable/Source/CustomizableObjectEditor/Private/MuCOE/CustomizableObjectEditorViewportClient.cpp
@@ -836,6 +836,7 @@ void FCustomizableObjectEditorViewportClient::OnPreSetSkeletalMesh(const FPreSet
 				DebugComponent->bCastInsetShadow = true; // For better quality shadows in the editor previews, more similar to the in-game ones
 				DebugComponent->bCanHighlightSelectedSections = false;
 				DebugComponent->bComponentUseFixedSkelBounds = true; // First bounds computed would be using physics asset
+				DebugComponent->bSyncAttachParentLOD = false; // Needed for the "LOD Auto" display mode to work in the preview
 				DebugComponent->MarkRenderStateDirty();
 				DebugComponent->AttachToComponent(Actor->GetRootComponent(), FAttachmentTransformRules::KeepRelativeTransform);
 				DebugComponent->RegisterComponent();
diff --git a/Engine/Plugins/Mutable/Source/CustomizableObjectEditor/Private/MuCOE/SCustomizableObjectEditorViewport.cpp b/Engine/Plugins/Mutable/Source/CustomizableObjectEditor/Private/MuCOE/SCustomizableObjectEditorViewport.cpp
--- a/Engine/Plugins/Mutable/Source/CustomizableObjectEditor/Private/MuCOE/SCustomizableObjectEditorViewport.cpp
+++ b/Engine/Plugins/Mutable/Source/CustomizableObjectEditor/Private/MuCOE/SCustomizableObjectEditorViewport.cpp
@@ -308,10 +308,9 @@ void SCustomizableObjectEditorViewportTabBody::Tick(const FGeometry& AllottedGeo
 {
 	SCompoundWidget::Tick(AllottedGeometry, InCurrentTime, InDeltaTime);
 
-	// Needed?
+	// Needed for the forced "LOD 0, 1, 2..." display mode to work in the preview
 	PreviewScenePtr->GetWorld()->Tick(LEVELTICK_All, InDeltaTime);
 
-
 	const TMap<FName, TWeakObjectPtr<UDebugSkelMeshComponent>>& PreviewSkeletalMeshComponents = LevelViewportClient->GetPreviewMeshComponents();
 
 	// Update the material list. Not ideal to do it every tick, but tracking changes on materials in the current instance is not easy right now.
