[mutable] Move image operations that were still using the ASTOpFixed representation to their own ASTOp subclasses.
#rnx

diff --git a/Engine/Plugins/Mutable/Source/CustomizableObject/Internal/MuCO/CustomizableObjectPrivate.h b/Engine/Plugins/Mutable/Source/CustomizableObject/Internal/MuCO/CustomizableObjectPrivate.h
--- a/Engine/Plugins/Mutable/Source/CustomizableObject/Internal/MuCO/CustomizableObjectPrivate.h
+++ b/Engine/Plugins/Mutable/Source/CustomizableObject/Internal/MuCO/CustomizableObjectPrivate.h
@@ -1820,6 +1820,8 @@ public:
 		RemoveObsoleteDataTypesFromEnum,
 		ConvertModelResourcesToUObject,
 
+		RemoveObsoletImageGradientOp,
+
 		// -----<new versions can be added above this line>--------
 		LastCustomizableObjectVersion
 	};
diff --git a/Engine/Plugins/Mutable/Source/CustomizableObjectEditor/Private/MuCOE/SMutableCodeViewer.cpp b/Engine/Plugins/Mutable/Source/CustomizableObjectEditor/Private/MuCOE/SMutableCodeViewer.cpp
--- a/Engine/Plugins/Mutable/Source/CustomizableObjectEditor/Private/MuCOE/SMutableCodeViewer.cpp
+++ b/Engine/Plugins/Mutable/Source/CustomizableObjectEditor/Private/MuCOE/SMutableCodeViewer.cpp
@@ -3012,14 +3012,14 @@ void FMutableCodeTreeElement::GenerateLabelText()
 	case mu::EOpType::IM_RESIZE:
 	{
 		mu::OP::ImageResizeArgs Args = Program.GetOpArgs<mu::OP::ImageResizeArgs>(MutableOperation);
-		OpName += FString::Printf(TEXT(" %d x %d"), int32(Args.size[0]), int32(Args.size[1]));
+		OpName += FString::Printf(TEXT(" %d x %d"), int32(Args.Size[0]), int32(Args.Size[1]));
 		break;
 	}
 
 	case mu::EOpType::IM_RESIZEREL:
 	{
 		mu::OP::ImageResizeRelArgs Args = Program.GetOpArgs<mu::OP::ImageResizeRelArgs>(MutableOperation);
-		OpName += FString::Printf(TEXT(" %.3f x %.3f"), Args.factor[0], Args.factor[1]);
+		OpName += FString::Printf(TEXT(" %.3f x %.3f"), Args.Factor[0], Args.Factor[1]);
 		break;
 	}
 
@@ -3063,10 +3063,10 @@ void FMutableCodeTreeElement::GenerateLabelText()
 
 	case mu::EOpType::IM_PLAINCOLOUR:
 	{
-		mu::OP::ImagePlainColourArgs Args = Program.GetOpArgs<mu::OP::ImagePlainColourArgs>(MutableOperation);
+		mu::OP::ImagePlainColorArgs Args = Program.GetOpArgs<mu::OP::ImagePlainColorArgs>(MutableOperation);
 		OpName += TEXT(" format: ");
-		OpName += StringCast<TCHAR>(mu::TypeInfo::s_imageFormatName[int32(Args.format)]).Get();
-		OpName += FString::Printf(TEXT(" size %d x %d"), Args.size[0], Args.size[1]);
+		OpName += StringCast<TCHAR>(mu::TypeInfo::s_imageFormatName[int32(Args.Format)]).Get();
+		OpName += FString::Printf(TEXT(" size %d x %d"), Args.Size[0], Args.Size[1]);
 		OpName += FString::Printf(TEXT(" mips %d"), Args.LODs);
 		break;
 	}
diff --git a/Engine/Plugins/Mutable/Source/MutableRuntime/Internal/MuR/Operations.h b/Engine/Plugins/Mutable/Source/MutableRuntime/Internal/MuR/Operations.h
--- a/Engine/Plugins/Mutable/Source/MutableRuntime/Internal/MuR/Operations.h
+++ b/Engine/Plugins/Mutable/Source/MutableRuntime/Internal/MuR/Operations.h
@@ -179,9 +179,6 @@ namespace mu
         //! grey-level.
         IM_COLOURMAP,
 
-        //! Build a horizontal gradient image from two colours
-        IM_GRADIENT,
-
         //! Generate a black and white image from an image and a threshold.
         IM_BINARISE,
 
@@ -618,19 +615,17 @@ namespace mu
 
         struct ImageResizeArgs
         {
-            ADDRESS source;
-            uint16 size[2];
-
-            inline uint16 GetSize(int i) const { return size[i]; }
+            ADDRESS Source;
+            uint16 Size[2];
         };
 
         struct ImageResizeLikeArgs
         {
             //! Image that will be resized
-            ADDRESS source;
+            ADDRESS Source;
 
             //! Image whose size will be used to resize the source.
-            ADDRESS sizeSource;
+            ADDRESS SizeSource;
         };
 
         struct ImageResizeVarArgs
@@ -645,25 +640,25 @@ namespace mu
         struct ImageResizeRelArgs
         {
             //! Image that will be resized
-            ADDRESS source;
+            ADDRESS Source;
 
             //! Factor for each axis.
-            float factor[2];
+            float Factor[2];
         };
 
         struct ImageBlankLayoutArgs
         {
-            ADDRESS layout;
+            ADDRESS Layout;
 
-            //! Size of a layout block in pixels.
-            uint16 blockSize[2];
-			EImageFormat format;
+            /** Size of a layout block in pixels. */
+            uint16 BlockSize[2];
+			EImageFormat Format;
 
-            //! If true, generate mipmaps
-            uint8 generateMipmaps;
+            /** If true, generate mipmaps. */
+            uint8 GenerateMipmaps;
 
-            //! Mipmaps to generate if mipmaps are to be generated. 0 means all.
-            uint8 mipmapCount;
+            /** Mipmaps to generate if mipmaps are to be generated. 0 means all. */
+            uint8 MipmapCount;
         };
 
         struct ImageComposeArgs
@@ -675,30 +670,23 @@ namespace mu
 
         struct ImageInterpolateArgs
         {
-            ADDRESS factor;
-            ADDRESS targets[ MUTABLE_OP_MAX_INTERPOLATE_COUNT ];
-        };
-
-        struct ImageInterpolate3Args
-        {
-            ADDRESS factor1, factor2;
-            ADDRESS target0, target1, target2;
+            ADDRESS Factor;
+            ADDRESS Targets[ MUTABLE_OP_MAX_INTERPOLATE_COUNT ];
         };
 
         struct ImageSaturateArgs
         {
-            //! Image to modify.
-            ADDRESS base;
-
-            //! Saturation factor: 0 desaturates, 1 leaves the same, >1 saturates
-            ADDRESS factor;
+            /** Image to modify. */
+            ADDRESS Base;
 
+            /** Saturation factor : 0 desaturates, 1 leaves the same, >1 saturates */
+            ADDRESS Factor;
         };
 
         struct ImageLuminanceArgs
         {
             //! Image to modify.
-            ADDRESS base;
+            ADDRESS Base;
         };
 
         struct ImageSwizzleArgs
@@ -710,39 +698,27 @@ namespace mu
 
         struct ImageColourMapArgs
         {
-            ADDRESS base;
-            ADDRESS mask;
-            ADDRESS map;
-        };
-
-        struct ImageGradientArgs
-        {
-            ADDRESS colour0;
-            ADDRESS colour1;
-            uint16 size[2];
+            ADDRESS Base;
+            ADDRESS Mask;
+            ADDRESS Map;
         };
 
         struct ImageBinariseArgs
         {
-            ADDRESS base;
-            ADDRESS threshold;
+            ADDRESS Base;
+            ADDRESS Threshold;
         };
 
-        struct ImagePlainColourArgs
+        struct ImagePlainColorArgs
         {
-            ADDRESS colour;
-			EImageFormat format;
-            uint16 size[2];
+            ADDRESS Color;
+			EImageFormat Format;
+            uint16 Size[2];
 
 			/** Number of mipmaps to generate. 0 means all the chain. */
 			uint8 LODs;
         };
 
-        struct ImageGPUArgs
-        {
-            ADDRESS program;
-        };
-
         struct ImageCropArgs
         {
             ADDRESS source;
@@ -796,13 +772,13 @@ namespace mu
 
         struct ImageDisplaceArgs
         {
-            ADDRESS source;
-            ADDRESS displacementMap;
+            ADDRESS Source;
+            ADDRESS DisplacementMap;
         };
 
 		struct ImageInvertArgs
 		{
-			ADDRESS base;
+			ADDRESS Base;
 		};
 
         struct ImageNormalCompositeArgs
@@ -1119,24 +1095,6 @@ namespace mu
             ColourFromScalarsArgs ColourFromScalars;
             ArithmeticArgs ColourArithmetic;
 
-            //-------------------------------------------------------------------------------------
-            ImageResizeArgs ImageResize;
-            ImageResizeLikeArgs ImageResizeLike;
-            ImageResizeVarArgs ImageResizeVar;
-            ImageResizeRelArgs ImageResizeRel;
-            ImageBlankLayoutArgs ImageBlankLayout;
-            ImageInterpolateArgs ImageInterpolate;
-            ImageInterpolate3Args ImageInterpolate3;
-            ImageSaturateArgs ImageSaturate;
-            ImageLuminanceArgs ImageLuminance;
-            ImageColourMapArgs ImageColourMap;
-            ImageGradientArgs ImageGradient;
-            ImageBinariseArgs ImageBinarise;
-            ImagePlainColourArgs ImagePlainColour;
-            ImageGPUArgs ImageGPU;
-            ImageDisplaceArgs ImageDisplace;
-			ImageInvertArgs ImageInvert;
-
         } FArgs;
 
         FArgs args;
@@ -1148,10 +1106,6 @@ namespace mu
     MUTABLE_DEFINE_POD_SERIALISABLE(OP);
     MUTABLE_DEFINE_POD_VECTOR_SERIALISABLE(OP);
     
-    // Check that we didn't go out of control with the operation size
-    static_assert( sizeof(OP::FArgs)==28, "Argument union has an unexpected size." );
-    static_assert( sizeof(OP)==32, "Operation struct has an unexpected size." );
-
 
 	//! Types of data handled by the Mutable runtime.
 	enum class EDataType : uint8
diff --git a/Engine/Plugins/Mutable/Source/MutableRuntime/Private/MuR/CodeRunner.cpp b/Engine/Plugins/Mutable/Source/MutableRuntime/Private/MuR/CodeRunner.cpp
--- a/Engine/Plugins/Mutable/Source/MutableRuntime/Private/MuR/CodeRunner.cpp
+++ b/Engine/Plugins/Mutable/Source/MutableRuntime/Private/MuR/CodeRunner.cpp
@@ -3369,7 +3369,7 @@ namespace mu
             switch (Item.Stage)
             {
             case 0:
-                AddOp( FScheduledOp( Item.At, Item, 1), FScheduledOp( Args.source, Item) );
+                AddOp( FScheduledOp( Item.At, Item, 1), FScheduledOp( Args.Source, Item) );
                 break;
 
             case 1:
@@ -3392,16 +3392,16 @@ namespace mu
             {
             case 0:
                 AddOp(FScheduledOp(Item.At, Item, 1),
-                      	FScheduledOp(Args.source, Item),
-                        FScheduledOp(Args.sizeSource, Item));
+                      	FScheduledOp(Args.Source, Item),
+                        FScheduledOp(Args.SizeSource, Item));
                 break;
 
             case 1:
             {
             	MUTABLE_CPUPROFILER_SCOPE(IM_RESIZELIKE_1)
             	
-                TSharedPtr<const FImage> Base = LoadImage( FCacheAddress(Args.source,Item) );
-                TSharedPtr<const FImage> SizeBase = LoadImage( FCacheAddress(Args.sizeSource,Item) );
+                TSharedPtr<const FImage> Base = LoadImage( FCacheAddress(Args.Source,Item) );
+                TSharedPtr<const FImage> SizeBase = LoadImage( FCacheAddress(Args.SizeSource,Item) );
 				FImageSize DestSize = SizeBase->GetSize();
 				Release(SizeBase);
 
@@ -3450,7 +3450,7 @@ namespace mu
             switch (Item.Stage)
             {
             case 0:
-                AddOp( FScheduledOp( Item.At, Item, 1), FScheduledOp( Args.source, Item) );
+                AddOp( FScheduledOp( Item.At, Item, 1), FScheduledOp( Args.Source, Item) );
                 break;
 
             case 1:
@@ -3472,18 +3472,18 @@ namespace mu
             switch (Item.Stage)
             {
             case 0:
-                AddOp( FScheduledOp( Item.At, Item, 1), FScheduledOp::FromOpAndOptions(Args.layout, Item, 0));
+                AddOp( FScheduledOp( Item.At, Item, 1), FScheduledOp::FromOpAndOptions(Args.Layout, Item, 0));
                 break;
 
             case 1:
             {
             	MUTABLE_CPUPROFILER_SCOPE(IM_BLANKLAYOUT_1)
             		
-                TSharedPtr<const FLayout> pLayout = LoadLayout(FScheduledOp::FromOpAndOptions(Args.layout, Item, 0));
+                TSharedPtr<const FLayout> pLayout = LoadLayout(FScheduledOp::FromOpAndOptions(Args.Layout, Item, 0));
 
                 FIntPoint SizeInBlocks = pLayout->GetGridSize();
 
-				FIntPoint BlockSizeInPixels(Args.blockSize[0], Args.blockSize[1]);
+				FIntPoint BlockSizeInPixels(Args.BlockSize[0], Args.BlockSize[1]);
 
 				// Image size if we don't skip any mipmap
 				FIntPoint FullImageSizeInPixels = SizeInBlocks * BlockSizeInPixels;
@@ -3516,20 +3516,20 @@ namespace mu
 				}
 
                 int32 MipsToGenerate = 1;
-                if (Args.generateMipmaps)
+                if (Args.GenerateMipmaps)
                 {
-                    if (Args.mipmapCount == 0)
+                    if (Args.MipmapCount == 0)
                     {
 						MipsToGenerate = FImage::GetMipmapCount(ImageSizeInPixels.X, ImageSizeInPixels.Y);
                     }
                     else
                     {
-						MipsToGenerate = FMath::Max(Args.mipmapCount - MipsToSkip, 1);
+						MipsToGenerate = FMath::Max(Args.MipmapCount - MipsToSkip, 1);
                     }
                 }
 
 				// It needs to be initialized in case it has gaps.
-                TSharedPtr<FImage> New = CreateImage(ImageSizeInPixels.X, ImageSizeInPixels.Y, MipsToGenerate, EImageFormat(Args.format), EInitializationType::Black );
+                TSharedPtr<FImage> New = CreateImage(ImageSizeInPixels.X, ImageSizeInPixels.Y, MipsToGenerate, EImageFormat(Args.Format), EInitializationType::Black );
                 StoreImage(Item, New);
                 break;
             }
@@ -3645,7 +3645,7 @@ namespace mu
             {
             case 0:
                 AddOp( FScheduledOp( Item.At, Item, 1),
-                       FScheduledOp( Args.factor, Item) );
+                       FScheduledOp( Args.Factor, Item) );
                 break;
 
             case 1:
@@ -3655,13 +3655,13 @@ namespace mu
                 // Targets must be consecutive
                 int32 count = 0;
                 for ( int32 i=0
-                    ; i<MUTABLE_OP_MAX_INTERPOLATE_COUNT && Args.targets[i]
+                    ; i<MUTABLE_OP_MAX_INTERPOLATE_COUNT && Args.Targets[i]
                     ; ++i )
                 {
                     count++;
                 }
 
-                float factor = LoadScalar( FCacheAddress(Args.factor,Item) );
+                float factor = LoadScalar( FCacheAddress(Args.Factor,Item) );
 
                 float delta = 1.0f/(count-1);
                 int32 min = (int32)floorf( factor/delta );
@@ -3678,18 +3678,18 @@ namespace mu
                 if ( bifactor < UE_SMALL_NUMBER )
                 {
                     AddOp( FScheduledOp( Item.At, Item, 2, dataPos),
-                            FScheduledOp( Args.targets[min], Item) );
+                            FScheduledOp( Args.Targets[min], Item) );
                 }
                 else if ( bifactor > 1.0f-UE_SMALL_NUMBER )
                 {
                     AddOp( FScheduledOp( Item.At, Item, 2, dataPos),
-                            FScheduledOp( Args.targets[max], Item) );
+                            FScheduledOp( Args.Targets[max], Item) );
                 }
                 else
                 {
                     AddOp( FScheduledOp( Item.At, Item, 2, dataPos),
-                            FScheduledOp( Args.targets[min], Item),
-                            FScheduledOp( Args.targets[max], Item) );
+                            FScheduledOp( Args.Targets[min], Item),
+                            FScheduledOp( Args.Targets[max], Item) );
                 }
                 break;
             }
@@ -3701,7 +3701,7 @@ namespace mu
                 // Targets must be consecutive
                 int32 count = 0;
                 for ( int32 i=0
-                    ; i<MUTABLE_OP_MAX_INTERPOLATE_COUNT && Args.targets[i]
+                    ; i<MUTABLE_OP_MAX_INTERPOLATE_COUNT && Args.Targets[i]
                     ; ++i )
                 {
                     count++;
@@ -3715,18 +3715,18 @@ namespace mu
 
                 if (Bifactor < UE_SMALL_NUMBER)
                 {
-                    TSharedPtr<const FImage> Source = LoadImage(FCacheAddress(Args.targets[Min], Item));
+                    TSharedPtr<const FImage> Source = LoadImage(FCacheAddress(Args.Targets[Min], Item));
 					StoreImage(Item, Source);
 				}
                 else if (Bifactor > 1.0f - UE_SMALL_NUMBER)
                 {
-                    TSharedPtr<const FImage> Source = LoadImage(FCacheAddress(Args.targets[Max], Item));
+                    TSharedPtr<const FImage> Source = LoadImage(FCacheAddress(Args.Targets[Max], Item));
 					StoreImage(Item, Source);
 				}
                 else
                 {
-					TSharedPtr<const FImage> pMin = LoadImage(FCacheAddress(Args.targets[Min], Item));
-                    TSharedPtr<const FImage> pMax = LoadImage(FCacheAddress(Args.targets[Max], Item));
+					TSharedPtr<const FImage> pMin = LoadImage(FCacheAddress(Args.Targets[Min], Item));
+                    TSharedPtr<const FImage> pMax = LoadImage(FCacheAddress(Args.Targets[Max], Item));
 
                     if (pMin && pMax)
                     {						
@@ -3819,8 +3819,8 @@ namespace mu
             {
             case 0:
                 AddOp( FScheduledOp( Item.At, Item, 1),
-                        FScheduledOp( Args.base, Item ),
-                        FScheduledOp::FromOpAndOptions( Args.factor, Item, 0 ));
+                        FScheduledOp( Args.Base, Item ),
+                        FScheduledOp::FromOpAndOptions( Args.Factor, Item, 0 ));
                 break;
 
             case 1:
@@ -3843,14 +3843,14 @@ namespace mu
             {
             case 0:
                 AddOp( FScheduledOp( Item.At, Item, 1),
-                        FScheduledOp( Args.base, Item ) );
+                        FScheduledOp( Args.Base, Item ) );
                 break;
 
             case 1:
             {
            		MUTABLE_CPUPROFILER_SCOPE(IM_LUMINANCE_1)
             		
-                TSharedPtr<const FImage> Base = LoadImage( FCacheAddress(Args.base,Item) );
+                TSharedPtr<const FImage> Base = LoadImage( FCacheAddress(Args.Base,Item) );
 
 				TSharedPtr<FImage> Result = CreateImage(Base->GetSizeX(), Base->GetSizeY(), Base->GetLODCount(), EImageFormat::L_UByte, EInitializationType::NotInitialized);
                 ImageLuminance( Result.Get(),Base.Get() );
@@ -3899,18 +3899,18 @@ namespace mu
             {
             case 0:
                     AddOp( FScheduledOp( Item.At, Item, 1),
-                           FScheduledOp( Args.base, Item ),
-                           FScheduledOp( Args.mask, Item ),
-                           FScheduledOp( Args.map, Item ) );
+                           FScheduledOp( Args.Base, Item ),
+                           FScheduledOp( Args.Mask, Item ),
+                           FScheduledOp( Args.Map, Item ) );
                 break;
 
             case 1:
             {
             	MUTABLE_CPUPROFILER_SCOPE(IM_COLOURMAP_1)
             		
-                TSharedPtr<const FImage> Source = LoadImage( FCacheAddress(Args.base,Item) );
-                TSharedPtr<const FImage> Mask = LoadImage( FCacheAddress(Args.mask,Item) );
-                TSharedPtr<const FImage> Map = LoadImage( FCacheAddress(Args.map,Item) );
+                TSharedPtr<const FImage> Source = LoadImage( FCacheAddress(Args.Base,Item) );
+                TSharedPtr<const FImage> Mask = LoadImage( FCacheAddress(Args.Mask,Item) );
+                TSharedPtr<const FImage> Map = LoadImage( FCacheAddress(Args.Map,Item) );
 
 				bool bOnlyOneMip = (Mask->GetLODCount() < Source->GetLODCount());
 
@@ -3948,38 +3948,6 @@ namespace mu
             break;
         }
 
-        case EOpType::IM_GRADIENT:
-        {
-			OP::ImageGradientArgs Args = Program.GetOpArgs<OP::ImageGradientArgs>(Item.At);
-            switch (Item.Stage)
-            {
-            case 0:
-                AddOp( FScheduledOp( Item.At, Item, 1),
-                        FScheduledOp::FromOpAndOptions( Args.colour0, Item, 0 ),
-                        FScheduledOp::FromOpAndOptions( Args.colour1, Item, 0 ) );
-                break;
-
-            case 1:
-            {
-				MUTABLE_CPUPROFILER_SCOPE(IM_GRADIENT_1)
-            		
-				FVector4f colour0 = LoadColor(FScheduledOp::FromOpAndOptions(Args.colour0, Item, 0));
-				FVector4f colour1 = LoadColor(FScheduledOp::FromOpAndOptions(Args.colour1, Item, 0));
-
-				TSharedPtr<FImage> Result = CreateImage(Args.size[0], Args.size[1], 1, EImageFormat::RGB_UByte, EInitializationType::NotInitialized);
-                ImageGradient( Result.Get(), colour0, colour1 );
-
-				StoreImage( Item, Result );
-                break;
-            }
-
-            default:
-                check(false);
-            }
-
-            break;
-        }
-
         case EOpType::IM_BINARISE:
         {
 			OP::ImageBinariseArgs Args = Program.GetOpArgs<OP::ImageBinariseArgs>(Item.At);
@@ -3987,17 +3955,17 @@ namespace mu
             {
             case 0:
                 AddOp( FScheduledOp( Item.At, Item, 1),
-                        FScheduledOp( Args.base, Item ),
-                        FScheduledOp::FromOpAndOptions( Args.threshold, Item, 0 ) );
+                        FScheduledOp( Args.Base, Item ),
+                        FScheduledOp::FromOpAndOptions( Args.Threshold, Item, 0 ) );
                 break;
 
             case 1:
             {
             	MUTABLE_CPUPROFILER_SCOPE(IM_BINARISE_1)
             		
-                TSharedPtr<const FImage> pA = LoadImage( FCacheAddress(Args.base,Item) );
+                TSharedPtr<const FImage> pA = LoadImage( FCacheAddress(Args.Base,Item) );
 
-                float c = LoadScalar(FScheduledOp::FromOpAndOptions(Args.threshold, Item, 0));
+                float c = LoadScalar(FScheduledOp::FromOpAndOptions(Args.Threshold, Item, 0));
 
                 TSharedPtr<FImage> Result = CreateImage(pA->GetSizeX(), pA->GetSizeY(), pA->GetLODCount(), EImageFormat::L_UByte, EInitializationType::NotInitialized);
 				ImageBinarise( Result.Get(), pA.Get(), c );
@@ -4020,7 +3988,7 @@ namespace mu
 			switch (Item.Stage)
 			{
 			case 0:
-				AddOp(FScheduledOp(Item.At, Item, 1), FScheduledOp(Args.base, Item));
+				AddOp(FScheduledOp(Item.At, Item, 1), FScheduledOp(Args.Base, Item));
 				break;
 
 			case 1:
@@ -4038,21 +4006,21 @@ namespace mu
 
         case EOpType::IM_PLAINCOLOUR:
         {
-			OP::ImagePlainColourArgs Args = Program.GetOpArgs<OP::ImagePlainColourArgs>(Item.At);
+			OP::ImagePlainColorArgs Args = Program.GetOpArgs<OP::ImagePlainColorArgs>(Item.At);
             switch (Item.Stage)
             {
             case 0:
-				AddOp( FScheduledOp( Item.At, Item, 1), FScheduledOp::FromOpAndOptions( Args.colour, Item, 0 ) );
+				AddOp( FScheduledOp( Item.At, Item, 1), FScheduledOp::FromOpAndOptions( Args.Color, Item, 0 ) );
                 break;
 
             case 1:
             {
             	MUTABLE_CPUPROFILER_SCOPE(IM_PLAINCOLOUR_1)
             		
-				FVector4f c = LoadColor(FScheduledOp::FromOpAndOptions(Args.colour, Item, 0));
+				FVector4f c = LoadColor(FScheduledOp::FromOpAndOptions(Args.Color, Item, 0));
 
-				uint16 SizeX = Args.size[0];
-				uint16 SizeY = Args.size[1];
+				uint16 SizeX = Args.Size[0];
+				uint16 SizeY = Args.Size[1];
 				int32 LODs = Args.LODs;
 				
 				// This means all the mip chain
@@ -4068,7 +4036,7 @@ namespace mu
 					--LODs;
 				}
 
-                TSharedPtr<FImage> pA = CreateImage( SizeX, SizeY, FMath::Max(LODs,1), EImageFormat(Args.format), EInitializationType::NotInitialized );
+                TSharedPtr<FImage> pA = CreateImage( SizeX, SizeY, FMath::Max(LODs,1), EImageFormat(Args.Format), EInitializationType::NotInitialized );
 
 				ImOp.FillColor(pA.Get(), c);
 
@@ -4651,16 +4619,16 @@ namespace mu
             {
             case 0:
                 AddOp( FScheduledOp( Item.At, Item, 1),
-                        FScheduledOp( Args.source, Item ),
-                        FScheduledOp( Args.displacementMap, Item ) );
+                        FScheduledOp( Args.Source, Item ),
+                        FScheduledOp( Args.DisplacementMap, Item ) );
 				break;
 
             case 1:
             {
                 MUTABLE_CPUPROFILER_SCOPE(IM_DISPLACE_1)
 
-                TSharedPtr<const FImage> Source = LoadImage( FCacheAddress(Args.source,Item) );
-                TSharedPtr<const FImage> pMap = LoadImage( FCacheAddress(Args.displacementMap,Item) );
+                TSharedPtr<const FImage> Source = LoadImage( FCacheAddress(Args.Source,Item) );
+                TSharedPtr<const FImage> pMap = LoadImage( FCacheAddress(Args.DisplacementMap,Item) );
 
 				if (!Source)
 				{
@@ -6406,12 +6374,12 @@ namespace mu
 			{
 			case 0:
 				AddOp(FScheduledOp(Item.At, Item, 1, Item.CustomState), 
-						FScheduledOp(Args.source, Item, 0, Item.CustomState));
+						FScheduledOp(Args.Source, Item, 0, Item.CustomState));
 				break;
 
 			case 1:
-				HeapImageDesc[Item.CustomState].m_size[0] = Args.size[0];
-				HeapImageDesc[Item.CustomState].m_size[1] = Args.size[1];
+				HeapImageDesc[Item.CustomState].m_size[0] = Args.Size[0];
+				HeapImageDesc[Item.CustomState].m_size[1] = Args.Size[1];
 				StoreValidDesc(Item);
 				break;
 
@@ -6438,8 +6406,8 @@ namespace mu
 				int32 ResultArrayAddress = HeapData.Add(ResultArrayData);
 
 				AddOp(FScheduledOp(Item.At, Item, 1, ResultArrayAddress), 
-						FScheduledOp(Args.source, Item, 0, ResultArrayData.ReturnArray.Begin + 0),
-						FScheduledOp(Args.sizeSource, Item, 0, ResultArrayData.ReturnArray.Begin + 1));
+						FScheduledOp(Args.Source, Item, 0, ResultArrayData.ReturnArray.Begin + 0),
+						FScheduledOp(Args.SizeSource, Item, 0, ResultArrayData.ReturnArray.Begin + 1));
 
 				break;
 			}
@@ -6478,15 +6446,15 @@ namespace mu
 			{
 			case 0:
 				AddOp(FScheduledOp(Item.At, Item, 1, Item.CustomState), 
-						FScheduledOp(Args.source, Item, 0, Item.CustomState));
+						FScheduledOp(Args.Source, Item, 0, Item.CustomState));
 				break;
 
 			case 1:
 			{
 				FExtendedImageDesc& ResultAndBaseDesc = HeapImageDesc[Item.CustomState];
 				FImageSize DestSize(
-					uint16(ResultAndBaseDesc.m_size[0] * Args.factor[0] + 0.5f),
-					uint16(ResultAndBaseDesc.m_size[1] * Args.factor[1] + 0.5f));
+					uint16(ResultAndBaseDesc.m_size[0] * Args.Factor[0] + 0.5f),
+					uint16(ResultAndBaseDesc.m_size[1] * Args.Factor[1] + 0.5f));
 				ResultAndBaseDesc.m_size = DestSize;
 				StoreValidDesc(Item);
 				break;
@@ -6507,7 +6475,7 @@ namespace mu
 			case 0:
 			{
 				// We need to run the full layout
-				FScheduledOp FullLayoutOp(Args.layout, Item);
+				FScheduledOp FullLayoutOp(Args.Layout, Item);
 				FullLayoutOp.Type = FScheduledOp::EType::Full;
 				AddOp(FScheduledOp(Item.At, Item, 1, Item.CustomState), FullLayoutOp);
 				break;
@@ -6515,26 +6483,26 @@ namespace mu
 
 			case 1:
 			{
-				TSharedPtr<const FLayout> pLayout = LoadLayout(FCacheAddress(Args.layout, Item));
+				TSharedPtr<const FLayout> pLayout = LoadLayout(FCacheAddress(Args.Layout, Item));
 
 				FIntPoint SizeInBlocks = pLayout->GetGridSize();
-				FIntPoint BlockSizeInPixels(Args.blockSize[0], Args.blockSize[1]);
+				FIntPoint BlockSizeInPixels(Args.BlockSize[0], Args.BlockSize[1]);
 				FIntPoint ImageSizeInPixels = SizeInBlocks * BlockSizeInPixels;
 
 				FExtendedImageDesc& ResultAndBaseDesc = HeapImageDesc[Item.CustomState];
 				FImageSize DestSize(uint16(ImageSizeInPixels.X), uint16(ImageSizeInPixels.Y));
 				ResultAndBaseDesc.m_size = DestSize;
-				ResultAndBaseDesc.m_format = Args.format;
+				ResultAndBaseDesc.m_format = Args.Format;
 				
-				if (Args.generateMipmaps)
+				if (Args.GenerateMipmaps)
 				{
-					if (Args.mipmapCount == 0)
+					if (Args.MipmapCount == 0)
 					{
 						ResultAndBaseDesc.m_lods = FImage::GetMipmapCount(ImageSizeInPixels.X, ImageSizeInPixels.Y);
 					}
 					else
 					{
-						ResultAndBaseDesc.m_lods = Args.mipmapCount;
+						ResultAndBaseDesc.m_lods = Args.MipmapCount;
 					}
 				}
 				StoreValidDesc(Item);
@@ -6602,7 +6570,7 @@ namespace mu
 			int32 NumImages = 0;
 			for (; NumImages < MUTABLE_OP_MAX_INTERPOLATE_COUNT; ++NumImages)
 			{
-				if (!Args.targets[NumImages])
+				if (!Args.Targets[NumImages])
 				{
 					break;
 				}
@@ -6624,7 +6592,7 @@ namespace mu
 
 				for (int32 Index = 0; Index < NumImages; ++Index)
 				{
-					Deps.Add(FScheduledOp(Args.targets[Index], Item, 0, ResultArrayData.ReturnArray.Begin + Index));
+					Deps.Add(FScheduledOp(Args.Targets[Index], Item, 0, ResultArrayData.ReturnArray.Begin + Index));
 				}
 
 				AddOp(FScheduledOp(Item.At, Item, 1, ResultArrayAddress), Deps); 
@@ -6662,7 +6630,7 @@ namespace mu
 			case 0: 
 			{
 				AddOp(FScheduledOp(Item.At, Item, 1, Item.CustomState), 
-						FScheduledOp(Args.base, Item, 0, Item.CustomState)); 
+						FScheduledOp(Args.Base, Item, 0, Item.CustomState)); 
 				break;
 			}
 			case 1: StoreValidDesc(Item); break;
@@ -6679,7 +6647,7 @@ namespace mu
 			case 0:
 			{
 				AddOp(FScheduledOp(Item.At, Item, 1, Item.CustomState), 
-						FScheduledOp(Args.base, Item, 0, Item.CustomState));
+						FScheduledOp(Args.Base, Item, 0, Item.CustomState));
 				break;
 			}
 			case 1:
@@ -6766,7 +6734,7 @@ namespace mu
 			case 0: 
 			{	
 				AddOp(FScheduledOp(Item.At, Item, 1, Item.CustomState), 
-						FScheduledOp(Args.base, Item, 0, Item.CustomState));
+						FScheduledOp(Args.Base, Item, 0, Item.CustomState));
 				break;
 			}
 			case 1: StoreValidDesc(Item); break;
@@ -6775,17 +6743,6 @@ namespace mu
 			break;
 		}
 
-		case EOpType::IM_GRADIENT:
-		{
-			OP::ImageGradientArgs Args = Program.GetOpArgs<OP::ImageGradientArgs>(Item.At);
-			HeapImageDesc[Item.CustomState].m_size[0] = Args.size[0];
-			HeapImageDesc[Item.CustomState].m_size[1] = Args.size[1];
-			HeapImageDesc[Item.CustomState].m_lods = 1;
-			HeapImageDesc[Item.CustomState].m_format = EImageFormat::RGB_UByte;
-			StoreValidDesc(Item);
-			break;
-		}
-
 		case EOpType::IM_BINARISE:
 		{
 			OP::ImageBinariseArgs Args = Program.GetOpArgs<OP::ImageBinariseArgs>(Item.At);
@@ -6794,7 +6751,7 @@ namespace mu
 			case 0:
 			{
 				AddOp(FScheduledOp(Item.At, Item, 1, Item.CustomState), 
-						FScheduledOp(Args.base, Item, 0, Item.CustomState));
+						FScheduledOp(Args.Base, Item, 0, Item.CustomState));
 				break;
 			}
 			case 1:
@@ -6817,7 +6774,7 @@ namespace mu
 			case 0: 
 			{
 				AddOp(FScheduledOp(Item.At, Item, 1, Item.CustomState), 
-						FScheduledOp(Args.base, Item, 0, Item.CustomState)); 
+						FScheduledOp(Args.Base, Item, 0, Item.CustomState)); 
 				break;
 			}
 			case 1: StoreValidDesc(Item); break;
@@ -6828,11 +6785,11 @@ namespace mu
 
 		case EOpType::IM_PLAINCOLOUR:
 		{
-			OP::ImagePlainColourArgs Args = Program.GetOpArgs<OP::ImagePlainColourArgs>(Item.At);
-			HeapImageDesc[Item.CustomState].m_size[0] = Args.size[0];
-			HeapImageDesc[Item.CustomState].m_size[1] = Args.size[1];
+			OP::ImagePlainColorArgs Args = Program.GetOpArgs<OP::ImagePlainColorArgs>(Item.At);
+			HeapImageDesc[Item.CustomState].m_size[0] = Args.Size[0];
+			HeapImageDesc[Item.CustomState].m_size[1] = Args.Size[1];
 			HeapImageDesc[Item.CustomState].m_lods = Args.LODs;
-			HeapImageDesc[Item.CustomState].m_format = Args.format;
+			HeapImageDesc[Item.CustomState].m_format = Args.Format;
 			StoreValidDesc(Item);
 			break;
 		}
@@ -6950,7 +6907,7 @@ namespace mu
 			case 0: 
 			{
 				AddOp(FScheduledOp(Item.At, Item, 1, Item.CustomState),
-						FScheduledOp(Args.source, Item, 0, Item.CustomState));
+						FScheduledOp(Args.Source, Item, 0, Item.CustomState));
 				break;
 			}
 			case 1: StoreValidDesc(Item); break;
diff --git a/Engine/Plugins/Mutable/Source/MutableRuntime/Private/MuR/CodeRunnerTasked.cpp b/Engine/Plugins/Mutable/Source/MutableRuntime/Private/MuR/CodeRunnerTasked.cpp
--- a/Engine/Plugins/Mutable/Source/MutableRuntime/Private/MuR/CodeRunnerTasked.cpp
+++ b/Engine/Plugins/Mutable/Source/MutableRuntime/Private/MuR/CodeRunnerTasked.cpp
@@ -1453,8 +1453,8 @@ namespace mu
 		MUTABLE_CPUPROFILER_SCOPE(FImageSaturateTask_Prepare);
 		bOutFailed = false;
 
-		TSharedPtr<const FImage> Source = Runner->LoadImage(FCacheAddress(Args.base, Op));
-		Factor = Runner->LoadScalar(FScheduledOp::FromOpAndOptions(Args.factor, Op, 0));
+		TSharedPtr<const FImage> Source = Runner->LoadImage(FCacheAddress(Args.Base, Op));
+		Factor = Runner->LoadScalar(FScheduledOp::FromOpAndOptions(Args.Factor, Op, 0));
 
 		if (!Source)
 		{
@@ -1531,14 +1531,14 @@ namespace mu
 
 		ImageCompressionQuality = Runner->Settings.ImageCompressionQuality;
 
-		FImageSize destSize = FImageSize(Args.size[0], Args.size[1]);
+		FImageSize destSize = FImageSize(Args.Size[0], Args.Size[1]);
 
 		// Apply the mips-to-skip to the dest size
 		int32 MipsToSkip = Op.ExecutionOptions;
 		destSize[0] = FMath::Max(destSize[0] >> MipsToSkip, 1);
 		destSize[1] = FMath::Max(destSize[1] >> MipsToSkip, 1);
 
-		Base = Runner->LoadImage(FCacheAddress(Args.source, Op));
+		Base = Runner->LoadImage(FCacheAddress(Args.Source, Op));
 		if (!Base 
 			|| 
 			( Base->GetSizeX()==destSize[0] && Base->GetSizeY()==destSize[1] )
@@ -1581,7 +1581,7 @@ namespace mu
 		// This runs on a random worker thread
 		MUTABLE_CPUPROFILER_SCOPE(FImageResizeTask);
 
-		FImageSize destSize = FImageSize( Args.size[0], Args.size[1]);
+		FImageSize destSize = FImageSize( Args.Size[0], Args.Size[1]);
 
 		// Apply the mips-to-skip to the dest size
 		int32 MipsToSkip = Op.ExecutionOptions;
@@ -1652,7 +1652,7 @@ namespace mu
 
 		ImageCompressionQuality = Runner->Settings.ImageCompressionQuality;
 
-		Base = Runner->LoadImage(FCacheAddress(Args.source, Op));
+		Base = Runner->LoadImage(FCacheAddress(Args.Source, Op));
 		if (!Base)
 		{
 			Runner->StoreImage(Op, Base);
@@ -1660,8 +1660,8 @@ namespace mu
 		}
 
 		DestSize = FImageSize(
-			uint16(FMath::Max(1.0, Base->GetSizeX() * Args.factor[0] + 0.5f)),
-			uint16(FMath::Max(1.0, Base->GetSizeY() * Args.factor[1] + 0.5f)));
+			uint16(FMath::Max(1.0, Base->GetSizeX() * Args.Factor[0] + 0.5f)),
+			uint16(FMath::Max(1.0, Base->GetSizeY() * Args.Factor[1] + 0.5f)));
 
 		if (Base->GetSizeX() == DestSize[0] && Base->GetSizeY() == DestSize[1])
 		{
@@ -1757,7 +1757,7 @@ namespace mu
 		MUTABLE_CPUPROFILER_SCOPE(FImageInvertTask_Prepare);
 		bOutFailed = false;
 
-		TSharedPtr<const FImage> Source = Runner->LoadImage({ Args.base, Op.ExecutionIndex, Op.ExecutionOptions });
+		TSharedPtr<const FImage> Source = Runner->LoadImage({ Args.Base, Op.ExecutionIndex, Op.ExecutionOptions });
 
 		// Create destination data
 		Result = Runner->CloneOrTakeOver(Source);
diff --git a/Engine/Plugins/Mutable/Source/MutableRuntime/Private/MuR/Operations.cpp b/Engine/Plugins/Mutable/Source/MutableRuntime/Private/MuR/Operations.cpp
--- a/Engine/Plugins/Mutable/Source/MutableRuntime/Private/MuR/Operations.cpp
+++ b/Engine/Plugins/Mutable/Source/MutableRuntime/Private/MuR/Operations.cpp
@@ -88,7 +88,6 @@ namespace mu
 		{ EDataType::Image,			},	// IM_LUMINANCE
 		{ EDataType::Image,			},	// IM_SWIZZLE
 		{ EDataType::Image,			},	// IM_COLOURMAP
-		{ EDataType::Image,			},	// IM_GRADIENT
 		{ EDataType::Image,			},	// IM_BINARISE
 		{ EDataType::Image,			},	// IM_PLAINCOLOUR
 		{ EDataType::Image,			},	// IM_CROP
@@ -236,71 +235,6 @@ namespace mu
             f(&op.args.ColourArithmetic.b);
             break;
 
-        //-------------------------------------------------------------------------------------
-        case EOpType::IM_RESIZE:
-            f(&op.args.ImageResize.source );
-            break;
-
-        case EOpType::IM_RESIZELIKE:
-            f(&op.args.ImageResizeLike.source );
-            f(&op.args.ImageResizeLike.sizeSource );
-            break;
-
-        case EOpType::IM_RESIZEREL:
-            f(&op.args.ImageResizeRel.source );
-            break;
-
-        case EOpType::IM_BLANKLAYOUT:
-            f(&op.args.ImageBlankLayout.layout );
-            break;
-
-        case EOpType::IM_INTERPOLATE:
-            f(&op.args.ImageInterpolate.factor );
-
-            for (int t=0;t<MUTABLE_OP_MAX_INTERPOLATE_COUNT;++t)
-            {
-                f(&op.args.ImageInterpolate.targets[t] );
-            }
-            break;
-
-        case EOpType::IM_SATURATE:
-            f(&op.args.ImageSaturate.base );
-            f(&op.args.ImageSaturate.factor );
-            break;
-
-        case EOpType::IM_LUMINANCE:
-            f(&op.args.ImageLuminance.base );
-            break;
-
-        case EOpType::IM_COLOURMAP:
-            f(&op.args.ImageColourMap.base );
-            f(&op.args.ImageColourMap.mask );
-            f(&op.args.ImageColourMap.map );
-            break;
-
-        case EOpType::IM_GRADIENT:
-            f(&op.args.ImageGradient.colour0 );
-            f(&op.args.ImageGradient.colour1 );
-            break;
-
-        case EOpType::IM_BINARISE:
-            f(&op.args.ImageBinarise.base );
-            f(&op.args.ImageBinarise.threshold );
-            break;
-
-        case EOpType::IM_PLAINCOLOUR:
-            f(&op.args.ImagePlainColour.colour );
-            break;
-
-        case EOpType::IM_DISPLACE:
-            f(&op.args.ImageDisplace.source );
-            f(&op.args.ImageDisplace.displacementMap );
-            break;
-
-		case EOpType::IM_INVERT:
-			f(&op.args.ImageInvert.base);
-			break;
-
         default:
 			check( false );
             break;
@@ -562,30 +496,30 @@ namespace mu
 
         case EOpType::IM_RESIZE:
         {
-			OP::ImageResizeArgs args = program.GetOpArgs<OP::ImageResizeArgs>(at);
-            f(args.source );
+			OP::ImageResizeArgs Args = program.GetOpArgs<OP::ImageResizeArgs>(at);
+            f(Args.Source );
             break;
         }
 
         case EOpType::IM_RESIZELIKE:
         {
-			OP::ImageResizeLikeArgs args = program.GetOpArgs<OP::ImageResizeLikeArgs>(at);
-            f(args.source );
-            f(args.sizeSource );
+			OP::ImageResizeLikeArgs Args = program.GetOpArgs<OP::ImageResizeLikeArgs>(at);
+            f(Args.Source );
+            f(Args.SizeSource );
             break;
         }
 
         case EOpType::IM_RESIZEREL:
         {
-			OP::ImageResizeRelArgs args = program.GetOpArgs<OP::ImageResizeRelArgs>(at);
-            f(args.source );
+			OP::ImageResizeRelArgs Args = program.GetOpArgs<OP::ImageResizeRelArgs>(at);
+            f(Args.Source );
             break;
         }
 
         case EOpType::IM_BLANKLAYOUT:
         {
-			OP::ImageBlankLayoutArgs args = program.GetOpArgs<OP::ImageBlankLayoutArgs>(at);
-            f(args.layout );
+			OP::ImageBlankLayoutArgs Args = program.GetOpArgs<OP::ImageBlankLayoutArgs>(at);
+            f(Args.Layout );
             break;
         }
 
@@ -601,12 +535,12 @@ namespace mu
 
         case EOpType::IM_INTERPOLATE:
         {
-			OP::ImageInterpolateArgs args = program.GetOpArgs<OP::ImageInterpolateArgs>(at);
-            f(args.factor );
+			OP::ImageInterpolateArgs Args = program.GetOpArgs<OP::ImageInterpolateArgs>(at);
+            f(Args.Factor );
 
-            for (int t=0;t<MUTABLE_OP_MAX_INTERPOLATE_COUNT;++t)
+            for (int32 TargetIndex=0; TargetIndex <MUTABLE_OP_MAX_INTERPOLATE_COUNT;++TargetIndex)
             {
-                f(args.targets[t] );
+                f(Args.Targets[TargetIndex]);
             }
             break;
         }
@@ -623,48 +557,40 @@ namespace mu
 
         case EOpType::IM_SATURATE:
         {
-			OP::ImageSaturateArgs args = program.GetOpArgs<OP::ImageSaturateArgs>(at);
-            f(args.base );
-            f(args.factor );
+			OP::ImageSaturateArgs Args = program.GetOpArgs<OP::ImageSaturateArgs>(at);
+            f(Args.Base );
+            f(Args.Factor );
             break;
         }
 
         case EOpType::IM_LUMINANCE:
         {
-			OP::ImageLuminanceArgs args = program.GetOpArgs<OP::ImageLuminanceArgs>(at);
-            f(args.base );
+			OP::ImageLuminanceArgs Args = program.GetOpArgs<OP::ImageLuminanceArgs>(at);
+            f(Args.Base );
             break;
         }
 
         case EOpType::IM_COLOURMAP:
         {
-			OP::ImageColourMapArgs args = program.GetOpArgs<OP::ImageColourMapArgs>(at);
-            f(args.base );
-            f(args.mask );
-            f(args.map );
-            break;
-        }
-
-        case EOpType::IM_GRADIENT:
-        {
-			OP::ImageGradientArgs args = program.GetOpArgs<OP::ImageGradientArgs>(at);
-            f(args.colour0 );
-            f(args.colour1 );
+			OP::ImageColourMapArgs Args = program.GetOpArgs<OP::ImageColourMapArgs>(at);
+            f(Args.Base );
+            f(Args.Mask );
+            f(Args.Map );
             break;
         }
 
         case EOpType::IM_BINARISE:
         {
-			OP::ImageBinariseArgs args = program.GetOpArgs<OP::ImageBinariseArgs>(at);
-            f(args.base );
-            f(args.threshold );
+			OP::ImageBinariseArgs Args = program.GetOpArgs<OP::ImageBinariseArgs>(at);
+            f(Args.Base );
+            f(Args.Threshold );
             break;
         }
 
         case EOpType::IM_PLAINCOLOUR:
         {
-			OP::ImagePlainColourArgs args = program.GetOpArgs<OP::ImagePlainColourArgs>(at);
-            f(args.colour );
+			OP::ImagePlainColorArgs args = program.GetOpArgs<OP::ImagePlainColorArgs>(at);
+            f(args.Color );
             break;
         }
 
@@ -703,16 +629,16 @@ namespace mu
 
         case EOpType::IM_DISPLACE:
         {
-			OP::ImageDisplaceArgs args = program.GetOpArgs<OP::ImageDisplaceArgs>(at);
-            f(args.source );
-            f(args.displacementMap );
+			OP::ImageDisplaceArgs Args = program.GetOpArgs<OP::ImageDisplaceArgs>(at);
+            f(Args.Source );
+            f(Args.DisplacementMap );
             break;
         }
 
 		case EOpType::IM_INVERT:
 		{
 			OP::ImageInvertArgs args = program.GetOpArgs<OP::ImageInvertArgs>(at);
-			f(args.base);
+			f(args.Base);
 			break;
 		}
 
diff --git a/Engine/Plugins/Mutable/Source/MutableTools/Internal/MuT/AST.h b/Engine/Plugins/Mutable/Source/MutableTools/Internal/MuT/AST.h
--- a/Engine/Plugins/Mutable/Source/MutableTools/Internal/MuT/AST.h
+++ b/Engine/Plugins/Mutable/Source/MutableTools/Internal/MuT/AST.h
@@ -1168,17 +1168,10 @@ namespace mu
 		virtual void Link( FProgram& program, FLinkerOptions* Options) override;
 		virtual bool IsEqual(const ASTOp& otherUntyped) const override;
 		virtual uint64 Hash() const override;
-		virtual FImageDesc GetImageDesc( bool returnBestOption, FGetImageDescContext* context ) const override;
-		virtual void GetLayoutBlockSize( int* pBlockX, int* pBlockY ) override;
-		virtual Ptr<ASTOp> OptimiseSize() const override;
 		virtual Ptr<ASTOp> OptimiseSemantic(const FModelOptimizationOptions&, int32 Pass) const override;
-		virtual Ptr<ASTOp> OptimiseSink(const FModelOptimizationOptions&, FOptimizeSinkContext&) const override;
 		virtual FBoolEvalResult EvaluateBool( ASTOpList& facts, FEvaluateBoolCache* cache ) const override;
 		virtual int32 EvaluateInt( ASTOpList& facts, bool &unknown ) const override;
-		virtual bool IsImagePlainConstant(FVector4f& colour ) const override;
 		virtual bool IsColourConstant(FVector4f& colour ) const override;
-		virtual Ptr<ImageSizeExpression> GetImageSizeExpression() const override;
-		virtual FSourceDataDescriptor GetSourceDataDescriptor(FGetSourceDataDescriptorContext* Context) const override;
 
         /** Add child for OP fixed nodes. */
         inline void SetChild( OP::ADDRESS& at, const Ptr<ASTOp>& child )
@@ -1212,103 +1205,17 @@ namespace mu
             SetChild(at,child.child());
         }
 
-    private:
-
-        //
-        inline FImageDesc GetImageDesc( OP::ADDRESS at, bool returnBestOption=false, FGetImageDescContext* context=nullptr ) const
-        {
-			FImageDesc res;
-
-            if (children[at])
-            {
-                res = children[at]->GetImageDesc( returnBestOption, context );
-            }
-
-            return res;
-        }
-
-		//
-		inline FSourceDataDescriptor GetSourceDataDescriptor(OP::ADDRESS at, FGetSourceDataDescriptorContext* Context = nullptr) const
-		{
-			FSourceDataDescriptor Result;
-
-			if (children[at])
-			{
-				Result = children[at]->GetSourceDataDescriptor(Context);
-			}
-
-			return Result;
-		}
-
-        inline void GetLayoutBlockSize( OP::ADDRESS at, int* pBlockX, int* pBlockY )
-        {
-            if (children[at])
-            {
-                children[at]->GetLayoutBlockSize(pBlockX,pBlockY);
-            }
-        }
-
     };
 
 
-	/** This operation represents another operation that is still being generated. */
-	//class ASTOpFuture : public ASTOp
-	//{
-	//public:
-
-	//	//! Fixed size op structure. Any code ADDRESS in it, refers to the children TArray below.
-	//	OP op;
-
-	//public:
-	//	ASTOpFuture();
-
-	//	// ASTOp interface
-	//	virtual EOpType GetOpType() const override
-	//	{
-	//		EnsureReady();
-	//		return Real->GetOpType();
-	//	}
-	//	virtual Ptr<ASTOp> Clone(MapChildFuncRef mapChild) const override;
-	//	virtual void ForEachChild(const TFunctionRef<void(ASTChild&)> f) override;
-	//	virtual void Link(FProgram& program, FLinkerOptions* Options) override;
-	//	virtual bool IsEqual(const ASTOp& otherUntyped) const override;
-	//	virtual uint64 Hash() const override;
-	//	virtual FImageDesc GetImageDesc(bool returnBestOption, FGetImageDescContext* context) const override;
-	//	virtual void GetLayoutBlockSize(int* pBlockX, int* pBlockY) override;
-	//	virtual Ptr<ASTOp> OptimiseSize() const override;
-	//	virtual Ptr<ASTOp> OptimiseSemantic(const FModelOptimizationOptions&, int32 Pass) const override;
-	//	virtual Ptr<ASTOp> OptimiseSink(const FModelOptimizationOptions&, FOptimizeSinkContext&) const override;
-	//	virtual FBoolEvalResult EvaluateBool(ASTOpList& facts, FEvaluateBoolCache* cache) const override;
-	//	virtual int32 EvaluateInt(ASTOpList& facts, bool& unknown) const override;
-	//	virtual bool IsImagePlainConstant(FVector4f& colour) const override;
-	//	virtual bool IsColourConstant(FVector4f& colour) const override;
-	//	virtual Ptr<ImageSizeExpression> GetImageSizeExpression() const override;
-	//	virtual FSourceDataDescriptor GetSourceDataDescriptor(FGetSourceDataDescriptorContext* Context) const override;
-
-	//	// Own interface
-
-	//	/** Wait until the real operation has been completely generated. */
-	//	void EnsureReady() const
-	//	{
-
-	//	}
-
-	//private:
-
-	//};
-
-
-    //---------------------------------------------------------------------------------------------
-    //---------------------------------------------------------------------------------------------
-    //---------------------------------------------------------------------------------------------
     inline void ASTChild::AddParent()
     {
         ParentIndexInChild = Child->Parents.Num();
         Child->Parents.Add(Parent);
     }
 
-    //---------------------------------------------------------------------------------------------
-    inline void ASTChild::ClearParent()
+
+	inline void ASTChild::ClearParent()
     {
         check( ParentIndexInChild<Child->Parents.Num() );
 		// Can't do this, because the indices are stored in children.
diff --git a/Engine/Plugins/Mutable/Source/MutableTools/Private/MuT/AST.cpp b/Engine/Plugins/Mutable/Source/MutableTools/Private/MuT/AST.cpp
--- a/Engine/Plugins/Mutable/Source/MutableTools/Private/MuT/AST.cpp
+++ b/Engine/Plugins/Mutable/Source/MutableTools/Private/MuT/AST.cpp
@@ -1066,184 +1066,6 @@ uint64 ASTOpFixed::Hash() const
 }
 
 
-//!
-FImageDesc ASTOpFixed::GetImageDesc( bool returnBestOption, FGetImageDescContext* context ) const
-{
-    FImageDesc res;
-
-    // Local context in case it is necessary
-    FGetImageDescContext localContext;
-    if (!context)
-    {
-      context = &localContext;
-    }
-    else
-    {
-        // Cached result?
-		FImageDesc* PtrValue = context->m_results.Find(this);
-        if (PtrValue)
-        {
-            return *PtrValue;
-        }
-    }
-
-    EOpType type = op.type;
-
-    switch ( type )
-    {
-
-    case EOpType::NONE:
-        break;
-
-    case EOpType::IM_SATURATE:
-        res = GetImageDesc( op.args.ImageSaturate.base, returnBestOption, context );
-        break;
-
-    case EOpType::IM_LUMINANCE:
-        res = GetImageDesc( op.args.ImageLuminance.base,returnBestOption, context );
-        res.m_format = EImageFormat::L_UByte;
-        break;
-
-    case EOpType::IM_INTERPOLATE:
-        res = GetImageDesc(op.args.ImageInterpolate.targets[0], returnBestOption, context );
-        break;
-
-    case EOpType::IM_PLAINCOLOUR:
-        res.m_format = EImageFormat(op.args.ImagePlainColour.format);
-        res.m_size[0] = op.args.ImagePlainColour.size[0];
-        res.m_size[1] = op.args.ImagePlainColour.size[1];
-		res.m_lods = op.args.ImagePlainColour.LODs;
-        check( res.m_format != EImageFormat::None );
-        break;
-
-    case EOpType::IM_RESIZE:
-        res = GetImageDesc( op.args.ImageResize.source, returnBestOption, context );
-
-        res.m_size = FImageSize
-            (
-                op.args.ImageResize.size[0],
-                op.args.ImageResize.size[1]
-            );
-        break;
-
-    case EOpType::IM_RESIZEREL:
-        res = GetImageDesc( op.args.ImageResizeRel.source, returnBestOption, context );
-        res.m_size[0] = (uint16)( res.m_size[0]*op.args.ImageResizeRel.factor[0] );
-        res.m_size[1] = (uint16)( res.m_size[1]*op.args.ImageResizeRel.factor[1] );
-        break;
-
-    case EOpType::IM_RESIZELIKE:
-        res = GetImageDesc( op.args.ImageResizeLike.source, returnBestOption, context );
-        if (op.args.ImageResizeLike.sizeSource)
-        {
-            res.m_size = children[op.args.ImageResizeLike.sizeSource]->GetImageDesc( returnBestOption, context ).m_size;
-        }
-        break;
-
-    case EOpType::IM_GRADIENT:
-        res.m_size[0] = op.args.ImageGradient.size[0];
-        res.m_size[1] = op.args.ImageGradient.size[1];
-        res.m_format = EImageFormat::RGB_UByte;
-        break;
-
-    case EOpType::IM_BLANKLAYOUT:
-        // TODO: We would need to process the layout to find the grid size, and then use
-        // the block size with it.
-        res.m_size = FImageSize(0,0);
-        res.m_format = op.args.ImageBlankLayout.format;
-        break;
-
-    case EOpType::IM_BINARISE:
-        res = GetImageDesc( op.args.ImageBinarise.base, returnBestOption, context );
-        res.m_format = EImageFormat::L_UByte;
-        break;
-
-    case EOpType::IM_DISPLACE:
-        res = GetImageDesc( op.args.ImageDisplace.source, returnBestOption, context );
-        break;
-
-	case EOpType::IM_INVERT:
-		res = GetImageDesc(op.args.ImageInvert.base, returnBestOption, context);
-		break;
-
-	case EOpType::IM_COLOURMAP:
-		res = GetImageDesc( op.args.ImageColourMap.base, returnBestOption, context );
-		break;
-
-    default:
-        check(false);
-    }
-
-    // Cache the result
-    if (context)
-    {
-        context->m_results.Add(this,res);
-    }
-
-    return res;
-}
-
-
-void ASTOpFixed::GetLayoutBlockSize( int* pBlockX, int* pBlockY )
-{
-    switch ( op.type )
-    {
-    
-	case EOpType::IM_RESIZEREL:
-	{
-		GetLayoutBlockSize(op.args.ImageResizeRel.source, pBlockX, pBlockY);
-		(*pBlockX) = int(float(*pBlockX) * op.args.ImageResizeRel.factor[0]);
-		(*pBlockY) = int(float(*pBlockX) * op.args.ImageResizeRel.factor[1]);
-
-		break;
-	}
-
-	case EOpType::IM_RESIZE:
-	{
-		GetLayoutBlockSize(op.args.ImageResize.source, pBlockX, pBlockY);
-
-		if (*pBlockX > 0 && *pBlockY> 0)
-		{
-			FImageDesc sourceDesc = GetImageDesc(op.args.ImageResize.source);
-			if (sourceDesc.m_size[0]>0 && sourceDesc.m_size[1]>0)
-			{
-				float factorX = float(op.args.ImageResize.size[0]) / float(sourceDesc.m_size[0]);
-				float factorY = float(op.args.ImageResize.size[1]) / float(sourceDesc.m_size[1]);
-				*pBlockX = int(*pBlockX * factorX);
-				*pBlockY = int(*pBlockX * factorY);
-			}
-			else
-			{
-				*pBlockX = 0;
-				*pBlockY = 0;
-			}
-		}
-
-		break;
-	}
-
-	case EOpType::IM_BLANKLAYOUT:
-	{
-		*pBlockX = int(op.args.ImageBlankLayout.blockSize[0]);
-		*pBlockY = int(op.args.ImageBlankLayout.blockSize[1]);
-		break;
-	}
-
-	case EOpType::IM_PLAINCOLOUR:
-	{
-		*pBlockX = 0;
-		*pBlockY = 0;
-		break;
-	}
-
-    default:
-        checkf( false, TEXT("Instruction not supported") );
-    }
-
-    //check( *pBlockX!=0 && *pBlockY!=0 );
-}
-
-
 ASTOp::FBoolEvalResult ASTOpFixed::EvaluateBool( ASTOpList& facts, FEvaluateBoolCache* cache ) const
 {
     FEvaluateBoolCache localCache;
@@ -1430,47 +1252,6 @@ int ASTOpFixed::EvaluateInt( ASTOpList& /*facts*/, bool &unknown ) const
 }
 
 
-bool ASTOpFixed::IsImagePlainConstant(FVector4f& colour) const
-{
-    bool res = false;
-    switch( op.type )
-    {
-
-    case EOpType::IM_BLANKLAYOUT:
-        res = true;
-        colour[0] = 0;
-        colour[1] = 0;
-        colour[2] = 0;
-        colour[3] = 0;
-        break;
-
-    case EOpType::IM_RESIZE:
-        res = children[op.args.ImageResize.source]->IsImagePlainConstant( colour );
-        break;
-
-    case EOpType::IM_RESIZEREL:
-        res = children[op.args.ImageResizeRel.source]->IsImagePlainConstant( colour );
-        break;
-
-    case EOpType::IM_RESIZELIKE:
-        res = children[op.args.ImageResizeLike.source]->IsImagePlainConstant(  colour );
-        break;
-
-    case EOpType::IM_PLAINCOLOUR:
-        res = children[op.args.ImagePlainColour.colour]->IsColourConstant( colour );
-        break;
-
-    default:
-        // TODO: Improve this test with more operations
-        //check( false );
-        res = false;
-        break;
-    }
-
-    return res;
-}
-
-
 //-------------------------------------------------------------------------------------------------
 bool ASTOpFixed::IsColourConstant(FVector4f& colour) const
 {
@@ -1501,220 +1282,4 @@ bool ASTOpFixed::IsColourConstant(FVector4f& colour) const
 }
 
 
-//-------------------------------------------------------------------------------------------------
-mu::Ptr<ImageSizeExpression> ASTOpFixed::GetImageSizeExpression() const
-{
-    Ptr<ImageSizeExpression> pRes = new ImageSizeExpression;
-
-    EOpType type = GetOpType();
-    switch ( type )
-    {
-    case EOpType::NONE:
-    {
-        pRes->type = ImageSizeExpression::ISET_CONSTANT;
-        pRes->size[0]=0;
-        pRes->size[1]=0;
-        break;
-    }
-
-    case EOpType::IM_RESIZE:
-        pRes->type = ImageSizeExpression::ISET_CONSTANT;
-        pRes->size[0] = op.args.ImageResize.size[0];
-        pRes->size[1] = op.args.ImageResize.size[1];
-        break;
-
-    case EOpType::IM_RESIZEREL:
-        if ( children[op.args.ImageResizeRel.source] )
-        {
-            pRes = children[op.args.ImageResizeRel.source].child()->GetImageSizeExpression();
-            if (pRes->type == ImageSizeExpression::ISET_CONSTANT)
-            {
-                pRes->size[0] = uint16( pRes->size[0] * op.args.ImageResizeRel.factor[0] );
-                pRes->size[1] = uint16( pRes->size[1] * op.args.ImageResizeRel.factor[1] );
-            }
-            else
-            {
-                // TODO: Proportional factor
-                pRes = new ImageSizeExpression();
-                pRes->type = ImageSizeExpression::ISET_UNKNOWN;
-            }
-        }
-        break;
-
-    case EOpType::IM_RESIZELIKE:
-        if ( children[op.args.ImageResizeLike.sizeSource] )
-        {
-            pRes = children[op.args.ImageResizeLike.sizeSource].child()->GetImageSizeExpression();
-        }
-        break;
-
-    case EOpType::IM_PLAINCOLOUR:
-        pRes->type = ImageSizeExpression::ISET_CONSTANT;
-        pRes->size[0] = op.args.ImagePlainColour.size[0];
-        pRes->size[1] = op.args.ImagePlainColour.size[1];
-        break;
-
-    case EOpType::IM_BLANKLAYOUT:
-        pRes->type = ImageSizeExpression::ISET_LAYOUTFACTOR;
-        pRes->layout = children[op.args.ImageBlankLayout.layout].child();
-        pRes->factor[0] = op.args.ImageBlankLayout.blockSize[0];
-        pRes->factor[1] = op.args.ImageBlankLayout.blockSize[1];
-        break;
-
-    case EOpType::IM_INTERPOLATE:
-        if ( children[op.args.ImageInterpolate.targets[0]] )
-        {
-            pRes = children[op.args.ImageInterpolate.targets[0]].child()->GetImageSizeExpression();
-        }
-        break;
-
-    case EOpType::IM_SATURATE:
-        if ( children[op.args.ImageSaturate.base] )
-        {
-            pRes = children[op.args.ImageSaturate.base].child()->GetImageSizeExpression();
-        }
-        break;
-
-    case EOpType::IM_LUMINANCE:
-        if ( children[op.args.ImageLuminance.base] )
-        {
-            pRes = children[op.args.ImageLuminance.base].child()->GetImageSizeExpression();
-        }
-        break;
-
-    case EOpType::IM_BINARISE:
-        if ( children[op.args.ImageBinarise.base] )
-        {
-            pRes = children[op.args.ImageBinarise.base].child()->GetImageSizeExpression();
-        }
-        break;
-
-    case EOpType::IM_COLOURMAP:
-        if ( children[op.args.ImageColourMap.base] )
-        {
-            pRes = children[op.args.ImageColourMap.base].child()->GetImageSizeExpression();
-        }
-        break;
-
-    case EOpType::IM_GRADIENT:
-        pRes->type = ImageSizeExpression::ISET_CONSTANT;
-        pRes->size[0] = op.args.ImageGradient.size[0];
-        pRes->size[1] = op.args.ImageGradient.size[1];
-        break;
-
-    case EOpType::IM_DISPLACE:
-        if ( children[op.args.ImageDisplace.source] )
-        {
-            pRes = children[op.args.ImageDisplace.source].child()->GetImageSizeExpression();
-        }
-        break;
-
-	case EOpType::IM_INVERT:
-		if (children[op.args.ImageInvert.base])
-		{
-			pRes = children[op.args.ImageInvert.base].child()->GetImageSizeExpression();
-		}
-		break;
- 
-    default:
-    	check( false );
-    	break;
-    }
-
-    return pRes;
-}
-
-
-
-
-//!
-FSourceDataDescriptor ASTOpFixed::GetSourceDataDescriptor(FGetSourceDataDescriptorContext* Context) const
-{
-	FSourceDataDescriptor  Result;
-
-	// Cache management
-	TUniquePtr<FGetSourceDataDescriptorContext> LocalContext;
-	if (!Context)
-	{
-		LocalContext.Reset(new FGetSourceDataDescriptorContext);
-		Context = LocalContext.Get();
-	}
-
-	FSourceDataDescriptor* Found = Context->Cache.Find(this);
-	if (Found)
-	{
-		return *Found;
-	}
-
-	// Calculate
-	EOpType Type = op.type;
-
-	switch (Type)
-	{
-
-	case EOpType::NONE:
-	case EOpType::IM_PLAINCOLOUR:
-	case EOpType::IM_GRADIENT:
-	case EOpType::IM_BLANKLAYOUT:
-		break;
-
-	case EOpType::IM_SATURATE:
-		Result = GetSourceDataDescriptor(op.args.ImageSaturate.base, Context);
-		break;
-
-	case EOpType::IM_LUMINANCE:
-		Result = GetSourceDataDescriptor(op.args.ImageLuminance.base, Context);
-		break;
-
-	case EOpType::IM_INTERPOLATE:
-	{
-		for (int32 SourceIndex = 0; SourceIndex < MUTABLE_OP_MAX_INTERPOLATE_COUNT; ++SourceIndex)
-		{
-			FSourceDataDescriptor SourceDesc = GetSourceDataDescriptor(op.args.ImageInterpolate.targets[SourceIndex], Context);
-			Result.CombineWith(SourceDesc);
-		}
-		break;
-	}
-
-	case EOpType::IM_RESIZE:
-		Result = GetSourceDataDescriptor(op.args.ImageResize.source, Context);
-		break;
-
-	case EOpType::IM_RESIZEREL:
-		Result = GetSourceDataDescriptor(op.args.ImageResizeRel.source, Context);
-		break;
-
-	case EOpType::IM_RESIZELIKE:
-		Result = GetSourceDataDescriptor(op.args.ImageResizeLike.source, Context);
-		break;
-
-	case EOpType::IM_BINARISE:
-		Result = GetSourceDataDescriptor(op.args.ImageBinarise.base, Context);
-		break;
-
-	case EOpType::IM_DISPLACE:
-		Result = GetSourceDataDescriptor(op.args.ImageDisplace.source, Context);
-		break;
-
-	case EOpType::IM_INVERT:
-		Result = GetSourceDataDescriptor(op.args.ImageInvert.base, Context);
-		break;
-
-	case EOpType::IM_COLOURMAP:
-		Result = GetSourceDataDescriptor(op.args.ImageColourMap.base, Context);
-		break;
-
-	default:
-		// This should probably be implemented
-		ensure(false);
-		break;
-	}
-
-	// Cache the result
-	Context->Cache.Add(this, Result);
-
-	return Result;
-}
-
-
 }
diff --git a/Engine/Plugins/Mutable/Source/MutableTools/Private/MuT/ASTOpImageBinarize.cpp b/Engine/Plugins/Mutable/Source/MutableTools/Private/MuT/ASTOpImageBinarize.cpp
new file mode 100644
--- /dev/null
+++ b/Engine/Plugins/Mutable/Source/MutableTools/Private/MuT/ASTOpImageBinarize.cpp
@@ -0,0 +1,165 @@
+// Copyright Epic Games, Inc. All Rights Reserved.
+
+#include "MuT/ASTOpImageBinarize.h"
+
+#include "Containers/Map.h"
+#include "MuR/ModelPrivate.h"
+#include "MuR/RefCounted.h"
+#include "MuR/Types.h"
+#include "MuR/ImagePrivate.h"
+#include "MuT/ASTOpConditional.h"
+#include "MuT/ASTOpSwitch.h"
+#include "MuT/ASTOpImagePatch.h"
+#include "MuT/ASTOpImagePixelFormat.h"
+#include "MuT/ASTOpImageLayer.h"
+#include "MuT/ASTOpImageLayerColor.h"
+#include "MuT/ASTOpImageRasterMesh.h"
+#include "MuT/ASTOpImageTransform.h"
+#include "MuT/ASTOpImageSwizzle.h"
+#include "MuT/ASTOpImageCompose.h"
+#include "MuT/ASTOpImageMakeGrowMap.h"
+#include "MuT/ASTOpImageResize.h"
+
+
+namespace mu
+{
+
+	ASTOpImageBinarize::ASTOpImageBinarize()
+		: Base(this),
+		Threshold(this)
+	{
+	}
+
+
+	ASTOpImageBinarize::~ASTOpImageBinarize()
+	{
+		// Explicit call needed to avoid recursive destruction
+		ASTOp::RemoveChildren();
+	}
+
+
+	bool ASTOpImageBinarize::IsEqual(const ASTOp& InOther) const
+	{
+		if (InOther.GetOpType()==GetOpType())
+		{
+			const ASTOpImageBinarize* Other = static_cast<const ASTOpImageBinarize*>(&InOther);
+			return Base == Other->Base &&
+				Threshold == Other->Threshold;
+		}
+		return false;
+	}
+
+
+	uint64 ASTOpImageBinarize::Hash() const
+	{
+		uint64 Res = std::hash<EOpType>()(GetOpType());
+		hash_combine(Res, Base.child().get());
+		hash_combine(Res, Threshold.child().get());
+		return Res;
+	}
+
+
+	Ptr<ASTOp> ASTOpImageBinarize::Clone(MapChildFuncRef MapChild) const
+	{
+		Ptr<ASTOpImageBinarize> New = new ASTOpImageBinarize();
+		New->Base = MapChild(Base.child());
+		New->Threshold = MapChild(Threshold.child());
+		return New;
+	}
+
+
+	void ASTOpImageBinarize::ForEachChild(const TFunctionRef<void(ASTChild&)> Func)
+	{
+		Func(Base);
+		Func(Threshold);
+	}
+
+
+	void ASTOpImageBinarize::Link(FProgram& Program, FLinkerOptions*)
+	{
+		// Already linked?
+		if (!linkedAddress)
+		{
+			OP::ImageBinariseArgs Args = {};
+
+			if (Base)
+			{
+				Args.Base = Base->linkedAddress;
+			}
+			if (Threshold)
+			{
+				Args.Threshold = Threshold->linkedAddress;
+			}
+
+			linkedAddress = (OP::ADDRESS)Program.OpAddress.Num();
+			Program.OpAddress.Add(Program.ByteCode.Num());
+			AppendCode(Program.ByteCode, GetOpType());
+			AppendCode(Program.ByteCode, Args);
+		}
+	}
+
+
+	FImageDesc ASTOpImageBinarize::GetImageDesc(bool bReturnBestOption, FGetImageDescContext* Context) const
+	{
+		FImageDesc Result;
+
+		// Local context in case it is necessary
+		FGetImageDescContext LocalContext;
+		if (!Context)
+		{
+			Context = &LocalContext;
+		}
+		else
+		{
+			// Cached result?
+			FImageDesc* PtrValue = Context->m_results.Find(this);
+			if (PtrValue)
+			{
+				return *PtrValue;
+			}
+		}
+
+		// Actual work
+		if (Base)
+		{
+			Result = Base->GetImageDesc(bReturnBestOption, Context);
+		}
+
+		if (Threshold)
+		{
+			FImageDesc SizeResult = Threshold->GetImageDesc(bReturnBestOption, Context);
+			Result.m_size = SizeResult.m_size;
+		}
+
+		// Cache the result
+		if (Context)
+		{
+			Context->m_results.Add(this, Result);
+		}
+
+		return Result;
+	}
+
+
+	Ptr<ImageSizeExpression> ASTOpImageBinarize::GetImageSizeExpression() const
+	{
+		if (Base)
+		{
+			return Base->GetImageSizeExpression();
+		}
+
+		return new ImageSizeExpression;
+	}
+
+
+	FSourceDataDescriptor ASTOpImageBinarize::GetSourceDataDescriptor(FGetSourceDataDescriptorContext* Context) const
+	{
+		if (Base)
+		{
+			return Base->GetSourceDataDescriptor(Context);
+		}
+
+		return {};
+	}
+
+}
diff --git a/Engine/Plugins/Mutable/Source/MutableTools/Private/MuT/ASTOpImageBinarize.h b/Engine/Plugins/Mutable/Source/MutableTools/Private/MuT/ASTOpImageBinarize.h
new file mode 100644
--- /dev/null
+++ b/Engine/Plugins/Mutable/Source/MutableTools/Private/MuT/ASTOpImageBinarize.h
@@ -0,0 +1,44 @@
+// Copyright Epic Games, Inc. All Rights Reserved.
+
+#pragma once
+
+#include "HAL/PlatformMath.h"
+#include "MuR/Image.h"
+#include "MuR/MutableMath.h"
+#include "MuR/Operations.h"
+#include "MuR/Ptr.h"
+#include "MuT/AST.h"
+
+
+namespace mu
+{
+	struct FProgram;
+
+	class ASTOpImageBinarize final : public ASTOp
+	{
+	public:
+
+		ASTChild Base;
+		ASTChild Threshold;
+
+	public:
+
+		ASTOpImageBinarize();
+		ASTOpImageBinarize(const ASTOpImageBinarize&) = delete;
+		~ASTOpImageBinarize();
+
+		virtual EOpType GetOpType() const override { return EOpType::IM_BINARISE; }
+		virtual uint64 Hash() const override;
+		virtual bool IsEqual(const ASTOp&) const override;
+		virtual Ptr<ASTOp> Clone(MapChildFuncRef) const override;
+		virtual void ForEachChild(const TFunctionRef<void(ASTChild&)>) override;
+		virtual void Link(FProgram&, FLinkerOptions*) override;
+		virtual FImageDesc GetImageDesc(bool bReturnBestOption, FGetImageDescContext*) const override;
+		virtual Ptr<ImageSizeExpression> GetImageSizeExpression() const override;
+		virtual FSourceDataDescriptor GetSourceDataDescriptor(FGetSourceDataDescriptorContext*) const override;
+
+	};
+
+
+}
+
diff --git a/Engine/Plugins/Mutable/Source/MutableTools/Private/MuT/ASTOpImageBlankLayout.cpp b/Engine/Plugins/Mutable/Source/MutableTools/Private/MuT/ASTOpImageBlankLayout.cpp
new file mode 100644
--- /dev/null
+++ b/Engine/Plugins/Mutable/Source/MutableTools/Private/MuT/ASTOpImageBlankLayout.cpp
@@ -0,0 +1,170 @@
+// Copyright Epic Games, Inc. All Rights Reserved.
+
+#include "MuT/ASTOpImageBlankLayout.h"
+
+#include "Containers/Map.h"
+#include "MuR/ModelPrivate.h"
+#include "MuR/RefCounted.h"
+#include "MuR/Types.h"
+#include "MuT/ASTOpConditional.h"
+#include "MuT/ASTOpSwitch.h"
+#include "MuT/ASTOpImagePatch.h"
+#include "MuT/ASTOpImagePixelFormat.h"
+#include "MuT/ASTOpImageLayer.h"
+#include "MuT/ASTOpImageLayerColor.h"
+#include "MuT/ASTOpImageRasterMesh.h"
+
+
+namespace mu
+{
+
+	ASTOpImageBlankLayout::ASTOpImageBlankLayout()
+		: Layout(this)
+	{
+	}
+
+
+	ASTOpImageBlankLayout::~ASTOpImageBlankLayout()
+	{
+		// Explicit call needed to avoid recursive destruction
+		ASTOp::RemoveChildren();
+	}
+
+
+	bool ASTOpImageBlankLayout::IsEqual(const ASTOp& InOther) const
+	{
+		if (InOther.GetOpType()==GetOpType())
+		{
+			const ASTOpImageBlankLayout* Other = static_cast<const ASTOpImageBlankLayout*>(&InOther);
+			return Layout == Other->Layout &&
+				BlockSize == Other->BlockSize &&
+				Format == Other->Format &&
+				GenerateMipmaps == Other->GenerateMipmaps &&
+				MipmapCount == Other->MipmapCount;
+		}
+		return false;
+	}
+
+
+	uint64 ASTOpImageBlankLayout::Hash() const
+	{
+		uint64 Result = std::hash<EOpType>()(GetOpType());
+		hash_combine(Result, Layout.child().get());
+		hash_combine(Result, BlockSize[0]);
+		hash_combine(Result, BlockSize[1]);
+		hash_combine(Result, Format);
+		hash_combine(Result, GenerateMipmaps);
+		hash_combine(Result, MipmapCount);
+		return Result;
+	}
+
+
+	Ptr<ASTOp> ASTOpImageBlankLayout::Clone(MapChildFuncRef MapChild) const
+	{
+		Ptr<ASTOpImageBlankLayout> n = new ASTOpImageBlankLayout();
+		n->Layout = MapChild(Layout.child());
+		n->BlockSize = BlockSize;
+		n->Format = Format;
+		n->GenerateMipmaps = GenerateMipmaps;
+		n->MipmapCount = MipmapCount;
+		return n;
+	}
+
+
+	void ASTOpImageBlankLayout::ForEachChild(const TFunctionRef<void(ASTChild&)> Func)
+	{
+		Func(Layout);
+	}
+
+
+	void ASTOpImageBlankLayout::Link(FProgram& Program, FLinkerOptions*)
+	{
+		// Already linked?
+		if (!linkedAddress)
+		{
+			OP::ImageBlankLayoutArgs Args={};
+
+			if (Layout)
+			{
+				Args.Layout = Layout->linkedAddress;
+			}
+
+			Args.BlockSize[0] = BlockSize[0];
+			Args.BlockSize[1] = BlockSize[1];
+			Args.Format = Format;
+			Args.GenerateMipmaps = GenerateMipmaps;
+			Args.MipmapCount = MipmapCount;
+
+			linkedAddress = (OP::ADDRESS)Program.OpAddress.Num();
+			Program.OpAddress.Add(Program.ByteCode.Num());
+			AppendCode(Program.ByteCode, GetOpType());
+			AppendCode(Program.ByteCode, Args);
+		}
+
+	}
+
+
+	FImageDesc ASTOpImageBlankLayout::GetImageDesc(bool bReturnBestOption, FGetImageDescContext* Context) const
+	{
+		FImageDesc Result;
+
+		// Local context in case it is necessary
+		FGetImageDescContext LocalContext;
+		if (!Context)
+		{
+			Context = &LocalContext;
+		}
+		else
+		{
+			// Cached result?
+			FImageDesc* PtrValue = Context->m_results.Find(this);
+			if (PtrValue)
+			{
+				return *PtrValue;
+			}
+		}
+
+		// Actual work
+	   // TODO: We would need to process the layout to find the grid size, and then use
+		// the block size with it.
+		Result.m_size = FImageSize(0, 0);
+		Result.m_format = Format;
+
+		// Cache the result
+		if (Context)
+		{
+			Context->m_results.Add(this, Result);
+		}
+
+		return Result;
+	}
+
+
+	Ptr<ImageSizeExpression> ASTOpImageBlankLayout::GetImageSizeExpression() const
+	{
+		Ptr<ImageSizeExpression> Res = new ImageSizeExpression;
+		Res->type = ImageSizeExpression::ISET_LAYOUTFACTOR;
+		Res->layout = Layout.child();
+		Res->factor[0] = BlockSize[0];
+		Res->factor[1] = BlockSize[1];
+		return Res;
+	}
+
+
+	void ASTOpImageBlankLayout::GetLayoutBlockSize(int32* OutBlockX, int32* OutBlockY)
+	{
+		*OutBlockX = int32(BlockSize[0]);
+		*OutBlockY = int32(BlockSize[1]);
+	}
+
+
+	bool ASTOpImageBlankLayout::IsImagePlainConstant(FVector4f& OutColour) const
+	{
+		OutColour[0] = 0;
+		OutColour[1] = 0;
+		OutColour[2] = 0;
+		OutColour[3] = 0;
+		return true;
+	}
+
+}
diff --git a/Engine/Plugins/Mutable/Source/MutableTools/Private/MuT/ASTOpImageBlankLayout.h b/Engine/Plugins/Mutable/Source/MutableTools/Private/MuT/ASTOpImageBlankLayout.h
new file mode 100644
--- /dev/null
+++ b/Engine/Plugins/Mutable/Source/MutableTools/Private/MuT/ASTOpImageBlankLayout.h
@@ -0,0 +1,55 @@
+// Copyright Epic Games, Inc. All Rights Reserved.
+
+#pragma once
+
+#include "HAL/PlatformMath.h"
+#include "MuR/Image.h"
+#include "MuR/MutableMath.h"
+#include "MuR/Operations.h"
+#include "MuR/Ptr.h"
+#include "MuT/AST.h"
+
+
+namespace mu
+{
+	struct FProgram;
+
+	class ASTOpImageBlankLayout final : public ASTOp
+	{
+	public:
+
+		ASTChild Layout;
+
+		/** Size of a layout block in pixels. */
+		UE::Math::TIntVector2<uint16> BlockSize;
+
+		EImageFormat Format = EImageFormat::None;
+
+		/** If true, generate mipmaps. */
+		uint8 GenerateMipmaps = 0;
+
+		/** Mipmaps to generate if mipmaps are to be generated. 0 means all. */
+		uint8 MipmapCount = 0;
+
+	public:
+
+		ASTOpImageBlankLayout();
+		ASTOpImageBlankLayout(const ASTOpImageBlankLayout&) = delete;
+		~ASTOpImageBlankLayout();
+
+		virtual EOpType GetOpType() const override { return EOpType::IM_BLANKLAYOUT; }
+		virtual uint64 Hash() const override;
+		virtual bool IsEqual(const ASTOp&) const override;
+		virtual Ptr<ASTOp> Clone(MapChildFuncRef ) const override;
+		virtual void ForEachChild(const TFunctionRef<void(ASTChild&)>) override;
+		virtual void Link(FProgram&, FLinkerOptions*) override;
+		virtual FImageDesc GetImageDesc(bool bReturnBestOption, FGetImageDescContext*) const override;
+		virtual Ptr<ImageSizeExpression> GetImageSizeExpression() const override;
+		virtual void GetLayoutBlockSize(int* OutBlockX, int* OutBlockY) override;
+		virtual bool IsImagePlainConstant(FVector4f& colour) const override;
+
+	};
+
+
+}
+
diff --git a/Engine/Plugins/Mutable/Source/MutableTools/Private/MuT/ASTOpImageColorMap.cpp b/Engine/Plugins/Mutable/Source/MutableTools/Private/MuT/ASTOpImageColorMap.cpp
new file mode 100644
--- /dev/null
+++ b/Engine/Plugins/Mutable/Source/MutableTools/Private/MuT/ASTOpImageColorMap.cpp
@@ -0,0 +1,176 @@
+// Copyright Epic Games, Inc. All Rights Reserved.
+
+#include "MuT/ASTOpImageColorMap.h"
+
+#include "Containers/Map.h"
+#include "MuR/ModelPrivate.h"
+#include "MuR/RefCounted.h"
+#include "MuR/Types.h"
+#include "MuR/ImagePrivate.h"
+#include "MuT/ASTOpConditional.h"
+#include "MuT/ASTOpSwitch.h"
+#include "MuT/ASTOpImagePatch.h"
+#include "MuT/ASTOpImagePixelFormat.h"
+#include "MuT/ASTOpImageLayer.h"
+#include "MuT/ASTOpImageLayerColor.h"
+#include "MuT/ASTOpImageRasterMesh.h"
+#include "MuT/ASTOpImageTransform.h"
+#include "MuT/ASTOpImageSwizzle.h"
+#include "MuT/ASTOpImageCompose.h"
+#include "MuT/ASTOpImageMakeGrowMap.h"
+#include "MuT/ASTOpImageResize.h"
+
+
+namespace mu
+{
+
+	ASTOpImageColorMap::ASTOpImageColorMap()
+		: Base(this),
+		Mask(this),
+		Map(this)
+	{
+	}
+
+
+	ASTOpImageColorMap::~ASTOpImageColorMap()
+	{
+		// Explicit call needed to avoid recursive destruction
+		ASTOp::RemoveChildren();
+	}
+
+
+	bool ASTOpImageColorMap::IsEqual(const ASTOp& InOther) const
+	{
+		if (InOther.GetOpType()==GetOpType())
+		{
+			const ASTOpImageColorMap* Other = static_cast<const ASTOpImageColorMap*>(&InOther);
+			return Base == Other->Base &&
+				Mask == Other->Mask &&
+				Map == Other->Map;
+		}
+		return false;
+	}
+
+
+	uint64 ASTOpImageColorMap::Hash() const
+	{
+		uint64 Res = std::hash<EOpType>()(GetOpType());
+		hash_combine(Res, Base.child().get());
+		hash_combine(Res, Mask.child().get());
+		hash_combine(Res, Map.child().get());
+		return Res;
+	}
+
+
+	Ptr<ASTOp> ASTOpImageColorMap::Clone(MapChildFuncRef MapChild) const
+	{
+		Ptr<ASTOpImageColorMap> New = new ASTOpImageColorMap();
+		New->Base = MapChild(Base.child());
+		New->Mask = MapChild(Mask.child());
+		New->Map = MapChild(Map.child());
+		return New;
+	}
+
+
+	void ASTOpImageColorMap::ForEachChild(const TFunctionRef<void(ASTChild&)> Func)
+	{
+		Func(Base);
+		Func(Mask);
+		Func(Map);
+	}
+
+
+	void ASTOpImageColorMap::Link(FProgram& Program, FLinkerOptions*)
+	{
+		// Already linked?
+		if (!linkedAddress)
+		{
+			OP::ImageColourMapArgs Args = {};
+
+			if (Base)
+			{
+				Args.Base = Base->linkedAddress;
+			}
+			if (Mask)
+			{
+				Args.Mask = Mask->linkedAddress;
+			}
+			if (Map)
+			{
+				Args.Map = Map->linkedAddress;
+			}
+
+			linkedAddress = (OP::ADDRESS)Program.OpAddress.Num();
+			Program.OpAddress.Add(Program.ByteCode.Num());
+			AppendCode(Program.ByteCode, GetOpType());
+			AppendCode(Program.ByteCode, Args);
+		}
+	}
+
+
+	FImageDesc ASTOpImageColorMap::GetImageDesc(bool bReturnBestOption, FGetImageDescContext* Context) const
+	{
+		FImageDesc Result;
+
+		// Local context in case it is necessary
+		FGetImageDescContext LocalContext;
+		if (!Context)
+		{
+			Context = &LocalContext;
+		}
+		else
+		{
+			// Cached result?
+			FImageDesc* PtrValue = Context->m_results.Find(this);
+			if (PtrValue)
+			{
+				return *PtrValue;
+			}
+		}
+
+		// Actual work
+		if (Base)
+		{
+			Result = Base->GetImageDesc(bReturnBestOption, Context);
+		}
+
+		// Cache the result
+		if (Context)
+		{
+			Context->m_results.Add(this, Result);
+		}
+
+		return Result;
+	}
+
+
+	Ptr<ImageSizeExpression> ASTOpImageColorMap::GetImageSizeExpression() const
+	{
+		if (Base)
+		{
+			return Base->GetImageSizeExpression();
+		}
+
+		return new ImageSizeExpression;
+	}
+
+
+	bool ASTOpImageColorMap::IsImagePlainConstant(FVector4f& OutColour) const
+	{
+		// Because we cannot calculate the OutColour
+		bool bResult = false;
+		return bResult;
+	}
+
+
+	FSourceDataDescriptor ASTOpImageColorMap::GetSourceDataDescriptor(FGetSourceDataDescriptorContext* Context) const
+	{
+		if (Base)
+		{
+			return Base->GetSourceDataDescriptor(Context);
+		}
+
+		return {};
+	}
+
+}
diff --git a/Engine/Plugins/Mutable/Source/MutableTools/Private/MuT/ASTOpImageColorMap.h b/Engine/Plugins/Mutable/Source/MutableTools/Private/MuT/ASTOpImageColorMap.h
new file mode 100644
--- /dev/null
+++ b/Engine/Plugins/Mutable/Source/MutableTools/Private/MuT/ASTOpImageColorMap.h
@@ -0,0 +1,46 @@
+// Copyright Epic Games, Inc. All Rights Reserved.
+
+#pragma once
+
+#include "HAL/PlatformMath.h"
+#include "MuR/Image.h"
+#include "MuR/MutableMath.h"
+#include "MuR/Operations.h"
+#include "MuR/Ptr.h"
+#include "MuT/AST.h"
+
+
+namespace mu
+{
+	struct FProgram;
+
+	class ASTOpImageColorMap final : public ASTOp
+	{
+	public:
+
+		ASTChild Base;
+		ASTChild Mask;
+		ASTChild Map;
+
+	public:
+
+		ASTOpImageColorMap();
+		ASTOpImageColorMap(const ASTOpImageColorMap&) = delete;
+		~ASTOpImageColorMap();
+
+		virtual EOpType GetOpType() const override { return EOpType::IM_COLOURMAP; }
+		virtual uint64 Hash() const override;
+		virtual bool IsEqual(const ASTOp&) const override;
+		virtual Ptr<ASTOp> Clone(MapChildFuncRef) const override;
+		virtual void ForEachChild(const TFunctionRef<void(ASTChild&)>) override;
+		virtual void Link(FProgram&, FLinkerOptions*) override;
+		virtual FImageDesc GetImageDesc(bool bReturnBestOption, FGetImageDescContext*) const override;
+		virtual Ptr<ImageSizeExpression> GetImageSizeExpression() const override;
+		virtual bool IsImagePlainConstant(FVector4f& colour) const override;
+		virtual FSourceDataDescriptor GetSourceDataDescriptor(FGetSourceDataDescriptorContext*) const override;
+
+	};
+
+
+}
+
diff --git a/Engine/Plugins/Mutable/Source/MutableTools/Private/MuT/ASTOpImageCompose.cpp b/Engine/Plugins/Mutable/Source/MutableTools/Private/MuT/ASTOpImageCompose.cpp
--- a/Engine/Plugins/Mutable/Source/MutableTools/Private/MuT/ASTOpImageCompose.cpp
+++ b/Engine/Plugins/Mutable/Source/MutableTools/Private/MuT/ASTOpImageCompose.cpp
@@ -14,6 +14,7 @@
 #include "MuR/Types.h"
 #include "MuT/ASTOpConstantResource.h"
 #include "MuT/ASTOpImagePixelFormat.h"
+#include "MuT/ASTOpImageResize.h"
 
 
 namespace mu
@@ -256,35 +257,34 @@ mu::Ptr<ASTOp> ASTOpImageCompose::OptimiseSemantic(const FModelOptimizationOptio
 		{
 			// We could only take the block, but we must make sure it will have the format
 			// and size of the base.
-			mu::FImageDesc baseDesc = baseAt->GetImageDesc(true);
-			mu::FImageDesc blockDesc = blockAt->GetImageDesc(true);
+			mu::FImageDesc BaseDesc = baseAt->GetImageDesc(true);
+			mu::FImageDesc BlockDesc = blockAt->GetImageDesc(true);
 
 			at = blockAt;
 
 			// \todo isn't this a common operation?
-			if (baseDesc.m_format != blockDesc.m_format
+			if (BaseDesc.m_format != BlockDesc.m_format
 				&&
-				baseDesc.m_format != EImageFormat::None)
+				BaseDesc.m_format != EImageFormat::None)
 			{
-				mu::Ptr<ASTOpImagePixelFormat> reformat = new ASTOpImagePixelFormat;
-				reformat->Format = baseDesc.m_format;
-				reformat->FormatIfAlpha = EImageFormat::None;
-				reformat->Source = at;
-				at = reformat;
+				mu::Ptr<ASTOpImagePixelFormat> Reformat = new ASTOpImagePixelFormat;
+				Reformat->Format = BaseDesc.m_format;
+				Reformat->FormatIfAlpha = EImageFormat::None;
+				Reformat->Source = at;
+				at = Reformat;
 			}
 
-			if (baseDesc.m_size != blockDesc.m_size
+			if (BaseDesc.m_size != BlockDesc.m_size
 				&&
-				baseDesc.m_size[0] != 0
+				BaseDesc.m_size[0] != 0
 				&&
-				baseDesc.m_size[1] != 0)
+				BaseDesc.m_size[1] != 0)
 			{
-				mu::Ptr<ASTOpFixed> resize = new ASTOpFixed;
-				resize->op.type = EOpType::IM_RESIZE;
-				resize->op.args.ImageResize.size[0] = baseDesc.m_size[0];
-				resize->op.args.ImageResize.size[1] = baseDesc.m_size[1];
-				resize->SetChild(resize->op.args.ImageResize.source, at);
-				at = resize;
+				mu::Ptr<ASTOpImageResize> Resize = new ASTOpImageResize;
+				Resize->Size[0] = BaseDesc.m_size[0];
+				Resize->Size[1] = BaseDesc.m_size[1];
+				Resize->Source = at;
+				at = Resize;
 			}
 		}
 	}
diff --git a/Engine/Plugins/Mutable/Source/MutableTools/Private/MuT/ASTOpImageCrop.cpp b/Engine/Plugins/Mutable/Source/MutableTools/Private/MuT/ASTOpImageCrop.cpp
--- a/Engine/Plugins/Mutable/Source/MutableTools/Private/MuT/ASTOpImageCrop.cpp
+++ b/Engine/Plugins/Mutable/Source/MutableTools/Private/MuT/ASTOpImageCrop.cpp
@@ -13,6 +13,9 @@
 #include "MuT/ASTOpImageLayer.h"
 #include "MuT/ASTOpImageLayerColor.h"
 #include "MuT/ASTOpImageRasterMesh.h"
+#include "MuT/ASTOpImageInterpolate.h"
+#include "MuT/ASTOpImagePlainColor.h"
+#include "MuT/ASTOpImageDisplace.h"
 
 
 namespace mu
@@ -46,7 +49,7 @@ namespace mu
 
 	uint64 ASTOpImageCrop::Hash() const
 	{
-		uint64 res = std::hash<EOpType>()(EOpType::IM_CROP);
+		uint64 res = std::hash<EOpType>()(GetOpType());
 		hash_combine(res, Source.child().get());
 		hash_combine(res, Min[0]);
 		hash_combine(res, Min[1]);
@@ -56,59 +59,62 @@ namespace mu
 	}
 
 
-	mu::Ptr<ASTOp> ASTOpImageCrop::Clone(MapChildFuncRef mapChild) const
+	Ptr<ASTOp> ASTOpImageCrop::Clone(MapChildFuncRef MapChild) const
 	{
 		Ptr<ASTOpImageCrop> n = new ASTOpImageCrop();
-		n->Source = mapChild(Source.child());
+		n->Source = MapChild(Source.child());
 		n->Min = Min;
 		n->Size = Size;
 		return n;
 	}
 
 
-	void ASTOpImageCrop::ForEachChild(const TFunctionRef<void(ASTChild&)> f)
+	void ASTOpImageCrop::ForEachChild(const TFunctionRef<void(ASTChild&)> Func)
 	{
-		f(Source);
+		Func(Source);
 	}
 
 
-	void ASTOpImageCrop::Link(FProgram& program, FLinkerOptions*)
+	void ASTOpImageCrop::Link(FProgram& Program, FLinkerOptions*)
 	{
 		// Already linked?
 		if (!linkedAddress)
 		{
-			OP::ImageCropArgs args;
-			memset(&args, 0, sizeof(args));
-
-			if (Source) args.source = Source->linkedAddress;
-			args.minX = Min[0];
-			args.minY = Min[1];
-			args.sizeX = Size[0];
-			args.sizeY = Size[1];
-
-			linkedAddress = (OP::ADDRESS)program.OpAddress.Num();
-			program.OpAddress.Add((uint32_t)program.ByteCode.Num());
-			AppendCode(program.ByteCode, EOpType::IM_CROP);
-			AppendCode(program.ByteCode, args);
+			OP::ImageCropArgs Args={};
+
+			if (Source)
+			{
+				Args.source = Source->linkedAddress;
+			}
+
+			Args.minX = Min[0];
+			Args.minY = Min[1];
+			Args.sizeX = Size[0];
+			Args.sizeY = Size[1];
+
+			linkedAddress = (OP::ADDRESS)Program.OpAddress.Num();
+			Program.OpAddress.Add(Program.ByteCode.Num());
+			AppendCode(Program.ByteCode, GetOpType());
+			AppendCode(Program.ByteCode, Args);
 		}
 
 	}
 
 
-	FImageDesc ASTOpImageCrop::GetImageDesc(bool returnBestOption, FGetImageDescContext* context) const
+	FImageDesc ASTOpImageCrop::GetImageDesc(bool bReturnBestOption, FGetImageDescContext* Context) const
 	{
 		FImageDesc Result;
 
 		// Local context in case it is necessary
-		FGetImageDescContext localContext;
-		if (!context)
+		FGetImageDescContext LocalContext;
+		if (!Context)
 		{
-			context = &localContext;
+			Context = &LocalContext;
 		}
 		else
 		{
 			// Cached result?
-			FImageDesc* PtrValue = context->m_results.Find(this);
+			FImageDesc* PtrValue = Context->m_results.Find(this);
 			if (PtrValue)
 			{
 				return *PtrValue;
@@ -118,43 +124,42 @@ namespace mu
 		// Actual work
 		if (Source)
 		{
-			Result = Source->GetImageDesc(returnBestOption, context);
+			Result = Source->GetImageDesc(bReturnBestOption, Context);
 			//Result.m_lods = 1;
 			Result.m_size[0] = Size[0];
 			Result.m_size[1] = Size[1];
 		}
 
-
 		// Cache the result
-		if (context)
+		if (Context)
 		{
-			context->m_results.Add(this, Result);
+			Context->m_results.Add(this, Result);
 		}
 
 		return Result;
 	}
 
 
-	mu::Ptr<ImageSizeExpression> ASTOpImageCrop::GetImageSizeExpression() const
+	Ptr<ImageSizeExpression> ASTOpImageCrop::GetImageSizeExpression() const
 	{
-		Ptr<ImageSizeExpression> pRes = new ImageSizeExpression;
-		pRes->type = ImageSizeExpression::ISET_CONSTANT;
-		pRes->size[0] = Size[0];
-		pRes->size[1] = Size[1];
-		return pRes;
+		Ptr<ImageSizeExpression> Res = new ImageSizeExpression;
+		Res->type = ImageSizeExpression::ISET_CONSTANT;
+		Res->size[0] = Size[0];
+		Res->size[1] = Size[1];
+		return Res;
 	}
 
 
-	void ASTOpImageCrop::GetLayoutBlockSize(int* pBlockX, int* pBlockY)
+	void ASTOpImageCrop::GetLayoutBlockSize(int32* OutBlockX, int32* OutBlockY)
 	{
 		// We didn't find any layout yet.
-		*pBlockX = 0;
-		*pBlockY = 0;
+		*OutBlockX = 0;
+		*OutBlockY = 0;
 
 		// Try the source
 		if (Source)
 		{
-			Source->GetLayoutBlockSize( pBlockX, pBlockY );
+			Source->GetLayoutBlockSize( OutBlockX, OutBlockY );
 		}
 	}
 
@@ -172,21 +177,21 @@ namespace mu
 
 	Ptr<ASTOp> ASTOpImageCrop::OptimiseSemantic(const FModelOptimizationOptions& options, int32 Pass) const
 	{
-		Ptr<ASTOp> at;
+		Ptr<ASTOp> Result;
 
-		Ptr<ASTOp> sourceAt = Source.child();
+		Ptr<ASTOp> SourceAt = Source.child();
 
 		// The instruction can be sunk
-		EOpType sourceType = sourceAt->GetOpType();
+		EOpType sourceType = SourceAt->GetOpType();
 		switch (sourceType)
 		{
 		case EOpType::IM_PLAINCOLOUR:
 		{
-			Ptr<ASTOpFixed> NewOp = mu::Clone<ASTOpFixed>(sourceAt.get());
-			NewOp->op.args.ImagePlainColour.size[0] = Size[0];
-			NewOp->op.args.ImagePlainColour.size[1] = Size[1];
-			NewOp->op.args.ImagePlainColour.LODs = 1; // TODO
-			at = NewOp;
+			Ptr<ASTOpImagePlainColor> NewOp = mu::Clone<ASTOpImagePlainColor>(SourceAt.get());
+			NewOp->Size[0] = Size[0];
+			NewOp->Size[1] = Size[1];
+			NewOp->LODs = 1; // TODO
+			Result = NewOp;
 			break;
 		}
 
@@ -195,33 +200,32 @@ namespace mu
 
 		}
 
-		return at;
+		return Result;
 	}
 
 
-	Ptr<ASTOp> ASTOpImageCrop::OptimiseSink(const FModelOptimizationOptions& options, FOptimizeSinkContext& context) const
+	Ptr<ASTOp> ASTOpImageCrop::OptimiseSink(const FModelOptimizationOptions&, FOptimizeSinkContext& context) const
 	{
-		Ptr<ASTOp> at;
+		Ptr<ASTOp> Result;
 
-		Ptr<ASTOp> sourceAt = Source.child();
+		Ptr<ASTOp> SourceAt = Source.child();
 
-		switch (sourceAt->GetOpType())
+		switch (SourceAt->GetOpType())
 		{
-			// In case we have other operations with special optimisation rules.
+		// In case we have other operations with special optimisation rules.
 		case EOpType::NONE:
 			break;
 
 		default:
 		{
-			at = context.ImageCropSinker.Apply(this);
+			Result = context.ImageCropSinker.Apply(this);
 
 			break;
 		} // default
 
 		}
 
-
-		return at;
+		return Result;
 	}
 
 
@@ -431,13 +435,13 @@ namespace mu
 		case EOpType::IM_DISPLACE:
 		{
 			// We move the op down the arguments
-			Ptr<ASTOpFixed> nop = mu::Clone<ASTOpFixed>(at);
+			Ptr<ASTOpImageDisplace> nop = mu::Clone<ASTOpImageDisplace>(at);
 
-			Ptr<ASTOp> aOp = nop->children[nop->op.args.ImageDisplace.source].child();
-			nop->SetChild(nop->op.args.ImageDisplace.source, Visit(aOp, currentCropOp));
+			Ptr<ASTOp> aOp = nop->Source.child();
+			nop->Source = Visit(aOp, currentCropOp);
 
-			Ptr<ASTOp> bOp = nop->children[nop->op.args.ImageDisplace.displacementMap].child();
-			nop->SetChild(nop->op.args.ImageDisplace.displacementMap, Visit(bOp, currentCropOp));
+			Ptr<ASTOp> bOp = nop->DisplacementMap.child();
+			nop->DisplacementMap = Visit(bOp, currentCropOp);
 
 			newAt = nop;
 			break;
@@ -473,16 +477,16 @@ namespace mu
 		case EOpType::IM_INTERPOLATE:
 		{
 			// Move the op  down all the paths
-			auto newOp = mu::Clone<ASTOpFixed>(at);
+			Ptr<ASTOpImageInterpolate> NewOp = mu::Clone<ASTOpImageInterpolate>(at);
 
-			for (int v = 0; v < MUTABLE_OP_MAX_INTERPOLATE_COUNT; ++v)
+			for (int32 v = 0; v < MUTABLE_OP_MAX_INTERPOLATE_COUNT; ++v)
 			{
-				Ptr<ASTOp> child = newOp->children[newOp->op.args.ImageInterpolate.targets[v]].child();
+				Ptr<ASTOp> child = NewOp->Targets[v].child();
 				Ptr<ASTOp> bOp = Visit(child, currentCropOp);
-				newOp->SetChild(newOp->op.args.ImageInterpolate.targets[v], bOp);
+				NewOp->Targets[v] = bOp;
 			}
 
-			newAt = newOp;
+			newAt = NewOp;
 			break;
 		}
 
diff --git a/Engine/Plugins/Mutable/Source/MutableTools/Private/MuT/ASTOpImageCrop.h b/Engine/Plugins/Mutable/Source/MutableTools/Private/MuT/ASTOpImageCrop.h
--- a/Engine/Plugins/Mutable/Source/MutableTools/Private/MuT/ASTOpImageCrop.h
+++ b/Engine/Plugins/Mutable/Source/MutableTools/Private/MuT/ASTOpImageCrop.h
@@ -30,16 +30,16 @@ namespace mu
 
 		virtual EOpType GetOpType() const override { return EOpType::IM_CROP; }
 		virtual uint64 Hash() const override;
-		virtual bool IsEqual(const ASTOp& otherUntyped) const override;
-		virtual Ptr<ASTOp> Clone(MapChildFuncRef mapChild) const override;
+		virtual bool IsEqual(const ASTOp&) const override;
+		virtual Ptr<ASTOp> Clone(MapChildFuncRef ) const override;
 		virtual void ForEachChild(const TFunctionRef<void(ASTChild&)>) override;
-		virtual void Link(FProgram& program, FLinkerOptions* Options) override;
-		virtual FImageDesc GetImageDesc(bool returnBestOption, FGetImageDescContext* context) const override;
+		virtual void Link(FProgram&, FLinkerOptions*) override;
+		virtual FImageDesc GetImageDesc(bool bReturnBestOption, FGetImageDescContext*) const override;
 		virtual Ptr<ImageSizeExpression> GetImageSizeExpression() const override;
-		virtual void GetLayoutBlockSize(int* pBlockX, int* pBlockY) override;
+		virtual void GetLayoutBlockSize(int* OutBlockX, int* OutBlockY) override;
 		//TODO: virtual bool IsImagePlainConstant(FVector4f& colour) const override;
-		virtual Ptr<ASTOp> OptimiseSemantic(const FModelOptimizationOptions& options, int32 Pass) const override;
-		virtual Ptr<ASTOp> OptimiseSink(const FModelOptimizationOptions& options, FOptimizeSinkContext& context) const override;
+		virtual Ptr<ASTOp> OptimiseSemantic(const FModelOptimizationOptions&, int32 Pass) const override;
+		virtual Ptr<ASTOp> OptimiseSink(const FModelOptimizationOptions&, FOptimizeSinkContext&) const override;
 		virtual FSourceDataDescriptor GetSourceDataDescriptor(FGetSourceDataDescriptorContext*) const override;
 
 	};
diff --git a/Engine/Plugins/Mutable/Source/MutableTools/Private/MuT/ASTOpImageDisplace.cpp b/Engine/Plugins/Mutable/Source/MutableTools/Private/MuT/ASTOpImageDisplace.cpp
new file mode 100644
--- /dev/null
+++ b/Engine/Plugins/Mutable/Source/MutableTools/Private/MuT/ASTOpImageDisplace.cpp
@@ -0,0 +1,277 @@
+// Copyright Epic Games, Inc. All Rights Reserved.
+
+#include "MuT/ASTOpImageDisplace.h"
+
+#include "Containers/Map.h"
+#include "MuR/ModelPrivate.h"
+#include "MuR/RefCounted.h"
+#include "MuR/Types.h"
+#include "MuR/ImagePrivate.h"
+#include "MuT/ASTOpConditional.h"
+#include "MuT/ASTOpSwitch.h"
+#include "MuT/ASTOpImagePatch.h"
+#include "MuT/ASTOpImagePixelFormat.h"
+#include "MuT/ASTOpImageLayer.h"
+#include "MuT/ASTOpImageLayerColor.h"
+#include "MuT/ASTOpImageRasterMesh.h"
+#include "MuT/ASTOpImageTransform.h"
+#include "MuT/ASTOpImageSwizzle.h"
+#include "MuT/ASTOpImageCompose.h"
+#include "MuT/ASTOpImageMakeGrowMap.h"
+#include "MuT/ASTOpImageResize.h"
+
+
+namespace mu
+{
+
+	ASTOpImageDisplace::ASTOpImageDisplace()
+		: Source(this),
+		DisplacementMap(this)
+	{
+	}
+
+
+	ASTOpImageDisplace::~ASTOpImageDisplace()
+	{
+		// Explicit call needed to avoid recursive destruction
+		ASTOp::RemoveChildren();
+	}
+
+
+	bool ASTOpImageDisplace::IsEqual(const ASTOp& InOther) const
+	{
+		if (InOther.GetOpType()==GetOpType())
+		{
+			const ASTOpImageDisplace* Other = static_cast<const ASTOpImageDisplace*>(&InOther);
+			return Source == Other->Source &&
+				DisplacementMap == Other->DisplacementMap;
+		}
+		return false;
+	}
+
+
+	uint64 ASTOpImageDisplace::Hash() const
+	{
+		uint64 Res = std::hash<EOpType>()(GetOpType());
+		hash_combine(Res, Source.child().get());
+		hash_combine(Res, DisplacementMap.child().get());
+		return Res;
+	}
+
+
+	Ptr<ASTOp> ASTOpImageDisplace::Clone(MapChildFuncRef MapChild) const
+	{
+		Ptr<ASTOpImageDisplace> New = new ASTOpImageDisplace();
+		New->Source = MapChild(Source.child());
+		New->DisplacementMap = MapChild(DisplacementMap.child());
+		return New;
+	}
+
+
+	void ASTOpImageDisplace::ForEachChild(const TFunctionRef<void(ASTChild&)> Func)
+	{
+		Func(Source);
+		Func(DisplacementMap);
+	}
+
+
+	void ASTOpImageDisplace::Link(FProgram& Program, FLinkerOptions*)
+	{
+		// Already linked?
+		if (!linkedAddress)
+		{
+			OP::ImageDisplaceArgs Args = {};
+
+			if (Source)
+			{
+				Args.Source = Source->linkedAddress;
+			}
+			if (DisplacementMap)
+			{
+				Args.DisplacementMap = DisplacementMap->linkedAddress;
+			}
+
+			linkedAddress = (OP::ADDRESS)Program.OpAddress.Num();
+			Program.OpAddress.Add(Program.ByteCode.Num());
+			AppendCode(Program.ByteCode, GetOpType());
+			AppendCode(Program.ByteCode, Args);
+		}
+	}
+
+
+	FImageDesc ASTOpImageDisplace::GetImageDesc(bool bReturnBestOption, FGetImageDescContext* Context) const
+	{
+		FImageDesc Result;
+
+		// Local context in case it is necessary
+		FGetImageDescContext LocalContext;
+		if (!Context)
+		{
+			Context = &LocalContext;
+		}
+		else
+		{
+			// Cached result?
+			FImageDesc* PtrValue = Context->m_results.Find(this);
+			if (PtrValue)
+			{
+				return *PtrValue;
+			}
+		}
+
+		// Actual work
+		if (Source)
+		{
+			Result = Source->GetImageDesc(bReturnBestOption, Context);
+		}
+
+		if (DisplacementMap)
+		{
+			FImageDesc SizeResult = DisplacementMap->GetImageDesc(bReturnBestOption, Context);
+			Result.m_size = SizeResult.m_size;
+		}
+
+		// Cache the result
+		if (Context)
+		{
+			Context->m_results.Add(this, Result);
+		}
+
+		return Result;
+	}
+
+
+	Ptr<ImageSizeExpression> ASTOpImageDisplace::GetImageSizeExpression() const
+	{
+		if (Source)
+		{
+			return Source->GetImageSizeExpression();
+		}
+
+		return new ImageSizeExpression;
+	}
+
+
+	FSourceDataDescriptor ASTOpImageDisplace::GetSourceDataDescriptor(FGetSourceDataDescriptorContext* Context) const
+	{
+		if (Source)
+		{
+			return Source->GetSourceDataDescriptor(Context);
+		}
+
+		return {};
+	}
+
+
+	Ptr<ASTOp> ASTOpImageDisplace::OptimiseSink(const FModelOptimizationOptions&, FOptimizeSinkContext& context) const
+	{
+		Ptr<ASTOp> at;
+
+		Ptr<ASTOp> OriginalAt = at;
+		Ptr<ASTOp> sourceAt = Source.child();
+		Ptr<ASTOp> displaceMapAt = DisplacementMap.child();
+
+		switch (sourceAt->GetOpType())
+		{
+
+		case EOpType::IM_CONDITIONAL:
+		{
+			if (displaceMapAt->GetOpType() == EOpType::IM_CONDITIONAL)
+			{
+				const ASTOpConditional* typedSource = static_cast<const ASTOpConditional*>(sourceAt.get());
+				const ASTOpConditional* typedDisplacementMap = static_cast<const ASTOpConditional*>(displaceMapAt.get());
+
+				if (typedSource->condition == typedDisplacementMap->condition)
+				{
+					Ptr<ASTOpConditional> nop = mu::Clone<ASTOpConditional>(sourceAt);
+
+					Ptr<ASTOpImageDisplace> aOp = mu::Clone<ASTOpImageDisplace>(this);
+					aOp->Source = typedSource->yes.child();
+					aOp->DisplacementMap = typedDisplacementMap->yes.child();
+					nop->yes = aOp;
+
+					Ptr<ASTOpImageDisplace> bOp = mu::Clone<ASTOpImageDisplace>(this);
+					bOp->Source = typedSource->no.child();
+					bOp->DisplacementMap = typedDisplacementMap->no.child();
+					nop->no = bOp;
+
+					at = nop;
+				}
+			}
+			break;
+		}
+
+		case EOpType::IM_SWITCH:
+		{
+			if (displaceMapAt->GetOpType() == EOpType::IM_SWITCH)
+			{
+				const ASTOpSwitch* typedSource = static_cast<const ASTOpSwitch*>(sourceAt.get());
+				const ASTOpSwitch* typedDisplacementMap = static_cast<const ASTOpSwitch*>(displaceMapAt.get());
+
+				if (typedSource->IsCompatibleWith(typedDisplacementMap))
+				{
+					// Move the format down all the paths
+					auto nop = mu::Clone<ASTOpSwitch>(sourceAt);
+
+					if (nop->def)
+					{
+						Ptr<ASTOpImageDisplace> defOp = mu::Clone<ASTOpImageDisplace>(this);
+						defOp->Source = typedSource->def.child();
+						defOp->DisplacementMap = typedDisplacementMap->def.child();
+						nop->def = defOp;
+					}
+
+					// We need to copy the options because we change them
+					for (size_t v = 0; v < nop->cases.Num(); ++v)
+					{
+						if (nop->cases[v].branch)
+						{
+							Ptr<ASTOpImageDisplace> bOp = mu::Clone<ASTOpImageDisplace>(this);
+							bOp->Source = typedSource->cases[v].branch.child();
+							bOp->DisplacementMap = typedDisplacementMap->FindBranch(typedSource->cases[v].condition);
+							nop->cases[v].branch = bOp;
+						}
+					}
+
+					at = nop;
+				}
+			}
+
+			// If we didn't optimize already, try to simply sink the source.
+			if (OriginalAt == at)
+			{
+				// Make a project for every path
+				Ptr<ASTOpSwitch> nop = mu::Clone<ASTOpSwitch>(sourceAt.get());
+
+				if (nop->def)
+				{
+					Ptr<ASTOpImageDisplace> defOp = mu::Clone<ASTOpImageDisplace>(this);
+					defOp->Source = nop->def.child();
+					nop->def = defOp;
+				}
+
+				// We need to copy the options because we change them
+				for (size_t o = 0; o < nop->cases.Num(); ++o)
+				{
+					if (nop->cases[o].branch)
+					{
+						Ptr<ASTOpImageDisplace> bOp = mu::Clone<ASTOpImageDisplace>(this);
+						bOp->Source = nop->cases[o].branch.child();
+						nop->cases[o].branch = bOp;
+					}
+				}
+
+				at = nop;
+			}
+
+			break;
+		}
+
+		default:
+			break;
+		}
+
+		return at;
+	}
+
+}
diff --git a/Engine/Plugins/Mutable/Source/MutableTools/Private/MuT/ASTOpImageDisplace.h b/Engine/Plugins/Mutable/Source/MutableTools/Private/MuT/ASTOpImageDisplace.h
new file mode 100644
--- /dev/null
+++ b/Engine/Plugins/Mutable/Source/MutableTools/Private/MuT/ASTOpImageDisplace.h
@@ -0,0 +1,45 @@
+// Copyright Epic Games, Inc. All Rights Reserved.
+
+#pragma once
+
+#include "HAL/PlatformMath.h"
+#include "MuR/Image.h"
+#include "MuR/MutableMath.h"
+#include "MuR/Operations.h"
+#include "MuR/Ptr.h"
+#include "MuT/AST.h"
+
+
+namespace mu
+{
+	struct FProgram;
+
+	class ASTOpImageDisplace final : public ASTOp
+	{
+	public:
+
+		ASTChild Source;
+		ASTChild DisplacementMap;
+
+	public:
+
+		ASTOpImageDisplace();
+		ASTOpImageDisplace(const ASTOpImageDisplace&) = delete;
+		~ASTOpImageDisplace();
+
+		virtual EOpType GetOpType() const override { return EOpType::IM_DISPLACE; }
+		virtual uint64 Hash() const override;
+		virtual bool IsEqual(const ASTOp&) const override;
+		virtual Ptr<ASTOp> Clone(MapChildFuncRef) const override;
+		virtual void ForEachChild(const TFunctionRef<void(ASTChild&)>) override;
+		virtual void Link(FProgram&, FLinkerOptions*) override;
+		virtual FImageDesc GetImageDesc(bool bReturnBestOption, FGetImageDescContext*) const override;
+		virtual Ptr<ImageSizeExpression> GetImageSizeExpression() const override;
+		virtual FSourceDataDescriptor GetSourceDataDescriptor(FGetSourceDataDescriptorContext*) const override;
+		virtual Ptr<ASTOp> OptimiseSink(const FModelOptimizationOptions&, FOptimizeSinkContext&) const;
+
+	};
+
+
+}
+
diff --git a/Engine/Plugins/Mutable/Source/MutableTools/Private/MuT/ASTOpImageInterpolate.cpp b/Engine/Plugins/Mutable/Source/MutableTools/Private/MuT/ASTOpImageInterpolate.cpp
new file mode 100644
--- /dev/null
+++ b/Engine/Plugins/Mutable/Source/MutableTools/Private/MuT/ASTOpImageInterpolate.cpp
@@ -0,0 +1,217 @@
+// Copyright Epic Games, Inc. All Rights Reserved.
+
+#include "MuT/ASTOpImageInterpolate.h"
+
+#include "Containers/Map.h"
+#include "HAL/UnrealMemory.h"
+#include "Math/UnrealMathSSE.h"
+#include "Misc/AssertionMacros.h"
+#include "MuR/ImagePrivate.h"
+#include "MuR/ModelPrivate.h"
+#include "MuR/MutableMath.h"
+#include "MuR/RefCounted.h"
+#include "MuR/Types.h"
+#include "MuT/ASTOpConditional.h"
+#include "MuT/ASTOpImageCompose.h"
+#include "MuT/ASTOpImagePatch.h"
+#include "MuT/ASTOpImagePixelFormat.h"
+#include "MuT/ASTOpImageMultiLayer.h"
+#include "MuT/ASTOpImageLayer.h"
+#include "MuT/ASTOpImageLayerColor.h"
+#include "MuT/ASTOpImageRasterMesh.h"
+#include "MuT/ASTOpImageTransform.h"
+#include "MuT/ASTOpImageMipmap.h"
+#include "MuT/ASTOpSwitch.h"
+
+namespace mu
+{
+
+	ASTOpImageInterpolate::ASTOpImageInterpolate()
+		: Factor(this), Targets{ ASTChild(this),ASTChild(this),ASTChild(this),ASTChild(this),ASTChild(this),ASTChild(this) }
+	{
+	}
+
+
+	ASTOpImageInterpolate::~ASTOpImageInterpolate()
+	{
+		// Explicit call needed to avoid recursive destruction
+		ASTOp::RemoveChildren();
+	}
+
+
+	bool ASTOpImageInterpolate::IsEqual(const ASTOp& OtherUntyped) const
+	{
+		if (OtherUntyped.GetOpType()==GetOpType())
+		{
+			const ASTOpImageInterpolate* Other = static_cast<const ASTOpImageInterpolate*>(&OtherUntyped);
+			for (int32 i = 0; i<MUTABLE_OP_MAX_INTERPOLATE_COUNT; ++i)
+			{
+				if (!(Targets[i] == Other->Targets[i]))
+				{
+					return false;
+				}
+			}
+			return Factor == Other->Factor;
+		}
+		return false;
+	}
+
+
+	uint64 ASTOpImageInterpolate::Hash() const
+	{
+		uint64 res = std::hash<void*>()(Targets[0].child().get());
+		hash_combine(res, std::hash<void*>()(Targets[1].child().get()));
+		hash_combine(res, std::hash<void*>()(Targets[2].child().get()));
+		hash_combine(res, std::hash<void*>()(Targets[3].child().get()));
+		hash_combine(res, std::hash<void*>()(Targets[4].child().get()));
+		hash_combine(res, std::hash<void*>()(Targets[5].child().get()));
+		hash_combine(res, std::hash<void*>()(Factor.child().get()));
+		return res;
+	}
+
+
+	mu::Ptr<ASTOp> ASTOpImageInterpolate::Clone(MapChildFuncRef MapChild) const
+	{
+		mu::Ptr<ASTOpImageInterpolate> n = new ASTOpImageInterpolate();
+		n->Factor = MapChild(Factor.child());
+		for (int32 i = 0; i < MUTABLE_OP_MAX_INTERPOLATE_COUNT; ++i)
+		{
+			n->Targets[i] = MapChild(Targets[i].child());
+		}
+		n->Factor = Factor.child();
+		return n;
+	}
+
+
+	void ASTOpImageInterpolate::ForEachChild(const TFunctionRef<void(ASTChild&)> Func)
+	{
+		Func(Factor);
+		for (int32 i = 0; i < MUTABLE_OP_MAX_INTERPOLATE_COUNT; ++i)
+		{
+			Func(Targets[i]);
+		}
+	}
+
+
+	void ASTOpImageInterpolate::Link(FProgram& program, FLinkerOptions*)
+	{
+		// Already linked?
+		if (!linkedAddress)
+		{
+			OP::ImageInterpolateArgs Args = {};
+
+			if (Factor)
+			{
+				Args.Factor = Factor->linkedAddress;
+			}
+			for (int32 i = 0; i < MUTABLE_OP_MAX_INTERPOLATE_COUNT; ++i)
+			{
+				if (Targets[i]) Args.Targets[i] = Targets[i]->linkedAddress;
+			}			
+
+			linkedAddress = (OP::ADDRESS)program.OpAddress.Num();
+			program.OpAddress.Add((uint32)program.ByteCode.Num());
+			AppendCode(program.ByteCode, GetOpType());
+			AppendCode(program.ByteCode, Args);
+		}
+	}
+
+
+	//!
+	FImageDesc ASTOpImageInterpolate::GetImageDesc(bool bReturnBestOption, FGetImageDescContext* Context) const
+	{
+		FImageDesc res;
+
+		// Local context in case it is necessary
+		FGetImageDescContext localContext;
+		if (!Context)
+		{
+			Context = &localContext;
+		}
+		else
+		{
+			// Cached result?
+			FImageDesc* PtrValue = Context->m_results.Find(this);
+			if (PtrValue)
+			{
+				return *PtrValue;
+			}
+		}
+
+		if (Targets[0])
+		{
+			res = Targets[0]->GetImageDesc(bReturnBestOption, Context);
+		}
+
+		// Cache the result
+		if (Context)
+		{
+			Context->m_results.Add(this, res);
+		}
+
+		return res;
+	}
+
+
+	void ASTOpImageInterpolate::GetLayoutBlockSize(int32* pBlockX, int32* pBlockY)
+	{
+		if (Targets[0].child())
+		{
+			// Assume the block size of the biggest mip
+			Targets[0].child()->GetLayoutBlockSize(pBlockX, pBlockY);
+		}
+	}
+
+
+	mu::Ptr<ImageSizeExpression> ASTOpImageInterpolate::GetImageSizeExpression() const
+	{
+		mu::Ptr<ImageSizeExpression> pRes;
+
+		if (Targets[0].child())
+		{
+			pRes = Targets[0].child()->GetImageSizeExpression();
+		}
+		else
+		{
+			pRes = new ImageSizeExpression;
+		}
+
+		return pRes;
+	}
+
+
+	FSourceDataDescriptor ASTOpImageInterpolate::GetSourceDataDescriptor(FGetSourceDataDescriptorContext* Context) const
+	{
+		// Cache management
+		TUniquePtr<FGetSourceDataDescriptorContext> LocalContext;
+		if (!Context)
+		{
+			LocalContext.Reset(new FGetSourceDataDescriptorContext);
+			Context = LocalContext.Get();
+		}
+
+		FSourceDataDescriptor* Found= Context->Cache.Find(this);
+		if (Found)
+		{
+			return *Found;
+		}
+
+		// Not cached: calculate
+		FSourceDataDescriptor Result;
+
+		for( int32 TargetIndex=0; TargetIndex<MUTABLE_OP_MAX_INTERPOLATE_COUNT; ++TargetIndex)
+		{ 
+			if (Targets[TargetIndex])
+			{
+				FSourceDataDescriptor TargetDesc = Targets[TargetIndex]->GetSourceDataDescriptor(Context);
+				Result.CombineWith(TargetDesc);
+			}
+		}
+
+		Context->Cache.Add(this,Result);
+
+		return Result;
+	}
+
+
+}
diff --git a/Engine/Plugins/Mutable/Source/MutableTools/Private/MuT/ASTOpImageInterpolate.h b/Engine/Plugins/Mutable/Source/MutableTools/Private/MuT/ASTOpImageInterpolate.h
new file mode 100644
--- /dev/null
+++ b/Engine/Plugins/Mutable/Source/MutableTools/Private/MuT/ASTOpImageInterpolate.h
@@ -0,0 +1,43 @@
+// Copyright Epic Games, Inc. All Rights Reserved.
+
+#pragma once
+
+#include "HAL/PlatformMath.h"
+#include "MuR/Image.h"
+#include "MuR/Operations.h"
+#include "MuR/Ptr.h"
+#include "MuT/AST.h"
+
+
+namespace mu
+{
+	struct FProgram;
+
+	class ASTOpImageInterpolate final : public ASTOp
+	{
+	public:
+
+		ASTChild Factor;
+		ASTChild Targets[MUTABLE_OP_MAX_INTERPOLATE_COUNT];
+
+	public:
+
+		ASTOpImageInterpolate();
+		ASTOpImageInterpolate(const ASTOpImageInterpolate&) = delete;
+		~ASTOpImageInterpolate();
+
+		virtual EOpType GetOpType() const override { return EOpType::IM_INTERPOLATE; }
+		virtual uint64 Hash() const override;
+		virtual bool IsEqual(const ASTOp&) const override;
+		virtual Ptr<ASTOp> Clone(MapChildFuncRef) const override;
+		virtual void ForEachChild(const TFunctionRef<void(ASTChild&)>) override;
+		virtual void Link(FProgram&, FLinkerOptions*) override;
+		virtual FImageDesc GetImageDesc(bool bReturnBestOption, FGetImageDescContext*) const override;
+		virtual void GetLayoutBlockSize(int32* OutBlockX, int32* OutBlockY) override;
+		virtual Ptr<ImageSizeExpression> GetImageSizeExpression() const override;
+		virtual FSourceDataDescriptor GetSourceDataDescriptor(FGetSourceDataDescriptorContext*) const override;
+
+	};
+
+}
+
diff --git a/Engine/Plugins/Mutable/Source/MutableTools/Private/MuT/ASTOpImageInvert.cpp b/Engine/Plugins/Mutable/Source/MutableTools/Private/MuT/ASTOpImageInvert.cpp
new file mode 100644
--- /dev/null
+++ b/Engine/Plugins/Mutable/Source/MutableTools/Private/MuT/ASTOpImageInvert.cpp
@@ -0,0 +1,170 @@
+// Copyright Epic Games, Inc. All Rights Reserved.
+
+#include "MuT/ASTOpImageInvert.h"
+
+#include "Containers/Map.h"
+#include "MuR/ModelPrivate.h"
+#include "MuR/RefCounted.h"
+#include "MuR/Types.h"
+#include "MuR/ImagePrivate.h"
+#include "MuT/ASTOpConditional.h"
+#include "MuT/ASTOpSwitch.h"
+#include "MuT/ASTOpImagePatch.h"
+#include "MuT/ASTOpImagePixelFormat.h"
+#include "MuT/ASTOpImageLayer.h"
+#include "MuT/ASTOpImageLayerColor.h"
+#include "MuT/ASTOpImageRasterMesh.h"
+#include "MuT/ASTOpImageTransform.h"
+#include "MuT/ASTOpImageSwizzle.h"
+#include "MuT/ASTOpImageCompose.h"
+#include "MuT/ASTOpImageMakeGrowMap.h"
+#include "MuT/ASTOpImageResize.h"
+
+
+namespace mu
+{
+
+	ASTOpImageInvert::ASTOpImageInvert()
+		: Base(this)
+	{
+	}
+
+
+	ASTOpImageInvert::~ASTOpImageInvert()
+	{
+		// Explicit call needed to avoid recursive destruction
+		ASTOp::RemoveChildren();
+	}
+
+
+	bool ASTOpImageInvert::IsEqual(const ASTOp& InOther) const
+	{
+		if (InOther.GetOpType()==GetOpType())
+		{
+			const ASTOpImageInvert* Other = static_cast<const ASTOpImageInvert*>(&InOther);
+			return Base == Other->Base;
+		}
+		return false;
+	}
+
+
+	uint64 ASTOpImageInvert::Hash() const
+	{
+		uint64 Res = std::hash<EOpType>()(GetOpType());
+		hash_combine(Res, Base.child().get());
+		return Res;
+	}
+
+
+	Ptr<ASTOp> ASTOpImageInvert::Clone(MapChildFuncRef MapChild) const
+	{
+		Ptr<ASTOpImageInvert> New = new ASTOpImageInvert();
+		New->Base = MapChild(Base.child());
+		return New;
+	}
+
+
+	void ASTOpImageInvert::ForEachChild(const TFunctionRef<void(ASTChild&)> Func)
+	{
+		Func(Base);
+	}
+
+
+	void ASTOpImageInvert::Link(FProgram& Program, FLinkerOptions*)
+	{
+		// Already linked?
+		if (!linkedAddress)
+		{
+			OP::ImageInvertArgs Args = {};
+
+			if (Base)
+			{
+				Args.Base = Base->linkedAddress;
+			}
+
+			linkedAddress = (OP::ADDRESS)Program.OpAddress.Num();
+			Program.OpAddress.Add(Program.ByteCode.Num());
+			AppendCode(Program.ByteCode, GetOpType());
+			AppendCode(Program.ByteCode, Args);
+		}
+	}
+
+
+	FImageDesc ASTOpImageInvert::GetImageDesc(bool bReturnBestOption, FGetImageDescContext* Context) const
+	{
+		FImageDesc Result;
+
+		// Local context in case it is necessary
+		FGetImageDescContext LocalContext;
+		if (!Context)
+		{
+			Context = &LocalContext;
+		}
+		else
+		{
+			// Cached result?
+			FImageDesc* PtrValue = Context->m_results.Find(this);
+			if (PtrValue)
+			{
+				return *PtrValue;
+			}
+		}
+
+		// Actual work
+		if (Base)
+		{
+			Result = Base->GetImageDesc(bReturnBestOption, Context);
+		}
+
+		// Cache the result
+		if (Context)
+		{
+			Context->m_results.Add(this, Result);
+		}
+
+		return Result;
+	}
+
+
+	Ptr<ImageSizeExpression> ASTOpImageInvert::GetImageSizeExpression() const
+	{
+		if (Base)
+		{
+			return Base->GetImageSizeExpression();
+		}
+
+		return new ImageSizeExpression;
+	}
+
+
+	bool ASTOpImageInvert::IsImagePlainConstant(FVector4f& OutColour) const
+	{
+		bool bResult = true;
+		OutColour = FVector4f(1.0f,1.0f,1.0f,1.0f);
+
+		if (Base)
+		{
+			bResult = Base->IsImagePlainConstant(OutColour);
+			if (bResult)
+			{
+				OutColour[0] = 1.0f - OutColour[0];
+				OutColour[1] = 1.0f - OutColour[1];
+				OutColour[2] = 1.0f - OutColour[2];
+			}
+		}
+
+		return bResult;
+	}
+
+
+	FSourceDataDescriptor ASTOpImageInvert::GetSourceDataDescriptor(FGetSourceDataDescriptorContext* Context) const
+	{
+		if (Base)
+		{
+			return Base->GetSourceDataDescriptor(Context);
+		}
+
+		return {};
+	}
+
+}
diff --git a/Engine/Plugins/Mutable/Source/MutableTools/Private/MuT/ASTOpImageInvert.h b/Engine/Plugins/Mutable/Source/MutableTools/Private/MuT/ASTOpImageInvert.h
new file mode 100644
--- /dev/null
+++ b/Engine/Plugins/Mutable/Source/MutableTools/Private/MuT/ASTOpImageInvert.h
@@ -0,0 +1,44 @@
+// Copyright Epic Games, Inc. All Rights Reserved.
+
+#pragma once
+
+#include "HAL/PlatformMath.h"
+#include "MuR/Image.h"
+#include "MuR/MutableMath.h"
+#include "MuR/Operations.h"
+#include "MuR/Ptr.h"
+#include "MuT/AST.h"
+
+
+namespace mu
+{
+	struct FProgram;
+
+	class ASTOpImageInvert final : public ASTOp
+	{
+	public:
+
+		ASTChild Base;
+
+	public:
+
+		ASTOpImageInvert();
+		ASTOpImageInvert(const ASTOpImageInvert&) = delete;
+		~ASTOpImageInvert();
+
+		virtual EOpType GetOpType() const override { return EOpType::IM_INVERT; }
+		virtual uint64 Hash() const override;
+		virtual bool IsEqual(const ASTOp&) const override;
+		virtual Ptr<ASTOp> Clone(MapChildFuncRef) const override;
+		virtual void ForEachChild(const TFunctionRef<void(ASTChild&)>) override;
+		virtual void Link(FProgram&, FLinkerOptions*) override;
+		virtual FImageDesc GetImageDesc(bool bReturnBestOption, FGetImageDescContext*) const override;
+		virtual Ptr<ImageSizeExpression> GetImageSizeExpression() const override;
+		virtual bool IsImagePlainConstant(FVector4f& colour) const override;
+		virtual FSourceDataDescriptor GetSourceDataDescriptor(FGetSourceDataDescriptorContext*) const override;
+
+	};
+
+
+}
+
diff --git a/Engine/Plugins/Mutable/Source/MutableTools/Private/MuT/ASTOpImageLayer.cpp b/Engine/Plugins/Mutable/Source/MutableTools/Private/MuT/ASTOpImageLayer.cpp
--- a/Engine/Plugins/Mutable/Source/MutableTools/Private/MuT/ASTOpImageLayer.cpp
+++ b/Engine/Plugins/Mutable/Source/MutableTools/Private/MuT/ASTOpImageLayer.cpp
@@ -7,6 +7,9 @@
 #include "MuT/ASTOpImageSwizzle.h"
 #include "MuT/ASTOpImageRasterMesh.h"
 #include "MuT/ASTOpImageCrop.h"
+#include "MuT/ASTOpImageResize.h"
+#include "MuT/ASTOpImagePlainColor.h"
+#include "MuT/ASTOpImageDisplace.h"
 #include "MuT/ASTOpSwitch.h"
 #include "MuR/ModelPrivate.h"
 #include "MuR/RefCounted.h"
@@ -196,7 +199,7 @@ namespace mu
 			// TODO: May some flags be supported?
 			if (Flags == 0)
 			{
-				const ASTOpFixed* BlendPlainColor = static_cast<const ASTOpFixed*>(blendAt.get());
+				const ASTOpImagePlainColor* BlendPlainColor = static_cast<const ASTOpImagePlainColor*>(blendAt.get());
 
 				Ptr<ASTOpImageLayerColor> NewLayerColor = new ASTOpImageLayerColor;
 				NewLayerColor->base = baseAt;
@@ -204,7 +207,7 @@ namespace mu
 				NewLayerColor->blendType = blendType;
 				NewLayerColor->blendTypeAlpha = blendTypeAlpha;
 				NewLayerColor->BlendAlphaSourceChannel = BlendAlphaSourceChannel;
-				NewLayerColor->color = BlendPlainColor->children[BlendPlainColor->op.args.ImagePlainColour.colour].child();
+				NewLayerColor->color = BlendPlainColor->Color.child();
 				at = NewLayerColor;
 			}
 		}
@@ -298,16 +301,16 @@ namespace mu
 
 				case EOpType::IM_DISPLACE:
 				{
-					const ASTOpFixed* MaskDisplace = static_cast<const ASTOpFixed*>(CurrentMask.get());
-					const ASTOpFixed* BlendDisplace = static_cast<const ASTOpFixed*>(CurrentBlend.get());
+					const ASTOpImageDisplace* MaskDisplace = static_cast<const ASTOpImageDisplace*>(CurrentMask.get());
+					const ASTOpImageDisplace* BlendDisplace = static_cast<const ASTOpImageDisplace*>(CurrentBlend.get());
 					if (MaskDisplace && BlendDisplace 
 						&&
-						MaskDisplace->children[MaskDisplace->op.args.ImageDisplace.displacementMap].child()
+						MaskDisplace->DisplacementMap.child()
 						==
-						BlendDisplace->children[BlendDisplace->op.args.ImageDisplace.displacementMap].child())
+						BlendDisplace->DisplacementMap.child())
 					{
-						CurrentMask = MaskDisplace->children[MaskDisplace->op.args.ImageDisplace.source].child();
-						CurrentBlend = BlendDisplace->children[BlendDisplace->op.args.ImageDisplace.source].child();
+						CurrentMask = MaskDisplace->Source.child();
+						CurrentBlend = BlendDisplace->Source.child();
 						bMatchingAlphaExpression = true;
 					}
 					break;
@@ -341,16 +344,16 @@ namespace mu
 
 				case EOpType::IM_RESIZE:
 				{
-					const ASTOpFixed* MaskResize = static_cast<const ASTOpFixed*>(CurrentMask.get());
-					const ASTOpFixed* BlendResize = static_cast<const ASTOpFixed*>(CurrentBlend.get());
+					const ASTOpImageResize* MaskResize = static_cast<const ASTOpImageResize*>(CurrentMask.get());
+					const ASTOpImageResize* BlendResize = static_cast<const ASTOpImageResize*>(CurrentBlend.get());
 					if (MaskResize && BlendResize 
 						&&
-						MaskResize->op.args.ImageResize.size[0] == BlendResize->op.args.ImageResize.size[0]
+						MaskResize->Size[0] == BlendResize->Size[0]
 						&&
-						MaskResize->op.args.ImageResize.size[1] == BlendResize->op.args.ImageResize.size[1])
+						MaskResize->Size[1] == BlendResize->Size[1])
 					{
-						CurrentMask = MaskResize->children[MaskResize->op.args.ImageResize.source].child();
-						CurrentBlend = BlendResize->children[BlendResize->op.args.ImageResize.source].child();
+						CurrentMask = MaskResize->Source.child();
+						CurrentBlend = BlendResize->Source.child();
 						bMatchingAlphaExpression = true;
 					}
 					break;
diff --git a/Engine/Plugins/Mutable/Source/MutableTools/Private/MuT/ASTOpImageLayerColor.cpp b/Engine/Plugins/Mutable/Source/MutableTools/Private/MuT/ASTOpImageLayerColor.cpp
--- a/Engine/Plugins/Mutable/Source/MutableTools/Private/MuT/ASTOpImageLayerColor.cpp
+++ b/Engine/Plugins/Mutable/Source/MutableTools/Private/MuT/ASTOpImageLayerColor.cpp
@@ -5,6 +5,7 @@
 #include "MuT/ASTOpSwitch.h"
 #include "MuT/ASTOpImageSwizzle.h"
 #include "MuT/ASTOpImageRasterMesh.h"
+#include "MuT/ASTOpImageDisplace.h"
 #include "Containers/Map.h"
 #include "HAL/PlatformMath.h"
 #include "MuR/ModelPrivate.h"
@@ -399,12 +400,12 @@ namespace mu
 				break;
 			}
 
-			Ptr<ASTOpFixed> NewDisplace = mu::Clone<ASTOpFixed>(baseAt);
+			Ptr<ASTOpImageDisplace> NewDisplace = mu::Clone<ASTOpImageDisplace>(baseAt);
 
-			Ptr<ASTOp> sourceOp = NewDisplace->children[NewDisplace->op.args.ImageDisplace.source].child();
+			Ptr<ASTOp> SourceOp = NewDisplace->Source.child();
 			Ptr<ASTOpImageLayerColor> NewSource = mu::Clone<ASTOpImageLayerColor>(this);
-			NewSource->base = sourceOp;
-			NewDisplace->SetChild(NewDisplace->op.args.ImageDisplace.source, NewSource);
+			NewSource->base = SourceOp;
+			NewDisplace->Source = NewSource;
 
 			at = NewDisplace;
 			break;
diff --git a/Engine/Plugins/Mutable/Source/MutableTools/Private/MuT/ASTOpImageLuminance.cpp b/Engine/Plugins/Mutable/Source/MutableTools/Private/MuT/ASTOpImageLuminance.cpp
new file mode 100644
--- /dev/null
+++ b/Engine/Plugins/Mutable/Source/MutableTools/Private/MuT/ASTOpImageLuminance.cpp
@@ -0,0 +1,170 @@
+// Copyright Epic Games, Inc. All Rights Reserved.
+
+#include "MuT/ASTOpImageLuminance.h"
+
+#include "Containers/Map.h"
+#include "MuR/ModelPrivate.h"
+#include "MuR/RefCounted.h"
+#include "MuR/Types.h"
+#include "MuR/ImagePrivate.h"
+#include "MuT/ASTOpConditional.h"
+#include "MuT/ASTOpSwitch.h"
+#include "MuT/ASTOpImagePatch.h"
+#include "MuT/ASTOpImagePixelFormat.h"
+#include "MuT/ASTOpImageLayer.h"
+#include "MuT/ASTOpImageLayerColor.h"
+#include "MuT/ASTOpImageRasterMesh.h"
+#include "MuT/ASTOpImageTransform.h"
+#include "MuT/ASTOpImageSwizzle.h"
+#include "MuT/ASTOpImageCompose.h"
+#include "MuT/ASTOpImageMakeGrowMap.h"
+#include "MuT/ASTOpImageResize.h"
+
+
+namespace mu
+{
+
+	ASTOpImageLuminance::ASTOpImageLuminance()
+		: Base(this)
+	{
+	}
+
+
+	ASTOpImageLuminance::~ASTOpImageLuminance()
+	{
+		// Explicit call needed to avoid recursive destruction
+		ASTOp::RemoveChildren();
+	}
+
+
+	bool ASTOpImageLuminance::IsEqual(const ASTOp& InOther) const
+	{
+		if (InOther.GetOpType()==GetOpType())
+		{
+			const ASTOpImageLuminance* Other = static_cast<const ASTOpImageLuminance*>(&InOther);
+			return Base == Other->Base;
+		}
+		return false;
+	}
+
+
+	uint64 ASTOpImageLuminance::Hash() const
+	{
+		uint64 Res = std::hash<EOpType>()(GetOpType());
+		hash_combine(Res, Base.child().get());
+		return Res;
+	}
+
+
+	Ptr<ASTOp> ASTOpImageLuminance::Clone(MapChildFuncRef MapChild) const
+	{
+		Ptr<ASTOpImageLuminance> New = new ASTOpImageLuminance();
+		New->Base = MapChild(Base.child());
+		return New;
+	}
+
+
+	void ASTOpImageLuminance::ForEachChild(const TFunctionRef<void(ASTChild&)> Func)
+	{
+		Func(Base);
+	}
+
+
+	void ASTOpImageLuminance::Link(FProgram& Program, FLinkerOptions*)
+	{
+		// Already linked?
+		if (!linkedAddress)
+		{
+			OP::ImageLuminanceArgs Args = {};
+
+			if (Base)
+			{
+				Args.Base = Base->linkedAddress;
+			}
+
+			linkedAddress = (OP::ADDRESS)Program.OpAddress.Num();
+			Program.OpAddress.Add(Program.ByteCode.Num());
+			AppendCode(Program.ByteCode, GetOpType());
+			AppendCode(Program.ByteCode, Args);
+		}
+	}
+
+
+	FImageDesc ASTOpImageLuminance::GetImageDesc(bool bReturnBestOption, FGetImageDescContext* Context) const
+	{
+		FImageDesc Result;
+
+		// Local context in case it is necessary
+		FGetImageDescContext LocalContext;
+		if (!Context)
+		{
+			Context = &LocalContext;
+		}
+		else
+		{
+			// Cached result?
+			FImageDesc* PtrValue = Context->m_results.Find(this);
+			if (PtrValue)
+			{
+				return *PtrValue;
+			}
+		}
+
+		// Actual work
+		if (Base)
+		{
+			Result = Base->GetImageDesc(bReturnBestOption, Context);
+		}
+
+		// Cache the result
+		if (Context)
+		{
+			Context->m_results.Add(this, Result);
+		}
+
+		return Result;
+	}
+
+
+	Ptr<ImageSizeExpression> ASTOpImageLuminance::GetImageSizeExpression() const
+	{
+		if (Base)
+		{
+			return Base->GetImageSizeExpression();
+		}
+
+		return new ImageSizeExpression;
+	}
+
+
+	bool ASTOpImageLuminance::IsImagePlainConstant(FVector4f& OutColour) const
+	{
+		bool bResult = true;
+		OutColour = FVector4f(0.0f,0.0f,0.0f,1.0f);
+
+		if (Base)
+		{
+			bResult = Base->IsImagePlainConstant(OutColour);
+			if (bResult)
+			{
+				float Luminance = OutColour[0] * 77.0f + OutColour[1] * 150.0f + OutColour[2] * 29.0f;
+				Luminance /= 255.0f;
+				OutColour = FVector4f(Luminance, Luminance, Luminance, OutColour[3]);
+			}
+		}
+
+		return bResult;
+	}
+
+
+	FSourceDataDescriptor ASTOpImageLuminance::GetSourceDataDescriptor(FGetSourceDataDescriptorContext* Context) const
+	{
+		if (Base)
+		{
+			return Base->GetSourceDataDescriptor(Context);
+		}
+
+		return {};
+	}
+
+}
diff --git a/Engine/Plugins/Mutable/Source/MutableTools/Private/MuT/ASTOpImageLuminance.h b/Engine/Plugins/Mutable/Source/MutableTools/Private/MuT/ASTOpImageLuminance.h
new file mode 100644
--- /dev/null
+++ b/Engine/Plugins/Mutable/Source/MutableTools/Private/MuT/ASTOpImageLuminance.h
@@ -0,0 +1,44 @@
+// Copyright Epic Games, Inc. All Rights Reserved.
+
+#pragma once
+
+#include "HAL/PlatformMath.h"
+#include "MuR/Image.h"
+#include "MuR/MutableMath.h"
+#include "MuR/Operations.h"
+#include "MuR/Ptr.h"
+#include "MuT/AST.h"
+
+
+namespace mu
+{
+	struct FProgram;
+
+	class ASTOpImageLuminance final : public ASTOp
+	{
+	public:
+
+		ASTChild Base;
+
+	public:
+
+		ASTOpImageLuminance();
+		ASTOpImageLuminance(const ASTOpImageLuminance&) = delete;
+		~ASTOpImageLuminance();
+
+		virtual EOpType GetOpType() const override { return EOpType::IM_LUMINANCE; }
+		virtual uint64 Hash() const override;
+		virtual bool IsEqual(const ASTOp&) const override;
+		virtual Ptr<ASTOp> Clone(MapChildFuncRef) const override;
+		virtual void ForEachChild(const TFunctionRef<void(ASTChild&)>) override;
+		virtual void Link(FProgram&, FLinkerOptions*) override;
+		virtual FImageDesc GetImageDesc(bool bReturnBestOption, FGetImageDescContext*) const override;
+		virtual Ptr<ImageSizeExpression> GetImageSizeExpression() const override;
+		virtual bool IsImagePlainConstant(FVector4f& colour) const override;
+		virtual FSourceDataDescriptor GetSourceDataDescriptor(FGetSourceDataDescriptorContext*) const override;
+
+	};
+
+
+}
+
diff --git a/Engine/Plugins/Mutable/Source/MutableTools/Private/MuT/ASTOpImageMipmap.cpp b/Engine/Plugins/Mutable/Source/MutableTools/Private/MuT/ASTOpImageMipmap.cpp
--- a/Engine/Plugins/Mutable/Source/MutableTools/Private/MuT/ASTOpImageMipmap.cpp
+++ b/Engine/Plugins/Mutable/Source/MutableTools/Private/MuT/ASTOpImageMipmap.cpp
@@ -15,6 +15,8 @@
 #include "MuT/ASTOpImageCompose.h"
 #include "MuT/ASTOpImagePatch.h"
 #include "MuT/ASTOpImagePixelFormat.h"
+#include "MuT/ASTOpImageBlankLayout.h"
+#include "MuT/ASTOpImagePlainColor.h"
 #include "MuT/ASTOpSwitch.h"
 
 
@@ -136,18 +138,18 @@ namespace mu
 		case EOpType::IM_BLANKLAYOUT:
 		{
 			// Set the mipmap generation on the blank layout operation
-			Ptr<ASTOpFixed> mop = mu::Clone<ASTOpFixed>(sourceAt);
-			mop->op.args.ImageBlankLayout.generateMipmaps = 1;    // true
-			mop->op.args.ImageBlankLayout.mipmapCount = Levels;
-			at = mop;
+			Ptr<ASTOpImageBlankLayout> NewOp = mu::Clone<ASTOpImageBlankLayout>(sourceAt);
+			NewOp->GenerateMipmaps = 1;    // true
+			NewOp->MipmapCount = Levels;
+			at = NewOp;
 			break;
 		}
 
 		case EOpType::IM_PLAINCOLOUR:
 		{
 			// Set the mipmap generation on the plaincolour operation
-			Ptr<ASTOpFixed> mop = mu::Clone<ASTOpFixed>(sourceAt);
-			mop->op.args.ImagePlainColour.LODs = Levels;
+			Ptr<ASTOpImagePlainColor> mop = mu::Clone<ASTOpImagePlainColor>(sourceAt);
+			mop->LODs = Levels;
 			at = mop;
 			break;
 		}
diff --git a/Engine/Plugins/Mutable/Source/MutableTools/Private/MuT/ASTOpImagePatch.h b/Engine/Plugins/Mutable/Source/MutableTools/Private/MuT/ASTOpImagePatch.h
--- a/Engine/Plugins/Mutable/Source/MutableTools/Private/MuT/ASTOpImagePatch.h
+++ b/Engine/Plugins/Mutable/Source/MutableTools/Private/MuT/ASTOpImagePatch.h
@@ -30,13 +30,13 @@ struct FProgram;
 
 		virtual EOpType GetOpType() const override { return EOpType::IM_PATCH; }
 		virtual uint64 Hash() const override;
-		virtual bool IsEqual(const ASTOp& otherUntyped) const override;
-		virtual Ptr<ASTOp> Clone(MapChildFuncRef mapChild) const override;
+		virtual bool IsEqual(const ASTOp&) const override;
+		virtual Ptr<ASTOp> Clone(MapChildFuncRef) const override;
 		virtual void ForEachChild(const TFunctionRef<void(ASTChild&)>) override;
-		virtual void Link(FProgram& program, FLinkerOptions* Options) override;
-		virtual FImageDesc GetImageDesc(bool returnBestOption, FGetImageDescContext* context) const override;
+		virtual void Link(FProgram&, FLinkerOptions*) override;
+		virtual FImageDesc GetImageDesc(bool bReturnBestOption, FGetImageDescContext*) const override;
 		virtual Ptr<ImageSizeExpression> GetImageSizeExpression() const override;
-		virtual void GetLayoutBlockSize(int* pBlockX, int* pBlockY) override;
+		virtual void GetLayoutBlockSize(int32* OutBlockX, int32* OutBlockY) override;
 		//TODO: virtual bool IsImagePlainConstant(FVector4f& colour) const override;
 		virtual FSourceDataDescriptor GetSourceDataDescriptor(FGetSourceDataDescriptorContext*) const override;
 
diff --git a/Engine/Plugins/Mutable/Source/MutableTools/Private/MuT/ASTOpImagePixelFormat.cpp b/Engine/Plugins/Mutable/Source/MutableTools/Private/MuT/ASTOpImagePixelFormat.cpp
--- a/Engine/Plugins/Mutable/Source/MutableTools/Private/MuT/ASTOpImagePixelFormat.cpp
+++ b/Engine/Plugins/Mutable/Source/MutableTools/Private/MuT/ASTOpImagePixelFormat.cpp
@@ -17,6 +17,11 @@
 #include "MuT/ASTOpImageLayer.h"
 #include "MuT/ASTOpImageLayerColor.h"
 #include "MuT/ASTOpImageRasterMesh.h"
+#include "MuT/ASTOpImageBlankLayout.h"
+#include "MuT/ASTOpImageInterpolate.h"
+#include "MuT/ASTOpImagePlainColor.h"
+#include "MuT/ASTOpImageDisplace.h"
+#include "MuT/ASTOpImageInvert.h"
 #include "MuT/ASTOpSwitch.h"
 #include "MuT/CompilerPrivate.h"
 
@@ -145,11 +150,11 @@ namespace mu
 			// This op doesn't support compressed formats
 			if (!bIsCompressedFormat)
 			{
-				mu::Ptr<ASTOpFixed> newOp = mu::Clone<ASTOpFixed>(sourceAt);
+				mu::Ptr<ASTOpImageDisplace> newOp = mu::Clone<ASTOpImageDisplace>(sourceAt);
 
 				mu::Ptr<ASTOpImagePixelFormat> fop = mu::Clone<ASTOpImagePixelFormat>(this);
-				fop->Source = newOp->children[newOp->op.args.ImageDisplace.source].child();
-				newOp->SetChild(newOp->op.args.ImageDisplace.source, fop);
+				fop->Source = newOp->Source.child();
+				newOp->Source = fop;
 
 				at = newOp;
 			}
@@ -161,11 +166,11 @@ namespace mu
 			// This op doesn't support compressed formats
 			if (!bIsCompressedFormat)
 			{
-				mu::Ptr<ASTOpFixed> newOp = mu::Clone<ASTOpFixed>(sourceAt);
+				mu::Ptr<ASTOpImageInvert> newOp = mu::Clone<ASTOpImageInvert>(sourceAt);
 
 				mu::Ptr<ASTOpImagePixelFormat> fop = mu::Clone<ASTOpImagePixelFormat>(this);
-				fop->Source = newOp->children[newOp->op.args.ImageInvert.base].child();
-				newOp->SetChild(newOp->op.args.ImageInvert.base, fop);
+				fop->Source = newOp->Base.child();
+				newOp->Base = fop;
 
 				at = newOp;
 			}
@@ -193,9 +198,9 @@ namespace mu
 		case EOpType::IM_BLANKLAYOUT:
 		{
 			// Just make sure the layout format is the right one and forget the op
-			auto nop = mu::Clone<ASTOpFixed>(sourceAt);
+			Ptr<ASTOpImageBlankLayout> NewOp = mu::Clone<ASTOpImageBlankLayout>(sourceAt);
 
-			EImageFormat layoutFormat = (EImageFormat)nop->op.args.ImageBlankLayout.format;
+			EImageFormat layoutFormat = NewOp->Format;
 			if (FormatIfAlpha != EImageFormat::None
 				&&
 				GetImageFormatData(layoutFormat).Channels > 3)
@@ -203,26 +208,26 @@ namespace mu
 				format = FormatIfAlpha;
 			}
 
-			nop->op.args.ImageBlankLayout.format = format;
-			at = nop;
+			NewOp->Format = format;
+			at = NewOp;
 			break;
 		}
 
 		case EOpType::IM_PLAINCOLOUR:
 		{
 			// Just make sure the format is the right one and forget the op
-			auto nop = mu::Clone<ASTOpFixed>(sourceAt);
+			Ptr<ASTOpImagePlainColor> NewOp = mu::Clone<ASTOpImagePlainColor>(sourceAt);
 
-			EImageFormat layoutFormat = (EImageFormat)nop->op.args.ImagePlainColour.format;
+			EImageFormat LayoutFormat = NewOp->Format;
 			if (FormatIfAlpha != EImageFormat::None
 				&&
-				GetImageFormatData(layoutFormat).Channels > 3)
+				GetImageFormatData(LayoutFormat).Channels > 3)
 			{
 				format = FormatIfAlpha;
 			}
 
-			nop->op.args.ImagePlainColour.format = format;
-			at = nop;
+			NewOp->Format = format;
+			at = NewOp;
 			break;
 		}
 
@@ -469,13 +474,13 @@ namespace mu
 			if (!bIsCompressedFormat)
 			{
 				// Move the format down all the paths
-				Ptr<ASTOpFixed> newOp = mu::Clone<ASTOpFixed>(at);
+				Ptr<ASTOpImageInterpolate> newOp = mu::Clone<ASTOpImageInterpolate>(at);
 
 				for (int v = 0; v < MUTABLE_OP_MAX_INTERPOLATE_COUNT; ++v)
 				{
-					Ptr<ASTOp> child = newOp->children[newOp->op.args.ImageInterpolate.targets[v]].child();
-					Ptr<ASTOp> bOp = Visit(child, currentFormatOp);
-					newOp->SetChild(newOp->op.args.ImageInterpolate.targets[v], bOp);
+					Ptr<ASTOp> Child = newOp->Targets[v].child();
+					Ptr<ASTOp> FormattedChild = Visit(Child, currentFormatOp);
+					newOp->Targets[v] = FormattedChild;
 				}
 
 				newAt = newOp;
diff --git a/Engine/Plugins/Mutable/Source/MutableTools/Private/MuT/ASTOpImagePlainColor.cpp b/Engine/Plugins/Mutable/Source/MutableTools/Private/MuT/ASTOpImagePlainColor.cpp
new file mode 100644
--- /dev/null
+++ b/Engine/Plugins/Mutable/Source/MutableTools/Private/MuT/ASTOpImagePlainColor.cpp
@@ -0,0 +1,178 @@
+// Copyright Epic Games, Inc. All Rights Reserved.
+
+#include "MuT/ASTOpImagePlainColor.h"
+
+#include "Containers/Map.h"
+#include "MuR/ModelPrivate.h"
+#include "MuR/RefCounted.h"
+#include "MuR/Types.h"
+#include "MuT/ASTOpConditional.h"
+#include "MuT/ASTOpSwitch.h"
+#include "MuT/ASTOpImagePatch.h"
+#include "MuT/ASTOpImagePixelFormat.h"
+#include "MuT/ASTOpImageLayer.h"
+#include "MuT/ASTOpImageLayerColor.h"
+#include "MuT/ASTOpImageRasterMesh.h"
+#include "MuT/ASTOpImageInterpolate.h"
+
+
+namespace mu
+{
+
+	ASTOpImagePlainColor::ASTOpImagePlainColor()
+		: Color(this)
+	{
+	}
+
+
+	ASTOpImagePlainColor::~ASTOpImagePlainColor()
+	{
+		// Explicit call needed to avoid recursive destruction
+		ASTOp::RemoveChildren();
+	}
+
+
+	bool ASTOpImagePlainColor::IsEqual(const ASTOp& InOther) const
+	{
+		if (InOther.GetOpType()==GetOpType())
+		{
+			const ASTOpImagePlainColor* other = static_cast<const ASTOpImagePlainColor*>(&InOther);
+			return Color == other->Color &&
+				Format == other->Format &&
+				Size == other->Size &&
+				LODs == other->LODs;
+		}
+		return false;
+	}
+
+
+	uint64 ASTOpImagePlainColor::Hash() const
+	{
+		uint64 res = std::hash<EOpType>()(GetOpType());
+		hash_combine(res, Color.child().get());
+		hash_combine(res, Size[0]);
+		hash_combine(res, Size[1]);
+		hash_combine(res, Format);
+		hash_combine(res, LODs);
+		return res;
+	}
+
+
+	Ptr<ASTOp> ASTOpImagePlainColor::Clone(MapChildFuncRef MapChild) const
+	{
+		Ptr<ASTOpImagePlainColor> New = new ASTOpImagePlainColor();
+		New->Color = MapChild(Color.child());
+		New->Format = Format;
+		New->Size = Size;
+		New->LODs = LODs;
+		return New;
+	}
+
+
+	void ASTOpImagePlainColor::ForEachChild(const TFunctionRef<void(ASTChild&)> Func)
+	{
+		Func(Color);
+	}
+
+
+	void ASTOpImagePlainColor::Link(FProgram& Program, FLinkerOptions*)
+	{
+		// Already linked?
+		if (!linkedAddress)
+		{
+			OP::ImagePlainColorArgs Args={};
+
+			if (Color)
+			{
+				Args.Color = Color->linkedAddress;
+			}
+
+			Args.Format = Format;
+			Args.Size[0] = Size[0];
+			Args.Size[1] = Size[1];
+			Args.LODs = LODs;
+
+			linkedAddress = (OP::ADDRESS)Program.OpAddress.Num();
+			Program.OpAddress.Add(Program.ByteCode.Num());
+			AppendCode(Program.ByteCode, GetOpType());
+			AppendCode(Program.ByteCode, Args);
+		}
+
+	}
+
+
+	FImageDesc ASTOpImagePlainColor::GetImageDesc(bool bReturnBestOption, FGetImageDescContext* Context) const
+	{
+		FImageDesc Result;
+
+		// Local context in case it is necessary
+		FGetImageDescContext LocalContext;
+		if (!Context)
+		{
+			Context = &LocalContext;
+		}
+		else
+		{
+			// Cached result?
+			FImageDesc* PtrValue = Context->m_results.Find(this);
+			if (PtrValue)
+			{
+				return *PtrValue;
+			}
+		}
+
+		// Actual work
+		Result.m_format = Format;
+		Result.m_size[0] = Size[0];
+		Result.m_size[1] = Size[1];
+		Result.m_lods = LODs;
+		check(Result.m_format != EImageFormat::None);
+
+		// Cache the result
+		if (Context)
+		{
+			Context->m_results.Add(this, Result);
+		}
+
+		return Result;
+	}
+
+
+	Ptr<ImageSizeExpression> ASTOpImagePlainColor::GetImageSizeExpression() const
+	{
+		Ptr<ImageSizeExpression> Res = new ImageSizeExpression;
+		Res->type = ImageSizeExpression::ISET_CONSTANT;
+		Res->size[0] = Size[0];
+		Res->size[1] = Size[1];
+		return Res;
+	}
+
+
+	void ASTOpImagePlainColor::GetLayoutBlockSize(int32* OutBlockX, int32* OutBlockY)
+	{
+		// We didn't find any layout yet.
+		*OutBlockX = 0;
+		*OutBlockY = 0;
+	}
+
+
+	bool ASTOpImagePlainColor::IsImagePlainConstant(FVector4f& OutColour) const
+	{
+		bool bResult = true;
+		OutColour = FVector4f( 0, 0, 0, 1 );
+
+		if (Color)
+		{
+			bResult = Color->IsColourConstant(OutColour);
+		}
+
+		return bResult;
+	}
+
+
+	FSourceDataDescriptor ASTOpImagePlainColor::GetSourceDataDescriptor(FGetSourceDataDescriptorContext* Context) const
+	{
+		return {};
+	}
+
+}
diff --git a/Engine/Plugins/Mutable/Source/MutableTools/Private/MuT/ASTOpImagePlainColor.h b/Engine/Plugins/Mutable/Source/MutableTools/Private/MuT/ASTOpImagePlainColor.h
new file mode 100644
--- /dev/null
+++ b/Engine/Plugins/Mutable/Source/MutableTools/Private/MuT/ASTOpImagePlainColor.h
@@ -0,0 +1,48 @@
+// Copyright Epic Games, Inc. All Rights Reserved.
+
+#pragma once
+
+#include "HAL/PlatformMath.h"
+#include "MuR/Image.h"
+#include "MuR/MutableMath.h"
+#include "MuR/Operations.h"
+#include "MuR/Ptr.h"
+#include "MuT/AST.h"
+
+
+namespace mu
+{
+	struct FProgram;
+
+	class ASTOpImagePlainColor final : public ASTOp
+	{
+	public:
+
+		ASTChild Color;
+		EImageFormat Format = EImageFormat::None;
+		UE::Math::TIntVector2<uint16> Size = UE::Math::TIntVector2<uint16>(0, 0);
+		uint8 LODs=1;
+
+	public:
+
+		ASTOpImagePlainColor();
+		ASTOpImagePlainColor(const ASTOpImagePlainColor&) = delete;
+		~ASTOpImagePlainColor();
+
+		virtual EOpType GetOpType() const override { return EOpType::IM_PLAINCOLOUR; }
+		virtual uint64 Hash() const override;
+		virtual bool IsEqual(const ASTOp&) const override;
+		virtual Ptr<ASTOp> Clone(MapChildFuncRef ) const override;
+		virtual void ForEachChild(const TFunctionRef<void(ASTChild&)>) override;
+		virtual void Link(FProgram&, FLinkerOptions*) override;
+		virtual FImageDesc GetImageDesc(bool bReturnBestOption, FGetImageDescContext*) const override;
+		virtual Ptr<ImageSizeExpression> GetImageSizeExpression() const override;
+		virtual void GetLayoutBlockSize(int32* OutBlockX, int32* OutBlockY) override;
+		virtual bool IsImagePlainConstant(FVector4f& OutColor) const override;
+		virtual FSourceDataDescriptor GetSourceDataDescriptor(FGetSourceDataDescriptorContext*) const override;
+
+	};
+
+
+}
+
diff --git a/Engine/Plugins/Mutable/Source/MutableTools/Private/MuT/ASTOpImageResize.cpp b/Engine/Plugins/Mutable/Source/MutableTools/Private/MuT/ASTOpImageResize.cpp
new file mode 100644
--- /dev/null
+++ b/Engine/Plugins/Mutable/Source/MutableTools/Private/MuT/ASTOpImageResize.cpp
@@ -0,0 +1,529 @@
+// Copyright Epic Games, Inc. All Rights Reserved.
+
+#include "MuT/ASTOpImageResize.h"
+
+#include "Containers/Map.h"
+#include "MuR/ModelPrivate.h"
+#include "MuR/RefCounted.h"
+#include "MuR/Types.h"
+#include "MuR/ImagePrivate.h"
+#include "MuT/ASTOpConditional.h"
+#include "MuT/ASTOpSwitch.h"
+#include "MuT/ASTOpImagePatch.h"
+#include "MuT/ASTOpImagePixelFormat.h"
+#include "MuT/ASTOpImageLayer.h"
+#include "MuT/ASTOpImageLayerColor.h"
+#include "MuT/ASTOpImageRasterMesh.h"
+#include "MuT/ASTOpImageTransform.h"
+#include "MuT/ASTOpImageSwizzle.h"
+#include "MuT/ASTOpImageCompose.h"
+#include "MuT/ASTOpImageMakeGrowMap.h"
+#include "MuT/ASTOpImageInterpolate.h"
+#include "MuT/ASTOpImagePlainColor.h"
+#include "MuT/ASTOpImageDisplace.h"
+#include "MuT/ASTOpImageInvert.h"
+
+
+namespace mu
+{
+
+	ASTOpImageResize::ASTOpImageResize()
+		: Source(this)
+	{
+	}
+
+
+	ASTOpImageResize::~ASTOpImageResize()
+	{
+		// Explicit call needed to avoid recursive destruction
+		ASTOp::RemoveChildren();
+	}
+
+
+	bool ASTOpImageResize::IsEqual(const ASTOp& InOther) const
+	{
+		if (InOther.GetOpType()==GetOpType())
+		{
+			const ASTOpImageResize* Other = static_cast<const ASTOpImageResize*>(&InOther);
+			return Source == Other->Source &&
+				Size == Other->Size;
+		}
+		return false;
+	}
+
+
+	uint64 ASTOpImageResize::Hash() const
+	{
+		uint64 Res = std::hash<EOpType>()(GetOpType());
+		hash_combine(Res, Source.child().get());
+		hash_combine(Res, Size[0]);
+		hash_combine(Res, Size[1]);
+		return Res;
+	}
+
+
+	Ptr<ASTOp> ASTOpImageResize::Clone(MapChildFuncRef MapChild) const
+	{
+		Ptr<ASTOpImageResize> New = new ASTOpImageResize();
+		New->Source = MapChild(Source.child());
+		New->Size = Size;
+		return New;
+	}
+
+
+	void ASTOpImageResize::ForEachChild(const TFunctionRef<void(ASTChild&)> Func)
+	{
+		Func(Source);
+	}
+
+
+	void ASTOpImageResize::Link(FProgram& Program, FLinkerOptions*)
+	{
+		// Already linked?
+		if (!linkedAddress)
+		{
+			OP::ImageResizeArgs Args = {};
+
+			if (Source)
+			{
+				Args.Source = Source->linkedAddress;
+			}
+			Args.Size[0] = Size[0];
+			Args.Size[1] = Size[1];
+
+			linkedAddress = (OP::ADDRESS)Program.OpAddress.Num();
+			Program.OpAddress.Add(Program.ByteCode.Num());
+			AppendCode(Program.ByteCode, GetOpType());
+			AppendCode(Program.ByteCode, Args);
+		}
+	}
+
+
+	FImageDesc ASTOpImageResize::GetImageDesc(bool bReturnBestOption, FGetImageDescContext* Context) const
+	{
+		FImageDesc Result;
+
+		// Local context in case it is necessary
+		FGetImageDescContext LocalContext;
+		if (!Context)
+		{
+			Context = &LocalContext;
+		}
+		else
+		{
+			// Cached result?
+			FImageDesc* PtrValue = Context->m_results.Find(this);
+			if (PtrValue)
+			{
+				return *PtrValue;
+			}
+		}
+
+		// Actual work
+		if (Source)
+		{
+			Result = Source->GetImageDesc(bReturnBestOption, Context);
+		}
+
+		Result.m_size[0] = Size[0];
+		Result.m_size[1] = Size[1];
+
+		// Cache the result
+		if (Context)
+		{
+			Context->m_results.Add(this, Result);
+		}
+
+		return Result;
+	}
+
+
+	Ptr<ImageSizeExpression> ASTOpImageResize::GetImageSizeExpression() const
+	{
+		Ptr<ImageSizeExpression> Res = new ImageSizeExpression;
+		Res->type = ImageSizeExpression::ISET_CONSTANT;
+		Res->size[0] = Size[0];
+		Res->size[1] = Size[1];
+		return Res;
+	}
+
+
+	void ASTOpImageResize::GetLayoutBlockSize(int32* OutBlockX, int32* OutBlockY)
+	{
+		// We didn't find any layout yet.
+		*OutBlockX = 0;
+		*OutBlockY = 0;
+
+		// Try the source
+		if (Source)
+		{
+			Source->GetLayoutBlockSize( OutBlockX, OutBlockY );
+
+			if (*OutBlockX > 0 && *OutBlockY > 0)
+			{
+				FImageDesc SourceDesc = Source->GetImageDesc(false);
+				if (SourceDesc.m_size[0] > 0 && SourceDesc.m_size[1] > 0)
+				{
+					float FactorX = float(Size[0]) / float(SourceDesc.m_size[0]);
+					float FactorY = float(Size[1]) / float(SourceDesc.m_size[1]);
+					*OutBlockX = int32(*OutBlockX * FactorX);
+					*OutBlockY = int32(*OutBlockX * FactorY);
+				}
+				else
+				{
+					*OutBlockX = 0;
+					*OutBlockY = 0;
+				}
+			}
+		}
+	}
+
+
+	bool ASTOpImageResize::IsImagePlainConstant(FVector4f& OutColour) const
+	{
+		bool bResult = true;
+		OutColour = FVector4f(0.0f,0.0f,0.0f,1.0f);
+
+		if (Source)
+		{
+			bResult = Source->IsImagePlainConstant(OutColour);
+		}
+
+		return bResult;
+	}
+
+
+	FSourceDataDescriptor ASTOpImageResize::GetSourceDataDescriptor(FGetSourceDataDescriptorContext* Context) const
+	{
+		if (Source)
+		{
+			return Source->GetSourceDataDescriptor(Context);
+		}
+
+		return {};
+	}
+
+
+	Ptr<ASTOp> ASTOpImageResize::OptimiseSemantic(const FModelOptimizationOptions&, int32 Pass) const
+	{
+		Ptr<ASTOp> Result;
+
+		Ptr<ASTOp> SourceAt = Source.child();
+
+		// The instruction can be sunk
+		EOpType SourceType = SourceAt->GetOpType();
+		switch (SourceType)
+		{
+		case EOpType::IM_PLAINCOLOUR:
+		{
+			Ptr<ASTOpImagePlainColor> NewOp = mu::Clone<ASTOpImagePlainColor>(SourceAt.get());
+			NewOp->Size[0] = Size[0];
+			NewOp->Size[1] = Size[1];
+			NewOp->LODs = 1; // TODO
+			Result = NewOp;
+			break;
+		}
+
+		default:
+			break;
+
+		}
+
+		return Result;
+	}
+
+
+	Ptr<ASTOp> ASTOpImageResize::OptimiseSize() const
+	{
+		Ptr<ASTOp> Result;
+
+		Ptr<ASTOp> SourceAt = Source.child();
+
+		EOpType SourceType = SourceAt->GetOpType();
+		switch (SourceType)
+		{
+		case EOpType::IM_RESIZE:
+		{
+			// Keep top resize
+			Ptr<const ASTOpImageResize> SourceOp = static_cast<const ASTOpImageResize*>(SourceAt.get());
+
+			Ptr<ASTOpImageResize> NewOp = mu::Clone<ASTOpImageResize>(this);
+			NewOp->Source = SourceOp->Source.child();
+
+			Result = NewOp;
+			break;
+		}
+
+		case EOpType::IM_PLAINCOLOUR:
+		{
+			// Set the size in the children and remove resize
+			Ptr<ASTOpImagePlainColor> SourceOp = mu::Clone<ASTOpImagePlainColor>(SourceAt.get());
+			SourceOp->Size[0] = Size[0];
+			SourceOp->Size[1] = Size[1];
+			SourceOp->LODs = 1; // TODO
+			Result = SourceOp;
+			break;
+		}
+
+		case EOpType::IM_TRANSFORM:
+		{
+			// Set the size in the children and remove resize
+			Ptr<ASTOpImageTransform> SourceOp = mu::Clone<ASTOpImageTransform>(SourceAt.get());
+			SourceOp->SizeX = Size[0];
+			SourceOp->SizeY = Size[1];
+			Result = SourceOp;
+			break;
+		}
+
+		case EOpType::IM_CONDITIONAL:
+		{
+			// We move the resize down the two paths
+			Ptr<ASTOpConditional> NewOp = mu::Clone<ASTOpConditional>(SourceAt);
+
+			Ptr<ASTOpImageResize> AOp = mu::Clone<ASTOpImageResize>(this);
+			AOp->Source = NewOp->yes.child();
+			NewOp->yes = AOp;
+
+			Ptr<ASTOpImageResize> BOp = mu::Clone<ASTOpImageResize>(this);
+			BOp->Source = NewOp->no.child();
+			NewOp->no = BOp;
+
+			Result = NewOp;
+			break;
+		}
+
+		case EOpType::IM_SWITCH:
+		{
+			// Move the resize down all the paths
+			Ptr<ASTOpSwitch> NewOp = mu::Clone<ASTOpSwitch>(SourceAt);
+
+			if (NewOp->def)
+			{
+				Ptr<ASTOpImageResize> DefOp = mu::Clone<ASTOpImageResize>(this);
+				DefOp->Source = NewOp->def.child();
+				NewOp->def = DefOp;
+			}
+
+			for (ASTOpSwitch::FCase& Case : NewOp->cases)
+			{
+				if (Case.branch)
+				{
+					Ptr<ASTOpImageResize> CaseOp = mu::Clone<ASTOpImageResize>(this);
+					CaseOp->Source = Case.branch.child();
+					Case.branch = CaseOp;
+				}
+			}
+
+			Result = NewOp;
+			break;
+		}
+
+		case EOpType::IM_SWIZZLE:
+		{
+			Ptr<ASTOpImageSwizzle> NewOp = mu::Clone<ASTOpImageSwizzle>(SourceAt);
+			for (int32 ChannelIndex = 0; ChannelIndex < MUTABLE_OP_MAX_SWIZZLE_CHANNELS; ++ChannelIndex)
+			{
+				Ptr<ASTOp> OldChannelOp = NewOp->Sources[ChannelIndex].child();
+				if (OldChannelOp)
+				{
+					Ptr<ASTOpImageResize> ChannelResize = mu::Clone<ASTOpImageResize>(this);
+					ChannelResize->Source = OldChannelOp;
+					NewOp->Sources[ChannelIndex] = ChannelResize;
+				}
+			}
+			Result = NewOp;
+			break;
+		}
+
+		case EOpType::IM_COMPOSE:
+		{
+			Ptr<ASTOpImageCompose> NewOp = mu::Clone<ASTOpImageCompose>(SourceAt);
+
+			Ptr<ASTOpImageResize> BaseOp = mu::Clone<ASTOpImageResize>(this);
+			BaseOp->Source = NewOp->Base.child();
+			NewOp->Base = BaseOp;
+
+			Ptr<ASTOpImageResize> BlockOp = mu::Clone<ASTOpImageResize>(this);
+			BlockOp->Source = NewOp->BlockImage.child();
+			NewOp->BlockImage = BlockOp;
+
+			if (NewOp->Mask)
+			{
+				Ptr<ASTOpImageResize> MaskOp = mu::Clone<ASTOpImageResize>(this);
+				MaskOp->Source = NewOp->Mask.child();
+				NewOp->Mask = MaskOp;
+			}
+
+			Result = NewOp;
+			break;
+		}
+
+		case EOpType::IM_RASTERMESH:
+		{
+			Ptr<ASTOpImageRasterMesh> NewOp = mu::Clone<ASTOpImageRasterMesh>(SourceAt);
+
+			//if ( NewOp->op.args.ImageRasterMesh.sizeX != op.args.ImageResize.size[0]
+			//     ||
+			//     NewOp->op.args.ImageRasterMesh.sizeY != op.args.ImageResize.size[1] )
+			{
+				NewOp->SizeX = Size[0];
+				NewOp->SizeY = Size[1];
+
+				if (NewOp->mask)
+				{
+					Ptr<ASTOpImageResize> MaskOp = mu::Clone<ASTOpImageResize>(this);
+					MaskOp->Source = NewOp->mask.child();
+					NewOp->mask = MaskOp;
+				}
+
+				// Don't apply absolute resizes to the image to raster: it could even enlarge it.
+				// This should only be scaled with relative resizes, which come from LOD biases, etc.
+				//if (NewOp->op.args.ImageRasterMesh.image)
+				//{
+				//	Ptr<ASTOpImageResize> mop = mu::Clone<ASTOpImageResize>(this);
+				//	mop->SetChild(mop->op.args.ImageResize.source, NewOp->children[NewOp->op.args.ImageRasterMesh.image].child());
+				//	NewOp->SetChild(NewOp->op.args.ImageRasterMesh.image, mop);
+				//}
+			}
+
+			Result = NewOp;
+			break;
+		}
+
+		case EOpType::IM_INTERPOLATE:
+		{
+			Ptr<ASTOpImageInterpolate> NewOp = mu::Clone<ASTOpImageInterpolate>(SourceAt);
+
+			for (int32 TargetIndex = 0; TargetIndex < MUTABLE_OP_MAX_INTERPOLATE_COUNT; ++TargetIndex)
+			{
+				Ptr<ASTOp> TargetAt = NewOp->Targets[TargetIndex].child();
+				if (TargetAt)
+				{
+					Ptr<ASTOpImageResize> SourceOp = mu::Clone<ASTOpImageResize>(this);
+					SourceOp->Source = TargetAt;
+					NewOp->Targets[TargetIndex] = SourceOp;
+				}
+			}
+
+			Result = NewOp;
+			break;
+		}
+
+		case EOpType::IM_INVERT:
+		{
+			Ptr<ASTOpImageInvert> NewOp = mu::Clone<ASTOpImageInvert>(SourceAt);
+			Ptr<ASTOp> BaseAt = NewOp->Base.child();
+
+			Ptr<ASTOpImageResize> NewBase = mu::Clone<ASTOpImageResize>(this);
+			NewBase->Source = BaseAt;
+
+			NewOp->Base = NewBase;
+
+			Result = NewOp;
+			break;
+		}
+
+		case EOpType::IM_PIXELFORMAT:
+		{
+			// \todo: only if shrinking?
+
+			// Only sink the resize if we know that the pixelformat source image is uncompressed.
+			Ptr<ASTOpImagePixelFormat> SourceTyped = static_cast<ASTOpImagePixelFormat*>(SourceAt.get());
+			FImageDesc PixelFormatSourceDesc = SourceTyped->Source->GetImageDesc();
+			if (PixelFormatSourceDesc.m_format != EImageFormat::None
+				&&
+				!mu::IsCompressedFormat(PixelFormatSourceDesc.m_format))
+			{
+				Ptr<ASTOpImagePixelFormat> NewOp = mu::Clone<ASTOpImagePixelFormat>(SourceAt);
+				Ptr<ASTOp> BaseAt = NewOp->Source.child();
+
+				Ptr<ASTOpImageResize> NewBase = mu::Clone<ASTOpImageResize>(this);
+				NewBase->Source = BaseAt;
+
+				NewOp->Source = NewBase;
+
+				Result = NewOp;
+			}
+			break;
+		}
+
+		case EOpType::IM_LAYER:
+		{
+			Ptr<ASTOpImageLayer> NewOp = mu::Clone<ASTOpImageLayer>(SourceAt);
+
+			Ptr<ASTOpImageResize> BaseOp = mu::Clone<ASTOpImageResize>(this);
+			BaseOp->Source = NewOp->base.child();
+			NewOp->base = BaseOp;
+
+			Ptr<ASTOpImageResize> blendOp = mu::Clone<ASTOpImageResize>(this);
+			blendOp->Source = NewOp->blend.child();
+			NewOp->blend = blendOp;
+
+			Ptr<ASTOp> MaskAt = NewOp->mask.child();
+			if (MaskAt)
+			{
+				Ptr<ASTOpImageResize> ResizedMaskOp = mu::Clone<ASTOpImageResize>(this);
+				ResizedMaskOp->Source = MaskAt;
+				NewOp->mask = ResizedMaskOp;
+			}
+
+			Result = NewOp;
+			break;
+		}
+
+		case EOpType::IM_LAYERCOLOUR:
+		{
+			Ptr<ASTOpImageLayerColor> NewOp = mu::Clone<ASTOpImageLayerColor>(SourceAt);
+
+			Ptr<ASTOpImageResize> BaseOp = mu::Clone<ASTOpImageResize>(this);
+			BaseOp->Source = NewOp->base.child();
+			NewOp->base = BaseOp;
+
+			Ptr<ASTOp> MaskAt = NewOp->mask.child();
+			if (MaskAt)
+			{
+				Ptr<ASTOpImageResize> ResizedMaskOp = mu::Clone<ASTOpImageResize>(this);
+				ResizedMaskOp->Source = MaskAt;
+				NewOp->mask = ResizedMaskOp;
+			}
+
+			Result = NewOp;
+			break;
+		}
+
+		case EOpType::IM_DISPLACE:
+		{
+			// In the size optimization phase we can optimize the resize with the displace
+			// because the constants have not been collapsed yet.
+			// We will still check it and sink the size directly below the IM_MAKEGROWMAP op
+			Ptr<ASTOpImageDisplace> SourceTyped = static_cast<ASTOpImageDisplace*>(SourceAt.get());
+			Ptr<ASTOp> OriginalDisplacementMapOp = SourceTyped->DisplacementMap.child();
+			if (OriginalDisplacementMapOp->GetOpType() == EOpType::IM_MAKEGROWMAP)
+			{
+				Ptr<ASTOpImageDisplace> NewOp = mu::Clone<ASTOpImageDisplace>(SourceAt);
+
+				Ptr<ASTOpImageResize> BaseOp = mu::Clone<ASTOpImageResize>(this);
+				BaseOp->Source = NewOp->Source.child();
+				NewOp->Source = BaseOp;
+
+				// Clone the map op and replace its children
+				Ptr<ASTOpImageMakeGrowMap> MapOp = mu::Clone<ASTOpImageMakeGrowMap>(OriginalDisplacementMapOp);
+				NewOp->DisplacementMap = MapOp;
+
+				Ptr<ASTOpImageResize> MapSourceOp = mu::Clone<ASTOpImageResize>(this);
+				MapSourceOp->Source = MapOp->Mask.child();
+				MapOp->Mask = MapSourceOp;
+
+				Result = NewOp;
+			}
+			break;
+		}
+
+		default:
+			break;
+		}
+
+		return Result;
+	}
+
+}
diff --git a/Engine/Plugins/Mutable/Source/MutableTools/Private/MuT/ASTOpImageResize.h b/Engine/Plugins/Mutable/Source/MutableTools/Private/MuT/ASTOpImageResize.h
new file mode 100644
--- /dev/null
+++ b/Engine/Plugins/Mutable/Source/MutableTools/Private/MuT/ASTOpImageResize.h
@@ -0,0 +1,48 @@
+// Copyright Epic Games, Inc. All Rights Reserved.
+
+#pragma once
+
+#include "HAL/PlatformMath.h"
+#include "MuR/Image.h"
+#include "MuR/MutableMath.h"
+#include "MuR/Operations.h"
+#include "MuR/Ptr.h"
+#include "MuT/AST.h"
+
+
+namespace mu
+{
+	struct FProgram;
+
+	class ASTOpImageResize final : public ASTOp
+	{
+	public:
+
+		ASTChild Source;
+		UE::Math::TIntVector2<uint16> Size = UE::Math::TIntVector2<uint16>(0, 0);
+
+	public:
+
+		ASTOpImageResize();
+		ASTOpImageResize(const ASTOpImageResize&) = delete;
+		~ASTOpImageResize();
+
+		virtual EOpType GetOpType() const override { return EOpType::IM_RESIZE; }
+		virtual uint64 Hash() const override;
+		virtual bool IsEqual(const ASTOp&) const override;
+		virtual Ptr<ASTOp> Clone(MapChildFuncRef) const override;
+		virtual void ForEachChild(const TFunctionRef<void(ASTChild&)>) override;
+		virtual void Link(FProgram&, FLinkerOptions*) override;
+		virtual FImageDesc GetImageDesc(bool bReturnBestOption, FGetImageDescContext*) const override;
+		virtual Ptr<ImageSizeExpression> GetImageSizeExpression() const override;
+		virtual void GetLayoutBlockSize(int32* OutBlockX, int32* OutBlockY) override;
+		virtual bool IsImagePlainConstant(FVector4f& colour) const override;
+		virtual Ptr<ASTOp> OptimiseSize() const override;
+		virtual Ptr<ASTOp> OptimiseSemantic(const FModelOptimizationOptions&, int32 Pass) const override;
+		virtual FSourceDataDescriptor GetSourceDataDescriptor(FGetSourceDataDescriptorContext*) const override;
+
+	};
+
+
+}
+
diff --git a/Engine/Plugins/Mutable/Source/MutableTools/Private/MuT/ASTOpImageResizeLike.cpp b/Engine/Plugins/Mutable/Source/MutableTools/Private/MuT/ASTOpImageResizeLike.cpp
new file mode 100644
--- /dev/null
+++ b/Engine/Plugins/Mutable/Source/MutableTools/Private/MuT/ASTOpImageResizeLike.cpp
@@ -0,0 +1,229 @@
+// Copyright Epic Games, Inc. All Rights Reserved.
+
+#include "MuT/ASTOpImageResizeLike.h"
+
+#include "Containers/Map.h"
+#include "MuR/ModelPrivate.h"
+#include "MuR/RefCounted.h"
+#include "MuR/Types.h"
+#include "MuR/ImagePrivate.h"
+#include "MuT/ASTOpConditional.h"
+#include "MuT/ASTOpSwitch.h"
+#include "MuT/ASTOpImagePatch.h"
+#include "MuT/ASTOpImagePixelFormat.h"
+#include "MuT/ASTOpImageLayer.h"
+#include "MuT/ASTOpImageLayerColor.h"
+#include "MuT/ASTOpImageRasterMesh.h"
+#include "MuT/ASTOpImageTransform.h"
+#include "MuT/ASTOpImageSwizzle.h"
+#include "MuT/ASTOpImageCompose.h"
+#include "MuT/ASTOpImageMakeGrowMap.h"
+#include "MuT/ASTOpImageResize.h"
+
+
+namespace mu
+{
+
+	ASTOpImageResizeLike::ASTOpImageResizeLike()
+		: Source(this),
+		SizeSource(this)
+	{
+	}
+
+
+	ASTOpImageResizeLike::~ASTOpImageResizeLike()
+	{
+		// Explicit call needed to avoid recursive destruction
+		ASTOp::RemoveChildren();
+	}
+
+
+	bool ASTOpImageResizeLike::IsEqual(const ASTOp& InOther) const
+	{
+		if (InOther.GetOpType()==GetOpType())
+		{
+			const ASTOpImageResizeLike* Other = static_cast<const ASTOpImageResizeLike*>(&InOther);
+			return Source == Other->Source &&
+				SizeSource == Other->SizeSource;
+		}
+		return false;
+	}
+
+
+	uint64 ASTOpImageResizeLike::Hash() const
+	{
+		uint64 Res = std::hash<EOpType>()(GetOpType());
+		hash_combine(Res, Source.child().get());
+		hash_combine(Res, SizeSource.child().get());
+		return Res;
+	}
+
+
+	Ptr<ASTOp> ASTOpImageResizeLike::Clone(MapChildFuncRef MapChild) const
+	{
+		Ptr<ASTOpImageResizeLike> New = new ASTOpImageResizeLike();
+		New->Source = MapChild(Source.child());
+		New->SizeSource = MapChild(SizeSource.child());
+		return New;
+	}
+
+
+	void ASTOpImageResizeLike::ForEachChild(const TFunctionRef<void(ASTChild&)> Func)
+	{
+		Func(Source);
+		Func(SizeSource);
+	}
+
+
+	void ASTOpImageResizeLike::Link(FProgram& Program, FLinkerOptions*)
+	{
+		// Already linked?
+		if (!linkedAddress)
+		{
+			OP::ImageResizeLikeArgs Args = {};
+
+			if (Source)
+			{
+				Args.Source = Source->linkedAddress;
+			}
+			if (SizeSource)
+			{
+				Args.SizeSource = SizeSource->linkedAddress;
+			}
+
+			linkedAddress = (OP::ADDRESS)Program.OpAddress.Num();
+			Program.OpAddress.Add(Program.ByteCode.Num());
+			AppendCode(Program.ByteCode, GetOpType());
+			AppendCode(Program.ByteCode, Args);
+		}
+	}
+
+
+	FImageDesc ASTOpImageResizeLike::GetImageDesc(bool bReturnBestOption, FGetImageDescContext* Context) const
+	{
+		FImageDesc Result;
+
+		// Local context in case it is necessary
+		FGetImageDescContext LocalContext;
+		if (!Context)
+		{
+			Context = &LocalContext;
+		}
+		else
+		{
+			// Cached result?
+			FImageDesc* PtrValue = Context->m_results.Find(this);
+			if (PtrValue)
+			{
+				return *PtrValue;
+			}
+		}
+
+		// Actual work
+		if (Source)
+		{
+			Result = Source->GetImageDesc(bReturnBestOption, Context);
+		}
+
+		if (SizeSource)
+		{
+			FImageDesc SizeResult = SizeSource->GetImageDesc(bReturnBestOption, Context);
+			Result.m_size = SizeResult.m_size;
+		}
+
+		// Cache the result
+		if (Context)
+		{
+			Context->m_results.Add(this, Result);
+		}
+
+		return Result;
+	}
+
+
+	Ptr<ImageSizeExpression> ASTOpImageResizeLike::GetImageSizeExpression() const
+	{
+		if (SizeSource)
+		{
+			return SizeSource->GetImageSizeExpression();
+		}
+
+		return new ImageSizeExpression;
+	}
+
+
+	bool ASTOpImageResizeLike::IsImagePlainConstant(FVector4f& OutColour) const
+	{
+		bool bResult = true;
+		OutColour = FVector4f(0.0f,0.0f,0.0f,1.0f);
+
+		if (Source)
+		{
+			bResult = Source->IsImagePlainConstant(OutColour);
+		}
+
+		return bResult;
+	}
+
+
+	FSourceDataDescriptor ASTOpImageResizeLike::GetSourceDataDescriptor(FGetSourceDataDescriptorContext* Context) const
+	{
+		if (Source)
+		{
+			return Source->GetSourceDataDescriptor(Context);
+		}
+
+		return {};
+	}
+
+
+	Ptr<ASTOp> ASTOpImageResizeLike::OptimiseSize() const
+	{
+		Ptr<ASTOp> Result;
+
+		if (!Source || !SizeSource)
+		{
+			return Result;
+		}
+
+		Ptr<ImageSizeExpression> SourceSize = Source->GetImageSizeExpression();
+		Ptr<ImageSizeExpression> SizeSourceSize = SizeSource->GetImageSizeExpression();
+
+		if (*SourceSize == *SizeSourceSize)
+		{
+			Result = Source.child();
+		}
+
+		else if (SizeSourceSize->type == ImageSizeExpression::ISET_CONSTANT)
+		{
+			Ptr<ASTOpImageResize> NewOp = new ASTOpImageResize;
+			NewOp->Source = Source.child();
+			NewOp->Size[0] = SizeSourceSize->size[0];
+			NewOp->Size[1] = SizeSourceSize->size[1];
+			Result = NewOp;
+		}
+		else if (SizeSourceSize->type == ImageSizeExpression::ISET_LAYOUTFACTOR)
+		{
+			// TODO
+			// Skip intermediate ops until the layout
+//                    if ( program.m_code[sizeSourceAt].type != EOpType::IM_BLANKLAYOUT )
+//                    {
+//                        m_modified = true;
+
+//                        OP blackLayoutOp;
+//                        blackLayoutOp.type = EOpType::IM_BLANKLAYOUT;
+//                        blackLayoutOp.args.ImageBlankLayout.layout = sizeSourceSize->layout;
+//                        blackLayoutOp.args.ImageBlankLayout.blockSize[0] = sizeSourceSize->factor[0];
+//                        blackLayoutOp.args.ImageBlankLayout.blockSize[1] = sizeSourceSize->factor[1];
+//                        blackLayoutOp.args.ImageBlankLayout.format = L_UByte;
+
+//                        OP opResize = program.m_code[at];
+//                        opResize.args.ImageResizeLike.sizeSource = program.AddOp( blackLayoutOp );
+//                        Result = program.AddOp( opResize );
+//                    }
+		}
+
+		return Result;
+	}
+
+}
diff --git a/Engine/Plugins/Mutable/Source/MutableTools/Private/MuT/ASTOpImageResizeLike.h b/Engine/Plugins/Mutable/Source/MutableTools/Private/MuT/ASTOpImageResizeLike.h
new file mode 100644
--- /dev/null
+++ b/Engine/Plugins/Mutable/Source/MutableTools/Private/MuT/ASTOpImageResizeLike.h
@@ -0,0 +1,46 @@
+// Copyright Epic Games, Inc. All Rights Reserved.
+
+#pragma once
+
+#include "HAL/PlatformMath.h"
+#include "MuR/Image.h"
+#include "MuR/MutableMath.h"
+#include "MuR/Operations.h"
+#include "MuR/Ptr.h"
+#include "MuT/AST.h"
+
+
+namespace mu
+{
+	struct FProgram;
+
+	class ASTOpImageResizeLike final : public ASTOp
+	{
+	public:
+
+		ASTChild Source;
+		ASTChild SizeSource;
+
+	public:
+
+		ASTOpImageResizeLike();
+		ASTOpImageResizeLike(const ASTOpImageResizeLike&) = delete;
+		~ASTOpImageResizeLike();
+
+		virtual EOpType GetOpType() const override { return EOpType::IM_RESIZELIKE; }
+		virtual uint64 Hash() const override;
+		virtual bool IsEqual(const ASTOp&) const override;
+		virtual Ptr<ASTOp> Clone(MapChildFuncRef) const override;
+		virtual void ForEachChild(const TFunctionRef<void(ASTChild&)>) override;
+		virtual void Link(FProgram&, FLinkerOptions*) override;
+		virtual FImageDesc GetImageDesc(bool bReturnBestOption, FGetImageDescContext*) const override;
+		virtual Ptr<ImageSizeExpression> GetImageSizeExpression() const override;
+		virtual bool IsImagePlainConstant(FVector4f& colour) const override;
+		virtual Ptr<ASTOp> OptimiseSize() const override;
+		virtual FSourceDataDescriptor GetSourceDataDescriptor(FGetSourceDataDescriptorContext*) const override;
+
+	};
+
+
+}
+
diff --git a/Engine/Plugins/Mutable/Source/MutableTools/Private/MuT/ASTOpImageResizeRel.cpp b/Engine/Plugins/Mutable/Source/MutableTools/Private/MuT/ASTOpImageResizeRel.cpp
new file mode 100644
--- /dev/null
+++ b/Engine/Plugins/Mutable/Source/MutableTools/Private/MuT/ASTOpImageResizeRel.cpp
@@ -0,0 +1,642 @@
+// Copyright Epic Games, Inc. All Rights Reserved.
+
+#include "MuT/ASTOpImageResizeRel.h"
+
+#include "Containers/Map.h"
+#include "MuR/ModelPrivate.h"
+#include "MuR/RefCounted.h"
+#include "MuR/Types.h"
+#include "MuR/ImagePrivate.h"
+#include "MuT/ASTOpConditional.h"
+#include "MuT/ASTOpSwitch.h"
+#include "MuT/ASTOpImagePatch.h"
+#include "MuT/ASTOpImagePixelFormat.h"
+#include "MuT/ASTOpImageLayer.h"
+#include "MuT/ASTOpImageLayerColor.h"
+#include "MuT/ASTOpImageRasterMesh.h"
+#include "MuT/ASTOpImageTransform.h"
+#include "MuT/ASTOpImageSwizzle.h"
+#include "MuT/ASTOpImageCompose.h"
+#include "MuT/ASTOpImageMakeGrowMap.h"
+#include "MuT/ASTOpImageResize.h"
+#include "MuT/ASTOpImageMultiLayer.h"
+#include "MuT/ASTOpImageMipmap.h"
+#include "MuT/ASTOpImageBlankLayout.h"
+#include "MuT/ASTOpImageInterpolate.h"
+#include "MuT/ASTOpImageSaturate.h"
+#include "MuT/ASTOpImagePlainColor.h"
+#include "MuT/ASTOpImageDisplace.h"
+#include "MuT/ASTOpImageInvert.h"
+
+
+namespace mu
+{
+
+	ASTOpImageResizeRel::ASTOpImageResizeRel()
+		: Source(this)
+	{
+	}
+
+
+	ASTOpImageResizeRel::~ASTOpImageResizeRel()
+	{
+		// Explicit call needed to avoid recursive destruction
+		ASTOp::RemoveChildren();
+	}
+
+
+	bool ASTOpImageResizeRel::IsEqual(const ASTOp& InOther) const
+	{
+		if (InOther.GetOpType()==GetOpType())
+		{
+			const ASTOpImageResizeRel* Other = static_cast<const ASTOpImageResizeRel*>(&InOther);
+			return Source == Other->Source &&
+				Factor == Other->Factor;
+		}
+		return false;
+	}
+
+
+	uint64 ASTOpImageResizeRel::Hash() const
+	{
+		uint64 Res = std::hash<EOpType>()(GetOpType());
+		hash_combine(Res, Source.child().get());
+		hash_combine(Res, Factor[0]);
+		hash_combine(Res, Factor[1]);
+		return Res;
+	}
+
+
+	Ptr<ASTOp> ASTOpImageResizeRel::Clone(MapChildFuncRef MapChild) const
+	{
+		Ptr<ASTOpImageResizeRel> New = new ASTOpImageResizeRel();
+		New->Source = MapChild(Source.child());
+		New->Factor = Factor;
+		return New;
+	}
+
+
+	void ASTOpImageResizeRel::ForEachChild(const TFunctionRef<void(ASTChild&)> Func)
+	{
+		Func(Source);
+	}
+
+
+	void ASTOpImageResizeRel::Link(FProgram& Program, FLinkerOptions*)
+	{
+		// Already linked?
+		if (!linkedAddress)
+		{
+			OP::ImageResizeRelArgs Args = {};
+
+			if (Source)
+			{
+				Args.Source = Source->linkedAddress;
+			}
+			Args.Factor[0] = Factor[0];
+			Args.Factor[1] = Factor[1];
+
+			linkedAddress = (OP::ADDRESS)Program.OpAddress.Num();
+			Program.OpAddress.Add(Program.ByteCode.Num());
+			AppendCode(Program.ByteCode, GetOpType());
+			AppendCode(Program.ByteCode, Args);
+		}
+	}
+
+
+	FImageDesc ASTOpImageResizeRel::GetImageDesc(bool bReturnBestOption, FGetImageDescContext* Context) const
+	{
+		FImageDesc Result;
+
+		// Local context in case it is necessary
+		FGetImageDescContext LocalContext;
+		if (!Context)
+		{
+			Context = &LocalContext;
+		}
+		else
+		{
+			// Cached result?
+			FImageDesc* PtrValue = Context->m_results.Find(this);
+			if (PtrValue)
+			{
+				return *PtrValue;
+			}
+		}
+
+		// Actual work
+		if (Source)
+		{
+			Result = Source->GetImageDesc(bReturnBestOption, Context);
+		}
+
+		Result.m_size[0] = uint16(Result.m_size[0] * Factor[0]);
+		Result.m_size[1] = uint16(Result.m_size[1] * Factor[1]);
+
+		// Cache the result
+		if (Context)
+		{
+			Context->m_results.Add(this, Result);
+		}
+
+		return Result;
+	}
+
+
+	Ptr<ImageSizeExpression> ASTOpImageResizeRel::GetImageSizeExpression() const
+	{
+		Ptr<ImageSizeExpression> Res;
+
+		if (Source)
+		{
+			Res = Source->GetImageSizeExpression();
+			if (Res->type == ImageSizeExpression::ISET_CONSTANT)
+			{
+				Res->size[0] = uint16(Res->size[0] * Factor[0]);
+				Res->size[1] = uint16(Res->size[1] * Factor[1]);
+			}
+			else
+			{
+				// TODO: Proportional factor
+				Res = new ImageSizeExpression();
+				Res->type = ImageSizeExpression::ISET_UNKNOWN;
+			}
+		}
+		else
+		{
+			Res = new ImageSizeExpression;
+		}
+
+		return Res;
+	}
+
+
+	void ASTOpImageResizeRel::GetLayoutBlockSize(int32* OutBlockX, int32* OutBlockY)
+	{
+		// We didn't find any layout yet.
+		*OutBlockX = 0;
+		*OutBlockY = 0;
+
+		// Try the source
+		if (Source)
+		{
+			Source->GetLayoutBlockSize( OutBlockX, OutBlockY );
+
+			if (*OutBlockX > 0 && *OutBlockY > 0)
+			{
+				FImageDesc SourceDesc = Source->GetImageDesc(false);
+				if (SourceDesc.m_size[0] > 0 && SourceDesc.m_size[1] > 0)
+				{
+					*OutBlockX = int32(*OutBlockX * Factor[0]);
+					*OutBlockY = int32(*OutBlockX * Factor[1]);
+				}
+				else
+				{
+					*OutBlockX = 0;
+					*OutBlockY = 0;
+				}
+			}
+		}
+	}
+
+
+	bool ASTOpImageResizeRel::IsImagePlainConstant(FVector4f& OutColour) const
+	{
+		bool bResult = true;
+		OutColour = FVector4f(0.0f,0.0f,0.0f,1.0f);
+
+		if (Source)
+		{
+			bResult = Source->IsImagePlainConstant(OutColour);
+		}
+
+		return bResult;
+	}
+
+
+	FSourceDataDescriptor ASTOpImageResizeRel::GetSourceDataDescriptor(FGetSourceDataDescriptorContext* Context) const
+	{
+		if (Source)
+		{
+			return Source->GetSourceDataDescriptor(Context);
+		}
+
+		return {};
+	}
+
+
+	Ptr<ASTOp> ASTOpImageResizeRel::OptimiseSemantic(const FModelOptimizationOptions&, int32 Pass) const
+	{
+		Ptr<ASTOp> Result;
+
+		Ptr<ASTOp> SourceAt = Source.child();
+
+		// The instruction can be sunk
+		EOpType SourceType = SourceAt->GetOpType();
+		switch (SourceType)
+		{
+		case EOpType::IM_RESIZE:
+		{
+			Ptr<ASTOpImageResize> NewOp = mu::Clone<ASTOpImageResize>(SourceAt.get());
+			NewOp->Size[0] = int16(NewOp->Size[0] * Factor[0]);
+			NewOp->Size[1] = int16(NewOp->Size[1] * Factor[1]);
+
+			Result = NewOp;
+			break;
+		}
+
+		default:
+			break;
+
+		}
+
+		return Result;
+	}
+
+
+	class Sink_ImageResizeRelAST
+	{
+	public:
+
+		Ptr<ASTOp> Apply(const ASTOpImageResizeRel* InRoot)
+		{
+			Root = InRoot;
+			OldToNew.Empty();
+
+			InitialSource = InRoot->Source.child();
+			Ptr<ASTOp> newSource = Visit(InitialSource, InRoot);
+
+			// If there is any change, it is the new root.
+			if (newSource != InitialSource)
+			{
+				return newSource;
+			}
+
+			return nullptr;
+		}
+
+	protected:
+
+		const ASTOp* Root;
+		Ptr<ASTOp> InitialSource;
+		//! For each operation we sink, the map from old instructions to new instructions.
+		TMap<FSinkerOldToNewKey, Ptr<ASTOp>> OldToNew;
+
+		Ptr<ASTOp> Visit(Ptr<ASTOp> at, const ASTOpImageResizeRel* CurrentSinkingOp)
+		{
+			if (!at) return nullptr;
+
+			// Already visited?
+			const Ptr<ASTOp>* Cached = OldToNew.Find({ at, CurrentSinkingOp });
+			if (Cached)
+			{
+				return *Cached;
+			}
+
+			float ScaleX = CurrentSinkingOp->Factor[0];
+			float ScaleY = CurrentSinkingOp->Factor[1];
+
+			Ptr<ASTOp> Result = at;
+			switch (at->GetOpType())
+			{
+
+			case EOpType::IM_CONDITIONAL:
+			{
+				// We move the mask creation down the two paths
+				auto NewOp = mu::Clone<ASTOpConditional>(at);
+				NewOp->yes = Visit(NewOp->yes.child(), CurrentSinkingOp);
+				NewOp->no = Visit(NewOp->no.child(), CurrentSinkingOp);
+				Result = NewOp;
+				break;
+			}
+
+			case EOpType::IM_PIXELFORMAT:
+			{
+				// \todo: only if shrinking?
+				Ptr<ASTOpImagePixelFormat> NewOp = mu::Clone<ASTOpImagePixelFormat>(at);
+				NewOp->Source = Visit(NewOp->Source.child(), CurrentSinkingOp);
+				Result = NewOp;
+				break;
+			}
+
+			case EOpType::IM_SWITCH:
+			{
+				// We move the mask creation down all the paths
+				Ptr<ASTOpSwitch> NewOp = mu::Clone<ASTOpSwitch>(at);
+				NewOp->def = Visit(NewOp->def.child(), CurrentSinkingOp);
+				for (ASTOpSwitch::FCase& c : NewOp->cases)
+				{
+					c.branch = Visit(c.branch.child(), CurrentSinkingOp);
+				}
+				Result = NewOp;
+				break;
+			}
+
+			case EOpType::IM_SWIZZLE:
+			{
+				Ptr<ASTOpImageSwizzle> NewOp = mu::Clone<ASTOpImageSwizzle>(at);
+				for (int32 s = 0; s < MUTABLE_OP_MAX_SWIZZLE_CHANNELS; ++s)
+				{
+					Ptr<ASTOp> channelOp = NewOp->Sources[s].child();
+					if (channelOp)
+					{
+						NewOp->Sources[s] = Visit(channelOp, CurrentSinkingOp);
+					}
+				}
+				Result = NewOp;
+				break;
+			}
+
+			case EOpType::IM_COMPOSE:
+			{
+				// We can only optimise if the layout grid blocks size in pixels is
+				// still an integer after relative scale
+				bool acceptable = false;
+				{
+					const ASTOpImageCompose* typedAt = static_cast<const ASTOpImageCompose*>(at.get());
+					Ptr<ASTOp> originalBaseOp = typedAt->Base.child();
+
+					// \todo: recursion-proof cache?
+					int32 layoutBlockPixelsX = 0;
+					int32 layoutBlockPixelsY = 0;
+					originalBaseOp->GetLayoutBlockSize(&layoutBlockPixelsX, &layoutBlockPixelsY);
+
+					int32 scaledLayoutBlockPixelsX = int32(layoutBlockPixelsX * ScaleX);
+					int32 scaledLayoutBlockPixelsY = int32(layoutBlockPixelsY * ScaleY);
+					int32 unscaledLayoutBlockPixelsX = int32(scaledLayoutBlockPixelsX / ScaleX);
+					int32 unscaledLayoutBlockPixelsY = int32(scaledLayoutBlockPixelsY / ScaleY);
+					acceptable =
+						(layoutBlockPixelsX != 0 && layoutBlockPixelsY != 0)
+						&&
+						(layoutBlockPixelsX == unscaledLayoutBlockPixelsX)
+						&&
+						(layoutBlockPixelsY == unscaledLayoutBlockPixelsY);
+				}
+
+				if (acceptable)
+				{
+					Ptr<ASTOpImageCompose> NewOp = mu::Clone<ASTOpImageCompose>(at);
+
+					Ptr<ASTOp> baseOp = NewOp->Base.child();
+					NewOp->Base = Visit(baseOp, CurrentSinkingOp);
+
+					Ptr<ASTOp> blockOp = NewOp->BlockImage.child();
+					NewOp->BlockImage = Visit(blockOp, CurrentSinkingOp);
+
+					Result = NewOp;
+				}
+
+				break;
+			}
+
+			case EOpType::IM_PATCH:
+			{
+				Ptr<ASTOpImagePatch> NewOp = mu::Clone<ASTOpImagePatch>(at);
+
+				NewOp->base = Visit(NewOp->base.child(), CurrentSinkingOp);
+				NewOp->patch = Visit(NewOp->patch.child(), CurrentSinkingOp);
+
+				// todo: review if this is always correct, or we need some "divisible" check
+				NewOp->location[0] = uint16(NewOp->location[0] * ScaleX);
+				NewOp->location[1] = uint16(NewOp->location[1] * ScaleY);
+
+				Result = NewOp;
+
+				break;
+			}
+
+			case EOpType::IM_MIPMAP:
+			{
+				Ptr<ASTOpImageMipmap> NewOp = mu::Clone<ASTOpImageMipmap>(at.get());
+				Ptr<ASTOp> baseOp = NewOp->Source.child();
+				NewOp->Source = Visit(baseOp, CurrentSinkingOp);
+				Result = NewOp;
+				break;
+			}
+
+			case EOpType::IM_INTERPOLATE:
+			{
+				Ptr<ASTOpImageInterpolate> NewOp = mu::Clone<ASTOpImageInterpolate>(at);
+
+				for (int32 v = 0; v < MUTABLE_OP_MAX_INTERPOLATE_COUNT; ++v)
+				{
+					Ptr<ASTOp> Child = NewOp->Targets[v].child();
+					Ptr<ASTOp> bOp = Visit(Child, CurrentSinkingOp);
+					NewOp->Targets[v] = bOp;
+				}
+
+				Result = NewOp;
+				break;
+			}
+
+			case EOpType::IM_MULTILAYER:
+			{
+				Ptr<ASTOpImageMultiLayer> nop = mu::Clone<ASTOpImageMultiLayer>(at);
+				nop->base = Visit(nop->base.child(), CurrentSinkingOp);
+				nop->mask = Visit(nop->mask.child(), CurrentSinkingOp);
+				nop->blend = Visit(nop->blend.child(), CurrentSinkingOp);
+				Result = nop;
+				break;
+			}
+
+			case EOpType::IM_LAYER:
+			{
+				Ptr<ASTOpImageLayer> nop = mu::Clone<ASTOpImageLayer>(at);
+
+				Ptr<ASTOp> aOp = nop->base.child();
+				nop->base = Visit(aOp, CurrentSinkingOp);
+
+				Ptr<ASTOp> bOp = nop->blend.child();
+				nop->blend = Visit(bOp, CurrentSinkingOp);
+
+				Ptr<ASTOp> mOp = nop->mask.child();
+				nop->mask = Visit(mOp, CurrentSinkingOp);
+
+				Result = nop;
+				break;
+			}
+
+			case EOpType::IM_LAYERCOLOUR:
+			{
+				Ptr<ASTOpImageLayerColor> nop = mu::Clone<ASTOpImageLayerColor>(at);
+
+				Ptr<ASTOp> aOp = nop->base.child();
+				nop->base = Visit(aOp, CurrentSinkingOp);
+
+				Ptr<ASTOp> mOp = nop->mask.child();
+				nop->mask = Visit(mOp, CurrentSinkingOp);
+
+				Result = nop;
+				break;
+			}
+
+			case EOpType::IM_DISPLACE:
+			{
+				Ptr<ASTOpImageDisplace> nop = mu::Clone<ASTOpImageDisplace>(at);
+
+				Ptr<ASTOp> sourceOp = nop->Source.child();
+				nop->Source = Visit(sourceOp, CurrentSinkingOp);
+
+				Ptr<ASTOp> mapOp = nop->DisplacementMap.child();
+				nop->DisplacementMap = Visit(mapOp, CurrentSinkingOp);
+
+				// Make sure we don't scale a constant displacement map, which is very wrong.
+				if (mapOp->GetOpType() == EOpType::IM_MAKEGROWMAP)
+				{
+					Result = nop;
+				}
+				else
+				{
+					// We cannot resize an already calculated displacement map.
+				}
+
+				break;
+			}
+
+			case EOpType::IM_MAKEGROWMAP:
+			{
+				Ptr<ASTOpImageMakeGrowMap> nop = mu::Clone<ASTOpImageMakeGrowMap>(at);
+				Ptr<ASTOp> maskOp = nop->Mask.child();
+				nop->Mask = Visit(maskOp, CurrentSinkingOp);
+				Result = nop;
+				break;
+			}
+
+			case EOpType::IM_INVERT:
+			{
+				Ptr<ASTOpImageInvert> nop = mu::Clone<ASTOpImageInvert>(at);
+				Ptr<ASTOp> maskOp = nop->Base.child();
+				nop->Base = Visit(maskOp, CurrentSinkingOp);
+				Result = nop;
+				break;
+			}
+
+			case EOpType::IM_SATURATE:
+			{
+				Ptr<ASTOpImageSaturate> NewOp = mu::Clone<ASTOpImageSaturate>(at);
+				NewOp->Base = Visit(NewOp->Base.child(), CurrentSinkingOp);
+				Result = NewOp;
+				break;
+			}
+
+			case EOpType::IM_TRANSFORM:
+			{
+				// It can only sink in the transform if it doesn't have it's own size.
+				const ASTOpImageTransform* TypedAt = static_cast<const ASTOpImageTransform*>(at.get());
+				if (TypedAt->SizeX == 0 && TypedAt->SizeY == 0)
+				{
+					Ptr<ASTOpImageTransform> NewOp = mu::Clone<ASTOpImageTransform>(at);
+					Ptr<ASTOp> MaskOp = NewOp->Base.child();
+					NewOp->Base = Visit(MaskOp, CurrentSinkingOp);
+					Result = NewOp;
+				}
+
+				break;
+			}
+
+			case EOpType::IM_RASTERMESH:
+			{
+				Ptr<ASTOpImageRasterMesh> nop = mu::Clone<ASTOpImageRasterMesh>(at);
+				Ptr<ASTOp> maskOp = nop->mask.child();
+				nop->mask = Visit(maskOp, CurrentSinkingOp);
+
+				// Resize the image to project as well, assuming that since the target has a different resolution
+				// it make sense for the source image to have a similar resize.
+				// Actually, don't do it because the LODBias will be applied separetely at graph generation time.
+				//auto imageOp = nop->image.child();
+				//nop->image = Visit(imageOp, CurrentSinkingOp);
+
+				nop->SizeX = uint16(nop->SizeX * ScaleX + 0.5f);
+				nop->SizeY = uint16(nop->SizeY * ScaleY + 0.5f);
+				Result = nop;
+				break;
+			}
+
+			default:
+				break;
+			}
+
+			// end on line, replace with sinking op
+			if (at == Result && at != InitialSource)
+			{
+				Ptr<ASTOpImageResizeRel> NewOp = mu::Clone<ASTOpImageResizeRel>(CurrentSinkingOp);
+				NewOp->Source = at;
+
+				Result = NewOp;
+			}
+
+			OldToNew.Add({ at, CurrentSinkingOp }, Result);
+
+			return Result;
+		}
+	};
+
+
+	Ptr<ASTOp> ASTOpImageResizeRel::OptimiseSize() const
+	{
+		Ptr<ASTOp> Result;
+
+		Ptr<ASTOp> SourceAt = Source.child();
+
+		// The instruction can be sunk
+		EOpType SourceType = SourceAt->GetOpType();
+		switch (SourceType)
+		{
+
+		case EOpType::IM_BLANKLAYOUT:
+		{
+			Ptr<ASTOpImageBlankLayout> NewOp = mu::Clone<ASTOpImageBlankLayout>(SourceAt);
+			NewOp->BlockSize[0] = uint16(NewOp->BlockSize[0] * Factor[0] + 0.5f);
+			NewOp->BlockSize[1] = uint16(NewOp->BlockSize[1] * Factor[1] + 0.5f);
+			Result = NewOp;
+			break;
+		}
+
+		case EOpType::IM_PLAINCOLOUR:
+		{
+			Ptr<ASTOpImagePlainColor> NewOp = mu::Clone<ASTOpImagePlainColor>(SourceAt);
+			NewOp->Size[0] = FMath::CeilToInt(NewOp->Size[0] * Factor[0]);
+			NewOp->Size[1] = FMath::CeilToInt(NewOp->Size[1] * Factor[1]);
+			NewOp->LODs = 1; // TODO
+			Result = NewOp;
+			break;
+		}
+
+		case EOpType::IM_TRANSFORM:
+		{
+			// We can only optimize here if we know the transform result size, otherwise, we will sink the op in the sinker.
+			const ASTOpImageTransform* typedAt = static_cast<const ASTOpImageTransform*>(SourceAt.get());
+			if (typedAt->SizeX != 0 && typedAt->SizeY != 0)
+			{
+				// Set the size in the children and remove resize
+				Ptr<ASTOpImageTransform> sourceOp = mu::Clone<ASTOpImageTransform>(SourceAt.get());
+				sourceOp->SizeX = FMath::CeilToInt32(sourceOp->SizeX * Factor[0]);
+				sourceOp->SizeY = FMath::CeilToInt32(sourceOp->SizeY * Factor[1]);
+				Result = sourceOp;
+			}
+			break;
+		}
+
+		// Don't combine. ResizeRel sometimes can resize more children than Resize can do. (see RasterMesh)
+		// It can be combined in an optimization step further in the process, when normal sizes may have been 
+		// optimized already (see OptimizeSemantic)
+//                case EOpType::IM_RESIZE:
+//                {
+//                    break;
+//                }
+
+
+		default:
+		{
+			Sink_ImageResizeRelAST sinker;
+			Result = sinker.Apply(this);
+
+			break;
+		}
+
+		}
+
+		return Result;
+	}
+
+}
diff --git a/Engine/Plugins/Mutable/Source/MutableTools/Private/MuT/ASTOpImageResizeRel.h b/Engine/Plugins/Mutable/Source/MutableTools/Private/MuT/ASTOpImageResizeRel.h
new file mode 100644
--- /dev/null
+++ b/Engine/Plugins/Mutable/Source/MutableTools/Private/MuT/ASTOpImageResizeRel.h
@@ -0,0 +1,48 @@
+// Copyright Epic Games, Inc. All Rights Reserved.
+
+#pragma once
+
+#include "HAL/PlatformMath.h"
+#include "MuR/Image.h"
+#include "MuR/MutableMath.h"
+#include "MuR/Operations.h"
+#include "MuR/Ptr.h"
+#include "MuT/AST.h"
+
+
+namespace mu
+{
+	struct FProgram;
+
+	class ASTOpImageResizeRel final : public ASTOp
+	{
+	public:
+
+		ASTChild Source;
+		FVector2f Factor = FVector2f(1.0f, 1.0f);
+
+	public:
+
+		ASTOpImageResizeRel();
+		ASTOpImageResizeRel(const ASTOpImageResizeRel&) = delete;
+		~ASTOpImageResizeRel();
+
+		virtual EOpType GetOpType() const override { return EOpType::IM_RESIZEREL; }
+		virtual uint64 Hash() const override;
+		virtual bool IsEqual(const ASTOp&) const override;
+		virtual Ptr<ASTOp> Clone(MapChildFuncRef) const override;
+		virtual void ForEachChild(const TFunctionRef<void(ASTChild&)>) override;
+		virtual void Link(FProgram&, FLinkerOptions*) override;
+		virtual FImageDesc GetImageDesc(bool bReturnBestOption, FGetImageDescContext*) const override;
+		virtual Ptr<ImageSizeExpression> GetImageSizeExpression() const override;
+		virtual void GetLayoutBlockSize(int32* OutBlockX, int32* OutBlockY) override;
+		virtual bool IsImagePlainConstant(FVector4f& colour) const override;
+		virtual Ptr<ASTOp> OptimiseSize() const override;
+		virtual Ptr<ASTOp> OptimiseSemantic(const FModelOptimizationOptions&, int32 Pass) const override;
+		virtual FSourceDataDescriptor GetSourceDataDescriptor(FGetSourceDataDescriptorContext*) const override;
+
+	};
+
+
+}
+
diff --git a/Engine/Plugins/Mutable/Source/MutableTools/Private/MuT/ASTOpImageSaturate.cpp b/Engine/Plugins/Mutable/Source/MutableTools/Private/MuT/ASTOpImageSaturate.cpp
new file mode 100644
--- /dev/null
+++ b/Engine/Plugins/Mutable/Source/MutableTools/Private/MuT/ASTOpImageSaturate.cpp
@@ -0,0 +1,173 @@
+// Copyright Epic Games, Inc. All Rights Reserved.
+
+#include "MuT/ASTOpImageSaturate.h"
+
+#include "Containers/Map.h"
+#include "MuR/ModelPrivate.h"
+#include "MuR/RefCounted.h"
+#include "MuR/Types.h"
+#include "MuR/ImagePrivate.h"
+#include "MuT/ASTOpConditional.h"
+#include "MuT/ASTOpSwitch.h"
+#include "MuT/ASTOpImagePatch.h"
+#include "MuT/ASTOpImagePixelFormat.h"
+#include "MuT/ASTOpImageLayer.h"
+#include "MuT/ASTOpImageLayerColor.h"
+#include "MuT/ASTOpImageRasterMesh.h"
+#include "MuT/ASTOpImageTransform.h"
+#include "MuT/ASTOpImageSwizzle.h"
+#include "MuT/ASTOpImageCompose.h"
+#include "MuT/ASTOpImageMakeGrowMap.h"
+#include "MuT/ASTOpImageResize.h"
+
+
+namespace mu
+{
+
+	ASTOpImageSaturate::ASTOpImageSaturate()
+		: Base(this),
+		Factor(this)
+	{
+	}
+
+
+	ASTOpImageSaturate::~ASTOpImageSaturate()
+	{
+		// Explicit call needed to avoid recursive destruction
+		ASTOp::RemoveChildren();
+	}
+
+
+	bool ASTOpImageSaturate::IsEqual(const ASTOp& InOther) const
+	{
+		if (InOther.GetOpType()==GetOpType())
+		{
+			const ASTOpImageSaturate* Other = static_cast<const ASTOpImageSaturate*>(&InOther);
+			return Base == Other->Base &&
+				Factor == Other->Factor;
+		}
+		return false;
+	}
+
+
+	uint64 ASTOpImageSaturate::Hash() const
+	{
+		uint64 Res = std::hash<EOpType>()(GetOpType());
+		hash_combine(Res, Base.child().get());
+		hash_combine(Res, Factor.child().get());
+		return Res;
+	}
+
+
+	Ptr<ASTOp> ASTOpImageSaturate::Clone(MapChildFuncRef MapChild) const
+	{
+		Ptr<ASTOpImageSaturate> New = new ASTOpImageSaturate();
+		New->Base = MapChild(Base.child());
+		New->Factor = MapChild(Factor.child());
+		return New;
+	}
+
+
+	void ASTOpImageSaturate::ForEachChild(const TFunctionRef<void(ASTChild&)> Func)
+	{
+		Func(Base);
+		Func(Factor);
+	}
+
+
+	void ASTOpImageSaturate::Link(FProgram& Program, FLinkerOptions*)
+	{
+		// Already linked?
+		if (!linkedAddress)
+		{
+			OP::ImageSaturateArgs Args = {};
+
+			if (Base)
+			{
+				Args.Base = Base->linkedAddress;
+			}
+			if (Factor)
+			{
+				Args.Factor = Factor->linkedAddress;
+			}
+
+			linkedAddress = (OP::ADDRESS)Program.OpAddress.Num();
+			Program.OpAddress.Add(Program.ByteCode.Num());
+			AppendCode(Program.ByteCode, GetOpType());
+			AppendCode(Program.ByteCode, Args);
+		}
+	}
+
+
+	FImageDesc ASTOpImageSaturate::GetImageDesc(bool bReturnBestOption, FGetImageDescContext* Context) const
+	{
+		FImageDesc Result;
+
+		// Local context in case it is necessary
+		FGetImageDescContext LocalContext;
+		if (!Context)
+		{
+			Context = &LocalContext;
+		}
+		else
+		{
+			// Cached result?
+			FImageDesc* PtrValue = Context->m_results.Find(this);
+			if (PtrValue)
+			{
+				return *PtrValue;
+			}
+		}
+
+		// Actual work
+		if (Base)
+		{
+			Result = Base->GetImageDesc(bReturnBestOption, Context);
+		}
+
+		// Cache the result
+		if (Context)
+		{
+			Context->m_results.Add(this, Result);
+		}
+
+		return Result;
+	}
+
+
+	Ptr<ImageSizeExpression> ASTOpImageSaturate::GetImageSizeExpression() const
+	{
+		if (Base)
+		{
+			return Base->GetImageSizeExpression();
+		}
+
+		return new ImageSizeExpression;
+	}
+
+
+	bool ASTOpImageSaturate::IsImagePlainConstant(FVector4f& OutColour) const
+	{
+		bool bResult = true;
+		OutColour = FVector4f(0.0f,0.0f,0.0f,1.0f);
+
+		if (Base)
+		{
+			bResult = Base->IsImagePlainConstant(OutColour);
+		}
+
+		return bResult;
+	}
+
+
+	FSourceDataDescriptor ASTOpImageSaturate::GetSourceDataDescriptor(FGetSourceDataDescriptorContext* Context) const
+	{
+		if (Base)
+		{
+			return Base->GetSourceDataDescriptor(Context);
+		}
+
+		return {};
+	}
+
+}
diff --git a/Engine/Plugins/Mutable/Source/MutableTools/Private/MuT/ASTOpImageSaturate.h b/Engine/Plugins/Mutable/Source/MutableTools/Private/MuT/ASTOpImageSaturate.h
new file mode 100644
--- /dev/null
+++ b/Engine/Plugins/Mutable/Source/MutableTools/Private/MuT/ASTOpImageSaturate.h
@@ -0,0 +1,45 @@
+// Copyright Epic Games, Inc. All Rights Reserved.
+
+#pragma once
+
+#include "HAL/PlatformMath.h"
+#include "MuR/Image.h"
+#include "MuR/MutableMath.h"
+#include "MuR/Operations.h"
+#include "MuR/Ptr.h"
+#include "MuT/AST.h"
+
+
+namespace mu
+{
+	struct FProgram;
+
+	class ASTOpImageSaturate final : public ASTOp
+	{
+	public:
+
+		ASTChild Base;
+		ASTChild Factor;
+
+	public:
+
+		ASTOpImageSaturate();
+		ASTOpImageSaturate(const ASTOpImageSaturate&) = delete;
+		~ASTOpImageSaturate();
+
+		virtual EOpType GetOpType() const override { return EOpType::IM_SATURATE; }
+		virtual uint64 Hash() const override;
+		virtual bool IsEqual(const ASTOp&) const override;
+		virtual Ptr<ASTOp> Clone(MapChildFuncRef) const override;
+		virtual void ForEachChild(const TFunctionRef<void(ASTChild&)>) override;
+		virtual void Link(FProgram&, FLinkerOptions*) override;
+		virtual FImageDesc GetImageDesc(bool bReturnBestOption, FGetImageDescContext*) const override;
+		virtual Ptr<ImageSizeExpression> GetImageSizeExpression() const override;
+		virtual bool IsImagePlainConstant(FVector4f& colour) const override;
+		virtual FSourceDataDescriptor GetSourceDataDescriptor(FGetSourceDataDescriptorContext*) const override;
+
+	};
+
+
+}
+
diff --git a/Engine/Plugins/Mutable/Source/MutableTools/Private/MuT/ASTOpImageSwizzle.cpp b/Engine/Plugins/Mutable/Source/MutableTools/Private/MuT/ASTOpImageSwizzle.cpp
--- a/Engine/Plugins/Mutable/Source/MutableTools/Private/MuT/ASTOpImageSwizzle.cpp
+++ b/Engine/Plugins/Mutable/Source/MutableTools/Private/MuT/ASTOpImageSwizzle.cpp
@@ -21,6 +21,11 @@
 #include "MuT/ASTOpImageRasterMesh.h"
 #include "MuT/ASTOpImageTransform.h"
 #include "MuT/ASTOpImageMipmap.h"
+#include "MuT/ASTOpImageInterpolate.h"
+#include "MuT/ASTOpImageSaturate.h"
+#include "MuT/ASTOpImagePlainColor.h"
+#include "MuT/ASTOpImageDisplace.h"
+#include "MuT/ASTOpImageInvert.h"
 #include "MuT/ASTOpSwitch.h"
 
 namespace mu
@@ -449,7 +454,7 @@ namespace mu
 			// Swizzle down compatible displaces.
 			if (!at && sourceType == EOpType::IM_DISPLACE)
 			{
-				const ASTOpFixed* FirstDisplace = static_cast<const ASTOpFixed*>(Sources[0].child().get());
+				const ASTOpImageDisplace* FirstDisplace = static_cast<const ASTOpImageDisplace*>(Sources[0].child().get());
 				check(FirstDisplace);
 
 				bool bAreAllDisplacesCompatible = true;
@@ -457,9 +462,9 @@ namespace mu
 				{
 					if (Sources[c])
 					{
-						const ASTOpFixed* Typed = static_cast<const ASTOpFixed*>(Sources[c].child().get());
+						const ASTOpImageDisplace* Typed = static_cast<const ASTOpImageDisplace*>(Sources[c].child().get());
 						check(Typed);
-						if (FirstDisplace->op.args.ImageDisplace.displacementMap != Typed->op.args.ImageDisplace.displacementMap)
+						if (FirstDisplace->DisplacementMap != Typed->DisplacementMap)
 						{
 							bAreAllDisplacesCompatible = false;
 							break;
@@ -470,19 +475,19 @@ namespace mu
 				if (bAreAllDisplacesCompatible)
 				{
 					// Move the swizzle down all the paths
-					Ptr<ASTOpFixed> NewDisplace = mu::Clone<ASTOpFixed>(FirstDisplace);
+					Ptr<ASTOpImageDisplace> NewDisplace = mu::Clone<ASTOpImageDisplace>(FirstDisplace);
 
 					Ptr<ASTOpImageSwizzle> SourceOp = mu::Clone<ASTOpImageSwizzle>(this);
 					for (int c = 0; c < MUTABLE_OP_MAX_SWIZZLE_CHANNELS; ++c)
 					{
-						const ASTOpFixed* ChannelDisplace = static_cast<const ASTOpFixed*>(Sources[c].child().get());
+						const ASTOpImageDisplace* ChannelDisplace = static_cast<const ASTOpImageDisplace*>(Sources[c].child().get());
 						if (ChannelDisplace)
 						{
-							SourceOp->Sources[c] = ChannelDisplace->children[ChannelDisplace->op.args.ImageDisplace.source].child();
+							SourceOp->Sources[c] = ChannelDisplace->Source.child();
 						}
 					}
 
-					NewDisplace->SetChild(NewDisplace->op.args.ImageDisplace.source, SourceOp);
+					NewDisplace->Source = SourceOp;
 
 					at = NewDisplace;
 				}
@@ -699,21 +704,20 @@ namespace mu
 			// Swizzle down plaincolours.
 			if (!at && sourceType == EOpType::IM_PLAINCOLOUR)
 			{
-				Ptr<ASTOpFixed> NewPlain = mu::Clone<ASTOpFixed>(channelSourceAt);
+				Ptr<ASTOpImagePlainColor> NewPlain = mu::Clone<ASTOpImagePlainColor>(channelSourceAt);
 				Ptr<ASTOpFixed> NewSwizzle = new ASTOpFixed;
 				NewSwizzle->op.type = EOpType::CO_SWIZZLE;
 				for (int c = 0; c < MUTABLE_OP_MAX_SWIZZLE_CHANNELS; ++c)
 				{
 					if (Sources[c])
 					{
-						const ASTOpFixed* TypedPlain = static_cast<const ASTOpFixed*>(Sources[c].child().get());
-
-						NewSwizzle->SetChild(NewSwizzle->op.args.ColourSwizzle.sources[c], TypedPlain->children[TypedPlain->op.args.ImagePlainColour.colour]);
+						const ASTOpImagePlainColor* TypedPlain = static_cast<const ASTOpImagePlainColor*>(Sources[c].child().get());
+						NewSwizzle->SetChild(NewSwizzle->op.args.ColourSwizzle.sources[c], TypedPlain->Color);
 					}
 					NewSwizzle->op.args.ColourSwizzle.sourceChannels[c] = SourceChannels[c];
 				}
-				NewPlain->SetChild(NewPlain->op.args.ImagePlainColour.colour, NewSwizzle);
-				NewPlain->op.args.ImagePlainColour.format = Format;
+				NewPlain->Color = NewSwizzle;
+				NewPlain->Format = Format;
 				at = NewPlain;
 			}
 
@@ -813,8 +817,8 @@ namespace mu
 						NewSwizzle = mu::Clone<ASTOpImageSwizzle>(this);
 					}
 
-					const ASTOpFixed* OldSaturate = static_cast<const ASTOpFixed*>(Sources[Channel].child().get());
-					Ptr<ASTOp> OldSaturateBase = OldSaturate->children[OldSaturate->op.args.ImageSaturate.base].child();
+					const ASTOpImageSaturate* OldSaturate = static_cast<const ASTOpImageSaturate*>(Sources[Channel].child().get());
+					Ptr<ASTOp> OldSaturateBase = OldSaturate->Base.child();
 
 					NewSwizzle->Sources[Channel] = OldSaturateBase;
 					at = NewSwizzle;
@@ -837,11 +841,11 @@ namespace mu
 			)
 		{
 			// Move the swizzle down 
-			Ptr<ASTOpFixed> NewSaturate = mu::Clone<ASTOpFixed>(Sources[0].child());
+			Ptr<ASTOpImageSaturate> NewSaturate = mu::Clone<ASTOpImageSaturate>(Sources[0].child());
 			check(NewSaturate);
 
 			Ptr<ASTOpImageSwizzle> NewSwizzle = mu::Clone<ASTOpImageSwizzle>(this);
-			Ptr<ASTOp> OldSaturateBase = NewSaturate->children[NewSaturate->op.args.ImageSaturate.base].child();
+			Ptr<ASTOp> OldSaturateBase = NewSaturate->Base.child();
 			NewSwizzle->Sources[0] = OldSaturateBase;
 			NewSwizzle->Sources[1] = OldSaturateBase;
 			NewSwizzle->Sources[2] = OldSaturateBase;
@@ -852,7 +856,7 @@ namespace mu
 				NewSwizzle->Sources[3] = OldSaturateBase;
 			}
 
-			NewSaturate->SetChild(NewSaturate->op.args.ImageSaturate.base, NewSwizzle);
+			NewSaturate->Base = NewSwizzle;
 
 			at = NewSaturate;
 		}
@@ -1240,16 +1244,14 @@ namespace mu
 
 		case EOpType::IM_INTERPOLATE:
 		{
-			Ptr<ASTOpFixed> newOp = mu::Clone<ASTOpFixed>(at);
+			Ptr<ASTOpImageInterpolate> NewOp = mu::Clone<ASTOpImageInterpolate>(at);
 
 			for (int32 v = 0; v < MUTABLE_OP_MAX_INTERPOLATE_COUNT; ++v)
 			{
-				Ptr<ASTOp> child = newOp->children[newOp->op.args.ImageInterpolate.targets[v]].child();
-				Ptr<ASTOp> bOp = Visit(child, CurrentSwizzleOp);
-				newOp->SetChild(newOp->op.args.ImageInterpolate.targets[v], bOp);
+				NewOp->Targets[v] = Visit(NewOp->Targets[v].child(), CurrentSwizzleOp);;
 			}
 
-			newAt = newOp;
+			newAt = NewOp;
 			break;
 		}
 
@@ -1283,11 +1285,11 @@ namespace mu
 
 		case EOpType::IM_DISPLACE:
 		{
-			Ptr<ASTOpFixed> NewOp = mu::Clone<ASTOpFixed>(at);
+			Ptr<ASTOpImageDisplace> NewOp = mu::Clone<ASTOpImageDisplace>(at);
 
-			Ptr<ASTOp> Child = NewOp->children[NewOp->op.args.ImageDisplace.source].child();
+			Ptr<ASTOp> Child = NewOp->Source.child();
 			Ptr<ASTOp> NewSource = Visit(Child, CurrentSwizzleOp);
-			NewOp->SetChild(NewOp->op.args.ImageDisplace.source, NewSource);
+			NewOp->Source = NewSource;
 
 			newAt = NewOp;
 			break;
@@ -1295,11 +1297,11 @@ namespace mu
 
 		case EOpType::IM_INVERT:
 		{
-			Ptr<ASTOpFixed> NewOp = mu::Clone<ASTOpFixed>(at);
+			Ptr<ASTOpImageInvert> NewOp = mu::Clone<ASTOpImageInvert>(at);
 
-			Ptr<ASTOp> Child = NewOp->children[NewOp->op.args.ImageInvert.base].child();
+			Ptr<ASTOp> Child = NewOp->Base.child();
 			Ptr<ASTOp> NewSource = Visit(Child, CurrentSwizzleOp);
-			NewOp->SetChild(NewOp->op.args.ImageInvert.base, NewSource);
+			NewOp->Base = NewSource;
 
 			newAt = NewOp;
 			break;
diff --git a/Engine/Plugins/Mutable/Source/MutableTools/Private/MuT/ASTOpImageSwizzle.h b/Engine/Plugins/Mutable/Source/MutableTools/Private/MuT/ASTOpImageSwizzle.h
--- a/Engine/Plugins/Mutable/Source/MutableTools/Private/MuT/ASTOpImageSwizzle.h
+++ b/Engine/Plugins/Mutable/Source/MutableTools/Private/MuT/ASTOpImageSwizzle.h
@@ -31,15 +31,15 @@ namespace mu
 
 		virtual EOpType GetOpType() const override { return EOpType::IM_SWIZZLE; }
 		virtual uint64 Hash() const override;
-		virtual bool IsEqual(const ASTOp& otherUntyped) const override;
-		virtual Ptr<ASTOp> Clone(MapChildFuncRef mapChild) const override;
+		virtual bool IsEqual(const ASTOp&) const override;
+		virtual Ptr<ASTOp> Clone(MapChildFuncRef) const override;
 		virtual void ForEachChild(const TFunctionRef<void(ASTChild&)>) override;
-		virtual void Link(FProgram& program, FLinkerOptions* Options) override;
+		virtual void Link(FProgram&, FLinkerOptions*) override;
 		virtual Ptr<ASTOp> OptimiseSemantic(const FModelOptimizationOptions&, int32 Pass) const override;
-		virtual Ptr<ASTOp> OptimiseSink(const FModelOptimizationOptions& options, FOptimizeSinkContext& context) const override;
-		virtual FImageDesc GetImageDesc(bool returnBestOption, FGetImageDescContext* context) const override;
-		virtual void GetLayoutBlockSize(int* pBlockX, int* pBlockY) override;
-		virtual bool IsImagePlainConstant(FVector4f& colour) const override;
+		virtual Ptr<ASTOp> OptimiseSink(const FModelOptimizationOptions&, FOptimizeSinkContext&) const override;
+		virtual FImageDesc GetImageDesc(bool bReturnBestOption, FGetImageDescContext*) const override;
+		virtual void GetLayoutBlockSize(int32* OutBlockX, int32* OutBlockY) override;
+		virtual bool IsImagePlainConstant(FVector4f& OutColour) const override;
 		virtual Ptr<ImageSizeExpression> GetImageSizeExpression() const override;
 		virtual FSourceDataDescriptor GetSourceDataDescriptor(FGetSourceDataDescriptorContext*) const override;
 
diff --git a/Engine/Plugins/Mutable/Source/MutableTools/Private/MuT/CodeGenerator.cpp b/Engine/Plugins/Mutable/Source/MutableTools/Private/MuT/CodeGenerator.cpp
--- a/Engine/Plugins/Mutable/Source/MutableTools/Private/MuT/CodeGenerator.cpp
+++ b/Engine/Plugins/Mutable/Source/MutableTools/Private/MuT/CodeGenerator.cpp
@@ -27,6 +27,8 @@
 #include "MuT/ASTOpImageLayerColor.h"
 #include "MuT/ASTOpImagePatch.h"
 #include "MuT/ASTOpImageCrop.h"
+#include "MuT/ASTOpImageBlankLayout.h"
+#include "MuT/ASTOpImagePlainColor.h"
 #include "MuT/ASTOpInstanceAdd.h"
 #include "MuT/ASTOpMeshBindShape.h"
 #include "MuT/ASTOpMeshClipDeform.h"
@@ -738,12 +740,11 @@ namespace mu
 			return Source;
 		}
 
-		Ptr<ASTOpFixed> BaseImage = new ASTOpFixed;
-		BaseImage->op.type = EOpType::IM_PLAINCOLOUR;
-		BaseImage->op.args.ImagePlainColour.size[0] = Size[0];
-		BaseImage->op.args.ImagePlainColour.size[1] = Size[1];
-		BaseImage->op.args.ImagePlainColour.format = Format;
-		BaseImage->op.args.ImagePlainColour.LODs = 1;
+		Ptr<ASTOpImagePlainColor> BaseImage = new ASTOpImagePlainColor;
+		BaseImage->Size[0] = Size[0];
+		BaseImage->Size[1] = Size[1];
+		BaseImage->Format = Format;
+		BaseImage->LODs = 1;
 
 		Ptr<ASTOp> CurrentImage = BaseImage;
 
@@ -1297,14 +1298,13 @@ namespace mu
 							bool bImageSizeWarning = false;
 
 							// Start with a blank image. It will be completed later with the blockSize, format and mips information
-							Ptr<ASTOpFixed> BlankImageOp;
+							Ptr<ASTOpImageBlankLayout> BlankImageOp;
 							Ptr<ASTOp> imageAd;
 							{
-								BlankImageOp = new ASTOpFixed();
-								BlankImageOp->op.type = EOpType::IM_BLANKLAYOUT;
-								BlankImageOp->SetChild(BlankImageOp->op.args.ImageBlankLayout.layout, MeshResults.LayoutOps[LayoutIndex]);
+								BlankImageOp = new ASTOpImageBlankLayout();
+								BlankImageOp->Layout = MeshResults.LayoutOps[LayoutIndex];
 								// The rest ok the op will be completed below
-								BlankImageOp->op.args.ImageBlankLayout.mipmapCount = 0;
+								BlankImageOp->MipmapCount = 0;
 								imageAd = BlankImageOp;
 							}
 
@@ -1382,11 +1382,11 @@ namespace mu
 								SurfaceNode->GetMessageContext());
 
 							// Complete the base op
-							BlankImageOp->op.args.ImageBlankLayout.blockSize[0] = uint16(LayoutBlockDesc.BlockPixelsX);
-							BlankImageOp->op.args.ImageBlankLayout.blockSize[1] = uint16(LayoutBlockDesc.BlockPixelsY);
-							BlankImageOp->op.args.ImageBlankLayout.format = GetUncompressedFormat(LayoutBlockDesc.FinalFormat);
-							BlankImageOp->op.args.ImageBlankLayout.generateMipmaps = LayoutBlockDesc.bBlocksHaveMips;
-							BlankImageOp->op.args.ImageBlankLayout.mipmapCount = 0;
+							BlankImageOp->BlockSize[0] = uint16(LayoutBlockDesc.BlockPixelsX);
+							BlankImageOp->BlockSize[1] = uint16(LayoutBlockDesc.BlockPixelsY);
+							BlankImageOp->Format = GetUncompressedFormat(LayoutBlockDesc.FinalFormat);
+							BlankImageOp->GenerateMipmaps = LayoutBlockDesc.bBlocksHaveMips;
+							BlankImageOp->MipmapCount = 0;
 
 							if (swizzleNode)
 							{
diff --git a/Engine/Plugins/Mutable/Source/MutableTools/Private/MuT/CodeGenerator.h b/Engine/Plugins/Mutable/Source/MutableTools/Private/MuT/CodeGenerator.h
--- a/Engine/Plugins/Mutable/Source/MutableTools/Private/MuT/CodeGenerator.h
+++ b/Engine/Plugins/Mutable/Source/MutableTools/Private/MuT/CodeGenerator.h
@@ -56,7 +56,6 @@ namespace mu
 	class NodeImageConditional;
 	class NodeImageConstant;
 	class NodeImageFormat;
-	class NodeImageGradient;
 	class NodeImageInterpolate;
 	class NodeImageInvert;
 	class NodeImageLayer;
@@ -453,7 +452,6 @@ namespace mu
 		void GenerateImage_Table(const FImageGenerationOptions&, FImageGenerationResult&, const NodeImageTable*);
 		void GenerateImage_Swizzle(const FImageGenerationOptions&, FImageGenerationResult&, const NodeImageSwizzle*);
 		void GenerateImage_ColourMap(const FImageGenerationOptions&, FImageGenerationResult&, const NodeImageColourMap*);
-		void GenerateImage_Gradient(const FImageGenerationOptions&, FImageGenerationResult&, const NodeImageGradient*);
 		void GenerateImage_Binarise(const FImageGenerationOptions&, FImageGenerationResult&, const NodeImageBinarise*);
 		void GenerateImage_Luminance(const FImageGenerationOptions&, FImageGenerationResult&, const NodeImageLuminance*);
 		void GenerateImage_Layer(const FImageGenerationOptions&, FImageGenerationResult&, const NodeImageLayer*);
diff --git a/Engine/Plugins/Mutable/Source/MutableTools/Private/MuT/CodeGenerator_Image.cpp b/Engine/Plugins/Mutable/Source/MutableTools/Private/MuT/CodeGenerator_Image.cpp
--- a/Engine/Plugins/Mutable/Source/MutableTools/Private/MuT/CodeGenerator_Image.cpp
+++ b/Engine/Plugins/Mutable/Source/MutableTools/Private/MuT/CodeGenerator_Image.cpp
@@ -28,6 +28,16 @@
 #include "MuT/ASTOpImageSwizzle.h"
 #include "MuT/ASTOpImageRasterMesh.h"
 #include "MuT/ASTOpImageCrop.h"
+#include "MuT/ASTOpImageResize.h"
+#include "MuT/ASTOpImageResizeRel.h"
+#include "MuT/ASTOpImageInterpolate.h"
+#include "MuT/ASTOpImageSaturate.h"
+#include "MuT/ASTOpImageLuminance.h"
+#include "MuT/ASTOpImageColorMap.h"
+#include "MuT/ASTOpImageBinarize.h"
+#include "MuT/ASTOpImagePlainColor.h"
+#include "MuT/ASTOpImageDisplace.h"
+#include "MuT/ASTOpImageInvert.h"
 #include "MuT/ASTOpMeshExtractLayoutBlocks.h"
 #include "MuT/ASTOpMeshFormat.h"
 #include "MuT/ASTOpMeshProject.h"
@@ -53,8 +63,6 @@
 #include "MuT/NodeImageConstant.h"
 #include "MuT/NodeImageFormat.h"
 #include "MuT/NodeImageFormatPrivate.h"
-#include "MuT/NodeImageGradient.h"
-#include "MuT/NodeImageGradientPrivate.h"
 #include "MuT/NodeImageInterpolate.h"
 #include "MuT/NodeImageInterpolatePrivate.h"
 #include "MuT/NodeImageInvert.h"
@@ -129,7 +137,6 @@ namespace mu
 			case Node::EType::ImageTable: GenerateImage_Table(Options, Result, static_cast<const NodeImageTable*>(Node)); break;
 			case Node::EType::ImageSwizzle: GenerateImage_Swizzle(Options, Result, static_cast<const NodeImageSwizzle*>(Node)); break;
 			case Node::EType::ImageColorMap: GenerateImage_ColourMap(Options, Result, static_cast<const NodeImageColourMap*>(Node)); break;
-			case Node::EType::ImageGradient: GenerateImage_Gradient(Options, Result, static_cast<const NodeImageGradient*>(Node)); break;
 			case Node::EType::ImageBinarise: GenerateImage_Binarise(Options, Result, static_cast<const NodeImageBinarise*>(Node)); break;
 			case Node::EType::ImageLuminance: GenerateImage_Luminance(Options, Result, static_cast<const NodeImageLuminance*>(Node)); break;
 			case Node::EType::ImageLayer: GenerateImage_Layer(Options, Result, static_cast<const NodeImageLayer*>(Node)); break;
@@ -688,19 +695,17 @@ namespace mu
 
 		MUTABLE_CPUPROFILER_SCOPE(NodeImageInterpolate);
 
-        Ptr<ASTOpFixed> op = new ASTOpFixed();
-        op->op.type = EOpType::IM_INTERPOLATE;
+        Ptr<ASTOpImageInterpolate> op = new ASTOpImageInterpolate();
 
         // Factor
-        if ( Node* pFactor = node.m_pFactor.get() )
+        if ( Node* Factor = node.m_pFactor.get() )
         {
-            op->SetChild( op->op.args.ImageInterpolate.factor, Generate_Generic( pFactor, Options));
+            op->Factor = Generate_Generic(Factor, Options);
         }
         else
         {
             // This argument is required
-            op->SetChild( op->op.args.ImageInterpolate.factor,
-                          GenerateMissingScalarCode(TEXT("Interpolation factor"), 0.5f, InNode->GetMessageContext() ));
+            op->Factor = GenerateMissingScalarCode(TEXT("Interpolation factor"), 0.5f, InNode->GetMessageContext() );
         }
 
         // Target images
@@ -708,16 +713,16 @@ namespace mu
 
 		UE::Math::TIntVector2<int32> FinalRectSize = Options.RectSize;
 
-        for ( std::size_t t=0
-            ; t< node.m_targets.Num() && numTargets<MUTABLE_OP_MAX_INTERPOLATE_COUNT
-            ; ++t )
+        for ( int32 TargetIndex=0
+            ; TargetIndex < node.m_targets.Num() && numTargets<MUTABLE_OP_MAX_INTERPOLATE_COUNT
+            ; ++TargetIndex)
         {
-            if ( node.m_targets[t] )
+            if ( node.m_targets[TargetIndex] )
             {
 				FImageGenerationOptions ChildOptions = Options;
 				ChildOptions.RectSize = FinalRectSize;
 				FImageGenerationResult BaseResult;
-				GenerateImage(ChildOptions, BaseResult, node.m_targets[t]);
+				GenerateImage(ChildOptions, BaseResult, node.m_targets[TargetIndex]);
 				Ptr<ASTOp> target = BaseResult.op;
 
 				if (FinalRectSize[0] == 0)
@@ -730,17 +735,17 @@ namespace mu
                 target = GenerateImageFormat( target, EImageFormat::RGB_UByte );
                 target = GenerateImageSize( target, FinalRectSize);
 
-                op->SetChild( op->op.args.ImageInterpolate.targets[numTargets], target);
+                op->Targets[numTargets] = target;
                 numTargets++;
             }
         }
 
         // At least one target is required
-        if (!op->op.args.ImageInterpolate.targets[0])
+        if (!op->Targets[0])
         {
             Ptr<ASTOp> target = GenerateMissingImageCode(TEXT("First interpolation image"), EImageFormat::RGB_UByte, InNode->GetMessageContext(), Options);
             target = GenerateImageSize( target, Options.RectSize);
-            op->SetChild( op->op.args.ImageInterpolate.targets[0], target);
+            op->Targets[0] = target;
         }
 
         Result.op = op;
@@ -902,38 +907,36 @@ namespace mu
 	{
 		const NodeImageSaturate::Private& node = *InNode->GetPrivate();
 
-		Ptr<ASTOpFixed> op = new ASTOpFixed();
-        op->op.type = EOpType::IM_SATURATE;
-
+		Ptr<ASTOpImageSaturate> op = new ASTOpImageSaturate();
 
         // Source image
-        Ptr<ASTOp> base = 0;
+        Ptr<ASTOp> Base;
         if ( node.m_pSource )
         {
 			FImageGenerationResult BaseResult;
 			GenerateImage(Options, BaseResult, node.m_pSource);
-			base = BaseResult.op;
+			Base = BaseResult.op;
 		}
         else
         {
             // This argument is required
-            base = GenerateMissingImageCode(TEXT("Saturate image"), EImageFormat::RGB_UByte, InNode->GetMessageContext(), Options);
+			Base = GenerateMissingImageCode(TEXT("Saturate image"), EImageFormat::RGB_UByte, InNode->GetMessageContext(), Options);
         }
 		
-        base = GenerateImageFormat(base, GetRGBOrRGBAFormat(base->GetImageDesc().m_format));
-        base = GenerateImageSize( base, Options.RectSize);
-        op->SetChild( op->op.args.ImageSaturate.base, base);
+		Base = GenerateImageFormat(Base, GetRGBOrRGBAFormat(Base->GetImageDesc().m_format));
+		Base = GenerateImageSize(Base, Options.RectSize);
+        op->Base = Base;
 
 
         // Factor
-        if ( Node* pFactor = node.m_pFactor.get() )
+        if ( Node* Factor = node.m_pFactor.get() )
         {
-            op->SetChild( op->op.args.ImageSaturate.factor, Generate_Generic( pFactor, Options));
+            op->Factor = Generate_Generic(Factor, Options);
         }
         else
         {
             // This argument is required
-            op->SetChild( op->op.args.ImageSaturate.factor, GenerateMissingScalarCode(TEXT("Saturation factor"), 0.5f, InNode->GetMessageContext() ) );
+            op->Factor = GenerateMissingScalarCode(TEXT("Saturation factor"), 0.5f, InNode->GetMessageContext() );
         }
 
         Result.op = op;
@@ -945,8 +948,7 @@ namespace mu
 	{
 		const NodeImageLuminance::Private& node = *InNode->GetPrivate();
 
-		Ptr<ASTOpFixed> op = new ASTOpFixed();
-        op->op.type = EOpType::IM_LUMINANCE;
+		Ptr<ASTOpImageLuminance> op = new ASTOpImageLuminance();
 
         // Source image
         Ptr<ASTOp> base;
@@ -962,7 +964,7 @@ namespace mu
             base = GenerateMissingImageCode(TEXT("Image luminance"), EImageFormat::RGB_UByte, InNode->GetMessageContext(), Options);
         }
         base = GenerateImageFormat( base, EImageFormat::RGB_UByte );
-        op->SetChild( op->op.args.ImageLuminance.base, base);
+        op->Base = base;
 
         Result.op = op;
     }
@@ -973,8 +975,7 @@ namespace mu
 	{
 		const NodeImageColourMap::Private& node = *InNode->GetPrivate();
 
-		Ptr<ASTOpFixed> op = new ASTOpFixed();
-        op->op.type = EOpType::IM_COLOURMAP;
+		Ptr<ASTOpImageColorMap> op = new ASTOpImageColorMap();
 
         // Base image
         Ptr<ASTOp> base ;
@@ -990,7 +991,7 @@ namespace mu
             base = GenerateMissingImageCode(TEXT("Colourmap base image"), EImageFormat::RGB_UByte, InNode->GetMessageContext(), Options);
         }
         base = GenerateImageSize( base, Options.RectSize);
-        op->SetChild( op->op.args.ImageColourMap.base, base);
+        op->Base = base;
 
         // Mask of the effect
         Ptr<ASTOp> mask;
@@ -1008,7 +1009,7 @@ namespace mu
         }
         mask = GenerateImageFormat( mask, EImageFormat::L_UByte );
         mask = GenerateImageSize( mask, Options.RectSize);
-        op->SetChild( op->op.args.ImageColourMap.mask, mask);
+        op->Mask = mask;
 
         // Map image
         Ptr<ASTOp> MapImageOp;
@@ -1030,48 +1031,7 @@ namespace mu
             // This argument is required
 			MapImageOp = GenerateMissingImageCode(TEXT("Map image"), EImageFormat::RGB_UByte, InNode->GetMessageContext(), Options);
         }
-        op->SetChild( op->op.args.ImageColourMap.map, MapImageOp);
-
-        Result.op = op;
-    }
-
-
-    //---------------------------------------------------------------------------------------------
-	void CodeGenerator::GenerateImage_Gradient(const FImageGenerationOptions& Options, FImageGenerationResult& Result, const NodeImageGradient* InNode)
-	{
-		const NodeImageGradient::Private& node = *InNode->GetPrivate();
-
-		Ptr<ASTOpFixed> op = new ASTOpFixed();
-        op->op.type = EOpType::IM_GRADIENT;
-
-        // First colour
-        Ptr<ASTOp> colour0 = 0;
-        if ( Node* pColour0 = node.m_pColour0.get() )
-        {
-            colour0 = Generate_Generic( pColour0, Options);
-        }
-        else
-        {
-            // This argument is required
-            colour0 = GenerateMissingColourCode(TEXT("Gradient colour 0"), InNode->GetMessageContext() );
-        }
-        op->SetChild( op->op.args.ImageGradient.colour0, colour0);
-
-        // Second colour
-        Ptr<ASTOp> colour1 = 0;
-        if ( Node* pColour1 = node.m_pColour1.get() )
-        {
-            colour1 = Generate_Generic( pColour1, Options);
-        }
-        else
-        {
-            // This argument is required
-            colour1 = GenerateMissingColourCode(TEXT("Gradient colour 1"), InNode->GetMessageContext() );
-        }
-        op->SetChild( op->op.args.ImageGradient.colour1, colour1);
-
-        op->op.args.ImageGradient.size[0] = (uint16)FMath::Max( 2, FMath::Min( 1024, node.m_size[0] ) );
-        op->op.args.ImageGradient.size[1] = (uint16)FMath::Max( 1, FMath::Min( 1024, node.m_size[1] ) );
+        op->Map = MapImageOp;
 
         Result.op = op;
     }
@@ -1082,38 +1042,37 @@ namespace mu
 	{
 		const NodeImageBinarise::Private& node = *InNode->GetPrivate();
 
-		Ptr<ASTOpFixed> op = new ASTOpFixed();
-        op->op.type = EOpType::IM_BINARISE;
+		Ptr<ASTOpImageBinarize> op = new ASTOpImageBinarize();
 
         // A image
-        Ptr<ASTOp> a;
+        Ptr<ASTOp> BaseOp;
         if ( node.m_pBase )
         {
 			FImageGenerationResult BaseResult;
 			GenerateImage(Options, BaseResult, node.m_pBase);
-			a = BaseResult.op;
+			BaseOp = BaseResult.op;
 		}
         else
         {
             // This argument is required
-            a = GenerateMissingImageCode(TEXT("Image Binarise Base"), EImageFormat::RGB_UByte, InNode->GetMessageContext(), Options);
+			BaseOp = GenerateMissingImageCode(TEXT("Image Binarise Base"), EImageFormat::RGB_UByte, InNode->GetMessageContext(), Options);
         }
-        a = GenerateImageFormat( a, EImageFormat::RGB_UByte );
-        a = GenerateImageSize( a, Options.RectSize);
-        op->SetChild( op->op.args.ImageBinarise.base, a );
+		BaseOp = GenerateImageFormat(BaseOp, EImageFormat::RGB_UByte );
+		BaseOp = GenerateImageSize(BaseOp, Options.RectSize);
+        op->Base = BaseOp;
 
         // Threshold
-        Ptr<ASTOp> b = 0;
+        Ptr<ASTOp> ThresholdOp = 0;
         if ( Node* pScalar = node.m_pThreshold.get() )
         {
-            b = Generate_Generic( pScalar, Options);
+			ThresholdOp = Generate_Generic( pScalar, Options);
         }
         else
         {
             // This argument is required
-            b = GenerateMissingScalarCode(TEXT("Image Binarise Threshold"), 0.5f, InNode->GetMessageContext() );
+			ThresholdOp = GenerateMissingScalarCode(TEXT("Image Binarise Threshold"), 0.5f, InNode->GetMessageContext() );
         }
-        op->SetChild( op->op.args.ImageBinarise.threshold, b );
+        op->Threshold = ThresholdOp;
 
         Result.op = op;
     }
@@ -1153,12 +1112,11 @@ namespace mu
         // Size
         if ( node.m_relative )
         {
-            Ptr<ASTOpFixed> op = new ASTOpFixed();
-            op->op.type = EOpType::IM_RESIZEREL;
-            op->op.args.ImageResizeRel.factor[0] = node.m_sizeX;
-            op->op.args.ImageResizeRel.factor[1] = node.m_sizeY;
-            op->SetChild( op->op.args.ImageResizeRel.source, base);
-            at = op;
+            Ptr<ASTOpImageResizeRel> NewOp = new ASTOpImageResizeRel();
+			NewOp->Factor[0] = node.m_sizeX;
+			NewOp->Factor[1] = node.m_sizeY;
+			NewOp->Source = base;
+            at = NewOp;
         }
         else
         {
@@ -1183,12 +1141,11 @@ namespace mu
 				FinalImageSize[1] = (RectInCells.size[1] * FinalImageSize[1]) / grid[1];
 			}
 
-            Ptr<ASTOpFixed> op = new ASTOpFixed();
-            op->op.type = EOpType::IM_RESIZE;
-            op->op.args.ImageResize.size[0] = (uint16)FinalImageSize[0];
-            op->op.args.ImageResize.size[1] = (uint16)FinalImageSize[1];
-            op->SetChild( op->op.args.ImageResize.source, base);
-            at = op;
+            Ptr<ASTOpImageResize> ResizeOp = new ASTOpImageResize();
+			ResizeOp->Size[0] = (uint16)FinalImageSize[0];
+			ResizeOp->Size[1] = (uint16)FinalImageSize[1];
+			ResizeOp->Source = base;
+            at = ResizeOp;
         }
 
         Result.op = at;
@@ -1244,13 +1201,12 @@ namespace mu
 		}
 
 
-        Ptr<ASTOpFixed> op = new ASTOpFixed();
-        op->op.type = EOpType::IM_PLAINCOLOUR;
-        op->SetChild( op->op.args.ImagePlainColour.colour, base);
-		op->op.args.ImagePlainColour.format = node.Format;
-		op->op.args.ImagePlainColour.size[0] = FinalImageSize[0];
-		op->op.args.ImagePlainColour.size[1] = FinalImageSize[1];
-		op->op.args.ImagePlainColour.LODs = 1;
+        Ptr<ASTOpImagePlainColor> op = new ASTOpImagePlainColor();
+        op->Color = base;
+		op->Format = node.Format;
+		op->Size[0] = FinalImageSize[0];
+		op->Size[1] = FinalImageSize[1];
+		op->LODs = 1;
 
         Result.op = op;
     }
@@ -1671,10 +1627,9 @@ namespace mu
 				MakeGrowMapOp->Mask = MipMask;
 			}
 
-			Ptr<ASTOpFixed> DisplaceOp = new ASTOpFixed();
-			DisplaceOp->op.type = EOpType::IM_DISPLACE;
-			DisplaceOp->SetChild(DisplaceOp->op.args.ImageDisplace.displacementMap, MakeGrowMapOp);
-			DisplaceOp->SetChild(DisplaceOp->op.args.ImageDisplace.source, ImageRasterOp);
+			Ptr<ASTOpImageDisplace> DisplaceOp = new ASTOpImageDisplace;
+			DisplaceOp->DisplacementMap = MakeGrowMapOp;
+			DisplaceOp->Source = ImageRasterOp;
 
 			Result.op = DisplaceOp;
 		}
@@ -1754,10 +1709,8 @@ namespace mu
 	{
 		const NodeImageInvert::Private& node = *InNode->GetPrivate();
 
-		Ptr<ASTOpFixed> op = new ASTOpFixed();
-		op->op.type = EOpType::IM_INVERT;
+		Ptr<ASTOpImageInvert> op = new ASTOpImageInvert();
 
-		// A image
 		Ptr<ASTOp> a;
 		if (node.m_pBase)
 		{
@@ -1772,7 +1725,7 @@ namespace mu
 		}
 		a = GenerateImageFormat(a, EImageFormat::RGB_UByte);
 		a = GenerateImageSize(a, Options.RectSize);
-		op->SetChild(op->op.args.ImageInvert.base, a);
+		op->Base = a;
 
 		Result.op = op;
 	}
@@ -1853,11 +1806,10 @@ namespace mu
 
         if ( desc.m_size[0] != 0 && desc.m_size[1] != 0 )
         {
-            Ptr<ASTOpFixed> sop = new ASTOpFixed();
-            sop->op.type = EOpType::IM_RESIZE;
-            sop->op.args.ImageResize.size[0] = desc.m_size[0];
-            sop->op.args.ImageResize.size[1] = desc.m_size[1];
-            sop->SetChild( sop->op.args.ImageResize.source, currentOp );
+            Ptr<ASTOpImageResize> sop = new ASTOpImageResize();
+            sop->Size[0] = desc.m_size[0];
+            sop->Size[1] = desc.m_size[1];
+            sop->Source = currentOp;
             currentOp = sop;
         }
 
@@ -1923,12 +1875,11 @@ namespace mu
 					{
 						// Use a relative resize, because at this point we may be generating a layout block and not the full image
 						float Factor = FMath::Min(MaxTextureSize / (float)(pImage->GetSizeX()), MaxTextureSize / (float)(pImage->GetSizeY()));
-						Ptr<ASTOpFixed> op = new ASTOpFixed();
-						op->op.type = EOpType::IM_RESIZEREL;
-						op->op.args.ImageResizeRel.factor[0] = Factor;
-						op->op.args.ImageResizeRel.factor[1] = Factor;
-						op->SetChild(op->op.args.ImageResizeRel.source, ImageOp);
-						ImageOp = op;
+						Ptr<ASTOpImageResizeRel> ResizeOp = new ASTOpImageResizeRel();
+						ResizeOp->Factor[0] = Factor;
+						ResizeOp->Factor[1] = Factor;
+						ResizeOp->Source = ImageOp;
+						ImageOp = ResizeOp;
 					}
 				}
 
@@ -2115,11 +2066,10 @@ namespace mu
 		{
 			if (UE::Math::TIntVector2<int32>(at->GetImageDesc().m_size) != size)
 			{
-				Ptr<ASTOpFixed> op = new ASTOpFixed();
-				op->op.type = EOpType::IM_RESIZE;
-				op->SetChild(op->op.args.ImageResize.source, at);
-				op->op.args.ImageResize.size[0] = size[0];
-				op->op.args.ImageResize.size[1] = size[1];
+				Ptr<ASTOpImageResize> op = new ASTOpImageResize();
+				op->Source = at;
+				op->Size[0] = size[0];
+				op->Size[1] = size[1];
 				Result = op;
 			}
 		}
diff --git a/Engine/Plugins/Mutable/Source/MutableTools/Private/MuT/CodeOptimiser_Parameters.cpp b/Engine/Plugins/Mutable/Source/MutableTools/Private/MuT/CodeOptimiser_Parameters.cpp
--- a/Engine/Plugins/Mutable/Source/MutableTools/Private/MuT/CodeOptimiser_Parameters.cpp
+++ b/Engine/Plugins/Mutable/Source/MutableTools/Private/MuT/CodeOptimiser_Parameters.cpp
@@ -8,6 +8,9 @@
 #include "MuT/ASTOpImagePixelFormat.h"
 #include "MuT/ASTOpImageLayer.h"
 #include "MuT/ASTOpImageLayerColor.h"
+#include "MuT/ASTOpImageResizeLike.h"
+#include "MuT/ASTOpImagePlainColor.h"
+#include "MuT/ASTOpImageDisplace.h"
 #include "MuT/ASTOpMeshMorph.h"
 #include "MuT/ASTOpInstanceAdd.h"
 #include "MuT/ASTOpLayoutFromMesh.h"
@@ -222,9 +225,9 @@ namespace mu
     //---------------------------------------------------------------------------------------------
     //---------------------------------------------------------------------------------------------
     //---------------------------------------------------------------------------------------------
-    Ptr<ASTOp> EnsureValidMask( Ptr<ASTOp> mask, Ptr<ASTOp> base )
+    Ptr<ASTOp> EnsureValidMask( Ptr<ASTOp> Mask, Ptr<ASTOp> Base )
     {
-        if ( !mask )
+        if ( !Mask)
         {
             Ptr<ASTOpFixed> whiteOp = new ASTOpFixed;
             whiteOp->op.type = EOpType::CO_CONSTANT;
@@ -233,23 +236,21 @@ namespace mu
 			whiteOp->op.args.ColourConstant.value[2] = 1;
 			whiteOp->op.args.ColourConstant.value[3] = 1;
 
-            Ptr<ASTOpFixed> wplainOp = new ASTOpFixed;
-            wplainOp->op.type = EOpType::IM_PLAINCOLOUR;
-            wplainOp->SetChild( wplainOp->op.args.ImagePlainColour.colour, whiteOp );
-            wplainOp->op.args.ImagePlainColour.format = EImageFormat::L_UByte;
-            wplainOp->op.args.ImagePlainColour.size[0] = 4;
-			wplainOp->op.args.ImagePlainColour.size[1] = 4;
-			wplainOp->op.args.ImagePlainColour.LODs = 1;
+            Ptr<ASTOpImagePlainColor> wplainOp = new ASTOpImagePlainColor;
+            wplainOp->Color = whiteOp;
+            wplainOp->Format = EImageFormat::L_UByte;
+            wplainOp->Size[0] = 4;
+			wplainOp->Size[1] = 4;
+			wplainOp->LODs = 1;
 
-            Ptr<ASTOpFixed> wresizeOp = new ASTOpFixed;
-            wresizeOp->op.type = EOpType::IM_RESIZELIKE;
-            wresizeOp->SetChild( wresizeOp->op.args.ImageResizeLike.source, wplainOp );
-            wresizeOp->SetChild( wresizeOp->op.args.ImageResizeLike.sizeSource, base );
+            Ptr<ASTOpImageResizeLike> ResizeOp = new ASTOpImageResizeLike;
+			ResizeOp->Source = wplainOp;
+			ResizeOp->SizeSource = Base;
 
-            mask = wresizeOp;
+			Mask = ResizeOp;
         }
 
-        return mask;
+        return Mask;
     }
 
 
@@ -457,21 +458,19 @@ namespace mu
 							blackOp->op.args.ColourConstant.value[2] = 0;
 							blackOp->op.args.ColourConstant.value[3] = 1;
 
-                            Ptr<ASTOpFixed> plainOp = new ASTOpFixed;
-                            plainOp->op.type = EOpType::IM_PLAINCOLOUR;
-                            plainOp->SetChild( plainOp->op.args.ImagePlainColour.colour, blackOp );
-                            plainOp->op.args.ImagePlainColour.format = EImageFormat::L_UByte;
-                            plainOp->op.args.ImagePlainColour.size[0] = 4;
-							plainOp->op.args.ImagePlainColour.size[1] = 4;
-							plainOp->op.args.ImagePlainColour.LODs = 1;
+                            Ptr<ASTOpImagePlainColor> plainOp = new ASTOpImagePlainColor;
+                            plainOp->Color = blackOp;
+                            plainOp->Format = EImageFormat::L_UByte;
+                            plainOp->Size[0] = 4;
+							plainOp->Size[1] = 4;
+							plainOp->LODs = 1;
 
-                            Ptr<ASTOpFixed> resizeOp = new ASTOpFixed;
-                            resizeOp->op.type = EOpType::IM_RESIZELIKE;
-                            resizeOp->SetChild( resizeOp->op.args.ImageResizeLike.source, plainOp );
-                            resizeOp->SetChild( resizeOp->op.args.ImageResizeLike.sizeSource, typedYes->base );
+                            Ptr<ASTOpImageResizeLike> ResizeOp = new ASTOpImageResizeLike;
+							ResizeOp->Source = plainOp;
+							ResizeOp->SizeSource = typedYes->base.child();
 
                             Ptr<ASTOpConditional> maskOp = mu::Clone<ASTOpConditional>(TypedOp);
-                            maskOp->no = resizeOp;
+                            maskOp->no = ResizeOp;
 
                             // If there is no mask (because it is optional), we need to make a
                             // white plain image
@@ -564,21 +563,19 @@ namespace mu
 							blackOp->op.args.ColourConstant.value[2] = 0;
 							blackOp->op.args.ColourConstant.value[3] = 1;
 
-                            Ptr<ASTOpFixed> plainOp = new ASTOpFixed;
-                            plainOp->op.type = EOpType::IM_PLAINCOLOUR;
-                            plainOp->SetChild( plainOp->op.args.ImagePlainColour.colour, blackOp );
-                            plainOp->op.args.ImagePlainColour.format = EImageFormat::L_UByte;
-                            plainOp->op.args.ImagePlainColour.size[0] = 4;
-							plainOp->op.args.ImagePlainColour.size[1] = 4;
-							plainOp->op.args.ImagePlainColour.LODs = 1;
+                            Ptr<ASTOpImagePlainColor> plainOp = new ASTOpImagePlainColor;
+                            plainOp->Color = blackOp;
+                            plainOp->Format = EImageFormat::L_UByte;
+                            plainOp->Size[0] = 4;
+							plainOp->Size[1] = 4;
+							plainOp->LODs = 1;
 
-                            Ptr<ASTOpFixed> resizeOp = new ASTOpFixed;
-                            resizeOp->op.type = EOpType::IM_RESIZELIKE;
-                            resizeOp->SetChild( resizeOp->op.args.ImageResizeLike.source, plainOp );
-                            resizeOp->SetChild( resizeOp->op.args.ImageResizeLike.sizeSource, typedNo->base );
+                            Ptr<ASTOpImageResizeLike> ResizeOp = new ASTOpImageResizeLike;
+							ResizeOp->Source = plainOp;
+							ResizeOp->SizeSource = typedNo->base.child();
 
                             Ptr<ASTOpConditional> maskOp = mu::Clone<ASTOpConditional>(TypedOp);
-                            maskOp->no = resizeOp;
+                            maskOp->no = ResizeOp;
 
                             // If there is no mask (because it is optional), we need to make a
                             // white plain image
@@ -1028,20 +1025,18 @@ namespace mu
                         blackOp->op.args.ColourConstant.value[2] = 0;
 						blackOp->op.args.ColourConstant.value[3] = 1;
 
-                        Ptr<ASTOpFixed> plainOp = new ASTOpFixed;
-                        plainOp->op.type = EOpType::IM_PLAINCOLOUR;
-                        plainOp->SetChild( plainOp->op.args.ImagePlainColour.colour, blackOp );
-                        plainOp->op.args.ImagePlainColour.format = EImageFormat::L_UByte;
-                        plainOp->op.args.ImagePlainColour.size[0] = 4;
-						plainOp->op.args.ImagePlainColour.size[1] = 4;
-						plainOp->op.args.ImagePlainColour.LODs = 1;
+                        Ptr<ASTOpImagePlainColor> plainOp = new ASTOpImagePlainColor;
+                        plainOp->Color = blackOp;
+                        plainOp->Format = EImageFormat::L_UByte;
+                        plainOp->Size[0] = 4;
+						plainOp->Size[1] = 4;
+						plainOp->LODs = 1;
 
-                        Ptr<ASTOpFixed> baseResizeOp = new ASTOpFixed;
-                        baseResizeOp->op.type = EOpType::IM_RESIZELIKE;
-                        baseResizeOp->SetChild( baseResizeOp->op.args.ImageResizeLike.sizeSource, baseAt );
-                        baseResizeOp->SetChild( baseResizeOp->op.args.ImageResizeLike.source, plainOp );
+                        Ptr<ASTOpImageResizeLike> BaseResizeOp = new ASTOpImageResizeLike;
+						BaseResizeOp->SizeSource = baseAt;
+						BaseResizeOp->Source = plainOp;
 
-                        maskOp->Base = baseResizeOp;
+                        maskOp->Base = BaseResizeOp;
                     }
 
                     // The base is composition of the layer base on the compose base
@@ -1085,20 +1080,18 @@ namespace mu
                         blackOp->op.args.ColourConstant.value[2] = 0;
 						blackOp->op.args.ColourConstant.value[3] = 1;
 
-                        Ptr<ASTOpFixed> plainOp = new ASTOpFixed;
-                        plainOp->op.type = EOpType::IM_PLAINCOLOUR;
-                        plainOp->SetChild( plainOp->op.args.ImagePlainColour.colour, blackOp );
-                        plainOp->op.args.ImagePlainColour.format = EImageFormat::L_UByte;
-                        plainOp->op.args.ImagePlainColour.size[0] = 4;
-						plainOp->op.args.ImagePlainColour.size[1] = 4;
-						plainOp->op.args.ImagePlainColour.LODs = 1;
+                        Ptr<ASTOpImagePlainColor> plainOp = new ASTOpImagePlainColor;
+                        plainOp->Color = blackOp;
+                        plainOp->Format = EImageFormat::L_UByte;
+                        plainOp->Size[0] = 4;
+						plainOp->Size[1] = 4;
+						plainOp->LODs = 1;
 
-                        Ptr<ASTOpFixed> baseResizeOp = new ASTOpFixed;
-                        baseResizeOp->op.type = EOpType::IM_RESIZELIKE;
-                        baseResizeOp->SetChild( baseResizeOp->op.args.ImageResizeLike.sizeSource, baseAt );
-                        baseResizeOp->SetChild( baseResizeOp->op.args.ImageResizeLike.source, plainOp );
+                        Ptr<ASTOpImageResizeLike> BaseResizeOp = new ASTOpImageResizeLike;
+						BaseResizeOp->SizeSource = baseAt;
+						BaseResizeOp->Source = plainOp;
 
-                        maskOp->Base = baseResizeOp;
+                        maskOp->Base = BaseResizeOp;
                     }
 
                     // The blended is a compose of the blended image on a blank image
@@ -1113,21 +1106,19 @@ namespace mu
                         blackOp->op.args.ColourConstant.value[2] = 0;
 						blackOp->op.args.ColourConstant.value[3] = 1;
 
-                        Ptr<ASTOpFixed> plainOp = new ASTOpFixed;
-                        plainOp->op.type = EOpType::IM_PLAINCOLOUR;
-                        plainOp->SetChild( plainOp->op.args.ImagePlainColour.colour, blackOp );
+                        Ptr<ASTOpImagePlainColor> plainOp = new ASTOpImagePlainColor;
+                        plainOp->Color = blackOp;
                         FImageDesc blendedDesc = baseAt->GetImageDesc();
-                        plainOp->op.args.ImagePlainColour.format = blendedDesc.m_format;
-                        plainOp->op.args.ImagePlainColour.size[0] = 4;
-						plainOp->op.args.ImagePlainColour.size[1] = 4;
-						plainOp->op.args.ImagePlainColour.LODs = 1;
+                        plainOp->Format = blendedDesc.m_format;
+                        plainOp->Size[0] = 4;
+						plainOp->Size[1] = 4;
+						plainOp->LODs = 1;
 
-                        Ptr<ASTOpFixed> resizeOp = new ASTOpFixed;
-                        resizeOp->op.type = EOpType::IM_RESIZELIKE;
-                        resizeOp->SetChild( resizeOp->op.args.ImageResizeLike.sizeSource, baseAt );
-                        resizeOp->SetChild( resizeOp->op.args.ImageResizeLike.source, plainOp );
+                        Ptr<ASTOpImageResizeLike> ResizeOp = new ASTOpImageResizeLike;
+						ResizeOp->SizeSource = baseAt;
+						ResizeOp->Source = plainOp;
 
-                        blendedOp->Base = resizeOp;
+                        blendedOp->Base = ResizeOp;
                     }
 
                     // The base is composition of the softlight base on the compose base
@@ -1198,23 +1189,21 @@ namespace mu
                         blackOp->op.args.ColourConstant.value[2] = 0;
 						blackOp->op.args.ColourConstant.value[3] = 1;
 
-                        Ptr<ASTOpFixed> plainOp = new ASTOpFixed;
-                        plainOp->op.type = EOpType::IM_PLAINCOLOUR;
-                        plainOp->SetChild( plainOp->op.args.ImagePlainColour.colour, blackOp );
-                        plainOp->op.args.ImagePlainColour.format = EImageFormat::L_UByte; //TODO: FORMAT_LIKE
-                        plainOp->op.args.ImagePlainColour.size[0] = 4;
-						plainOp->op.args.ImagePlainColour.size[1] = 4;
-						plainOp->op.args.ImagePlainColour.LODs = 1;
+                        Ptr<ASTOpImagePlainColor> plainOp = new ASTOpImagePlainColor;
+						plainOp->Color = blackOp;
+                        plainOp->Format = EImageFormat::L_UByte; //TODO: FORMAT_LIKE
+                        plainOp->Size[0] = 4;
+						plainOp->Size[1] = 4;
+						plainOp->LODs = 1;
 
-                        Ptr<ASTOpFixed> blockResizeOp = new ASTOpFixed;
-                        blockResizeOp->op.type = EOpType::IM_RESIZELIKE;
-                        blockResizeOp->SetChild( blockResizeOp->op.args.ImageResizeLike.sizeSource, blockAt );
-                        blockResizeOp->SetChild( blockResizeOp->op.args.ImageResizeLike.source, plainOp );
+                        Ptr<ASTOpImageResizeLike> BlockResizeOp = new ASTOpImageResizeLike;
+						BlockResizeOp->SizeSource = blockAt;
+						BlockResizeOp->Source = plainOp;
 
                         // Blank out the block from the mask
 						Ptr<ASTOp> newMaskBase = EnsureValidMask( typedBaseAt->mask.child(), baseAt );
                         maskOp->Base = newMaskBase;
-                        maskOp->BlockImage = blockResizeOp;
+                        maskOp->BlockImage = BlockResizeOp;
                     }
 
                     // The base is composition of the softlight base on the compose base
@@ -1245,23 +1234,21 @@ namespace mu
                         blackOp->op.args.ColourConstant.value[2] = 0;
 						blackOp->op.args.ColourConstant.value[3] = 1;
 
-                        Ptr<ASTOpFixed> plainOp = new ASTOpFixed;
-                        plainOp->op.type = EOpType::IM_PLAINCOLOUR;
-                        plainOp->SetChild( plainOp->op.args.ImagePlainColour.colour, blackOp );
-                        plainOp->op.args.ImagePlainColour.format = EImageFormat::L_UByte; //TODO: FORMAT_LIKE
-                        plainOp->op.args.ImagePlainColour.size[0] = 4;
-						plainOp->op.args.ImagePlainColour.size[1] = 4;
-						plainOp->op.args.ImagePlainColour.LODs = 1;
+                        Ptr<ASTOpImagePlainColor> plainOp = new ASTOpImagePlainColor;
+                        plainOp->Color = blackOp;
+                        plainOp->Format = EImageFormat::L_UByte; //TODO: FORMAT_LIKE
+                        plainOp->Size[0] = 4;
+						plainOp->Size[1] = 4;
+						plainOp->LODs = 1;
 
-                        Ptr<ASTOpFixed> blockResizeOp = new ASTOpFixed;
-                        blockResizeOp->op.type = EOpType::IM_RESIZELIKE;
-                        blockResizeOp->SetChild( blockResizeOp->op.args.ImageResizeLike.sizeSource, blockAt );
-                        blockResizeOp->SetChild( blockResizeOp->op.args.ImageResizeLike.source, plainOp );
+                        Ptr<ASTOpImageResizeLike> BlockResizeOp = new ASTOpImageResizeLike;
+						BlockResizeOp->SizeSource = blockAt;
+						BlockResizeOp->Source = plainOp;
 
                         // Blank out the block from the mask
                         Ptr<ASTOp> newMaskBase = EnsureValidMask( typedBaseAt->mask.child(), baseAt );
                         maskOp->Base = newMaskBase;
-                        maskOp->BlockImage = blockResizeOp;
+                        maskOp->BlockImage = BlockResizeOp;
                     }
 
                     // The base is composition of the effect base on the compose base
@@ -1548,15 +1535,15 @@ namespace mu
 
             case EOpType::IM_DISPLACE:
             {
-				const ASTOpFixed* TypedOp = static_cast<const ASTOpFixed*>(At.get());
+				const ASTOpImageDisplace* TypedOp = static_cast<const ASTOpImageDisplace*>(At.get());
 
-                RecurseWithCurrentState( TypedOp->children[TypedOp->op.args.ImageDisplace.source].child() );
+                RecurseWithCurrentState( TypedOp->Source.child() );
 
                 std::array<uint8, SIZE_T(EImageFormat::Count)> NewState;
                 NewState.fill(0);
                 NewState[SIZE_T(EImageFormat::L_UByte) ] = 1;
                 NewState[SIZE_T(EImageFormat::L_UByteRLE) ] = 1;
-                RecurseWithState( TypedOp->children[TypedOp->op.args.ImageDisplace.displacementMap].child(), NewState );
+                RecurseWithState( TypedOp->DisplacementMap.child(), NewState );
                 break;
             }
 
diff --git a/Engine/Plugins/Mutable/Source/MutableTools/Private/MuT/CodeOptimiser_Semantic.cpp b/Engine/Plugins/Mutable/Source/MutableTools/Private/MuT/CodeOptimiser_Semantic.cpp
--- a/Engine/Plugins/Mutable/Source/MutableTools/Private/MuT/CodeOptimiser_Semantic.cpp
+++ b/Engine/Plugins/Mutable/Source/MutableTools/Private/MuT/CodeOptimiser_Semantic.cpp
@@ -23,6 +23,7 @@
 #include "MuT/ASTOpImageSwizzle.h"
 #include "MuT/ASTOpImageTransform.h"
 #include "MuT/ASTOpImageRasterMesh.h"
+#include "MuT/ASTOpImageResize.h"
 #include "MuT/ASTOpMeshMorph.h"
 #include "MuT/ASTOpMeshAddTags.h"
 #include "MuT/ASTOpSwitch.h"
@@ -359,46 +360,6 @@ namespace mu
 			break;
 		}
 
-		//-------------------------------------------------------------------------------------
-		case EOpType::IM_RESIZEREL:
-		{
-			Ptr<ASTOp> sourceAt = children[op.args.ImageResizeRel.source].child();
-
-			// The instruction can be sunk
-			EOpType sourceType = sourceAt->GetOpType();
-			switch (sourceType)
-			{
-			// This is done here instead of in the OptimizeSize step to prevent removing 
-			// resize_rel too early.
-			case EOpType::IM_RESIZE:
-			{
-				Ptr<ASTOpFixed> NewOp = mu::Clone<ASTOpFixed>(sourceAt.get());
-				NewOp->op.args.ImageResize.size[0] = int16(NewOp->op.args.ImageResize.size[0] * op.args.ImageResizeRel.factor[0]);
-				NewOp->op.args.ImageResize.size[1] = int16(NewOp->op.args.ImageResize.size[1] * op.args.ImageResizeRel.factor[1]);
-
-				at = NewOp;
-				break;
-			}
-
-
-			default:
-				break;
-
-			}
-
-			break;
-		}
-
-        //-------------------------------------------------------------------------------------
-        case EOpType::IM_PIXELFORMAT:
-		case EOpType::IM_SWIZZLE:
-		case EOpType::IM_LAYER:
-		{
-			check(false); // moved to its own operation class
-			break;
-		}
-
-
         default:
             break;
         }
@@ -437,151 +398,6 @@ namespace mu
         return bModified;
     }
 
-
-	//---------------------------------------------------------------------------------------------
-	//---------------------------------------------------------------------------------------------
-	//---------------------------------------------------------------------------------------------
-	Ptr<ASTOp> ASTOpFixed::OptimiseSink( const FModelOptimizationOptions& options, FOptimizeSinkContext& context ) const
-    {
-        Ptr<ASTOp> at;
-
-        EOpType type = GetOpType();
-        switch ( type )
-        {
-
-        //-----------------------------------------------------------------------------------------
-        //-----------------------------------------------------------------------------------------
-        //-----------------------------------------------------------------------------------------
-		case EOpType::IM_SWIZZLE:
-		case EOpType::LA_REMOVEBLOCKS:
-		case EOpType::IM_RASTERMESH:
-		{
-			// Moved to their own operation class
-			check(false);
-            break;
-        }
-
-
-        //-----------------------------------------------------------------------------------------
-        //-----------------------------------------------------------------------------------------
-        //-----------------------------------------------------------------------------------------
-        case EOpType::IM_DISPLACE:
-        {
-			Ptr<ASTOp> OriginalAt = at;
-			Ptr<ASTOp> sourceAt = children[op.args.ImageDisplace.source].child();
-			Ptr<ASTOp> displaceMapAt = children[op.args.ImageDisplace.displacementMap].child();
-
-            switch ( sourceAt->GetOpType() )
-            {
-
-            case EOpType::IM_CONDITIONAL:
-            {
-                if (displaceMapAt->GetOpType()==EOpType::IM_CONDITIONAL)
-                {
-					const ASTOpConditional* typedSource = static_cast<const ASTOpConditional*>(sourceAt.get());
-					const ASTOpConditional* typedDisplacementMap = static_cast<const ASTOpConditional*>(displaceMapAt.get());
-
-                    if (typedSource->condition==typedDisplacementMap->condition)
-                    {
-                        auto nop = mu::Clone<ASTOpConditional>(sourceAt);
-
-                        auto aOp = mu::Clone<ASTOpFixed>(this);
-                        aOp->SetChild( aOp->op.args.ImageDisplace.source, typedSource->yes );
-                        aOp->SetChild( aOp->op.args.ImageDisplace.displacementMap, typedDisplacementMap->yes );
-                        nop->yes = aOp;
-
-                        auto bOp = mu::Clone<ASTOpFixed>(this);
-                        bOp->SetChild( bOp->op.args.ImageDisplace.source, typedSource->no);
-                        bOp->SetChild( bOp->op.args.ImageDisplace.displacementMap, typedDisplacementMap->no);
-                        nop->no = bOp;
-
-                        at = nop;
-                    }
-                }
-                break;
-            }
-
-            case EOpType::IM_SWITCH:
-            {
-                if (displaceMapAt->GetOpType()==EOpType::IM_SWITCH)
-                {
-					const ASTOpSwitch* typedSource = static_cast<const ASTOpSwitch*>(sourceAt.get());
-					const ASTOpSwitch* typedDisplacementMap = static_cast<const ASTOpSwitch*>(displaceMapAt.get());
-
-                    if (typedSource->IsCompatibleWith(typedDisplacementMap))
-                    {
-                        // Move the format down all the paths
-                        auto nop = mu::Clone<ASTOpSwitch>(sourceAt);
-
-                        if (nop->def)
-                        {
-							Ptr<ASTOpFixed> defOp = mu::Clone<ASTOpFixed>(this);
-                            defOp->SetChild( defOp->op.args.ImageDisplace.source, typedSource->def );
-                            defOp->SetChild( defOp->op.args.ImageDisplace.displacementMap, typedDisplacementMap->def );
-                            nop->def = defOp;
-                        }
-
-                        // We need to copy the options because we change them
-                        for ( size_t v=0; v<nop->cases.Num(); ++v )
-                        {
-                            if ( nop->cases[v].branch )
-                            {
-								Ptr<ASTOpFixed> bOp = mu::Clone<ASTOpFixed>(this);
-                                bOp->SetChild( bOp->op.args.ImageDisplace.source, typedSource->cases[v].branch );
-                                bOp->SetChild( bOp->op.args.ImageDisplace.displacementMap, typedDisplacementMap->FindBranch(typedSource->cases[v].condition) );
-                                nop->cases[v].branch = bOp;
-                            }
-                        }
-
-                        at = nop;
-                    }
-                }
-
-				// If we didn't optimize already, try to simply sink the source.
-				if (OriginalAt == at)
-				{
-					// Make a project for every path
-					auto nop = mu::Clone<ASTOpSwitch>(sourceAt.get());
-
-					if (nop->def)
-					{
-						Ptr<ASTOpFixed> defOp = mu::Clone<ASTOpFixed>(this);
-						defOp->SetChild(defOp->op.args.ImageDisplace.source, nop->def);
-						nop->def = defOp;
-					}
-
-					// We need to copy the options because we change them
-					for (size_t o = 0; o < nop->cases.Num(); ++o)
-					{
-						if (nop->cases[o].branch)
-						{
-							Ptr<ASTOpFixed> bOp = mu::Clone<ASTOpFixed>(this);
-							bOp->SetChild(bOp->op.args.ImageDisplace.source, nop->cases[o].branch);
-							nop->cases[o].branch = bOp;
-						}
-					}
-
-					at = nop;
-				}
-
-                break;
-            }
-
-            default:
-                break;
-            }
-
-            break;
-        }
-
-
-        default:
-            break;
-        }
-
-        return at;
-    }
-
     //---------------------------------------------------------------------------------------------
     //---------------------------------------------------------------------------------------------
     //---------------------------------------------------------------------------------------------
@@ -606,765 +422,4 @@ namespace mu
     }
 
 
-    //---------------------------------------------------------------------------------------------
-    //!
-    //---------------------------------------------------------------------------------------------
-    class Sink_ImageResizeRelAST
-    {
-    public:
-
-        Ptr<ASTOp> Apply( const ASTOp* root )
-        {
-            Root = root;
-            OldToNew.Empty();
-
-            check(root->GetOpType()==EOpType::IM_RESIZEREL);
-
-			const ASTOpFixed* typedRoot = static_cast<const ASTOpFixed*>(root);
-            InitialSource = typedRoot->children[typedRoot->op.args.ImageResizeRel.source].child();
-            Ptr<ASTOp> newSource = Visit( InitialSource, typedRoot );
-
-            // If there is any change, it is the new root.
-            if (newSource!=InitialSource)
-            {
-                return newSource;
-            }
-
-            return nullptr;
-        }
-
-    protected:
-
-        const ASTOp* Root;
-        Ptr<ASTOp> InitialSource;
-        //! For each operation we sink, the map from old instructions to new instructions.
-		TMap<FSinkerOldToNewKey, Ptr<ASTOp>> OldToNew;
-
-		Ptr<ASTOp> Visit(Ptr<ASTOp> at, const ASTOpFixed* currentSinkingOp)
-		{
-			if (!at) return nullptr;
-
-			// Already visited?
-			const Ptr<ASTOp>* Cached = OldToNew.Find({ at,currentSinkingOp });
-			if (Cached)
-			{
-				return *Cached;
-			}
-
-			float scaleX = currentSinkingOp->op.args.ImageResizeRel.factor[0];
-			float scaleY = currentSinkingOp->op.args.ImageResizeRel.factor[1];
-
-			Ptr<ASTOp> newAt = at;
-			switch (at->GetOpType())
-			{
-
-			case EOpType::IM_CONDITIONAL:
-			{
-				// We move the mask creation down the two paths
-				auto newOp = mu::Clone<ASTOpConditional>(at);
-				newOp->yes = Visit(newOp->yes.child(), currentSinkingOp);
-				newOp->no = Visit(newOp->no.child(), currentSinkingOp);
-				newAt = newOp;
-				break;
-			}
-
-			case EOpType::IM_PIXELFORMAT:
-			{
-				// \todo: only if shrinking?
-				Ptr<ASTOpImagePixelFormat> newOp = mu::Clone<ASTOpImagePixelFormat>(at);
-				newOp->Source = Visit(newOp->Source.child(), currentSinkingOp);
-				newAt = newOp;
-				break;
-			}
-
-			case EOpType::IM_SWITCH:
-			{
-				// We move the mask creation down all the paths
-				Ptr<ASTOpSwitch> newOp = mu::Clone<ASTOpSwitch>(at);
-				newOp->def = Visit(newOp->def.child(), currentSinkingOp);
-				for (ASTOpSwitch::FCase& c : newOp->cases)
-				{
-					c.branch = Visit(c.branch.child(), currentSinkingOp);
-				}
-				newAt = newOp;
-				break;
-			}
-
-			case EOpType::IM_SWIZZLE:
-			{
-				Ptr<ASTOpImageSwizzle> newOp = mu::Clone<ASTOpImageSwizzle>(at);
-				for (int s = 0; s < MUTABLE_OP_MAX_SWIZZLE_CHANNELS; ++s)
-				{
-					Ptr<ASTOp> channelOp = newOp->Sources[s].child();
-					if (channelOp)
-					{
-						newOp->Sources[s] = Visit(channelOp, currentSinkingOp);
-					}
-				}
-				newAt = newOp;
-				break;
-			}
-
-			case EOpType::IM_COMPOSE:
-			{
-				// We can only optimise if the layout grid blocks size in pixels is
-				// still an integer after relative scale
-				bool acceptable = false;
-				{
-					const ASTOpImageCompose* typedAt = static_cast<const ASTOpImageCompose*>(at.get());
-					Ptr<ASTOp> originalBaseOp = typedAt->Base.child();
-
-					// \todo: recursion-proof cache?
-					int layoutBlockPixelsX = 0;
-					int layoutBlockPixelsY = 0;
-					originalBaseOp->GetLayoutBlockSize(&layoutBlockPixelsX, &layoutBlockPixelsY);
-
-					int scaledLayoutBlockPixelsX = int(layoutBlockPixelsX * scaleX);
-					int scaledLayoutBlockPixelsY = int(layoutBlockPixelsY * scaleY);
-					int unscaledLayoutBlockPixelsX = int(scaledLayoutBlockPixelsX / scaleX);
-					int unscaledLayoutBlockPixelsY = int(scaledLayoutBlockPixelsY / scaleY);
-					acceptable =
-						(layoutBlockPixelsX != 0 && layoutBlockPixelsY != 0)
-						&&
-						(layoutBlockPixelsX == unscaledLayoutBlockPixelsX)
-						&&
-						(layoutBlockPixelsY == unscaledLayoutBlockPixelsY);
-				}
-
-				if (acceptable)
-				{
-					Ptr<ASTOpImageCompose> newOp = mu::Clone<ASTOpImageCompose>(at);
-
-					Ptr<ASTOp> baseOp = newOp->Base.child();
-					newOp->Base = Visit(baseOp, currentSinkingOp);
-
-					Ptr<ASTOp> blockOp = newOp->BlockImage.child();
-					newOp->BlockImage = Visit(blockOp, currentSinkingOp);
-
-					newAt = newOp;
-				}
-
-				break;
-			}
-
-			case EOpType::IM_PATCH:
-			{
-				Ptr<ASTOpImagePatch> newOp = mu::Clone<ASTOpImagePatch>(at);
-
-				newOp->base = Visit(newOp->base.child(), currentSinkingOp);
-				newOp->patch = Visit(newOp->patch.child(), currentSinkingOp);
-
-				// todo: review if this is always correct, or we need some "divisible" check
-				newOp->location[0] = uint16(newOp->location[0] * scaleX);
-				newOp->location[1] = uint16(newOp->location[1] * scaleY);
-
-				newAt = newOp;
-
-				break;
-			}
-
-			case EOpType::IM_MIPMAP:
-			{
-				Ptr<ASTOpImageMipmap> newOp = mu::Clone<ASTOpImageMipmap>(at.get());
-				Ptr<ASTOp> baseOp = newOp->Source.child();
-				newOp->Source = Visit(baseOp, currentSinkingOp);
-				newAt = newOp;
-				break;
-			}
-
-			case EOpType::IM_INTERPOLATE:
-			{
-				Ptr<ASTOpFixed> newOp = mu::Clone<ASTOpFixed>(at);
-
-				for (int v = 0; v < MUTABLE_OP_MAX_INTERPOLATE_COUNT; ++v)
-				{
-					Ptr<ASTOp> child = newOp->children[newOp->op.args.ImageInterpolate.targets[v]].child();
-					Ptr<ASTOp> bOp = Visit(child, currentSinkingOp);
-					newOp->SetChild(newOp->op.args.ImageInterpolate.targets[v], bOp);
-				}
-
-				newAt = newOp;
-				break;
-			}
-
-			case EOpType::IM_MULTILAYER:
-			{
-				Ptr<ASTOpImageMultiLayer> nop = mu::Clone<ASTOpImageMultiLayer>(at);
-				nop->base = Visit(nop->base.child(), currentSinkingOp);
-				nop->mask = Visit(nop->mask.child(), currentSinkingOp);
-				nop->blend = Visit(nop->blend.child(), currentSinkingOp);
-				newAt = nop;
-				break;
-			}
-
-			case EOpType::IM_LAYER:
-			{
-				Ptr<ASTOpImageLayer> nop = mu::Clone<ASTOpImageLayer>(at);
-
-				Ptr<ASTOp> aOp = nop->base.child();
-				nop->base = Visit(aOp, currentSinkingOp);
-
-				Ptr<ASTOp> bOp = nop->blend.child();
-				nop->blend = Visit(bOp, currentSinkingOp);
-
-				Ptr<ASTOp> mOp = nop->mask.child();
-				nop->mask = Visit(mOp, currentSinkingOp);
-
-				newAt = nop;
-				break;
-			}
-
-			case EOpType::IM_LAYERCOLOUR:
-			{
-				Ptr<ASTOpImageLayerColor> nop = mu::Clone<ASTOpImageLayerColor>(at);
-
-				Ptr<ASTOp> aOp = nop->base.child();
-				nop->base = Visit(aOp, currentSinkingOp);
-
-				Ptr<ASTOp> mOp = nop->mask.child();
-				nop->mask = Visit(mOp, currentSinkingOp);
-
-				newAt = nop;
-				break;
-			}
-
-			case EOpType::IM_DISPLACE:
-			{
-				Ptr<ASTOpFixed> nop = mu::Clone<ASTOpFixed>(at);
-
-				Ptr<ASTOp> sourceOp = nop->children[nop->op.args.ImageDisplace.source].child();
-				nop->SetChild(nop->op.args.ImageDisplace.source, Visit(sourceOp, currentSinkingOp));
-
-				Ptr<ASTOp> mapOp = nop->children[nop->op.args.ImageDisplace.displacementMap].child();
-				nop->SetChild(nop->op.args.ImageDisplace.displacementMap, Visit(mapOp, currentSinkingOp));
-
-				// Make sure we don't scale a constant displacement map, which is very wrong.
-				if (mapOp->GetOpType() == EOpType::IM_MAKEGROWMAP)
-				{
-					newAt = nop;
-				}
-				else
-				{
-					// We cannot resize an already calculated displacement map.
-				}
-
-				break;
-			}
-
-			case EOpType::IM_MAKEGROWMAP:
-			{
-				Ptr<ASTOpImageMakeGrowMap> nop = mu::Clone<ASTOpImageMakeGrowMap>(at);
-				Ptr<ASTOp> maskOp = nop->Mask.child();
-				nop->Mask = Visit(maskOp, currentSinkingOp);
-				newAt = nop;
-				break;
-			}
-
-			case EOpType::IM_INVERT:
-			{
-				Ptr<ASTOpFixed> nop = mu::Clone<ASTOpFixed>(at);
-				Ptr<ASTOp> maskOp = nop->children[nop->op.args.ImageInvert.base].child();
-				nop->SetChild(nop->op.args.ImageInvert.base, Visit(maskOp, currentSinkingOp));
-				newAt = nop;
-				break;
-			}
-
-			case EOpType::IM_SATURATE:
-			{
-				Ptr<ASTOpFixed> nop = mu::Clone<ASTOpFixed>(at);
-				Ptr<ASTOp> maskOp = nop->children[nop->op.args.ImageSaturate.base].child();
-				nop->SetChild(nop->op.args.ImageSaturate.base, Visit(maskOp, currentSinkingOp));
-				newAt = nop;
-				break;
-			}
-
-			case EOpType::IM_TRANSFORM:
-			{
-				// It can only sink in the transform if it doesn't have it's own size.
-				const ASTOpImageTransform* TypedAt = static_cast<const ASTOpImageTransform*>(at.get());
-				if (TypedAt->SizeX == 0 && TypedAt->SizeY == 0)
-				{
-					Ptr<ASTOpImageTransform> NewOp = mu::Clone<ASTOpImageTransform>(at);
-					Ptr<ASTOp> MaskOp = NewOp->Base.child();
-					NewOp->Base = Visit(MaskOp, currentSinkingOp);
-					newAt = NewOp;
-				}
-
-				break;
-			}
-
-			case EOpType::IM_RASTERMESH:
-			{
-				Ptr<ASTOpImageRasterMesh> nop = mu::Clone<ASTOpImageRasterMesh>(at);
-				Ptr<ASTOp> maskOp = nop->mask.child();
-				nop->mask = Visit(maskOp, currentSinkingOp);
-
-				// Resize the image to project as well, assuming that since the target has a different resolution
-				// it make sense for the source image to have a similar resize.
-				// Actually, don't do it because the LODBias will be applied separetely at graph generation time.
-				//auto imageOp = nop->image.child();
-				//nop->image = Visit(imageOp, currentSinkingOp);
-
-				nop->SizeX = uint16(nop->SizeX * scaleX + 0.5f);
-				nop->SizeY = uint16(nop->SizeY * scaleY + 0.5f);
-				newAt = nop;
-				break;
-			}
-
-			default:
-				break;
-			}
-
-			// end on line, replace with sinking op
-			if (at == newAt && at != InitialSource)
-			{
-				Ptr<ASTOpFixed> newOp = mu::Clone<ASTOpFixed>(currentSinkingOp);
-				check(newOp->GetOpType() == EOpType::IM_RESIZEREL);
-
-				newOp->SetChild(newOp->op.args.ImageResizeRel.source, at);
-
-				newAt = newOp;
-			}
-
-			OldToNew.Add({ at, currentSinkingOp }, newAt);
-
-            return newAt;
-        }
-    };
-
-
-    //---------------------------------------------------------------------------------------------
-    Ptr<ASTOp> ASTOpFixed::OptimiseSize() const
-    {
-        Ptr<ASTOp> at;
-
-        EOpType type = GetOpType();
-        switch ( type )
-        {
-
-        //-------------------------------------------------------------------------------------
-        case EOpType::IM_RESIZE:
-        {
-            auto sourceAt = children[op.args.ImageResize.source].child();
-
-            // The instruction can be sunk
-            EOpType sourceType = sourceAt->GetOpType();
-            switch ( sourceType )
-            {
-            case EOpType::IM_RESIZE:
-            {
-                // Keep top resize
-                Ptr<const ASTOpFixed> sourceOp = static_cast<const ASTOpFixed*>(sourceAt.get());
-
-                Ptr<ASTOpFixed> newOp = mu::Clone<ASTOpFixed>(this);
-                newOp->SetChild( newOp->op.args.ImageResize.source, sourceOp->children[sourceOp->op.args.ImageResize.source]);
-
-                at = newOp;
-                break;
-            }
-
-            case EOpType::IM_PLAINCOLOUR:
-            {
-                // Set the size in the children and remove resize
-				Ptr<ASTOpFixed> sourceOp = mu::Clone<ASTOpFixed>(sourceAt.get());
-                sourceOp->op.args.ImagePlainColour.size[0] = op.args.ImageResize.size[0];
-				sourceOp->op.args.ImagePlainColour.size[1] = op.args.ImageResize.size[1];
-				sourceOp->op.args.ImagePlainColour.LODs = 1; // TODO
-				at = sourceOp;
-                break;
-            }
-
-			case EOpType::IM_TRANSFORM:
-			{
-				// Set the size in the children and remove resize
-				Ptr<ASTOpImageTransform> sourceOp = mu::Clone<ASTOpImageTransform>(sourceAt.get());
-				sourceOp->SizeX = op.args.ImageResize.size[0];
-				sourceOp->SizeY = op.args.ImageResize.size[1];
-				at = sourceOp;
-				break;
-			}
-
-            case EOpType::IM_CONDITIONAL:
-            {
-                // We move the resize down the two paths
-                auto newOp = mu::Clone<ASTOpConditional>(sourceAt);
-
-                Ptr<ASTOpFixed> aOp = mu::Clone<ASTOpFixed>(this);
-                aOp->SetChild( aOp->op.args.ImageResize.source,newOp->yes);
-                newOp->yes = aOp;
-
-                Ptr<ASTOpFixed> bOp = mu::Clone<ASTOpFixed>(this);
-                bOp->SetChild( bOp->op.args.ImageResize.source,newOp->no);
-                newOp->no = bOp;
-
-                at = newOp;
-                break;
-            }
-
-            case EOpType::IM_SWITCH:
-            {
-                // Move the resize down all the paths
-                Ptr<ASTOpSwitch> newOp = mu::Clone<ASTOpSwitch>(sourceAt);
-
-                if (newOp->def)
-                {
-                    Ptr<ASTOpFixed> defOp = mu::Clone<ASTOpFixed>(this);
-                    defOp->SetChild( defOp->op.args.ImageResize.source, newOp->def);
-                    newOp->def = defOp;
-                }
-
-                for ( auto& cas:newOp->cases )
-                {
-                    if ( cas.branch )
-                    {
-                        Ptr<ASTOpFixed> bOp = mu::Clone<ASTOpFixed>(this);
-                        bOp->SetChild( bOp->op.args.ImageResize.source, cas.branch);
-                        cas.branch = bOp;
-                    }
-                }
-
-                at = newOp;
-                break;
-            }
-
-			case EOpType::IM_SWIZZLE:
-			{
-				Ptr<ASTOpImageSwizzle> newOp = mu::Clone<ASTOpImageSwizzle>(sourceAt);
-				for (int s = 0; s < MUTABLE_OP_MAX_SWIZZLE_CHANNELS; ++s)
-				{
-					Ptr<ASTOp> OldChannelOp = newOp->Sources[s].child();
-					if (OldChannelOp)
-					{
-						Ptr<ASTOpFixed> ChannelResize = mu::Clone<ASTOpFixed>(this);
-						ChannelResize->SetChild(ChannelResize->op.args.ImageResize.source, OldChannelOp);
-						newOp->Sources[s] = ChannelResize;
-					}
-				}
-				at = newOp;
-				break;
-			}
-
-            case EOpType::IM_COMPOSE:
-            {
-                Ptr<ASTOpImageCompose> newOp = mu::Clone<ASTOpImageCompose>(sourceAt);
-
-                auto aOp = mu::Clone<ASTOpFixed>(this);
-                aOp->SetChild( aOp->op.args.ImageResize.source, newOp->Base.child());
-                newOp->Base = aOp;
-
-                auto bOp = mu::Clone<ASTOpFixed>(this);
-                bOp->SetChild( bOp->op.args.ImageResize.source, newOp->BlockImage.child());
-                newOp->BlockImage = bOp;
-
-                if ( newOp->Mask )
-                {
-                    auto maskOp = mu::Clone<ASTOpFixed>(this);
-                    maskOp->SetChild( maskOp->op.args.ImageResize.source, newOp->Mask.child());
-                    newOp->Mask = maskOp;
-                }
-
-                at = newOp;
-                break;
-            }
-
-            case EOpType::IM_RASTERMESH:
-            {
-                Ptr<ASTOpImageRasterMesh> newOp = mu::Clone<ASTOpImageRasterMesh>(sourceAt);
-
-                //if ( newOp->op.args.ImageRasterMesh.sizeX != op.args.ImageResize.size[0]
-                //     ||
-                //     newOp->op.args.ImageRasterMesh.sizeY != op.args.ImageResize.size[1] )
-                {
-                    newOp->SizeX = op.args.ImageResize.size[0];
-                    newOp->SizeY = op.args.ImageResize.size[1];
-
-					if (newOp->mask)
-					{
-						Ptr<ASTOpFixed> mop = mu::Clone<ASTOpFixed>(this);
-						mop->SetChild(mop->op.args.ImageResize.source, newOp->mask.child());
-						newOp->mask = mop;
-					}
-
-					// Don't apply absolute resizes to the image to raster: it could even enlarge it.
-					// This should only be scaled with relative resizes, which come from LOD biases, etc.
-					//if (newOp->op.args.ImageRasterMesh.image)
-					//{
-					//	Ptr<ASTOpFixed> mop = mu::Clone<ASTOpFixed>(this);
-					//	mop->SetChild(mop->op.args.ImageResize.source, newOp->children[newOp->op.args.ImageRasterMesh.image].child());
-					//	newOp->SetChild(newOp->op.args.ImageRasterMesh.image, mop);
-					//}
-				}
-
-                at = newOp;
-                break;
-            }
-
-            case EOpType::IM_INTERPOLATE:
-            {
-                Ptr<ASTOpFixed> newOp = mu::Clone<ASTOpFixed>(sourceAt);
-
-                for ( int s=0; s<MUTABLE_OP_MAX_INTERPOLATE_COUNT; ++s )
-                {
-                    auto targetAt = newOp->children[ newOp->op.args.ImageInterpolate.targets[s] ].child();
-                    if ( targetAt )
-                    {
-                        Ptr<ASTOpFixed> sourceOp = mu::Clone<ASTOpFixed>(this);
-                        sourceOp->SetChild( sourceOp->op.args.ImageResize.source, targetAt);
-                        newOp->SetChild( newOp->op.args.ImageInterpolate.targets[s], sourceOp );
-                    }
-                }
-
-                at = newOp;
-                break;
-            }
-
-			case EOpType::IM_INVERT:
-			{
-				Ptr<ASTOpFixed> NewOp = mu::Clone<ASTOpFixed>(sourceAt);
-				Ptr<ASTOp> BaseAt = NewOp->children[NewOp->op.args.ImageInvert.base].child();
-
-				Ptr<ASTOpFixed> NewBase = mu::Clone<ASTOpFixed>(this);
-				NewBase->SetChild(NewBase->op.args.ImageResize.source, BaseAt);
-
-				NewOp->SetChild(NewOp->op.args.ImageInvert.base, NewBase);
-
-				at = NewOp;
-				break;
-			}
-
-			case EOpType::IM_PIXELFORMAT:
-			{				
-				// \todo: only if shrinking?
-				
-				// Only sink the resize if we know that the pixelformat source image is uncompressed.
-				Ptr<ASTOpImagePixelFormat> SourceTyped = static_cast<ASTOpImagePixelFormat*>(sourceAt.get());
-				FImageDesc PixelFormatSourceDesc = SourceTyped ->Source->GetImageDesc();
-				if (PixelFormatSourceDesc.m_format!=EImageFormat::None
-					&&
-					!mu::IsCompressedFormat(PixelFormatSourceDesc.m_format) )
-				{
-					Ptr<ASTOpImagePixelFormat> NewOp = mu::Clone<ASTOpImagePixelFormat>(sourceAt);
-					Ptr<ASTOp> BaseAt = NewOp->Source.child();
-
-					Ptr<ASTOpFixed> NewBase = mu::Clone<ASTOpFixed>(this);
-					NewBase->SetChild(NewBase->op.args.ImageResize.source, BaseAt);
-
-					NewOp->Source = NewBase;
-
-					at = NewOp;
-				}
-				break;
-			}
-
-            case EOpType::IM_LAYER:
-            {
-                Ptr<ASTOpImageLayer> newOp = mu::Clone<ASTOpImageLayer>(sourceAt);
-
-                Ptr<ASTOpFixed> baseOp = mu::Clone<ASTOpFixed>(this);
-                baseOp->SetChild( baseOp->op.args.ImageResize.source, newOp->base );
-                newOp->base = baseOp;
-
-                Ptr<ASTOpFixed> blendOp = mu::Clone<ASTOpFixed>(this);
-                blendOp->SetChild( blendOp->op.args.ImageResize.source, newOp->blend );
-                newOp->blend = blendOp;
-
-                auto maskAt = newOp->mask.child();
-                if (maskAt)
-                {
-                    Ptr<ASTOpFixed> maskOp = mu::Clone<ASTOpFixed>(this);
-                    maskOp->SetChild( maskOp->op.args.ImageResize.source, maskAt );
-                    newOp->mask = maskOp;
-                }
-
-                at = newOp;
-                break;
-            }
-
-            case EOpType::IM_LAYERCOLOUR:
-            {
-                Ptr<ASTOpImageLayerColor> newOp = mu::Clone<ASTOpImageLayerColor>(sourceAt);
-
-                Ptr<ASTOpFixed> baseOp = mu::Clone<ASTOpFixed>(this);
-                baseOp->SetChild( baseOp->op.args.ImageResize.source, newOp->base );
-                newOp->base = baseOp;
-
-                auto maskAt = newOp->mask.child();
-                if (maskAt)
-                {
-                    Ptr<ASTOpFixed> maskOp = mu::Clone<ASTOpFixed>(this);
-                    maskOp->SetChild( maskOp->op.args.ImageResize.source, maskAt );
-                    newOp->mask = maskOp;
-                }
-
-                at = newOp;
-                break;
-            }
-
-			case EOpType::IM_DISPLACE:
-			{
-				// In the size optimization phase we can optimize the resize with the displace
-				// because the constants have not been collapsed yet.
-				// We will still check it and sink the size directly below the IM_MAKEGROWMAP op
-				Ptr<ASTOpFixed> SourceTyped = static_cast<ASTOpFixed*>(sourceAt.get());
-				Ptr<ASTOp> OriginalDisplacementMapOp = SourceTyped->children[SourceTyped->op.args.ImageDisplace.displacementMap].Child;
-				if (OriginalDisplacementMapOp->GetOpType() == EOpType::IM_MAKEGROWMAP)
-				{
-					Ptr<ASTOpFixed> newOp = mu::Clone<ASTOpFixed>(sourceAt);
-
-					Ptr<ASTOpFixed> baseOp = mu::Clone<ASTOpFixed>(this);
-					baseOp->SetChild(baseOp->op.args.ImageResize.source,
-						newOp->children[newOp->op.args.ImageDisplace.source]);
-					newOp->SetChild(newOp->op.args.ImageDisplace.source, baseOp);
-
-					// Clone the map op and replace its children
-					Ptr<ASTOpImageMakeGrowMap> mapOp = mu::Clone<ASTOpImageMakeGrowMap>(OriginalDisplacementMapOp);
-					newOp->SetChild(newOp->op.args.ImageDisplace.displacementMap, mapOp);
-
-					Ptr<ASTOpFixed> mapSourceOp = mu::Clone<ASTOpFixed>(this);
-					mapSourceOp->SetChild(mapSourceOp->op.args.ImageResize.source, mapOp->Mask);
-					mapOp->Mask = mapSourceOp;
-
-					at = newOp;
-				}
-				break;
-			}
-
-            default:
-                break;
-            }
-
-            break;
-        }
-
-
-
-        //-------------------------------------------------------------------------------------
-        case EOpType::IM_RESIZEREL:
-        {
-            auto sourceAt = children[op.args.ImageResizeRel.source].child();
-
-            // The instruction can be sunk
-            EOpType sourceType = sourceAt->GetOpType();
-            switch ( sourceType )
-            {
-
-			case EOpType::IM_BLANKLAYOUT:
-			{
-				Ptr<ASTOpFixed> newOp = mu::Clone<ASTOpFixed>(sourceAt);
-
-				newOp->op.args.ImageBlankLayout.blockSize[0] =
-					uint16(newOp->op.args.ImageBlankLayout.blockSize[0]
-						* op.args.ImageResizeRel.factor[0]
-						+ 0.5f);
-				newOp->op.args.ImageBlankLayout.blockSize[1] =
-					uint16(newOp->op.args.ImageBlankLayout.blockSize[1]
-						* op.args.ImageResizeRel.factor[1]
-						+ 0.5f);
-				at = newOp;
-				break;
-			}
-
-			case EOpType::IM_PLAINCOLOUR:
-			{
-				Ptr<ASTOpFixed> newOp = mu::Clone<ASTOpFixed>(sourceAt);
-
-				newOp->op.args.ImagePlainColour.size[0] = FMath::CeilToInt( newOp->op.args.ImagePlainColour.size[0] * op.args.ImageResizeRel.factor[0] );
-				newOp->op.args.ImagePlainColour.size[1] = FMath::CeilToInt( newOp->op.args.ImagePlainColour.size[1] * op.args.ImageResizeRel.factor[1] );
-				newOp->op.args.ImagePlainColour.LODs = 1; // TODO
-				at = newOp;
-				break;
-			}
-
-			case EOpType::IM_TRANSFORM:
-			{
-				// We can only optimize here if we know the transform result size, otherwise, we will sink the op in the sinker.
-				const ASTOpImageTransform* typedAt = static_cast<const ASTOpImageTransform*>(sourceAt.get());
-				if (typedAt->SizeX != 0 && typedAt->SizeY != 0)
-				{
-					// Set the size in the children and remove resize
-					Ptr<ASTOpImageTransform> sourceOp = mu::Clone<ASTOpImageTransform>(sourceAt.get());
-					sourceOp->SizeX = FMath::CeilToInt32(sourceOp->SizeX * op.args.ImageResizeRel.factor[0]);
-					sourceOp->SizeY = FMath::CeilToInt32(sourceOp->SizeY * op.args.ImageResizeRel.factor[1]);
-					at = sourceOp;
-				}
-				break;
-			}
-
-				// Don't combine. ResizeRel sometimes can resize more children than Resize can do. (see RasterMesh)
-				// It can be combined in an optimization step further in the process, when normal sizes may have been 
-				// optimized already (see OptimizeSemantic)
-//                case EOpType::IM_RESIZE:
-//                {
-//                    break;
-//                }
-
-
-            default:
-            {
-                Sink_ImageResizeRelAST sinker;
-                at = sinker.Apply(this);
-
-                break;
-            }
-
-            }
-
-            break;
-        }
-
-
-        //-------------------------------------------------------------------------------------
-        case EOpType::IM_RESIZELIKE:
-        {
-            Ptr<ImageSizeExpression> sourceSize =
-                    children[op.args.ImageResizeLike.source]->GetImageSizeExpression();
-            Ptr<ImageSizeExpression> sizeSourceSize =
-                    children[op.args.ImageResizeLike.sizeSource]->GetImageSizeExpression();
-
-            if ( *sourceSize == *sizeSourceSize )
-            {
-                at = children[op.args.ImageResizeLike.source].child();
-            }
-
-            else if ( sizeSourceSize->type == ImageSizeExpression::ISET_CONSTANT )
-            {
-                Ptr<ASTOpFixed> newAt = new ASTOpFixed;
-                newAt->op.type = EOpType::IM_RESIZE;
-                newAt->SetChild( newAt->op.args.ImageResize.source, children[op.args.ImageResizeLike.source] );
-                newAt->op.args.ImageResize.size[0] = sizeSourceSize->size[0];
-                newAt->op.args.ImageResize.size[1] = sizeSourceSize->size[1];
-                at = newAt;
-            }
-            else if ( sizeSourceSize->type == ImageSizeExpression::ISET_LAYOUTFACTOR )
-            {
-                // TODO
-                // Skip intermediate ops until the layout
-//                    if ( program.m_code[sizeSourceAt].type != EOpType::IM_BLANKLAYOUT )
-//                    {
-//                        bModified = true;
-
-//                        OP blackLayoutOp;
-//                        blackLayoutOp.type = EOpType::IM_BLANKLAYOUT;
-//                        blackLayoutOp.args.ImageBlankLayout.layout = sizeSourceSize->layout;
-//                        blackLayoutOp.args.ImageBlankLayout.blockSize[0] = sizeSourceSize->factor[0];
-//                        blackLayoutOp.args.ImageBlankLayout.blockSize[1] = sizeSourceSize->factor[1];
-//                        blackLayoutOp.args.ImageBlankLayout.format = L_UByte;
-
-//                        OP opResize = program.m_code[at];
-//                        opResize.args.ImageResizeLike.sizeSource = program.AddOp( blackLayoutOp );
-//                        at = program.AddOp( opResize );
-//                    }
-            }
-
-            break;
-        }
-
-        default:
-            break;
-
-        }
-
-        return at;
-    }
-
 }
diff --git a/Engine/Plugins/Mutable/Source/MutableTools/Private/MuT/Compiler.cpp b/Engine/Plugins/Mutable/Source/MutableTools/Private/MuT/Compiler.cpp
--- a/Engine/Plugins/Mutable/Source/MutableTools/Private/MuT/Compiler.cpp
+++ b/Engine/Plugins/Mutable/Source/MutableTools/Private/MuT/Compiler.cpp
@@ -1,6 +1,5 @@
 // Copyright Epic Games, Inc. All Rights Reserved.
 
-
 #include "MuT/Compiler.h"
 
 #include "MuT/CompilerPrivate.h"
@@ -127,7 +126,6 @@ namespace mu
 		{ true,		{ 0, 1, 1, 0, 0, 0, 0, 0, 0, 0, 0, 0, 0, 0, 0, 0, 0, 0, 0, 0, 0, 0, 0, 0, 0, 0, 0, 0, 0 }		},	// IM_LUMINANCE
 		{ true,		{ 0, 1, 1, 1, 0, 0, 0, 0, 0, 0, 0, 0, 0, 0, 0, 0, 0, 0, 0, 0, 0, 0, 0, 0, 0, 0, 0, 0, 0 }		},	// IM_SWIZZLE
 		{ true,		{ 0, 1, 1, 1, 0, 0, 0, 0, 0, 0, 0, 0, 0, 0, 0, 0, 0, 0, 0, 0, 0, 0, 0, 0, 0, 0, 0, 0, 0 }		},	// IM_COLOURMAP
-		{ true,		{ 0, 0, 0, 0, 0, 0, 0, 0, 0, 0, 0, 0, 0, 0, 0, 0, 0, 0, 0, 0, 0, 0, 0, 0, 0, 0, 0, 0, 0 }		},	// IM_GRADIENT
 		{ true,		{ 0, 1, 1, 1, 0, 0, 0, 0, 0, 0, 0, 0, 0, 0, 0, 0, 0, 0, 0, 0, 0, 0, 0, 0, 0, 0, 0, 0, 0 }		},	// IM_BINARISE
 		{ true,		{ 0, 0, 0, 0, 0, 0, 0, 0, 0, 0, 0, 0, 0, 0, 0, 0, 0, 0, 0, 0, 0, 0, 0, 0, 0, 0, 0, 0, 0 }		},	// IM_PLAINCOLOUR
 		{ true,		{ 0, 0, 0, 0, 0, 0, 0, 0, 0, 0, 0, 0, 0, 0, 0, 0, 0, 0, 0, 0, 0, 0, 0, 0, 0, 0, 0, 0, 0 }		},	// IM_CROP
diff --git a/Engine/Plugins/Mutable/Source/MutableTools/Private/MuT/DataPacker.cpp b/Engine/Plugins/Mutable/Source/MutableTools/Private/MuT/DataPacker.cpp
--- a/Engine/Plugins/Mutable/Source/MutableTools/Private/MuT/DataPacker.cpp
+++ b/Engine/Plugins/Mutable/Source/MutableTools/Private/MuT/DataPacker.cpp
@@ -20,6 +20,8 @@
 #include "MuT/ASTOpImageLayer.h"
 #include "MuT/ASTOpImageLayerColor.h"
 #include "MuT/ASTOpImageMultiLayer.h"
+#include "MuT/ASTOpImagePlainColor.h"
+#include "MuT/ASTOpImageDisplace.h"
 #include "MuT/ASTOpInstanceAdd.h"
 #include "MuT/ASTOpMeshExtractLayoutBlocks.h"
 #include "MuT/ASTOpMeshRemoveMask.h"
@@ -440,15 +442,15 @@ namespace mu
             case EOpType::IM_DISPLACE:
             {
                 recurse = false;
-				const ASTOpFixed* op = static_cast<const ASTOpFixed*>(node.get());
+				const ASTOpImageDisplace* op = static_cast<const ASTOpImageDisplace*>(node.get());
 				TArray<bool> newState;
 				newState.Init(false, int32(EImageFormat::Count));
-				RecurseWithState( op->children[op->op.args.ImageDisplace.source].child(), newState );
+				RecurseWithState( op->Source.child(), newState );
 
                 newState[(int32)EImageFormat::L_UByte ] = true;
                 newState[(int32)EImageFormat::L_UByteRLE ] = true;
 
-                RecurseWithState( op->children[op->op.args.ImageDisplace.displacementMap].child(), newState );
+                RecurseWithState( op->DisplacementMap.child(), newState );
                 break;
             }
 
@@ -774,13 +776,12 @@ namespace mu
 					NewColor->op.args.ColourConstant.value[2] = PlainColor[2];
 					NewColor->op.args.ColourConstant.value[3] = PlainColor[3];
 
-					Ptr<ASTOpFixed> NewPlain = new ASTOpFixed;
-					NewPlain->op.type = EOpType::IM_PLAINCOLOUR;
-					NewPlain->SetChild( NewPlain->op.args.ImagePlainColour.colour, NewColor );
-					NewPlain->op.args.ImagePlainColour.format = pOld->GetFormat();
-					NewPlain->op.args.ImagePlainColour.size[0] = pOld->GetSizeX();
-					NewPlain->op.args.ImagePlainColour.size[1] = pOld->GetSizeY();
-					NewPlain->op.args.ImagePlainColour.LODs = 1;
+					Ptr<ASTOpImagePlainColor> NewPlain = new ASTOpImagePlainColor;
+					NewPlain->Color = NewColor;
+					NewPlain->Format = pOld->GetFormat();
+					NewPlain->Size[0] = pOld->GetSizeX();
+					NewPlain->Size[1] = pOld->GetSizeY();
+					NewPlain->LODs = 1;
 
 					ASTOp::Replace(n, NewPlain);
 
diff --git a/Engine/Plugins/Mutable/Source/MutableTools/Private/MuT/ErrorLog.cpp b/Engine/Plugins/Mutable/Source/MutableTools/Private/MuT/ErrorLog.cpp
--- a/Engine/Plugins/Mutable/Source/MutableTools/Private/MuT/ErrorLog.cpp
+++ b/Engine/Plugins/Mutable/Source/MutableTools/Private/MuT/ErrorLog.cpp
@@ -272,7 +272,6 @@ namespace mu
 		TEXT("IM_LUMINANCE     "),
 		TEXT("IM_SWIZZLE       "),
 		TEXT("IM_COLOURMAP     "),
-		TEXT("IM_GRADIENT      "),
 		TEXT("IM_BINARISE      "),
 		TEXT("IM_PLAINCOLOUR   "),
 		TEXT("IM_CROP          "),
diff --git a/Engine/Plugins/Mutable/Source/MutableTools/Private/MuT/NodeImageGradient.cpp b/Engine/Plugins/Mutable/Source/MutableTools/Private/MuT/NodeImageGradient.cpp
deleted file mode 100644
diff --git a/Engine/Plugins/Mutable/Source/MutableTools/Private/MuT/NodeImageGradientPrivate.h b/Engine/Plugins/Mutable/Source/MutableTools/Private/MuT/NodeImageGradientPrivate.h
deleted file mode 100644
diff --git a/Engine/Plugins/Mutable/Source/MutableTools/Public/MuT/NodeImageGradient.h b/Engine/Plugins/Mutable/Source/MutableTools/Public/MuT/NodeImageGradient.h
deleted file mode 100644
