[Mutable] Fix Image descriptor not being computed correctly on some cases.

- CL39367032 introduced a bug where the image descriptor result of an operation would not be computed if the operation address was not unique.

- Now the  image decriptor intermediate results are stored in a map. If the result for a given address is already computed it will be reused.

#rnx

diff --git a/Engine/Plugins/Mutable/Source/CustomizableObject/Private/MuCO/UnrealMutableModelDiskStreamer.cpp b/Engine/Plugins/Mutable/Source/CustomizableObject/Private/MuCO/UnrealMutableModelDiskStreamer.cpp
--- a/Engine/Plugins/Mutable/Source/CustomizableObject/Private/MuCO/UnrealMutableModelDiskStreamer.cpp
+++ b/Engine/Plugins/Mutable/Source/CustomizableObject/Private/MuCO/UnrealMutableModelDiskStreamer.cpp
@@ -234,6 +234,8 @@ FORCEINLINE FString MakeBlockFilePath(const FString& Prefix, const FMutableStrea
 
 bool FUnrealMutableModelBulkReader::DoesBlockExist(const mu::FModel* Model, uint32 Key)
 {
+	MUTABLE_CPUPROFILER_SCOPE(FUnrealMutableModelBulkReader::DoesBlockExist)
+
 	FObjectData* ObjectData = Objects.FindByPredicate(
 			[Model](const FObjectData& ObjectData) { return ObjectData.Model.Pin().Get() == Model; });
 
@@ -251,6 +253,8 @@ bool FUnrealMutableModelBulkReader::DoesBlockExist(const mu::FModel* Model, uint
 
 	if (ObjectData->ModelStreamableBulkData->StreamableBulkData.IsValidIndex(Block->FileId))
 	{	
+		MUTABLE_CPUPROFILER_SCOPE(FUnrealMutableModelBulkReader::DoesBlockExist_Bulk)
+
 		FByteBulkData& BulkData = ObjectData->ModelStreamableBulkData->StreamableBulkData[Block->FileId];
 		return BulkData.DoesExist();
 	}
@@ -259,11 +263,14 @@ bool FUnrealMutableModelBulkReader::DoesBlockExist(const mu::FModel* Model, uint
 #if WITH_EDITOR
 		if (ObjectData->bIsStoredInDDC)
 		{
+			MUTABLE_CPUPROFILER_SCOPE(FUnrealMutableModelBulkReader::DoesBlockExist_DDC)
 			return true;
 		}
 		else
 #endif
 		{
+			MUTABLE_CPUPROFILER_SCOPE(FUnrealMutableModelBulkReader::DoesBlockExist_File)
+
 			const FString FilePath = Private::MakeBlockFilePath(ObjectData->BulkFilePrefix, *Block);
 			return FPlatformFileManager::Get().GetPlatformFile().FileExists(*FilePath);
 		}
diff --git a/Engine/Plugins/Mutable/Source/MutableRuntime/Private/MuR/CodeRunner.cpp b/Engine/Plugins/Mutable/Source/MutableRuntime/Private/MuR/CodeRunner.cpp
--- a/Engine/Plugins/Mutable/Source/MutableRuntime/Private/MuR/CodeRunner.cpp
+++ b/Engine/Plugins/Mutable/Source/MutableRuntime/Private/MuR/CodeRunner.cpp
@@ -135,9 +135,16 @@ namespace mu
 		, Params(InParams)
 		, LODMask(InLodMask)
 	{
+		MUTABLE_CPUPROFILER_SCOPE(CodeRunner_Create)
+
 		const FProgram& program = Model->GetPrivate()->Program;
 		ScheduledStagePerOp.resize(program.OpAddress.Num());
-
+		
+		if (Type == FScheduledOp::EType::ImageDesc)
+		{
+			ImageDescResults.Reserve(64);
+		}
+	
 		// We will read this in the end, so make sure we keep it.
    		if (Type == FScheduledOp::EType::Full)
    		{
@@ -5965,24 +5972,12 @@ namespace mu
         }
     }
 
-	//---------------------------------------------------------------------------------------------
+	/** */
 	void CodeRunner::RunCodeImageDesc(const FScheduledOp& Item, const FParameters* pParams, const FModel* InModel,  uint32 LodMask )
 	{
 		MUTABLE_CPUPROFILER_SCOPE(RunCodeImageDesc);
-
 		check(Item.Type == FScheduledOp::EType::ImageDesc);
 
-		// Add the result heap address if not present.
-		if (HeapImageDesc.Num() == 0)
-		{
-			HeapImageDesc.AddDefaulted();
-		}
-
-		// Item.CustomState can point to HeapData or HeapImageDesc depending on the operation.
-		// Can not check here. Assume access to the arrays will do the check. 
-		// We could set the Item.CustomState highest bit to distinguish if needed.
-		//check((int32)Item.CustomState < HeapImageDesc.Num());
-
 		const FProgram& Program = Model->GetPrivate()->Program;
 
 		EOpType Type = Program.GetOpType(Item.At);
@@ -5991,11 +5986,13 @@ namespace mu
 
 		case EOpType::IM_CONSTANT:
 		{
+			MUTABLE_CPUPROFILER_SCOPE(RunCodeImageDesc_IM_CONSTANT);
+
 			check(Item.Stage == 0);
 			OP::ResourceConstantArgs Args = Program.GetOpArgs<OP::ResourceConstantArgs>(Item.At);
 			int32 ImageIndex = Args.value;
 
-			FExtendedImageDesc& Result = HeapImageDesc[Item.CustomState];
+			FExtendedImageDesc& Result = ImageDescResults.FindOrAdd(Item.At);
 
 			Result.m_format = Program.ConstantImages[ImageIndex].ImageFormat;
 			Result.m_size[0] = Program.ConstantImages[ImageIndex].ImageSizeX;
@@ -6034,21 +6031,27 @@ namespace mu
 
 		case EOpType::IM_PARAMETER:
 		{
+			MUTABLE_CPUPROFILER_SCOPE(RunCodeImageDesc_IM_PARAMETER);
 			check(Item.Stage == 0);
 
 			OP::ParameterArgs Args = Program.GetOpArgs<OP::ParameterArgs>(Item.At);
 			FName Id = pParams->GetImageValue(Args.variable);
-			HeapImageDesc[Item.CustomState] = GetExternalImageDesc(Id);
+
+			FExtendedImageDesc& Result = ImageDescResults.FindOrAdd(Item.At);
+			Result = GetExternalImageDesc(Id);
+
 			StoreValidDesc(Item);
 			break;
 		}
 
 		case EOpType::IM_REFERENCE:
 		{
+			MUTABLE_CPUPROFILER_SCOPE(RunCodeImageDesc_IM_REFERENCE);
 			check(Item.Stage == 0);
 
 			OP::ResourceReferenceArgs Args = Program.GetOpArgs<OP::ResourceReferenceArgs>(Item.At);
-			FExtendedImageDesc& Result = HeapImageDesc[Item.CustomState];
+			FExtendedImageDesc& Result = ImageDescResults.FindOrAdd(Item.At);
+			
 			Result = FExtendedImageDesc{ Args.ImageDesc, 0 };
 
 			StoreValidDesc(Item);
@@ -6057,6 +6060,7 @@ namespace mu
 
 		case EOpType::IM_CONDITIONAL:
 		{
+			MUTABLE_CPUPROFILER_SCOPE(RunCodeImageDesc_IM_CONDITIONAL);
 			OP::ConditionalArgs Args = Program.GetOpArgs<OP::ConditionalArgs>(Item.At);
 			switch (Item.Stage)
 			{
@@ -6066,7 +6070,7 @@ namespace mu
 				FScheduledOp FullConditionOp(Args.condition, Item);
 				FullConditionOp.Type = FScheduledOp::EType::Full;
 
-				AddOp(FScheduledOp(Item.At, Item, 1, Item.CustomState), FullConditionOp);
+				AddOp(FScheduledOp(Item.At, Item, 1), FullConditionOp);
 				break;
 			}
 
@@ -6075,12 +6079,19 @@ namespace mu
 				bool bValue = LoadBool(FCacheAddress(Args.condition, Item.ExecutionIndex, Item.ExecutionOptions));
 				OP::ADDRESS ResultAt = bValue ? Args.yes : Args.no;
 
-				AddOp(FScheduledOp(Item.At, Item, 2, Item.CustomState), 
-						FScheduledOp(ResultAt, Item, 0, Item.CustomState));
+				AddOp(FScheduledOp(Item.At, Item, 2, ResultAt), 
+						FScheduledOp(ResultAt, Item, 0));
 				break;
 			}
 
-			case 2: StoreValidDesc(Item); break;
+			case 2: 
+			{
+				FExtendedImageDesc& Result =  ImageDescResults.FindOrAdd(Item.At);
+				Result = ImageDescResults.FindChecked(Item.CustomState);
+
+				StoreValidDesc(Item); 
+				break;
+			}
 			default: check(false);
 			}
 			break;
@@ -6088,6 +6099,7 @@ namespace mu
 
 		case EOpType::IM_SWITCH:
 		{
+			MUTABLE_CPUPROFILER_SCOPE(RunCodeImageDesc_IM_SWITCH);
 			const uint8* Data = Program.GetOpArgsPointer(Item.At);
 		
 			OP::ADDRESS VarAddress;
@@ -6111,10 +6123,11 @@ namespace mu
 					// We need to run the full condition result
 					FScheduledOp FullVariableOp(VarAddress, Item);
 					FullVariableOp.Type = FScheduledOp::EType::Full;
-					AddOp(FScheduledOp(Item.At, Item, 1, Item.CustomState), FullVariableOp);
+					AddOp(FScheduledOp(Item.At, Item, 1), FullVariableOp);
 				}
 				else
 				{
+					ImageDescResults.FindOrAdd(Item.At);
 					StoreValidDesc(Item);
 				}
 				break;
@@ -6143,13 +6156,20 @@ namespace mu
 					}
 				}
 
-				AddOp(FScheduledOp(Item.At, Item, 2, Item.CustomState),
-						FScheduledOp(ValueAt, Item, 0, Item.CustomState));
+				AddOp(FScheduledOp(Item.At, Item, 2, ValueAt),
+						FScheduledOp(ValueAt, Item, 0));
 
 				break;
 			}
 
-			case 2: StoreValidDesc(Item); break;
+			case 2: 
+			{
+				FExtendedImageDesc& Result = ImageDescResults.FindOrAdd(Item.At);
+				Result = ImageDescResults.FindChecked(Item.CustomState);
+
+				StoreValidDesc(Item); 
+				break;
+			}
 			default: check(false); break;
 			}
 			break;
@@ -6157,43 +6177,29 @@ namespace mu
 
 		case EOpType::IM_LAYERCOLOUR:
 		{
+			MUTABLE_CPUPROFILER_SCOPE(RunCodeImageDesc_IM_LAYERCOLOUR);
 			OP::ImageLayerColourArgs Args = Program.GetOpArgs<OP::ImageLayerColourArgs>(Item.At);
+
 			switch (Item.Stage)
 			{
 			case 0:
 			{
-				FScheduledOpData ResultArrayData;
-				ResultArrayData.ReturnArray.ReturnAddress = Item.CustomState;
-				ResultArrayData.ReturnArray.Begin = HeapImageDesc.Num();
-				ResultArrayData.ReturnArray.Num = 2;
-
-				HeapImageDesc.AddDefaulted(ResultArrayData.ReturnArray.Num);
-
-				int32 ResultArrayAddress = HeapData.Add(ResultArrayData);
-
-				AddOp(FScheduledOp(Item.At, Item, 1, ResultArrayAddress), 
-						FScheduledOp(Args.base, Item, 0, ResultArrayData.ReturnArray.Begin + 0),
-						FScheduledOp(Args.mask, Item, 0, ResultArrayData.ReturnArray.Begin + 1));
+				AddOp(FScheduledOp(Item.At, Item, 1), 
+						FScheduledOp(Args.base, Item, 0), 
+						FScheduledOp(Args.mask, Item, 0));
 
 				break;
 			}
 			case 1:
 			{
-				FScheduledOpData ResultArrayData = HeapData[Item.CustomState];
-				TArrayView<FExtendedImageDesc> DepResults(
-						HeapImageDesc.GetData() + ResultArrayData.ReturnArray.Begin, 
-						ResultArrayData.ReturnArray.Num);
-				
-				check(DepResults.Num() == 2);
-	
-				FExtendedImageDesc& BaseImageDesc = DepResults[0];
-				FExtendedImageDesc& MaskImageDesc = DepResults[1];
+				FExtendedImageDesc& Result = ImageDescResults.FindOrAdd(Item.At); 
+				Result = ImageDescResults.FindChecked(Args.base);
 
-				FExtendedImageDesc& Result = HeapImageDesc[ResultArrayData.ReturnArray.ReturnAddress]; 
-				Result = BaseImageDesc;
-
-				Result.FirstLODAvailable = FMath::Max(
-							Result.FirstLODAvailable, MaskImageDesc.FirstLODAvailable);
+				if (Args.mask)
+				{
+					const FExtendedImageDesc& MaskResult = ImageDescResults.FindChecked(Args.mask);
+					Result.FirstLODAvailable = FMath::Max(Result.FirstLODAvailable, MaskResult.FirstLODAvailable);
+				}
 
 				StoreValidDesc(Item); 
 
@@ -6206,47 +6212,36 @@ namespace mu
 
 		case EOpType::IM_LAYER:
 		{
-			OP::ImageLayerArgs Args = Program.GetOpArgs<OP::ImageLayerArgs>(Item.At);
+			MUTABLE_CPUPROFILER_SCOPE(RunCodeImageDesc_IM_LAYER);
+			OP::ImageLayerArgs Args = Program.GetOpArgs<OP::ImageLayerArgs>(Item.At);	
 			switch (Item.Stage)
 			{
 			case 0:
-			{
-				FScheduledOpData ResultArrayData;
-				ResultArrayData.ReturnArray.ReturnAddress = Item.CustomState;
-				ResultArrayData.ReturnArray.Begin = HeapImageDesc.Num();
-				ResultArrayData.ReturnArray.Num = 3;
-
-				HeapImageDesc.AddDefaulted(ResultArrayData.ReturnArray.Num);
-
-				int32 ResultArrayAddress = HeapData.Add(ResultArrayData);
+			{	
+				AddOp(FScheduledOp(Item.At, Item, 1), 
+						FScheduledOp(Args.base, Item, 0),
+						FScheduledOp(Args.mask, Item, 0),
+						FScheduledOp(Args.blended, Item, 0));
 
-				AddOp(FScheduledOp(Item.At, Item, 1, ResultArrayAddress), 
-						FScheduledOp(Args.base, Item, 0, ResultArrayData.ReturnArray.Begin + 0),
-						FScheduledOp(Args.mask, Item, 0, ResultArrayData.ReturnArray.Begin + 1),
-						FScheduledOp(Args.blended, Item, 0, ResultArrayData.ReturnArray.Begin + 2));
 				break;
 			}
 			case 1:
 			{
-				FScheduledOpData ResultArrayData = HeapData[Item.CustomState];
-				TArrayView<FExtendedImageDesc> DepResults(
-						HeapImageDesc.GetData() + ResultArrayData.ReturnArray.Begin, 
-						ResultArrayData.ReturnArray.Num);
-				
-				check(DepResults.Num() == 3);
-	
-				FExtendedImageDesc& BaseImageDesc = DepResults[0];
-				FExtendedImageDesc& MaskImageDesc = DepResults[1];
-				FExtendedImageDesc& BlendedImageDesc = DepResults[2];
+				FExtendedImageDesc& Result = ImageDescResults.FindOrAdd(Item.At); 
+				Result = ImageDescResults.FindChecked(Args.base);
 
-				FExtendedImageDesc& Result = HeapImageDesc[ResultArrayData.ReturnArray.ReturnAddress]; 
-				Result = BaseImageDesc;
-
-				Result.FirstLODAvailable = FMath::Max3(
-							Result.FirstLODAvailable, 
-							MaskImageDesc.FirstLODAvailable,
-							BlendedImageDesc.FirstLODAvailable);
+				if (Args.mask)
+				{
+					const FExtendedImageDesc& MaskResult = ImageDescResults.FindChecked(Args.mask);
+					Result.FirstLODAvailable = FMath::Max(Result.FirstLODAvailable, MaskResult.FirstLODAvailable);
+				}
 
+				if (Args.blended)
+				{
+					const FExtendedImageDesc& BlendedResult = ImageDescResults.FindChecked(Args.blended);
+					Result.FirstLODAvailable = FMath::Max(Result.FirstLODAvailable, BlendedResult.FirstLODAvailable);
+				}
+		
 				StoreValidDesc(Item); 
 				break;
 			}
@@ -6257,6 +6252,7 @@ namespace mu
 
 		case EOpType::IM_MULTILAYER:
 		{
+			MUTABLE_CPUPROFILER_SCOPE(RunCodeImageDesc_IM_MULTILAYER);
 			OP::ImageMultiLayerArgs Args = Program.GetOpArgs<OP::ImageMultiLayerArgs>(Item.At);
 			switch (Item.Stage)
 			{
@@ -6264,11 +6260,18 @@ namespace mu
 			{
 				//TODO: For now multilayer operations will only check the base to get the descriptor.
 				// but all iterations should be checked for available mips. 
-				AddOp(FScheduledOp(Item.At, Item, 1, Item.CustomState), 
-						FScheduledOp(Args.base, Item, 0, Item.CustomState));
+				AddOp(FScheduledOp(Item.At, Item, 1), 
+						FScheduledOp(Args.base, Item, 0));
+				break;
+			}
+			case 1: 
+			{
+				FExtendedImageDesc& Result = ImageDescResults.FindOrAdd(Item.At); 
+				Result = ImageDescResults.FindChecked(Args.base);
+
+				StoreValidDesc(Item); 
 				break;
 			}
-			case 1: StoreValidDesc(Item); break;
 			default: check(false);
 			}
 			break;
@@ -6276,16 +6279,24 @@ namespace mu
 
 		case EOpType::IM_NORMALCOMPOSITE:
 		{
+			MUTABLE_CPUPROFILER_SCOPE(RunCodeImageDesc_IM_NORMALCOMPOSITE);
 			OP::ImageNormalCompositeArgs Args = Program.GetOpArgs<OP::ImageNormalCompositeArgs>(Item.At);
 			switch (Item.Stage)
 			{
 			case 0: 
 			{
-				AddOp(FScheduledOp(Item.At, Item, 1, Item.CustomState), 
-						FScheduledOp(Args.base, Item, 0, Item.CustomState)); 
+				AddOp(FScheduledOp(Item.At, Item, 1), 
+						FScheduledOp(Args.base, Item, 0)); 
+				break;
+			}
+			case 1: 
+			{
+				FExtendedImageDesc& Result = ImageDescResults.FindOrAdd(Item.At); 
+				Result = ImageDescResults.FindChecked(Args.base);
+				
+				StoreValidDesc(Item); 
 				break;
 			}
-			case 1: StoreValidDesc(Item); break;
 			default: check(false);
 			}
 			break;
@@ -6293,26 +6304,30 @@ namespace mu
 
 		case EOpType::IM_PIXELFORMAT:
 		{
+			MUTABLE_CPUPROFILER_SCOPE(RunCodeImageDesc_IM_PIXELFORMAT);
 			OP::ImagePixelFormatArgs Args = Program.GetOpArgs<OP::ImagePixelFormatArgs>(Item.At);
 			switch (Item.Stage)
 			{
 			case 0:
-				AddOp(FScheduledOp(Item.At, Item, 1, Item.CustomState), 
-						FScheduledOp(Args.source, Item, 0, Item.CustomState));
+				AddOp(FScheduledOp(Item.At, Item, 1), 
+						FScheduledOp(Args.source, Item, 0));
 				break;
 
 			case 1:
 			{
-				// Update directly in the heap
-				EImageFormat OldFormat = HeapImageDesc[Item.CustomState].m_format;
+
+				FExtendedImageDesc& Result = ImageDescResults.FindOrAdd(Item.At); 
+				Result = ImageDescResults.FindChecked(Args.source);
+
+				EImageFormat OldFormat = Result.m_format;
 				EImageFormat NewFormat = Args.format;
-				if (Args.formatIfAlpha != EImageFormat::None
-					&&
-					GetImageFormatData(OldFormat).Channels > 3)
+				if (Args.formatIfAlpha != EImageFormat::None && GetImageFormatData(OldFormat).Channels > 3)
 				{
 					NewFormat = Args.formatIfAlpha;
 				}
-				HeapImageDesc[Item.CustomState].m_format = NewFormat;				
+
+				Result.m_format = NewFormat;
+
 				StoreValidDesc(Item);
 				break;
 			}
@@ -6326,20 +6341,23 @@ namespace mu
 
 		case EOpType::IM_MIPMAP:
 		{
+			MUTABLE_CPUPROFILER_SCOPE(RunCodeImageDesc_IM_MIPMAP);
 			OP::ImageMipmapArgs Args = Program.GetOpArgs<OP::ImageMipmapArgs>(Item.At);
 			switch (Item.Stage)
 			{
 			case 0:
-				AddOp(FScheduledOp(Item.At, Item, 1, Item.CustomState), 
-						FScheduledOp(Args.source, Item, 0, Item.CustomState));
+				AddOp(FScheduledOp(Item.At, Item, 1), 
+						FScheduledOp(Args.source, Item, 0));
 				break;
 
 			case 1:
 			{
 				// Somewhat synched with Full op execution code.
-				FExtendedImageDesc BaseDesc = HeapImageDesc[Item.CustomState];
+				FExtendedImageDesc& Result = ImageDescResults.FindOrAdd(Item.At);
+				Result = ImageDescResults.FindChecked(Args.source);
+
 				int32 LevelCount = Args.levels;
-				int32 MaxLevelCount = FImage::GetMipmapCount(BaseDesc.m_size[0], BaseDesc.m_size[1]);
+				int32 MaxLevelCount = FImage::GetMipmapCount(Result.m_size[0], Result.m_size[1]);
 				if (LevelCount == 0)
 				{
 					LevelCount = MaxLevelCount;
@@ -6352,11 +6370,12 @@ namespace mu
 				}
 
 				// At least keep the levels we already have.
-				int32 StartLevel = BaseDesc.m_lods;
+				int32 StartLevel = Result.m_lods;
 				LevelCount = FMath::Max(StartLevel, LevelCount);
 
 				// Update result.
-				HeapImageDesc[Item.CustomState].m_lods = LevelCount;
+				Result.m_lods = LevelCount;
+
 				StoreValidDesc(Item);
 				break;
 			}
@@ -6369,20 +6388,28 @@ namespace mu
 
 		case EOpType::IM_RESIZE:
 		{
+			MUTABLE_CPUPROFILER_SCOPE(RunCodeImageDesc_IM_RESIZE);
 			OP::ImageResizeArgs Args = Program.GetOpArgs<OP::ImageResizeArgs>(Item.At);
 			switch (Item.Stage)
 			{
 			case 0:
-				AddOp(FScheduledOp(Item.At, Item, 1, Item.CustomState), 
-						FScheduledOp(Args.Source, Item, 0, Item.CustomState));
+			{
+				AddOp(FScheduledOp(Item.At, Item, 1), 
+						FScheduledOp(Args.Source, Item, 0));
 				break;
-
+			}
 			case 1:
-				HeapImageDesc[Item.CustomState].m_size[0] = Args.Size[0];
-				HeapImageDesc[Item.CustomState].m_size[1] = Args.Size[1];
+			{
+				FExtendedImageDesc& Result = ImageDescResults.FindOrAdd(Item.At);
+				Result = ImageDescResults.FindChecked(Args.Source);
+
+				Result.m_size[0] = Args.Size[0];
+				Result.m_size[1] = Args.Size[1];
+
 				StoreValidDesc(Item);
-				break;
 
+				break;
+			}
 			default:
 				check(false);
 			}
@@ -6391,42 +6418,29 @@ namespace mu
 
 		case EOpType::IM_RESIZELIKE:
 		{
+			MUTABLE_CPUPROFILER_SCOPE(RunCodeImageDesc_IM_RESIZELIKE);
 			OP::ImageResizeLikeArgs Args = Program.GetOpArgs<OP::ImageResizeLikeArgs>(Item.At);
 			switch (Item.Stage)
 			{
 			case 0:
 			{
-				FScheduledOpData ResultArrayData;
-				ResultArrayData.ReturnArray.ReturnAddress = Item.CustomState;
-				ResultArrayData.ReturnArray.Begin = HeapImageDesc.Num();
-				ResultArrayData.ReturnArray.Num = 2;
-
-				HeapImageDesc.AddDefaulted(ResultArrayData.ReturnArray.Num);
-
-				int32 ResultArrayAddress = HeapData.Add(ResultArrayData);
-
-				AddOp(FScheduledOp(Item.At, Item, 1, ResultArrayAddress), 
-						FScheduledOp(Args.Source, Item, 0, ResultArrayData.ReturnArray.Begin + 0),
-						FScheduledOp(Args.SizeSource, Item, 0, ResultArrayData.ReturnArray.Begin + 1));
+				AddOp(FScheduledOp(Item.At, Item, 1), 
+						FScheduledOp(Args.Source, Item, 0), 
+						FScheduledOp(Args.SizeSource, Item, 0));
 
 				break;
 			}
 
 			case 1:
 			{
-				FScheduledOpData ResultArrayData = HeapData[Item.CustomState];
-				TArrayView<FExtendedImageDesc> DepResults(
-						HeapImageDesc.GetData() + ResultArrayData.ReturnArray.Begin, 
-						ResultArrayData.ReturnArray.Num);
-				
-				check(DepResults.Num() == 2);
+				FExtendedImageDesc& Result = ImageDescResults.FindOrAdd(Item.At);
+				Result = ImageDescResults.FindChecked(Args.Source);
 	
-				FExtendedImageDesc& SourceImageDesc = DepResults[0];
-				FExtendedImageDesc& SizeImageDesc = DepResults[1];
-
-				FExtendedImageDesc& Result = HeapImageDesc[ResultArrayData.ReturnArray.ReturnAddress]; 
-				Result = SourceImageDesc;
-				Result.m_size = SizeImageDesc.m_size;
+				if (Args.SizeSource)
+				{
+					const FExtendedImageDesc& SizeSourceResult = ImageDescResults.FindChecked(Args.SizeSource);
+					Result.m_size = SizeSourceResult.m_size;
+				}
 
 				StoreValidDesc(Item);
 				break;
@@ -6441,21 +6455,26 @@ namespace mu
 
 		case EOpType::IM_RESIZEREL:
 		{
+			MUTABLE_CPUPROFILER_SCOPE(RunCodeImageDesc_IM_RESIZEREL);
 			OP::ImageResizeRelArgs Args = Program.GetOpArgs<OP::ImageResizeRelArgs>(Item.At);
 			switch (Item.Stage)
 			{
 			case 0:
-				AddOp(FScheduledOp(Item.At, Item, 1, Item.CustomState), 
-						FScheduledOp(Args.Source, Item, 0, Item.CustomState));
+				AddOp(FScheduledOp(Item.At, Item, 1), 
+						FScheduledOp(Args.Source, Item, 0));
 				break;
 
 			case 1:
 			{
-				FExtendedImageDesc& ResultAndBaseDesc = HeapImageDesc[Item.CustomState];
+				FExtendedImageDesc& Result = ImageDescResults.FindOrAdd(Item.At);
+				Result = ImageDescResults.FindChecked(Args.Source);
+				
 				FImageSize DestSize(
-					uint16(ResultAndBaseDesc.m_size[0] * Args.Factor[0] + 0.5f),
-					uint16(ResultAndBaseDesc.m_size[1] * Args.Factor[1] + 0.5f));
-				ResultAndBaseDesc.m_size = DestSize;
+					uint16(Result.m_size[0] * Args.Factor[0] + 0.5f),
+					uint16(Result.m_size[1] * Args.Factor[1] + 0.5f));
+
+				Result.m_size = DestSize;
+
 				StoreValidDesc(Item);
 				break;
 			}
@@ -6469,6 +6488,7 @@ namespace mu
 
 		case EOpType::IM_BLANKLAYOUT:
 		{
+			MUTABLE_CPUPROFILER_SCOPE(RunCodeImageDesc_IM_BLANKLAYOUT);
 			OP::ImageBlankLayoutArgs Args = Program.GetOpArgs<OP::ImageBlankLayoutArgs>(Item.At);
 			switch (Item.Stage)
 			{
@@ -6477,34 +6497,37 @@ namespace mu
 				// We need to run the full layout
 				FScheduledOp FullLayoutOp(Args.Layout, Item);
 				FullLayoutOp.Type = FScheduledOp::EType::Full;
-				AddOp(FScheduledOp(Item.At, Item, 1, Item.CustomState), FullLayoutOp);
+				AddOp(FScheduledOp(Item.At, Item, 1), FullLayoutOp);
 				break;
 			}
 
 			case 1:
 			{
-				TSharedPtr<const FLayout> pLayout = LoadLayout(FCacheAddress(Args.Layout, Item));
+				FExtendedImageDesc& Result = ImageDescResults.FindOrAdd(Item.At);
+				
+				TSharedPtr<const FLayout> Layout = LoadLayout(FCacheAddress(Args.Layout, Item));
 
-				FIntPoint SizeInBlocks = pLayout->GetGridSize();
+				FIntPoint SizeInBlocks = Layout->GetGridSize();
 				FIntPoint BlockSizeInPixels(Args.BlockSize[0], Args.BlockSize[1]);
 				FIntPoint ImageSizeInPixels = SizeInBlocks * BlockSizeInPixels;
 
-				FExtendedImageDesc& ResultAndBaseDesc = HeapImageDesc[Item.CustomState];
+				
 				FImageSize DestSize(uint16(ImageSizeInPixels.X), uint16(ImageSizeInPixels.Y));
-				ResultAndBaseDesc.m_size = DestSize;
-				ResultAndBaseDesc.m_format = Args.Format;
+				Result.m_size = DestSize;
+				Result.m_format = Args.Format;
 				
 				if (Args.GenerateMipmaps)
 				{
 					if (Args.MipmapCount == 0)
 					{
-						ResultAndBaseDesc.m_lods = FImage::GetMipmapCount(ImageSizeInPixels.X, ImageSizeInPixels.Y);
+						Result.m_lods = FImage::GetMipmapCount(ImageSizeInPixels.X, ImageSizeInPixels.Y);
 					}
 					else
 					{
-						ResultAndBaseDesc.m_lods = Args.MipmapCount;
+						Result.m_lods = Args.MipmapCount;
 					}
 				}
+
 				StoreValidDesc(Item);
 				break;
 			}
@@ -6518,42 +6541,27 @@ namespace mu
 
 		case EOpType::IM_COMPOSE:
 		{
+			MUTABLE_CPUPROFILER_SCOPE(RunCodeImageDesc_IM_COMPOSE);
 			OP::ImageComposeArgs Args = Program.GetOpArgs<OP::ImageComposeArgs>(Item.At);
 			switch (Item.Stage)
 			{
 			case 0: 
 			{
-				FScheduledOpData ResultArrayData;
-				ResultArrayData.ReturnArray.ReturnAddress = Item.CustomState;
-				ResultArrayData.ReturnArray.Begin = HeapImageDesc.Num();
-				ResultArrayData.ReturnArray.Num = 2;
-
-				HeapImageDesc.AddDefaulted(ResultArrayData.ReturnArray.Num);
-
-				int32 ResultArrayAddress = HeapData.Add(ResultArrayData);
-
-				AddOp(FScheduledOp(Item.At, Item, 1, ResultArrayAddress), 
-						FScheduledOp(Args.base, Item, 0, ResultArrayData.ReturnArray.Begin + 0),
-						FScheduledOp(Args.blockImage, Item, 0, ResultArrayData.ReturnArray.Begin + 1));
-
+				AddOp(FScheduledOp(Item.At, Item, 1), 
+						FScheduledOp(Args.base, Item, 0),
+						FScheduledOp(Args.blockImage, Item, 0)); 
 				break;
 			} 
 			case 1:
 			{
-				FScheduledOpData ResultArrayData = HeapData[Item.CustomState];
-				TArrayView<FExtendedImageDesc> DepResults(
-						HeapImageDesc.GetData() + ResultArrayData.ReturnArray.Begin, 
-						ResultArrayData.ReturnArray.Num);
-				
-				check(DepResults.Num() == 2);
-	
-				FExtendedImageDesc& BaseImageDesc = DepResults[0];
-				FExtendedImageDesc& BlockImageDesc = DepResults[1];
+				FExtendedImageDesc& Result = ImageDescResults.FindOrAdd(Item.At); 
+				Result = ImageDescResults.FindChecked(Args.base);
 
-				FExtendedImageDesc& Result = HeapImageDesc[ResultArrayData.ReturnArray.ReturnAddress]; 
-				Result = BaseImageDesc;
-				Result.FirstLODAvailable = 
-						FMath::Max(BaseImageDesc.FirstLODAvailable, BlockImageDesc.FirstLODAvailable);
+				if (Args.blockImage)
+				{
+					const FExtendedImageDesc& BlockResult = ImageDescResults.FindChecked(Args.blockImage);
+					Result.FirstLODAvailable = FMath::Max(Result.FirstLODAvailable, BlockResult.FirstLODAvailable);
+				}
 
 				StoreValidDesc(Item); 
 				break;
@@ -6565,6 +6573,7 @@ namespace mu
 
 		case EOpType::IM_INTERPOLATE:
 		{
+			MUTABLE_CPUPROFILER_SCOPE(RunCodeImageDesc_IM_INTERPOLATE);
 			OP::ImageInterpolateArgs Args = Program.GetOpArgs<OP::ImageInterpolateArgs>(Item.At);
 
 			int32 NumImages = 0;
@@ -6580,39 +6589,27 @@ namespace mu
 			{
 			case 0: 
 			{
-				FScheduledOpData ResultArrayData;
-				ResultArrayData.ReturnArray.ReturnAddress = Item.CustomState;
-				ResultArrayData.ReturnArray.Begin = HeapImageDesc.Num();
-				ResultArrayData.ReturnArray.Num = NumImages;
-
-				HeapImageDesc.AddDefaulted(ResultArrayData.ReturnArray.Num);
-				int32 ResultArrayAddress = HeapData.Add(ResultArrayData);
-
-				TArray<FScheduledOp, TInlineAllocator<MUTABLE_OP_MAX_INTERPOLATE_COUNT>> Deps;
-
-				for (int32 Index = 0; Index < NumImages; ++Index)
+				TArray<FScheduledOp, TFixedAllocator<MUTABLE_OP_MAX_INTERPOLATE_COUNT>> Deps;
+				for (int32 ImageIndex = 0; ImageIndex < NumImages; ++ImageIndex)
 				{
-					Deps.Add(FScheduledOp(Args.Targets[Index], Item, 0, ResultArrayData.ReturnArray.Begin + Index));
+					Deps.Add(FScheduledOp(Args.Targets[ImageIndex], Item, 0));
 				}
 
-				AddOp(FScheduledOp(Item.At, Item, 1, ResultArrayAddress), Deps); 
+				AddOp(FScheduledOp(Item.At, Item, 1), Deps); 
 				break;
 			}
 			case 1: 
 			{
-				FScheduledOpData ResultArrayData = HeapData[Item.CustomState];
-				TArrayView<FExtendedImageDesc> DepResults(
-						HeapImageDesc.GetData() + ResultArrayData.ReturnArray.Begin, 
-						ResultArrayData.ReturnArray.Num);
-				
-				check(DepResults.Num() == NumImages && NumImages > 0);
+				FExtendedImageDesc& Result = ImageDescResults.FindOrAdd(Item.At);
 
-				FExtendedImageDesc& Result = HeapImageDesc[ResultArrayData.ReturnArray.ReturnAddress]; 
-				Result = DepResults[0];
+				check(Args.Targets[0]);
 
-				for (const FExtendedImageDesc& ResultImageDesc : DepResults)
+				Result = ImageDescResults.FindChecked(Args.Targets[0]);
+				
+				for (int32 ImageIndex = 1; ImageIndex < NumImages; ++ImageIndex)
 				{
-					Result.FirstLODAvailable = FMath::Max(Result.FirstLODAvailable, ResultImageDesc.FirstLODAvailable);
+					const FExtendedImageDesc& TargetResult = ImageDescResults.FindChecked(Args.Targets[ImageIndex]);
+					Result.FirstLODAvailable = FMath::Max(Result.FirstLODAvailable, TargetResult.FirstLODAvailable);
 				}
 
 				StoreValidDesc(Item); break;
@@ -6624,16 +6621,24 @@ namespace mu
 
 		case EOpType::IM_SATURATE:
 		{
+			MUTABLE_CPUPROFILER_SCOPE(RunCodeImageDesc_IM_SATURATE);
 			OP::ImageSaturateArgs Args = Program.GetOpArgs<OP::ImageSaturateArgs>(Item.At);
 			switch (Item.Stage)
 			{
 			case 0: 
 			{
-				AddOp(FScheduledOp(Item.At, Item, 1, Item.CustomState), 
-						FScheduledOp(Args.Base, Item, 0, Item.CustomState)); 
+				AddOp(FScheduledOp(Item.At, Item, 1), 
+						FScheduledOp(Args.Base, Item, 0)); 
+				break;
+			}
+			case 1: 
+			{
+				FExtendedImageDesc& Result = ImageDescResults.FindOrAdd(Item.At);
+				Result = ImageDescResults.FindChecked(Args.Base);
+
+				StoreValidDesc(Item); 
 				break;
 			}
-			case 1: StoreValidDesc(Item); break;
 			default: check(false);
 			}
 			break;
@@ -6641,20 +6646,25 @@ namespace mu
 
 		case EOpType::IM_LUMINANCE:
 		{
+			MUTABLE_CPUPROFILER_SCOPE(RunCodeImageDesc_IM_LUMINANCE);
 			OP::ImageLuminanceArgs Args = Program.GetOpArgs<OP::ImageLuminanceArgs>(Item.At);
 			switch (Item.Stage)
 			{
 			case 0:
 			{
-				AddOp(FScheduledOp(Item.At, Item, 1, Item.CustomState), 
-						FScheduledOp(Args.Base, Item, 0, Item.CustomState));
+				AddOp(FScheduledOp(Item.At, Item, 1), 
+						FScheduledOp(Args.Base, Item, 0));
 				break;
 			}
 			case 1:
-				HeapImageDesc[Item.CustomState].m_format = EImageFormat::L_UByte;
+			{
+				FExtendedImageDesc& Result = ImageDescResults.FindOrAdd(Item.At);
+				Result = ImageDescResults.FindChecked(Args.Base);
+				Result.m_format = EImageFormat::L_UByte;
+				
 				StoreValidDesc(Item);
 				break;
-
+			}
 			default:
 				check(false);
 			}
@@ -6664,58 +6674,47 @@ namespace mu
 
 		case EOpType::IM_SWIZZLE:
 		{
+			MUTABLE_CPUPROFILER_SCOPE(RunCodeImageDesc_IM_SWIZZLE);
 			OP::ImageSwizzleArgs Args = Program.GetOpArgs<OP::ImageSwizzleArgs>(Item.At);
+			
+			TArray<OP::ADDRESS, TFixedAllocator<4>> ValidArgs;
+			for (int32 SourceIndex = 0; SourceIndex < 4; ++SourceIndex)
+			{
+				if (Args.sources[SourceIndex])
+				{
+					ValidArgs.AddUnique(Args.sources[SourceIndex]);
+				}
+			}
+
 			switch (Item.Stage)
 			{
 			case 0:
 			{
-				FScheduledOpData ResultArrayData;
-				ResultArrayData.ReturnArray.ReturnAddress = Item.CustomState;
-				ResultArrayData.ReturnArray.Begin = HeapImageDesc.Num();
-				ResultArrayData.ReturnArray.Num = 4;
-
-				HeapImageDesc.AddDefaulted(ResultArrayData.ReturnArray.Num);
+				TArray<FScheduledOp, TFixedAllocator<4>> Deps;
+				for (int32 ArgIndex = 0; ArgIndex < ValidArgs.Num(); ++ArgIndex)
+				{
+					Deps.Add(FScheduledOp(ValidArgs[ArgIndex], Item, 0));
+				}
 
-				int32 ResultArrayAddress = HeapData.Add(ResultArrayData);
+				AddOp(FScheduledOp(Item.At, Item, 1), Deps);
 
-				AddOp(FScheduledOp(Item.At, Item, 1, ResultArrayAddress), 
-						FScheduledOp(Args.sources[0], Item, 0, ResultArrayData.ReturnArray.Begin + 0),
-						FScheduledOp(Args.sources[1], Item, 0, ResultArrayData.ReturnArray.Begin + 1),
-						FScheduledOp(Args.sources[2], Item, 0, ResultArrayData.ReturnArray.Begin + 2),
-						FScheduledOp(Args.sources[3], Item, 0, ResultArrayData.ReturnArray.Begin + 3));
 				break;
 			}
 			case 1:
-			{
-				FScheduledOpData ResultArrayData = HeapData[Item.CustomState];
-				TArrayView<FExtendedImageDesc> DepResults(
-						HeapImageDesc.GetData() + ResultArrayData.ReturnArray.Begin, 
-						ResultArrayData.ReturnArray.Num);
-				
-				check(DepResults.Num() == 4);
-	
-				FExtendedImageDesc& Result = HeapImageDesc[ResultArrayData.ReturnArray.ReturnAddress];
+			{	
+				FExtendedImageDesc& Result = ImageDescResults.FindOrAdd(Item.At);
 
-				int32 FirstValidSourceIndex = 0;
-				for (; FirstValidSourceIndex < 4; ++FirstValidSourceIndex)
-				{
-					if (Args.sources[FirstValidSourceIndex])
-					{
-						Result = DepResults[FirstValidSourceIndex];
-						break;
-					}
-				}
+				check(ValidArgs.Num() > 0);
+
+				Result = ImageDescResults.FindChecked(ValidArgs[0]);
 				Result.m_format = Args.format;
 				
-				for (int32 Index = FirstValidSourceIndex + 1; Index < 4; ++Index)
+				for (int32 ArgIndex = 1; ArgIndex < ValidArgs.Num(); ++ArgIndex)
 				{
-					if (Args.sources[Index])
-					{
-						Result.FirstLODAvailable = FMath::Max(Result.FirstLODAvailable, DepResults[Index].FirstLODAvailable);
-					}
+					const FExtendedImageDesc& SourceResult = ImageDescResults.FindOrAdd(ValidArgs[ArgIndex]); 
+					Result.FirstLODAvailable = FMath::Max(Result.FirstLODAvailable, SourceResult.FirstLODAvailable);
 				}
 
-				HeapImageDesc[Item.CustomState].m_format = Args.format;
 				StoreValidDesc(Item);
 				break;
 			}
@@ -6728,16 +6727,24 @@ namespace mu
 
 		case EOpType::IM_COLOURMAP:
 		{
+			MUTABLE_CPUPROFILER_SCOPE(RunCodeImageDesc_IM_COLOURMAP);
 			OP::ImageColourMapArgs Args = Program.GetOpArgs<OP::ImageColourMapArgs>(Item.At);
 			switch (Item.Stage)
 			{
 			case 0: 
 			{	
-				AddOp(FScheduledOp(Item.At, Item, 1, Item.CustomState), 
-						FScheduledOp(Args.Base, Item, 0, Item.CustomState));
+				AddOp(FScheduledOp(Item.At, Item, 1), 
+						FScheduledOp(Args.Base, Item, 0));
+				break;
+			}
+			case 1: 
+			{	
+				FExtendedImageDesc& Result = ImageDescResults.FindOrAdd(Item.At);
+				Result = ImageDescResults.FindChecked(Args.Base);
+
+				StoreValidDesc(Item); 
 				break;
 			}
-			case 1: StoreValidDesc(Item); break;
 			default: check(false);
 			}
 			break;
@@ -6745,18 +6752,21 @@ namespace mu
 
 		case EOpType::IM_BINARISE:
 		{
+			MUTABLE_CPUPROFILER_SCOPE(RunCodeImageDesc_IM_BINARIZE);
 			OP::ImageBinariseArgs Args = Program.GetOpArgs<OP::ImageBinariseArgs>(Item.At);
 			switch (Item.Stage)
 			{
 			case 0:
 			{
-				AddOp(FScheduledOp(Item.At, Item, 1, Item.CustomState), 
-						FScheduledOp(Args.Base, Item, 0, Item.CustomState));
+				AddOp(FScheduledOp(Item.At, Item, 1), 
+						FScheduledOp(Args.Base, Item, 0));
 				break;
 			}
 			case 1:
 			{
-				HeapImageDesc[Item.CustomState].m_format = EImageFormat::L_UByte;
+				FExtendedImageDesc& Result = ImageDescResults.FindOrAdd(Item.At);
+				Result.m_format = EImageFormat::L_UByte;
+
 				StoreValidDesc(Item);
 				break;
 			}
@@ -6768,16 +6778,24 @@ namespace mu
 
 		case EOpType::IM_INVERT:
 		{
+			MUTABLE_CPUPROFILER_SCOPE(RunCodeImageDesc_IM_INVERT);
 			OP::ImageInvertArgs Args = Program.GetOpArgs<OP::ImageInvertArgs>(Item.At);
 			switch (Item.Stage)
 			{
 			case 0: 
 			{
-				AddOp(FScheduledOp(Item.At, Item, 1, Item.CustomState), 
-						FScheduledOp(Args.Base, Item, 0, Item.CustomState)); 
+				AddOp(FScheduledOp(Item.At, Item, 1), 
+						FScheduledOp(Args.Base, Item, 0)); 
+				break;
+			}
+			case 1: 
+			{
+				FExtendedImageDesc& Result = ImageDescResults.FindOrAdd(Item.At);
+				Result = ImageDescResults.FindChecked(Args.Base);
+
+				StoreValidDesc(Item);
 				break;
 			}
-			case 1: StoreValidDesc(Item); break;
 			default: check(false);
 			}
 			break;
@@ -6785,31 +6803,40 @@ namespace mu
 
 		case EOpType::IM_PLAINCOLOUR:
 		{
+			MUTABLE_CPUPROFILER_SCOPE(RunCodeImageDesc_IM_PLAINCOLOUR);
 			OP::ImagePlainColorArgs Args = Program.GetOpArgs<OP::ImagePlainColorArgs>(Item.At);
-			HeapImageDesc[Item.CustomState].m_size[0] = Args.Size[0];
-			HeapImageDesc[Item.CustomState].m_size[1] = Args.Size[1];
-			HeapImageDesc[Item.CustomState].m_lods = Args.LODs;
-			HeapImageDesc[Item.CustomState].m_format = Args.Format;
+			FExtendedImageDesc& Result = ImageDescResults.FindOrAdd(Item.At);
+			
+			Result.m_size[0] = Args.Size[0];
+			Result.m_size[1] = Args.Size[1];
+			Result.m_lods = Args.LODs;
+			Result.m_format = Args.Format;
+
 			StoreValidDesc(Item);
 			break;
 		}
 
 		case EOpType::IM_CROP:
 		{
+			MUTABLE_CPUPROFILER_SCOPE(RunCodeImageDesc_IM_CROP);
 			OP::ImageCropArgs Args = Program.GetOpArgs<OP::ImageCropArgs>(Item.At);
 			switch (Item.Stage)
 			{
 			case 0:
 			{
-				AddOp(FScheduledOp(Item.At, Item, 1, Item.CustomState), 
-						FScheduledOp(Args.source, Item, 0, Item.CustomState));
+				AddOp(FScheduledOp(Item.At, Item, 1), 
+						FScheduledOp(Args.source, Item, 0));
 				break;
 			}
 			case 1:
 			{
-				HeapImageDesc[Item.CustomState].m_size[0] = Args.sizeX;
-				HeapImageDesc[Item.CustomState].m_size[1] = Args.sizeY;
-				HeapImageDesc[Item.CustomState].m_lods = 1;
+				FExtendedImageDesc& Result = ImageDescResults.FindOrAdd(Item.At);
+				Result = ImageDescResults.FindChecked(Args.source);
+				
+				Result.m_size[0] = Args.sizeX;
+				Result.m_size[1] = Args.sizeY;
+				Result.m_lods = 1;
+
 				StoreValidDesc(Item);
 				break;
 			}
@@ -6821,42 +6848,30 @@ namespace mu
 
 		case EOpType::IM_PATCH:
 		{
+			MUTABLE_CPUPROFILER_SCOPE(RunCodeImageDesc_IM_PATCH);
 			OP::ImagePatchArgs Args = Program.GetOpArgs<OP::ImagePatchArgs>(Item.At);
 			switch (Item.Stage)
 			{
 			case 0: 
 			{
-				FScheduledOpData ResultArrayData;
-				ResultArrayData.ReturnArray.ReturnAddress = Item.CustomState;
-				ResultArrayData.ReturnArray.Begin = HeapImageDesc.Num();
-				ResultArrayData.ReturnArray.Num = 2;
-
-				HeapImageDesc.AddDefaulted(ResultArrayData.ReturnArray.Num);
-
-				int32 ResultArrayAddress = HeapData.Add(ResultArrayData);
-
-				AddOp(FScheduledOp(Item.At, Item, 1, ResultArrayAddress), 
-						FScheduledOp(Args.base, Item, 0, ResultArrayData.ReturnArray.Begin),
-						FScheduledOp(Args.patch, Item, 0, ResultArrayData.ReturnArray.Begin + 1));
+				AddOp(FScheduledOp(Item.At, Item, 1),
+						FScheduledOp(Args.base, Item, 0),
+						FScheduledOp(Args.patch, Item, 0)); 
 				break;
 			}
 			case 1: 
 			{
-				FScheduledOpData ResultArrayData = HeapData[Item.CustomState];
-				TArrayView<FExtendedImageDesc> DepResults(
-						HeapImageDesc.GetData() + ResultArrayData.ReturnArray.Begin, 
-						ResultArrayData.ReturnArray.Num);
-				
-				check(DepResults.Num() == 2);
-	
-				FExtendedImageDesc& BaseImageDesc = DepResults[0];
-				FExtendedImageDesc& PatchImageDesc = DepResults[1];
+				FExtendedImageDesc& Result = ImageDescResults.FindOrAdd(Item.At);
+				Result = ImageDescResults.FindChecked(Args.base);
 
-				FExtendedImageDesc& Result = HeapImageDesc[ResultArrayData.ReturnArray.ReturnAddress]; 
-				Result = BaseImageDesc;
-				Result.FirstLODAvailable = FMath::Max(Result.FirstLODAvailable, PatchImageDesc.FirstLODAvailable);
+				if (Args.patch)
+				{
+					FExtendedImageDesc& PatchImageDesc = ImageDescResults.FindChecked(Args.patch);
+					Result.FirstLODAvailable = FMath::Max(Result.FirstLODAvailable, PatchImageDesc.FirstLODAvailable);
+				}
 
-				StoreValidDesc(Item); break;
+				StoreValidDesc(Item); 
+				break;
 			}
 			default: check(false);
 			}
@@ -6865,30 +6880,39 @@ namespace mu
 
 		case EOpType::IM_RASTERMESH:
 		{
+			MUTABLE_CPUPROFILER_SCOPE(RunCodeImageDesc_IM_RASTERMESH);
 			OP::ImageRasterMeshArgs Args = Program.GetOpArgs<OP::ImageRasterMeshArgs>(Item.At);
-			HeapImageDesc[Item.CustomState].m_size[0] = Args.sizeX;
-			HeapImageDesc[Item.CustomState].m_size[1] = Args.sizeY;
-			HeapImageDesc[Item.CustomState].m_lods = 1;
-			HeapImageDesc[Item.CustomState].m_format = EImageFormat::L_UByte;
+			FExtendedImageDesc& Result = ImageDescResults.FindOrAdd(Item.At);
+			
+			Result.m_size[0] = Args.sizeX;
+			Result.m_size[1] = Args.sizeY;
+			Result.m_lods = 1;
+			Result.m_format = EImageFormat::L_UByte;
+
 			StoreValidDesc(Item);
 			break;
 		}
 
 		case EOpType::IM_MAKEGROWMAP:
 		{
+			MUTABLE_CPUPROFILER_SCOPE(RunCodeImageDesc_IM_MAKEGROWMAP);
 			OP::ImageMakeGrowMapArgs Args = Program.GetOpArgs<OP::ImageMakeGrowMapArgs>(Item.At);
 			switch (Item.Stage)
 			{
 			case 0:
 			{
-				AddOp(FScheduledOp(Item.At, Item, 1, Item.CustomState), 
-						FScheduledOp(Args.mask, Item, 0, Item.CustomState));
+				AddOp(FScheduledOp(Item.At, Item, 1), 
+						FScheduledOp(Args.mask, Item, 0));
 				break;
 			}
 			case 1:
 			{
-				HeapImageDesc[Item.CustomState].m_format = EImageFormat::L_UByte;
-				HeapImageDesc[Item.CustomState].m_lods = 1;
+				FExtendedImageDesc& Result = ImageDescResults.FindOrAdd(Item.At);
+				Result = ImageDescResults.FindChecked(Args.mask);
+
+				Result.m_format = EImageFormat::L_UByte;
+				Result.m_lods = 1;
+
 				StoreValidDesc(Item);
 				break;
 			}
@@ -6901,16 +6925,24 @@ namespace mu
 
 		case EOpType::IM_DISPLACE:
 		{
+			MUTABLE_CPUPROFILER_SCOPE(RunCodeImageDesc_IM_DISPLACE);
 			OP::ImageDisplaceArgs Args = Program.GetOpArgs<OP::ImageDisplaceArgs>(Item.At);
 			switch (Item.Stage)
 			{
 			case 0: 
 			{
-				AddOp(FScheduledOp(Item.At, Item, 1, Item.CustomState),
-						FScheduledOp(Args.Source, Item, 0, Item.CustomState));
+				AddOp(FScheduledOp(Item.At, Item, 1),
+						FScheduledOp(Args.Source, Item, 0));
+				break;
+			}
+			case 1: 
+			{
+				FExtendedImageDesc& Result = ImageDescResults.FindOrAdd(Item.At);
+				Result = ImageDescResults.FindChecked(Args.Source);
+
+				StoreValidDesc(Item); 
 				break;
 			}
-			case 1: StoreValidDesc(Item); break;
 			default: check(false);
 			}
 			break;
@@ -6918,26 +6950,29 @@ namespace mu
 
         case EOpType::IM_TRANSFORM:
         {
-
+			MUTABLE_CPUPROFILER_SCOPE(RunCodeImageDesc_IM_TRANSFORM);
 			OP::ImageTransformArgs Args = Program.GetOpArgs<OP::ImageTransformArgs>(Item.At);
 
             switch (Item.Stage)
             {
             case 0:
 			{
-				AddOp(FScheduledOp(Item.At, Item, 1, Item.CustomState), 
-						FScheduledOp(Args.Base, Item, 0, Item.CustomState));	
+				AddOp(FScheduledOp(Item.At, Item, 1), 
+						FScheduledOp(Args.Base, Item, 0));	
                 break;
 			}
             case 1:
             {
-				HeapImageDesc[Item.CustomState].m_lods = 1;
-				HeapImageDesc[Item.CustomState].m_format = GetUncompressedFormat(HeapImageDesc[Item.CustomState].m_format);
+				FExtendedImageDesc& Result = ImageDescResults.FindOrAdd(Item.At);
+				Result = ImageDescResults.FindChecked(Args.Base);
+			
+				Result.m_lods = 1;
+				Result.m_format = GetUncompressedFormat(Result.m_format);
 				
 				if (!(Args.SizeX == 0 && Args.SizeY == 0))
 				{
-					HeapImageDesc[Item.CustomState].m_size[0] = Args.SizeX;
-					HeapImageDesc[Item.CustomState].m_size[1] = Args.SizeY;
+					Result.m_size[0] = Args.SizeX;
+					Result.m_size[1] = Args.SizeY;
 				}
 
 				StoreValidDesc(Item);
@@ -6956,13 +6991,13 @@ namespace mu
 			{
 				// Operation not implemented
 				check(false);
-				HeapImageDesc[Item.CustomState] = FExtendedImageDesc();
+
+				ImageDescResults.FindOrAdd(Item.At);
 			}
 			break;
 		}
 	}
 
-
 	FImageOperator MakeImageOperator(CodeRunner* Runner)
 	{
 		return FImageOperator(
diff --git a/Engine/Plugins/Mutable/Source/MutableRuntime/Private/MuR/CodeRunner.h b/Engine/Plugins/Mutable/Source/MutableRuntime/Private/MuR/CodeRunner.h
--- a/Engine/Plugins/Mutable/Source/MutableRuntime/Private/MuR/CodeRunner.h
+++ b/Engine/Plugins/Mutable/Source/MutableRuntime/Private/MuR/CodeRunner.h
@@ -71,14 +71,7 @@ namespace  mu
         struct FScheduledOpData
         {
 			union
-			{
-				struct
-				{
-					int32 Begin;
-					int32 Num;
-					int32 ReturnAddress;
-				} ReturnArray;
-			
+			{	
 				struct
 				{
 					float Bifactor;
@@ -167,7 +160,9 @@ namespace  mu
          * It is not released until no operations are pending.
 		 */
 		TArray<FScheduledOpData> HeapData;
-		TArray<FExtendedImageDesc> HeapImageDesc;
+		
+		/** Image descriptor intermediate results. */
+		TMap<OP::ADDRESS, FExtendedImageDesc> ImageDescResults;
 
 		/** Only used for correct mip skipping with external images. It is the LOD for which the image is build. */
 		int32 ImageLOD;
@@ -180,8 +175,8 @@ namespace  mu
 
 		UE::Tasks::FTask StartRun(bool bForceInlineExecution);
 
-		//!
-		void GetImageDescResult(FExtendedImageDesc& OutDesc);
+		/** */
+		const FExtendedImageDesc& GetImageDescResult(OP::ADDRESS ResultAddres);
 
 		//!
 		FProgramCache& GetMemory();
@@ -395,7 +390,7 @@ namespace  mu
 		int32 GetInternalMemoryBytes() const
 		{
 			return sizeof(CodeRunner) 
-				+ HeapData.GetAllocatedSize() + HeapImageDesc.GetAllocatedSize()
+				+ HeapData.GetAllocatedSize() + ImageDescResults.GetAllocatedSize()
 				+ ClosedTasks.GetAllocatedSize() + OpenTasks.GetAllocatedSize() + ScheduledStagePerOp.GetAllocatedSize()
 				// this contains smart pointers, approximate size like this:
 				+ IssuedTasks.Max() * ( sizeof(FIssuedTask) + 16);
diff --git a/Engine/Plugins/Mutable/Source/MutableRuntime/Private/MuR/CodeRunnerTasked.cpp b/Engine/Plugins/Mutable/Source/MutableRuntime/Private/MuR/CodeRunnerTasked.cpp
--- a/Engine/Plugins/Mutable/Source/MutableRuntime/Private/MuR/CodeRunnerTasked.cpp
+++ b/Engine/Plugins/Mutable/Source/MutableRuntime/Private/MuR/CodeRunnerTasked.cpp
@@ -153,7 +153,7 @@ namespace mu
 		bUnrecoverableError = false;
 
 		HeapData.SetNum(0, EAllowShrinking::No);
-		HeapImageDesc.SetNum(0, EAllowShrinking::No);
+		ImageDescResults.Reset();
 		
 		RunnerCompletionEvent = UE::Tasks::FTaskEvent(TEXT("CodeRunnerCompletionEvent"));
 		
@@ -486,11 +486,10 @@ namespace mu
     }
 
 	
-	//---------------------------------------------------------------------------------------------
-	void CodeRunner::GetImageDescResult(FExtendedImageDesc& OutDesc)
+	/** */
+	const FExtendedImageDesc& CodeRunner::GetImageDescResult(OP::ADDRESS ResultAddress)
 	{
-		check(HeapImageDesc.Num() > 0);
-		OutDesc = HeapImageDesc[0];
+		return ImageDescResults.FindChecked(ResultAddress);
 	}
 
 	//---------------------------------------------------------------------------------------------
diff --git a/Engine/Plugins/Mutable/Source/MutableRuntime/Private/MuR/System.cpp b/Engine/Plugins/Mutable/Source/MutableRuntime/Private/MuR/System.cpp
--- a/Engine/Plugins/Mutable/Source/MutableRuntime/Private/MuR/System.cpp
+++ b/Engine/Plugins/Mutable/Source/MutableRuntime/Private/MuR/System.cpp
@@ -454,7 +454,7 @@ namespace mu
 			UE::Tasks::FTask CompletionEvent = Runner->StartRun(bForceInlineExecution);
 			check(CompletionEvent.IsCompleted());
 
-			Runner->GetImageDescResult(OutDesc);
+			OutDesc = Runner->GetImageDescResult(RootAddress);
 
 			m_pD->WorkingMemoryManager.EndRunnerThread();
 		}
@@ -504,10 +504,9 @@ namespace mu
 		UE::Tasks::FTask RunnerCompletionEvent = Runner->StartRun(bForceInlineExecution);
 		
 		return UE::Tasks::Launch(TEXT("FSystem::GetImageDescResultTask"),
-				[SystemPrivate = m_pD, Runner]() -> FExtendedImageDesc
+				[SystemPrivate = m_pD, Runner, RootAddress]() -> FExtendedImageDesc
 				{
-					FExtendedImageDesc Result;
-					Runner->GetImageDescResult(Result);
+					FExtendedImageDesc Result = Runner->GetImageDescResult(RootAddress);
 
 					SystemPrivate->WorkingMemoryManager.EndRunnerThread();
 					SystemPrivate->WorkingMemoryManager.CurrentInstanceCache = nullptr;
