[Mutable] Fix crash when computing image descriptor for IM_SWITCH ops with default null value.

#rnx

diff --git a/Engine/Plugins/Mutable/Source/MutableRuntime/Private/MuR/CodeRunner.cpp b/Engine/Plugins/Mutable/Source/MutableRuntime/Private/MuR/CodeRunner.cpp
--- a/Engine/Plugins/Mutable/Source/MutableRuntime/Private/MuR/CodeRunner.cpp
+++ b/Engine/Plugins/Mutable/Source/MutableRuntime/Private/MuR/CodeRunner.cpp
@@ -6156,14 +6156,23 @@ namespace mu
 					}
 				}
 
-				AddOp(FScheduledOp(Item.At, Item, 2, ValueAt),
-						FScheduledOp(ValueAt, Item, 0));
-
+				if (ValueAt)
+				{
+					AddOp(FScheduledOp(Item.At, Item, 2, ValueAt),
+							FScheduledOp(ValueAt, Item, 0));
+				}
+				else
+				{
+					ImageDescResults.FindOrAdd(Item.At);
+					StoreValidDesc(Item); 
+				}
 				break;
 			}
 
 			case 2: 
 			{
+				check(Item.CustomState);
+
 				FExtendedImageDesc& Result = ImageDescResults.FindOrAdd(Item.At);
 				Result = ImageDescResults.FindChecked(Item.CustomState);
 
