[mutable] Fix error with image descriptor generation at compile time. It was introduce in the recent ASTOp code refactor.
#rnx

diff --git a/Engine/Plugins/Mutable/Source/MutableTools/Private/MuT/ASTOpImageBinarize.cpp b/Engine/Plugins/Mutable/Source/MutableTools/Private/MuT/ASTOpImageBinarize.cpp
--- a/Engine/Plugins/Mutable/Source/MutableTools/Private/MuT/ASTOpImageBinarize.cpp
+++ b/Engine/Plugins/Mutable/Source/MutableTools/Private/MuT/ASTOpImageBinarize.cpp
@@ -125,12 +125,6 @@ namespace mu
 			Result = Base->GetImageDesc(bReturnBestOption, Context);
 		}
 
-		if (Threshold)
-		{
-			FImageDesc SizeResult = Threshold->GetImageDesc(bReturnBestOption, Context);
-			Result.m_size = SizeResult.m_size;
-		}
-
 		// Cache the result
 		if (Context)
 		{
diff --git a/Engine/Plugins/Mutable/Source/MutableTools/Private/MuT/ASTOpImageDisplace.cpp b/Engine/Plugins/Mutable/Source/MutableTools/Private/MuT/ASTOpImageDisplace.cpp
--- a/Engine/Plugins/Mutable/Source/MutableTools/Private/MuT/ASTOpImageDisplace.cpp
+++ b/Engine/Plugins/Mutable/Source/MutableTools/Private/MuT/ASTOpImageDisplace.cpp
@@ -125,12 +125,6 @@ namespace mu
 			Result = Source->GetImageDesc(bReturnBestOption, Context);
 		}
 
-		if (DisplacementMap)
-		{
-			FImageDesc SizeResult = DisplacementMap->GetImageDesc(bReturnBestOption, Context);
-			Result.m_size = SizeResult.m_size;
-		}
-
 		// Cache the result
 		if (Context)
 		{
