[Mutable] Fixed groom components not being destroyed when discarding instances.

Code was removed on a previous submit.

diff --git a/Engine/Plugins/Mutable/Source/CustomizableObject/Private/MuCO/CustomizableObjectSystem.cpp b/Engine/Plugins/Mutable/Source/CustomizableObject/Private/MuCO/CustomizableObjectSystem.cpp
--- a/Engine/Plugins/Mutable/Source/CustomizableObject/Private/MuCO/CustomizableObjectSystem.cpp
+++ b/Engine/Plugins/Mutable/Source/CustomizableObject/Private/MuCO/CustomizableObjectSystem.cpp
@@ -4104,7 +4104,6 @@ void UCustomizableObjectSystem::DiscardInstances()
 		Iterator && NumInstancesDiscarded < DiscardLimitPerTick;
 		++Iterator)
 	{
-
 		UCustomizableObjectInstance* COI = Iterator->CustomizableObjectInstance.Get();
 
 		const bool bUpdating = Private->CurrentMutableOperation && Private->CurrentMutableOperation->Instance == Iterator->CustomizableObjectInstance;
@@ -4122,7 +4121,6 @@ void UCustomizableObjectSystem::DiscardInstances()
 				{
 					if (const UModelResources* ModelResources = CustomizableObject->GetPrivate()->GetModelResources())
 					{
-
 						for (TObjectIterator<UCustomizableObjectInstanceUsage> It; It; ++It)
 						{
 							UCustomizableObjectInstanceUsage* InstanceUsage = *It;
@@ -4157,6 +4155,10 @@ void UCustomizableObjectSystem::DiscardInstances()
 								}
 							}
 
+							for (const UCustomizableObjectExtension* Extension : ICustomizableObjectModule::Get().GetRegisteredExtensions())
+							{
+								Extension->OnCustomizableObjectInstanceUsageDiscarded(*InstanceUsage);
+							}
 						}
 					}
 				}
