[Mutable] Change the default name of Component nodes.

diff --git a/Engine/Plugins/Mutable/Source/CustomizableObjectEditor/Private/MuCOE/Nodes/CustomizableObjectNodeComponentMesh.h b/Engine/Plugins/Mutable/Source/CustomizableObjectEditor/Private/MuCOE/Nodes/CustomizableObjectNodeComponentMesh.h
--- a/Engine/Plugins/Mutable/Source/CustomizableObjectEditor/Private/MuCOE/Nodes/CustomizableObjectNodeComponentMesh.h
+++ b/Engine/Plugins/Mutable/Source/CustomizableObjectEditor/Private/MuCOE/Nodes/CustomizableObjectNodeComponentMesh.h
@@ -59,7 +59,7 @@ public:
 	virtual void OnRenameNode(const FString& NewName) override;
 
 	UPROPERTY(EditAnywhere, Category = ComponentMesh)
-	FName ComponentName;
+	FName ComponentName = "Default name";
 	
 	/** All the Skeletal Meshes generated for this component will use the Reference Skeletal Mesh properties
 	 * for everything Mutable doesn't create or modify. This includes data like LOD distances, Physics
