[Mutable] Fix unwanted assets being added to package due to missing cook scope.

Wrapped all loading functions to ease maintenance.

diff --git a/Engine/Plugins/Mutable/Source/CustomizableObject/Internal/MuCO/LoadUtils.h b/Engine/Plugins/Mutable/Source/CustomizableObject/Internal/MuCO/LoadUtils.h
new file mode 100644
--- /dev/null
+++ b/Engine/Plugins/Mutable/Source/CustomizableObject/Internal/MuCO/LoadUtils.h
@@ -0,0 +1,67 @@
+﻿// Copyright Epic Games, Inc. All Rights Reserved.
+
+#pragma once
+
+#include "UObject/ICookInfo.h"
+#include "UObject/SoftObjectPtr.h"
+
+class UClass;
+class UObject;
+struct FAssetData;
+struct FSoftObjectPtr;
+
+
+namespace MutablePrivate
+{
+	CUSTOMIZABLEOBJECT_API UObject* LoadObject(const FAssetData& DataAsset);
+
+
+	CUSTOMIZABLEOBJECT_API UObject* LoadObject(const FSoftObjectPath& Path);
+	
+	
+	CUSTOMIZABLEOBJECT_API UObject* LoadObject(const FSoftObjectPtr& SoftObject);
+
+	
+	template <typename T>
+	T* LoadObject(UObject* Outer, const TCHAR* Name, const TCHAR* Filename = nullptr, uint32 LoadFlags = LOAD_None, UPackageMap* Sandbox = nullptr, const FLinkerInstancingContext* InstancingContext = nullptr)
+	{
+#if WITH_EDITOR
+		FCookLoadScope CookLoadScope(ECookLoadType::EditorOnly);
+#endif
+		return ::LoadObject<T>(Outer, Name, Filename, LoadFlags, Sandbox, InstancingContext);
+	}
+	
+
+	template<typename T>
+	T* LoadObject(const TSoftObjectPtr<T>& SoftObject)
+	{
+#if WITH_EDITOR
+		FCookLoadScope CookLoadScope(ECookLoadType::EditorOnly);
+#endif
+	
+		return SoftObject.LoadSynchronous();
+	}
+
+
+	template<typename T>
+	UClass* LoadClass(const FSoftClassPath& Path)
+	{
+#if WITH_EDITOR
+		FCookLoadScope CookLoadScope(ECookLoadType::EditorOnly);
+#endif
+		
+		return Path.TryLoadClass<T>();
+	}
+
+	
+	template<typename T>
+	UClass* LoadClass(const TSoftClassPtr<T>& SoftClass)
+	{
+#if WITH_EDITOR
+		FCookLoadScope CookLoadScope(ECookLoadType::EditorOnly);
+#endif
+
+		return SoftClass.LoadSynchronous();
+	}	
+}
+
diff --git a/Engine/Plugins/Mutable/Source/CustomizableObject/Internal/MuCO/MutableStreamableManager.h b/Engine/Plugins/Mutable/Source/CustomizableObject/Internal/MuCO/MutableStreamableManager.h
--- a/Engine/Plugins/Mutable/Source/CustomizableObject/Internal/MuCO/MutableStreamableManager.h
+++ b/Engine/Plugins/Mutable/Source/CustomizableObject/Internal/MuCO/MutableStreamableManager.h
@@ -5,6 +5,7 @@
 #include "Engine/StreamableManager.h"
 
 #include "Templates/SharedPointer.h"
+#include "UObject/ICookInfo.h"
 
 struct FStreamableHandle;
 
@@ -23,7 +24,7 @@ public:
 			DebugNameType&& DebugNameOrLocation = UE::FSourceLocation::Current())
 	{
 		check(IsInGameThread());
-
+		
 		TSharedRef<FMutableStreamableHandle> MutableStreamableHandle = MakeShared<FMutableStreamableHandle>();
 
 		if constexpr (std::is_constructible_v<FStreamableDelegate, FuncType>)
@@ -48,7 +49,7 @@ public:
 			}
 		});
 		
-		
+		FCookLoadScope CookLoadScope(ECookLoadType::EditorOnly);
 		TSharedPtr<FStreamableHandle> Handle = StreamableManager.RequestAsyncLoad(TargetsToStream, Delegate, Priority, bManageActiveHandle, bStartStalled, DebugNameOrLocation);
 		if (Handle)
 		{
@@ -65,6 +66,7 @@ public:
 		bool bManageActiveHandle = false,
 		DebugNameType&& DebugNameOrLocation = UE::FSourceLocation::Current())
 	{
+		FCookLoadScope CookLoadScope(ECookLoadType::EditorOnly);
 		return StreamableManager.RequestSyncLoad(TargetsToStream, bManageActiveHandle, DebugNameOrLocation);
 	}
 
diff --git a/Engine/Plugins/Mutable/Source/CustomizableObject/Private/MuCO/CustomizableObjectDGGUI.cpp b/Engine/Plugins/Mutable/Source/CustomizableObject/Private/MuCO/CustomizableObjectDGGUI.cpp
--- a/Engine/Plugins/Mutable/Source/CustomizableObject/Private/MuCO/CustomizableObjectDGGUI.cpp
+++ b/Engine/Plugins/Mutable/Source/CustomizableObject/Private/MuCO/CustomizableObjectDGGUI.cpp
@@ -4,6 +4,7 @@
 
 #include "Kismet/GameplayStatics.h"
 #include "Misc/ConfigCacheIni.h"
+#include "MuCO/LoadUtils.h"
 #include "UObject/UObjectIterator.h"
 
 #include UE_INLINE_GENERATED_CPP_BY_NAME(CustomizableObjectDGGUI)
@@ -25,7 +26,7 @@ void UDGGUI::OpenDGGUI(const int32 SlotID, UCustomizableObjectInstanceUsage* Sel
 		}
 
 		FSoftClassPath DGUIPath(DGGUIAssetPath);
-		if (UClass* DGUI = DGUIPath.TryLoadClass<UDGGUI>())
+		if (UClass* DGUI = MutablePrivate::LoadClass<UDGGUI>(DGUIPath))
 		{
 			UDGGUI* WDGUI = CreateWidget<UDGGUI>(Player, DGUI);
 			if (WDGUI)
diff --git a/Engine/Plugins/Mutable/Source/CustomizableObject/Private/MuCO/CustomizableObjectInstance.cpp b/Engine/Plugins/Mutable/Source/CustomizableObject/Private/MuCO/CustomizableObjectInstance.cpp
--- a/Engine/Plugins/Mutable/Source/CustomizableObject/Private/MuCO/CustomizableObjectInstance.cpp
+++ b/Engine/Plugins/Mutable/Source/CustomizableObject/Private/MuCO/CustomizableObjectInstance.cpp
@@ -56,6 +56,7 @@
 #include "MuCO/CustomizableObjectResourceData.h"
 #include "MuCO/CustomizableObjectStreamedResourceData.h"
 #include "MuCO/CustomizableObjectResourceDataTypes.h"
+#include "MuCO/LoadUtils.h"
 #include "Serialization/BulkData.h"
 
 #include UE_INLINE_GENERATED_CPP_BY_NAME(CustomizableObjectInstance)
@@ -2336,7 +2337,7 @@ bool UCustomizableInstancePrivate::UpdateSkeletalMesh_PostBeginUpdate0(UCustomiz
         		UE_LOG(LogMutable, Error, TEXT("Referenced mesh [%s] was not pre-loaded. It will be sync-loaded probably causing a hitch. CO [%s]"), *Ref.ToString(), *GetNameSafe(CustomizableObject));
         	}
 
-			UStreamableRenderAsset* Asset = Ref.LoadSynchronous();
+			UStreamableRenderAsset* Asset = MutablePrivate::LoadObject(Ref);
 			USkeletalMesh* SkeletalMesh = Cast<USkeletalMesh>(Asset);
 			if (SkeletalMesh)
 			{
@@ -2421,7 +2422,7 @@ bool UCustomizableInstancePrivate::UpdateSkeletalMesh_PostBeginUpdate0(UCustomiz
 			const FAnimBpOverridePhysicsAssetsInfo& Info = ModelResources.AnimBpOverridePhysiscAssetsInfo[PhysicsBodyExternalId];
 
 			// Make sure the AnimInstance class is loaded. It is expected to be already loaded at this point though. 
-			UClass* AnimInstanceClassLoaded = Info.AnimInstanceClass.LoadSynchronous();
+			UClass* AnimInstanceClassLoaded = MutablePrivate::LoadClass(Info.AnimInstanceClass);
 			TSubclassOf<UAnimInstance> AnimInstanceClass = TSubclassOf<UAnimInstance>(AnimInstanceClassLoaded);
 			if (!ensureAlways(AnimInstanceClass))
 			{
@@ -3131,8 +3132,9 @@ void UCustomizableInstancePrivate::SetReferenceSkeletalMesh() const
 		if (USkeletalMeshComponent* Parent = CustomizableObjectInstanceUsage->GetAttachParent())
 		{
 			Parent->EmptyOverrideMaterials();
-			
-			USkeletalMesh* SkeletalMesh = ModelResources->ReferenceSkeletalMeshesData[ObjectComponentIndex].SoftSkeletalMesh.LoadSynchronous();
+
+			TSoftObjectPtr<USkeletalMesh> SoftObjectPtr = ModelResources->ReferenceSkeletalMeshesData[ObjectComponentIndex].SoftSkeletalMesh;
+			USkeletalMesh* SkeletalMesh = MutablePrivate::LoadObject(SoftObjectPtr);
 			Parent->SetSkeletalMesh(SkeletalMesh);
 		}
 	}
@@ -6699,7 +6701,7 @@ void UCustomizableInstancePrivate::BuildMaterials(const TSharedRef<FUpdateContex
 										TEXT("Pass-through texture with name %s hasn't been loaded yet in BuildMaterials(). Forcing sync load."),
 										*Ref.ToSoftObjectPath().ToString());
 									ensure(false);
-									PassThroughTexture = Ref.LoadSynchronous();
+									PassThroughTexture = MutablePrivate::LoadObject(Ref);
 								}
 							}
 
@@ -7492,7 +7494,7 @@ namespace MutablePrivate
 			if (!AnimBP.IsNull())
 			{
 				// force a load right now - we don't know whether we would have loaded already - this could be called in editor
-				const TSubclassOf<UAnimInstance> LiveAnimBP = AnimBP.LoadSynchronous();
+				const TSubclassOf<UAnimInstance> LiveAnimBP = MutablePrivate::LoadClass(AnimBP);
 				if (LiveAnimBP)
 				{
 					Delegate.Execute(Index, LiveAnimBP);
diff --git a/Engine/Plugins/Mutable/Source/CustomizableObject/Private/MuCO/CustomizableObjectSystem.cpp b/Engine/Plugins/Mutable/Source/CustomizableObject/Private/MuCO/CustomizableObjectSystem.cpp
--- a/Engine/Plugins/Mutable/Source/CustomizableObject/Private/MuCO/CustomizableObjectSystem.cpp
+++ b/Engine/Plugins/Mutable/Source/CustomizableObject/Private/MuCO/CustomizableObjectSystem.cpp
@@ -42,6 +42,8 @@
 #include "Engine/Engine.h"
 #endif
 
+#include "MuCO/LoadUtils.h"
+
 #include UE_INLINE_GENERATED_CPP_BY_NAME(CustomizableObjectSystem)
 
 class AActor;
@@ -4147,7 +4149,8 @@ void UCustomizableObjectSystem::DiscardInstances()
 									if (IsReplaceDiscardedWithReferenceMeshEnabled())
 									{
 										// Force load the reference mesh if necessary. 
-										SkeletalMesh = ModelResources->ReferenceSkeletalMeshesData[ObjectComponentIndex].SoftSkeletalMesh.LoadSynchronous();
+										TSoftObjectPtr<USkeletalMesh> SoftObjectPtr = ModelResources->ReferenceSkeletalMeshesData[ObjectComponentIndex].SoftSkeletalMesh;
+										SkeletalMesh = MutablePrivate::LoadObject(SoftObjectPtr);
 									}
 
 									Parent->EmptyOverrideMaterials();
diff --git a/Engine/Plugins/Mutable/Source/CustomizableObject/Private/MuCO/EditorImageProvider.cpp b/Engine/Plugins/Mutable/Source/CustomizableObject/Private/MuCO/EditorImageProvider.cpp
--- a/Engine/Plugins/Mutable/Source/CustomizableObject/Private/MuCO/EditorImageProvider.cpp
+++ b/Engine/Plugins/Mutable/Source/CustomizableObject/Private/MuCO/EditorImageProvider.cpp
@@ -5,11 +5,12 @@
 
 #include "MuCO/CustomizableObjectSystem.h"
 #include "Engine/Texture2D.h"
+#include "MuCO/LoadUtils.h"
 
 
 UCustomizableSystemImageProvider::ValueType UEditorImageProvider::HasTextureParameterValue(const FName& ID)
 {
-	const UTexture2D* Texture = Cast<UTexture2D>(FSoftObjectPath(ID.ToString()).TryLoad());
+	const UTexture2D* Texture = Cast<UTexture2D>(MutablePrivate::LoadObject(FSoftObjectPath(ID.ToString())));
 	
 	return Texture ? ValueType::Unreal_Deferred : ValueType::None;
 }
@@ -17,6 +18,6 @@ UCustomizableSystemImageProvider::ValueType UEditorImageProvider::HasTexturePara
 
 UTexture2D* UEditorImageProvider::GetTextureParameterValue(const FName& ID)
 {
-	return Cast<UTexture2D>(FSoftObjectPath(ID.ToString()).TryLoad());
+	return Cast<UTexture2D>(MutablePrivate::LoadObject(FSoftObjectPath(ID.ToString())));
 }
 
diff --git a/Engine/Plugins/Mutable/Source/CustomizableObject/Private/MuCO/LoadUtils.cpp b/Engine/Plugins/Mutable/Source/CustomizableObject/Private/MuCO/LoadUtils.cpp
new file mode 100644
--- /dev/null
+++ b/Engine/Plugins/Mutable/Source/CustomizableObject/Private/MuCO/LoadUtils.cpp
@@ -0,0 +1,38 @@
+﻿// Copyright Epic Games, Inc. All Rights Reserved.
+
+#include "MuCO/LoadUtils.h"
+
+#include "AssetRegistry/AssetData.h"
+
+
+namespace MutablePrivate
+{
+	UObject* LoadObject(const FAssetData& DataAsset)
+	{
+#if	WITH_EDITOR
+		FCookLoadScope CookLoadScope(ECookLoadType::EditorOnly);
+#endif
+	
+		return DataAsset.GetAsset();
+	}
+
+
+	UObject* LoadObject(const FSoftObjectPath& Path)
+	{
+#if	WITH_EDITOR
+		FCookLoadScope CookLoadScope(ECookLoadType::EditorOnly);
+#endif
+	
+		return Path.TryLoad();
+	}
+	
+
+	UObject* LoadObject(const FSoftObjectPtr& SoftObject)
+	{
+#if	WITH_EDITOR
+		FCookLoadScope CookLoadScope(ECookLoadType::EditorOnly);
+#endif
+	
+		return SoftObject.LoadSynchronous();
+	}
+}
diff --git a/Engine/Plugins/Mutable/Source/CustomizableObject/Private/MuCO/MutableStreamableManager.cpp b/Engine/Plugins/Mutable/Source/CustomizableObject/Private/MuCO/MutableStreamableManager.cpp
--- a/Engine/Plugins/Mutable/Source/CustomizableObject/Private/MuCO/MutableStreamableManager.cpp
+++ b/Engine/Plugins/Mutable/Source/CustomizableObject/Private/MuCO/MutableStreamableManager.cpp
@@ -14,9 +14,10 @@ int32 FMutableStreamableManager::Tick(bool bBlocking)
 				TArray<FSoftObjectPath> RequestedAssets;
 				MutableStreamableHandle->Handle->GetRequestedAssets(RequestedAssets);
 
+				FCookLoadScope CookLoadScope(ECookLoadType::EditorOnly);
 				StreamableManager.RequestSyncLoad(RequestedAssets);
+
 				MutableStreamableHandle->Delegate.ExecuteIfBound();
-			
 				MutableStreamableHandle->Handle->CancelHandle();
 			}			
 		}
diff --git a/Engine/Plugins/Mutable/Source/CustomizableObject/Private/MuCO/UnrealMutableImageProvider.cpp b/Engine/Plugins/Mutable/Source/CustomizableObject/Private/MuCO/UnrealMutableImageProvider.cpp
--- a/Engine/Plugins/Mutable/Source/CustomizableObject/Private/MuCO/UnrealMutableImageProvider.cpp
+++ b/Engine/Plugins/Mutable/Source/CustomizableObject/Private/MuCO/UnrealMutableImageProvider.cpp
@@ -8,6 +8,7 @@
 #include "MuCO/CustomizableObjectPrivate.h"
 #include "MuCO/UnrealToMutableTextureConversionUtils.h"
 #include "TextureResource.h"
+#include "MuCO/LoadUtils.h"
 #include "MuR/Parameters.h"
 #include "MuR/ImageTypes.h"
 #include "MuR/Model.h"
@@ -749,7 +750,7 @@ void FUnrealMutableResourceProvider::CacheRuntimeReferencedImages(const TSharedR
 		{
 			UE_LOG(LogMutable, Warning, TEXT("Runtime Referenced Texture [%s] was not async loaded. Forcing load sync."), *RuntimeReferencedTexture->GetPathName());
 			
-			Texture = RuntimeReferencedTexture.LoadSynchronous();
+			Texture = MutablePrivate::LoadObject(RuntimeReferencedTexture);
 			if (!Texture)
 			{
 				UE_LOG(LogMutable, Warning, TEXT("Failed to force load sync [%s]."), *RuntimeReferencedTexture->GetPathName());
diff --git a/Engine/Plugins/Mutable/Source/CustomizableObjectEditor/Private/MuCOE/AssetDefinition_CustomizableObject.cpp b/Engine/Plugins/Mutable/Source/CustomizableObjectEditor/Private/MuCOE/AssetDefinition_CustomizableObject.cpp
--- a/Engine/Plugins/Mutable/Source/CustomizableObjectEditor/Private/MuCOE/AssetDefinition_CustomizableObject.cpp
+++ b/Engine/Plugins/Mutable/Source/CustomizableObjectEditor/Private/MuCOE/AssetDefinition_CustomizableObject.cpp
@@ -17,6 +17,7 @@
 #include "MuCO/ICustomizableObjectEditorModule.h"
 #include "Subsystems/AssetEditorSubsystem.h"
 #include "ToolMenus.h"
+#include "MuCO/LoadUtils.h"
 
 
 #define LOCTEXT_NAMESPACE "CustomizableObjectEditor"
@@ -55,7 +56,7 @@ EAssetCommandResult UAssetDefinition_CustomizableObject::ActivateAssets(const FA
 		if ( ActivateArgs.Assets.Num() == 1 )
 		{
 			const FAssetData* FirstAsset = AssetsToActivate.GetData();
-			GEditor->GetEditorSubsystem<UAssetEditorSubsystem>()->OpenEditorForAsset(FirstAsset->GetAsset());
+			GEditor->GetEditorSubsystem<UAssetEditorSubsystem>()->OpenEditorForAsset(MutablePrivate::LoadObject(*FirstAsset));
 		}
 		else if ( ActivateArgs.Assets.Num() > 1 )
 		{
@@ -64,7 +65,7 @@ EAssetCommandResult UAssetDefinition_CustomizableObject::ActivateAssets(const FA
 			ObjectsToOpen.Reserve(AssetsToActivate.Num());
 			for (const FAssetData& ToOpen : AssetsToActivate)
 			{
-				ObjectsToOpen.Add(ToOpen.GetAsset());
+				ObjectsToOpen.Add(MutablePrivate::LoadObject(ToOpen));
 			}
 			ObjectsToOpen.Shrink();
 			
@@ -241,7 +242,7 @@ namespace MenuExtension_CustomizableObject
 						continue;
 					}
 
-					UCustomizableObject* Object = Cast<UCustomizableObject>(Asset.GetAsset());
+					UCustomizableObject* Object = Cast<UCustomizableObject>(MutablePrivate::LoadObject(Asset));
 					if (Object)
 					{
 						TSharedRef<FCompilationRequest> NewRequest = MakeShared<FCompilationRequest>(*Object, true);
diff --git a/Engine/Plugins/Mutable/Source/CustomizableObjectEditor/Private/MuCOE/AssetDefinition_CustomizableObjectInstance.cpp b/Engine/Plugins/Mutable/Source/CustomizableObjectEditor/Private/MuCOE/AssetDefinition_CustomizableObjectInstance.cpp
--- a/Engine/Plugins/Mutable/Source/CustomizableObjectEditor/Private/MuCOE/AssetDefinition_CustomizableObjectInstance.cpp
+++ b/Engine/Plugins/Mutable/Source/CustomizableObjectEditor/Private/MuCOE/AssetDefinition_CustomizableObjectInstance.cpp
@@ -8,6 +8,7 @@
 #include "Editor.h"
 #include "ToolMenus.h"
 #include "MuCO/ICustomizableObjectEditorModule.h"
+#include "MuCO/LoadUtils.h"
 #include "Subsystems/AssetEditorSubsystem.h"
 
 #define LOCTEXT_NAMESPACE "CustomizableObjectInstance"
@@ -44,7 +45,7 @@ EAssetCommandResult UAssetDefinition_CustomizableObjectInstance::ActivateAssets(
 		if (ActivateArgs.Assets.Num() == 1)
 		{
 			const FAssetData* FirstAsset = ActivateArgs.Assets.GetData();
-			GEditor->GetEditorSubsystem<UAssetEditorSubsystem>()->OpenEditorForAsset(FirstAsset->GetAsset());
+			GEditor->GetEditorSubsystem<UAssetEditorSubsystem>()->OpenEditorForAsset(MutablePrivate::LoadObject(*FirstAsset));
 		}
 		else if (ActivateArgs.Assets.Num() > 1)
 		{
@@ -52,7 +53,7 @@ EAssetCommandResult UAssetDefinition_CustomizableObjectInstance::ActivateAssets(
 			ObjectsToOpen.Reserve(ActivateArgs.Assets.Num());
 			for (const FAssetData& ToOpen : ActivateArgs.Assets)
 			{
-				ObjectsToOpen.Add(ToOpen.GetAsset());
+				ObjectsToOpen.Add(MutablePrivate::LoadObject(ToOpen));
 			}
 			ObjectsToOpen.Shrink();
 			
diff --git a/Engine/Plugins/Mutable/Source/CustomizableObjectEditor/Private/MuCOE/CustomizableInstanceDetails.cpp b/Engine/Plugins/Mutable/Source/CustomizableObjectEditor/Private/MuCOE/CustomizableInstanceDetails.cpp
--- a/Engine/Plugins/Mutable/Source/CustomizableObjectEditor/Private/MuCOE/CustomizableInstanceDetails.cpp
+++ b/Engine/Plugins/Mutable/Source/CustomizableObjectEditor/Private/MuCOE/CustomizableInstanceDetails.cpp
@@ -25,6 +25,7 @@
 #include "AssetThumbnail.h"
 #include "Slate/DeferredCleanupSlateBrush.h"
 #include "SSearchableComboBox.h"
+#include "MuCO/LoadUtils.h"
 #include "Types/SlateEnums.h"
 #include "Widgets/Colors/SColorBlock.h"
 #include "Widgets/Colors/SColorPicker.h"
@@ -1076,7 +1077,7 @@ TSharedRef<SWidget> FCustomizableInstanceDetails::OnGenerateWidgetIntParameter(T
 
 			if (AssetData.IsValid())
 			{
-				if (UTexture2D* UIThumbnail = Cast<UTexture2D>(AssetData.GetAsset()))
+				if (UTexture2D* UIThumbnail = Cast<UTexture2D>(MutablePrivate::LoadObject(AssetData)))
 				{
 					// We need to store the generated texture.
 					TSharedRef<FDeferredCleanupSlateBrush> Brush = FDeferredCleanupSlateBrush::CreateBrush(UIThumbnail, FVector2D(68.f, 68.0f), /* Texture size (64) + padding (4) */
@@ -1151,7 +1152,7 @@ TSharedRef<SWidget> FCustomizableInstanceDetails::GenerateFloatWidget(const UCus
 	if (const TSoftObjectPtr<UObject>* FloatDecoratorAsset = UIMetadata.ExtraAssets.Find(UIMetadataKeyWords::FloatDecoratorName))
 	{
 		// Checking if there is an image decorator for the float slider:
-		if (UTexture2D* DecoratorTexture = Cast<UTexture2D>(FloatDecoratorAsset->LoadSynchronous()))
+		if (UTexture2D* DecoratorTexture = Cast<UTexture2D>(MutablePrivate::LoadObject(*FloatDecoratorAsset)))
 		{
 			TSharedRef<FDeferredCleanupSlateBrush> Brush = FDeferredCleanupSlateBrush::CreateBrush(DecoratorTexture, FVector2D(DecoratorTexture->GetSizeX(), 1.0f), FLinearColor(1.0f, 1.0f, 1.0f, 1.0f), ESlateBrushTileType::NoTile, ESlateBrushImageType::Linear);
 			DynamicBrushes.Add(Brush);
diff --git a/Engine/Plugins/Mutable/Source/CustomizableObjectEditor/Private/MuCOE/CustomizableObjectCompiler.cpp b/Engine/Plugins/Mutable/Source/CustomizableObjectEditor/Private/MuCOE/CustomizableObjectCompiler.cpp
--- a/Engine/Plugins/Mutable/Source/CustomizableObjectEditor/Private/MuCOE/CustomizableObjectCompiler.cpp
+++ b/Engine/Plugins/Mutable/Source/CustomizableObjectEditor/Private/MuCOE/CustomizableObjectCompiler.cpp
@@ -117,14 +117,12 @@ void FCustomizableObjectCompiler::PreloadingReferencerAssetsCallback(bool bAsync
 	
 	if (AsynchronousStreamableHandlePtr)
 	{
-		FCookLoadScope CookLoadScope(ECookLoadType::EditorOnly);
-
 		TArray<FSoftObjectPath> AssetsToStream;
 		AsynchronousStreamableHandlePtr->GetRequestedAssets(AssetsToStream);
 
 		for (FSoftObjectPath AssetToStream : AssetsToStream)
 		{
-			ArrayGCProtect.Add(AssetToStream.TryLoad()); // Already loaded.
+			ArrayGCProtect.Add(MutablePrivate::LoadObject(AssetToStream)); // Already loaded.
 		}
 
 		AsynchronousStreamableHandlePtr = nullptr;
@@ -281,7 +279,6 @@ void FCustomizableObjectCompiler::Compile(const TSharedRef<FCompilationRequest>&
 	{
 		// Customizations are marked as editoronly on load and are not packaged into the runtime game by default.
 		// The ones that need to be kept will be copied into SoftObjectPath on the object during save.
-		FCookLoadScope CookLoadScope(ECookLoadType::EditorOnly);
 
 		const TSharedRef<FMutableStreamableManager>& Streamable = System->GetPrivate()->StreamableManager;
 
@@ -368,8 +365,6 @@ void ProcessChildObjectsRecursively(const UCustomizableObject* ParentObject, FMu
 		return;
 	}
 	
-	FCookLoadScope CookLoadScope(ECookLoadType::EditorOnly);
-	
 	// Required to be deterministic.
 	ReferencedObjectNames.Sort([](const FName& A, const FName& B)
 	{
@@ -389,7 +384,7 @@ void ProcessChildObjectsRecursively(const UCustomizableObject* ParentObject, FMu
 	{
 		FSoftObjectPath SoftObjectPath = AssetData.GetSoftObjectPath();
 		
-		UCustomizableObject* ChildObject = Cast<UCustomizableObject>(SoftObjectPath.TryLoad());
+		UCustomizableObject* ChildObject = Cast<UCustomizableObject>(MutablePrivate::LoadObject(SoftObjectPath));
 		if (!ChildObject || ChildObject->HasAnyFlags(RF_Transient))
 		{
 			continue;
@@ -1004,8 +999,7 @@ void FCustomizableObjectCompiler::CompileInternal(bool bAsync)
 		for (const TPair<TSoftObjectPtr<const UTexture>, FMutableGraphGenerationContext::FGeneratedReferencedTexture>& Pair : GenerationContext.CompileTimeTextureMap)
 		{
 			check(Pair.Value.ID == NewCompileTimeReferencedTextures.Num());
-			FCookLoadScope CookLoadScope(ECookLoadType::EditorOnly);
-			FMutableSourceTextureData Tex(*Pair.Key.LoadSynchronous());
+			FMutableSourceTextureData Tex(*MutablePrivate::LoadObject(Pair.Key));
 			NewCompileTimeReferencedTextures.Add(Tex);
 		}
 
@@ -1277,7 +1271,7 @@ mu::Ptr<mu::Node> FCustomizableObjectCompiler::Export(UCustomizableObject* Objec
 	{
 		check(Pair.Value.ID == OutCompilerReferencedTextures.Num());
 
-		FMutableSourceTextureData Tex(*Pair.Key.LoadSynchronous());
+		FMutableSourceTextureData Tex(*MutablePrivate::LoadObject(Pair.Key));
 		OutCompilerReferencedTextures.Add(Tex);
 	}
 
diff --git a/Engine/Plugins/Mutable/Source/CustomizableObjectEditor/Private/MuCOE/CustomizableObjectCookPackageSplitter.cpp b/Engine/Plugins/Mutable/Source/CustomizableObjectEditor/Private/MuCOE/CustomizableObjectCookPackageSplitter.cpp
--- a/Engine/Plugins/Mutable/Source/CustomizableObjectEditor/Private/MuCOE/CustomizableObjectCookPackageSplitter.cpp
+++ b/Engine/Plugins/Mutable/Source/CustomizableObjectEditor/Private/MuCOE/CustomizableObjectCookPackageSplitter.cpp
@@ -5,6 +5,7 @@
 #include "Algo/Find.h"
 #include "MuCO/CustomizableObject.h"
 #include "MuCO/CustomizableObjectPrivate.h"
+#include "MuCO/LoadUtils.h"
 #include "UObject/NameTypes.h"
 #include "UObject/Package.h"
 
@@ -84,7 +85,7 @@ EMoveContainerError MoveContainerToNewOuter(
 	
 	OutContainer = nullptr;
 
-	UCustomizableObjectResourceDataContainer* Container = StreamedResourceData->GetPath().LoadSynchronous();
+	UCustomizableObjectResourceDataContainer* Container = MutablePrivate::LoadObject(StreamedResourceData->GetPath());
 	if (!Container)
 	{
 		return EMoveContainerError::FailedToLoadContainer;
@@ -387,7 +388,7 @@ bool FCustomizableObjectCookPackageSplitter::PopulateGeneratedPackage(
 
 	// [TEMP] Loading a package referencing the CO before PostSaveGeneratedPackage is called causes a name collision.
 	// Duplicate the object with the new outer instead of moving it until it is fixed.
-	UObject* Container = ResourceData->GetPath().LoadSynchronous();
+	UObject* Container = MutablePrivate::LoadObject(ResourceData->GetPath());
 	EMoveContainerError Error = Container ? EMoveContainerError::None : EMoveContainerError::FailedToLoadContainer;
 	if (Container)
 	{
diff --git a/Engine/Plugins/Mutable/Source/CustomizableObjectEditor/Private/MuCOE/CustomizableObjectCustomSettings.cpp b/Engine/Plugins/Mutable/Source/CustomizableObjectEditor/Private/MuCOE/CustomizableObjectCustomSettings.cpp
--- a/Engine/Plugins/Mutable/Source/CustomizableObjectEditor/Private/MuCOE/CustomizableObjectCustomSettings.cpp
+++ b/Engine/Plugins/Mutable/Source/CustomizableObjectEditor/Private/MuCOE/CustomizableObjectCustomSettings.cpp
@@ -14,6 +14,7 @@
 #include "PropertyCustomizationHelpers.h"
 #include "MuCO/CustomizableObject.h"
 #include "MuCO/CustomizableObjectInstance.h"
+#include "MuCO/LoadUtils.h"
 #include "MuCOE/ICustomizableObjectEditor.h"
 #include "MuCOE/ICustomizableObjectInstanceEditor.h"
 #include "Subsystems/ImportSubsystem.h"
@@ -453,7 +454,7 @@ void SCustomizableObjectCustomSettings::SetViewportLightsByAsset(const FAssetDat
 		return;	
 	}
 	
-	UCustomizableObjectEditorViewportLights* InLights = Cast<UCustomizableObjectEditorViewportLights>(InAsset.GetAsset());	
+	UCustomizableObjectEditorViewportLights* InLights = Cast<UCustomizableObjectEditorViewportLights>(MutablePrivate::LoadObject(InAsset));	
 	if (!InLights)
 	{
 		return;
diff --git a/Engine/Plugins/Mutable/Source/CustomizableObjectEditor/Private/MuCOE/CustomizableObjectEditor.cpp b/Engine/Plugins/Mutable/Source/CustomizableObjectEditor/Private/MuCOE/CustomizableObjectEditor.cpp
--- a/Engine/Plugins/Mutable/Source/CustomizableObjectEditor/Private/MuCOE/CustomizableObjectEditor.cpp
+++ b/Engine/Plugins/Mutable/Source/CustomizableObjectEditor/Private/MuCOE/CustomizableObjectEditor.cpp
@@ -1002,13 +1002,13 @@ void FCustomizableObjectEditor::ShowGizmoClipMesh(UCustomizableObjectNode& Node,
 
 			if (const UCustomizableObjectNodeStaticMesh* StaticMeshNode = Cast<UCustomizableObjectNodeStaticMesh>(ConnectedNode))
 			{
-				ClipMesh = StaticMeshNode->GetMesh().LoadSynchronous();
+				ClipMesh = MutablePrivate::LoadObject(StaticMeshNode->GetMesh());
 				StaticMeshNode->GetPinSection(*ConnectedPin, LODIndex, SectionIndex, DummyIndex);
 				MaterialSlotIndex = SectionIndex;
 			}
 			else if (const UCustomizableObjectNodeSkeletalMesh* SkeletalMeshNode = Cast<UCustomizableObjectNodeSkeletalMesh>(ConnectedNode))
 			{
-				ClipMesh = SkeletalMeshNode->GetMesh().LoadSynchronous();
+				ClipMesh = MutablePrivate::LoadObject(SkeletalMeshNode->GetMesh());
 				SkeletalMeshNode->GetPinSection(*ConnectedPin, LODIndex, SectionIndex, DummyIndex);
 				MaterialSlotIndex = SkeletalMeshNode->GetSkeletalMaterialIndexFor(*ConnectedPin);
 			}
diff --git a/Engine/Plugins/Mutable/Source/CustomizableObjectEditor/Private/MuCOE/CustomizableObjectEditorLevelProfile.cpp b/Engine/Plugins/Mutable/Source/CustomizableObjectEditor/Private/MuCOE/CustomizableObjectEditorLevelProfile.cpp
--- a/Engine/Plugins/Mutable/Source/CustomizableObjectEditor/Private/MuCOE/CustomizableObjectEditorLevelProfile.cpp
+++ b/Engine/Plugins/Mutable/Source/CustomizableObjectEditor/Private/MuCOE/CustomizableObjectEditorLevelProfile.cpp
@@ -11,6 +11,7 @@
 #include "Engine/TextureCube.h"
 #include "EngineUtils.h"
 #include "Misc/MessageDialog.h"
+#include "MuCO/LoadUtils.h"
 
 class UClass;
 
@@ -37,7 +38,7 @@ bool LevelProfileManage::LoadProfileUObjectNames(const FString& AssetPath,
 											     TArray<FString>& ArrayPostProcess,
 											     TArray<FString>& ArraySkyLight)
 {
-	UWorld* World = LoadObject<UWorld>(nullptr, *AssetPath, nullptr);
+	UWorld* World = MutablePrivate::LoadObject<UWorld>(nullptr, *AssetPath);
 
 	if (World == nullptr)
 	{
@@ -99,7 +100,7 @@ bool LevelProfileManage::LoadProfileUObjects(const FString& AssetPath,
 								             APostProcessVolume*& PostProcessVolume,
 								             ASkyLight*& SkyLight)
 {
-	UWorld* World = LoadObject<UWorld>(nullptr, *AssetPath, nullptr);
+	UWorld* World = MutablePrivate::LoadObject<UWorld>(nullptr, *AssetPath);
 
 	DirectionalLight = nullptr;
 	PostProcessVolume = nullptr;
diff --git a/Engine/Plugins/Mutable/Source/CustomizableObjectEditor/Private/MuCOE/CustomizableObjectEditorModule.cpp b/Engine/Plugins/Mutable/Source/CustomizableObjectEditor/Private/MuCOE/CustomizableObjectEditorModule.cpp
--- a/Engine/Plugins/Mutable/Source/CustomizableObjectEditor/Private/MuCOE/CustomizableObjectEditorModule.cpp
+++ b/Engine/Plugins/Mutable/Source/CustomizableObjectEditor/Private/MuCOE/CustomizableObjectEditorModule.cpp
@@ -539,8 +539,6 @@ void GetReferencingCustomizableObjects(FName CustomizableObjectName, TArray<FNam
 
 void GetReferencingPackages(const UCustomizableObject& Object, TArray<FAssetData>& ReferencingAssets)
 {
-	FCookLoadScope CookLoadScope(ECookLoadType::EditorOnly);
-
 	// Gather all child CustomizableObjects
 	TArray<FName> VisitedObjectNames;
 	GetReferencingCustomizableObjects(Object.GetPackage()->GetFName(), VisitedObjectNames, ReferencingAssets);
@@ -551,7 +549,7 @@ void GetReferencingPackages(const UCustomizableObject& Object, TArray<FAssetData
 	{
 		const TSoftObjectPtr SoftObjectPtr(ReferencingCustomizableObject.ToSoftObjectPath());
 
-		const UCustomizableObject* ChildCustomizableObject = Cast<UCustomizableObject>(SoftObjectPtr.LoadSynchronous());
+		const UCustomizableObject* ChildCustomizableObject = Cast<UCustomizableObject>(MutablePrivate::LoadObject(SoftObjectPtr));
 		if (!ChildCustomizableObject)
 		{
 			continue;
@@ -891,7 +889,7 @@ void FCustomizableObjectEditorModule::OnPreBeginPIE(const bool bIsSimulatingInEd
 			continue;
 		}
 
-		UCustomizableObject* Object = Cast<UCustomizableObject>(Asset.GetAsset());
+		UCustomizableObject* Object = Cast<UCustomizableObject>(MutablePrivate::LoadObject(Asset));
 		if (!Object || Object->IsCompiled() || Object->GetPrivate()->IsLocked())
 		{
 			continue;
diff --git a/Engine/Plugins/Mutable/Source/CustomizableObjectEditor/Private/MuCOE/CustomizableObjectEditorViewportClient.cpp b/Engine/Plugins/Mutable/Source/CustomizableObjectEditor/Private/MuCOE/CustomizableObjectEditorViewportClient.cpp
--- a/Engine/Plugins/Mutable/Source/CustomizableObjectEditor/Private/MuCOE/CustomizableObjectEditorViewportClient.cpp
+++ b/Engine/Plugins/Mutable/Source/CustomizableObjectEditor/Private/MuCOE/CustomizableObjectEditorViewportClient.cpp
@@ -53,6 +53,7 @@
 #include "MuCOE/CustomizableObjectEditorLogger.h"
 #include "MuCOE/UnrealBakeHelpers.h"
 #include "MuCO/CustomizableObjectPrivate.h"
+#include "MuCO/LoadUtils.h"
 #include "Widgets/Input/SButton.h"
 #include "Widgets/Input/SEditableTextBox.h"
 #include "Widgets/Layout/SUniformGridPanel.h"
@@ -142,7 +143,7 @@ FCustomizableObjectEditorViewportClient::FCustomizableObjectEditorViewportClient
 		PluginConfig->GetString(TEXT("EditorDefaults"), TEXT("ClipMorphMaterialName"), ClipMorphMaterialName);
 		if (!ClipMorphMaterialName.IsEmpty())
 		{
-			ClipMorphMaterial = LoadObject<UMaterial>(NULL, *ClipMorphMaterialName, NULL, LOAD_None, NULL);
+			ClipMorphMaterial = MutablePrivate::LoadObject<UMaterial>(nullptr, *ClipMorphMaterialName);
 			ensure(ClipMorphMaterial);
 		}
 
@@ -151,7 +152,7 @@ FCustomizableObjectEditorViewportClient::FCustomizableObjectEditorViewportClient
 		PluginConfig->GetString(TEXT("EditorDefaults"), TEXT("ClipMeshMaterialName"), ClipMeshMaterialName);
 		if (!ClipMeshMaterialName.IsEmpty())
 		{
-			ClipMeshMaterial = LoadObject<UMaterial>(nullptr, *ClipMeshMaterialName, nullptr, LOAD_None, nullptr);
+			ClipMeshMaterial = MutablePrivate::LoadObject<UMaterial>(nullptr, *ClipMeshMaterialName);
 			ensure(ClipMeshMaterial);
 		}
 	}
diff --git a/Engine/Plugins/Mutable/Source/CustomizableObjectEditor/Private/MuCOE/CustomizableObjectIdentifierCustomization.cpp b/Engine/Plugins/Mutable/Source/CustomizableObjectEditor/Private/MuCOE/CustomizableObjectIdentifierCustomization.cpp
--- a/Engine/Plugins/Mutable/Source/CustomizableObjectEditor/Private/MuCOE/CustomizableObjectIdentifierCustomization.cpp
+++ b/Engine/Plugins/Mutable/Source/CustomizableObjectEditor/Private/MuCOE/CustomizableObjectIdentifierCustomization.cpp
@@ -7,6 +7,7 @@
 #include "MuCOE/GraphTraversal.h"
 #include "MuCOE/Nodes/CustomizableObjectNodeObject.h"
 #include "PropertyCustomizationHelpers.h"
+#include "MuCO/LoadUtils.h"
 #include "MuCO/UnrealPortabilityHelpers.h"
 
 #define LOCTEXT_NAMESPACE "FCustomizableObjectIdentifierCustomization"
@@ -50,7 +51,7 @@ void FCustomizableObjectIdentifierCustomization::CustomizeHeader(TSharedRef<clas
 	{
 		for (auto Itr = OutAssetData.CreateIterator(); Itr; ++Itr)
 		{
-			UCustomizableObject* CustomizableObject = Cast<UCustomizableObject>(Itr->GetAsset());
+			UCustomizableObject* CustomizableObject = Cast<UCustomizableObject>(MutablePrivate::LoadObject(*Itr));
 
 			if (CustomizableObject)
 			{
@@ -95,7 +96,7 @@ FString FCustomizableObjectIdentifierCustomization::OnGetObjectPath() const
 
 void FCustomizableObjectIdentifierCustomization::OnEventChanged(const FAssetData& InAssetData)
 {
-	SelectedCustomizableObject = Cast<UCustomizableObject>(InAssetData.GetAsset());
+	SelectedCustomizableObject = Cast<UCustomizableObject>(MutablePrivate::LoadObject(InAssetData));
 
 	if (SelectedCustomizableObject)
 	{
diff --git a/Engine/Plugins/Mutable/Source/CustomizableObjectEditor/Private/MuCOE/CustomizableObjectInstanceFactory.cpp b/Engine/Plugins/Mutable/Source/CustomizableObjectEditor/Private/MuCOE/CustomizableObjectInstanceFactory.cpp
--- a/Engine/Plugins/Mutable/Source/CustomizableObjectEditor/Private/MuCOE/CustomizableObjectInstanceFactory.cpp
+++ b/Engine/Plugins/Mutable/Source/CustomizableObjectEditor/Private/MuCOE/CustomizableObjectInstanceFactory.cpp
@@ -13,6 +13,7 @@
 #include "MuCO/CustomizableSkeletalComponent.h"
 #include "MuCO/CustomizableSkeletalMeshActor.h"
 #include "MuCO/CustomizableSkeletalMeshActorPrivate.h"
+#include "MuCO/LoadUtils.h"
 #include "MuCOE/Nodes/CustomizableObjectNodeObject.h"
 #include "MuCO/UnrealPortabilityHelpers.h"
 
@@ -83,7 +84,7 @@ bool UCustomizableObjectInstanceFactory::CanCreateActorFrom(const FAssetData& As
 
     if (AssetData.GetClass()->IsChildOf(UCustomizableObjectInstance::StaticClass()))
     {
-        if (UCustomizableObjectInstance* CustomizableObjectInstance = Cast<UCustomizableObjectInstance>(AssetData.GetAsset()))
+        if (UCustomizableObjectInstance* CustomizableObjectInstance = Cast<UCustomizableObjectInstance>(MutablePrivate::LoadObject(AssetData)))
         {
 			TArray<FName> ComponentNames = CustomizableObjectInstance->GetComponentNames();
 
diff --git a/Engine/Plugins/Mutable/Source/CustomizableObjectEditor/Private/MuCOE/CustomizableObjectLayout.cpp b/Engine/Plugins/Mutable/Source/CustomizableObjectEditor/Private/MuCOE/CustomizableObjectLayout.cpp
--- a/Engine/Plugins/Mutable/Source/CustomizableObjectEditor/Private/MuCOE/CustomizableObjectLayout.cpp
+++ b/Engine/Plugins/Mutable/Source/CustomizableObjectEditor/Private/MuCOE/CustomizableObjectLayout.cpp
@@ -176,7 +176,7 @@ void UCustomizableObjectLayout::ConsolidateAutomaticBlocks()
 
 void UCustomizableObjectLayout::GetUVs(TArray<FVector2f>& UVs) const
 {
-	if(const UObject* Mesh = GetMesh().LoadSynchronous())
+	if (const UObject* Mesh = MutablePrivate::LoadObject(GetMesh()))
 	{
 		if (const USkeletalMesh* SkeletalMesh = Cast<const USkeletalMesh>(Mesh))
 		{
diff --git a/Engine/Plugins/Mutable/Source/CustomizableObjectEditor/Private/MuCOE/EdGraphSchema_CustomizableObject.cpp b/Engine/Plugins/Mutable/Source/CustomizableObjectEditor/Private/MuCOE/EdGraphSchema_CustomizableObject.cpp
--- a/Engine/Plugins/Mutable/Source/CustomizableObjectEditor/Private/MuCOE/EdGraphSchema_CustomizableObject.cpp
+++ b/Engine/Plugins/Mutable/Source/CustomizableObjectEditor/Private/MuCOE/EdGraphSchema_CustomizableObject.cpp
@@ -1073,7 +1073,7 @@ enum class ESpawnableObjectType : int32
 
 bool IsSpawnableAsset(const FAssetData& InAsset, ESpawnableObjectType& OutObjectType)
 {
-	UObject* Object = InAsset.GetAsset();
+	UObject* Object = MutablePrivate::LoadObject(InAsset);
 
 	// Type used to know what kind of UObject this asset is
 	OutObjectType = ESpawnableObjectType::None;
@@ -1125,7 +1125,7 @@ void UEdGraphSchema_CustomizableObject::DroppedAssetsOnGraph(const TArray<FAsset
 
 		// At this point we know we are working with an asset we can spawn as a mutable node.
 		
-		UObject* Object = Asset.GetAsset();
+		UObject* Object = MutablePrivate::LoadObject(Asset);
 		UEdGraphNode* GraphNode = nullptr;
 
 		// Depending on the UObjectType spawn one or another mutable node
diff --git a/Engine/Plugins/Mutable/Source/CustomizableObjectEditor/Private/MuCOE/GenerateMutableSource/GenerateMutableSource.cpp b/Engine/Plugins/Mutable/Source/CustomizableObjectEditor/Private/MuCOE/GenerateMutableSource/GenerateMutableSource.cpp
--- a/Engine/Plugins/Mutable/Source/CustomizableObjectEditor/Private/MuCOE/GenerateMutableSource/GenerateMutableSource.cpp
+++ b/Engine/Plugins/Mutable/Source/CustomizableObjectEditor/Private/MuCOE/GenerateMutableSource/GenerateMutableSource.cpp
@@ -455,10 +455,8 @@ FMutableComponentInfo* FMutableGraphGenerationContext::GetCurrentComponentInfo()
 
 UObject* FMutableGraphGenerationContext::LoadObject(const FSoftObjectPtr& SoftObject, bool bParticipatingObjectsPassLoad)
 {
-	FCookLoadScope CookLoadScope(ECookLoadType::EditorOnly);
-
 	return bLoadObjects && (!bParticipatingObjectsPass || bParticipatingObjectsPassLoad) ?
-		SoftObject.LoadSynchronous() :
+		MutablePrivate::LoadObject(SoftObject) :
 		nullptr;
 }
 
diff --git a/Engine/Plugins/Mutable/Source/CustomizableObjectEditor/Private/MuCOE/GenerateMutableSource/GenerateMutableSource.h b/Engine/Plugins/Mutable/Source/CustomizableObjectEditor/Private/MuCOE/GenerateMutableSource/GenerateMutableSource.h
--- a/Engine/Plugins/Mutable/Source/CustomizableObjectEditor/Private/MuCOE/GenerateMutableSource/GenerateMutableSource.h
+++ b/Engine/Plugins/Mutable/Source/CustomizableObjectEditor/Private/MuCOE/GenerateMutableSource/GenerateMutableSource.h
@@ -10,6 +10,7 @@
 #include "MuCO/CustomizableObjectIdentifier.h"
 #include "MuCO/CustomizableObjectCompilerTypes.h"
 #include "MuCO/UnrealToMutableTextureConversionUtils.h"
+#include "MuCO/LoadUtils.h"
 #include "MuCOE/CustomizableObjectEditorLogger.h"
 #include "MuCOE/ExtensionDataCompilerInterface.h"
 #include "MuCOE/Nodes/CustomizableObjectNodeMaterialBase.h"
@@ -29,7 +30,6 @@
 #include "MuR/Mesh.h"
 #include "MuR/Ptr.h"
 #include "UObject/Package.h"
-#include "UObject/ICookInfo.h"
 
 class UCustomizableObjectNodeGroupProjectorParameter;
 class UCustomizableObjectNodeMaterialBase;
@@ -58,6 +58,7 @@ struct FMutableSkinWeightProfileInfo;
 struct FMorphTargetVertexData;
 enum class EPinMode;
 
+
 struct FGeneratedImageProperties
 {
 	/** Name in the Material. */
@@ -551,8 +552,6 @@ struct FMutableGraphGenerationContext
 
 	/** Full hierarchy root. */
 	UCustomizableObjectNodeObject* Root = nullptr;
-
-public:
 	
 	// Cache of generated pins per LOD
 	TMap<FGeneratedKey, FGeneratedData> Generated;
@@ -760,27 +759,23 @@ public:
 	FMutableComponentInfo* GetCurrentComponentInfo();
 
 	UObject* LoadObject(const FSoftObjectPtr& SoftObject, bool bParticipatingObjectsPassLoad = false);
-
+	
 	template<typename T>
 	T* LoadObject(const TSoftObjectPtr<T>& SoftObject, bool bParticipatingObjectsPassLoad = false)
 	{
-		FCookLoadScope CookLoadScope(ECookLoadType::EditorOnly);
 		return bLoadObjects && (!bParticipatingObjectsPass || bParticipatingObjectsPassLoad) ?
-			SoftObject.LoadSynchronous() :
+			MutablePrivate::LoadObject<T>(SoftObject) :
 			nullptr;
 	}
 
 	template<typename T>
 	UClass* LoadClass(const TSoftClassPtr<T>& SoftClass, bool bParticipatingObjectsPassLoad = false)
 	{
-		FCookLoadScope CookLoadScope(ECookLoadType::EditorOnly);
 		return bLoadObjects && (!bParticipatingObjectsPass || bParticipatingObjectsPassLoad) ?
-			SoftClass.LoadSynchronous() :
+			MutablePrivate::LoadClass<T>(SoftClass) :
 			nullptr;
 	}
 
-public:
-
 	/** Only compiled components. All components types. Index is the ObjectComponentIndex. */
 	TArray<FName> ComponentNames;
 	
@@ -803,8 +798,7 @@ public:
 
 	// For a given material node (the key is node package path + node uid + image index in node) stores images generated for the same node at a higher quality LOD to reuse that image node
 	TMap<FString, FGroupProjectorImageInfo> GroupProjectorLODCache;
-
-
+	
 	TMap<uint32, FMutableSurfaceMetadata> SurfaceMetadata;
 
 	uint8 FromLOD = 0; // LOD to append to the CurrentLOD when using AutomaticLODs. 
diff --git a/Engine/Plugins/Mutable/Source/CustomizableObjectEditor/Private/MuCOE/GenerateMutableSource/GenerateMutableSourceLayout.cpp b/Engine/Plugins/Mutable/Source/CustomizableObjectEditor/Private/MuCOE/GenerateMutableSource/GenerateMutableSourceLayout.cpp
--- a/Engine/Plugins/Mutable/Source/CustomizableObjectEditor/Private/MuCOE/GenerateMutableSource/GenerateMutableSourceLayout.cpp
+++ b/Engine/Plugins/Mutable/Source/CustomizableObjectEditor/Private/MuCOE/GenerateMutableSource/GenerateMutableSourceLayout.cpp
@@ -86,7 +86,7 @@ mu::Ptr<mu::NodeLayout> CreateMutableLayoutNode(FMutableGraphGenerationContext&
 		TSharedPtr<mu::FMesh> MutableMesh;
 
 		// TODO: Don't force load the mesh
-		if (USkeletalMesh* SkeletalMesh = Cast<USkeletalMesh>(UnrealLayout->GetMesh().LoadSynchronous()))
+		if (USkeletalMesh* SkeletalMesh = Cast<USkeletalMesh>(GenerationContext.LoadObject(UnrealLayout->GetMesh())))
 		{
 			// We don't need all the data to generate the blocks
 			const EMutableMeshConversionFlags ShapeFlags =
@@ -100,7 +100,7 @@ mu::Ptr<mu::NodeLayout> CreateMutableLayoutNode(FMutableGraphGenerationContext&
 
 			GenerationContext.MeshGenerationFlags.Pop();
 		}
-		else if (UStaticMesh* StaticMesh = Cast<UStaticMesh>(UnrealLayout->GetMesh().LoadSynchronous()))
+		else if (UStaticMesh* StaticMesh = Cast<UStaticMesh>(GenerationContext.LoadObject(UnrealLayout->GetMesh())))
 		{
 			MutableMesh = ConvertStaticMeshToMutable(StaticMesh, UnrealLayout->GetLOD(), UnrealLayout->GetMaterial(), GenerationContext, nullptr);
 		}
diff --git a/Engine/Plugins/Mutable/Source/CustomizableObjectEditor/Private/MuCOE/GenerateMutableSource/GenerateMutableSourceMesh.cpp b/Engine/Plugins/Mutable/Source/CustomizableObjectEditor/Private/MuCOE/GenerateMutableSource/GenerateMutableSourceMesh.cpp
--- a/Engine/Plugins/Mutable/Source/CustomizableObjectEditor/Private/MuCOE/GenerateMutableSource/GenerateMutableSourceMesh.cpp
+++ b/Engine/Plugins/Mutable/Source/CustomizableObjectEditor/Private/MuCOE/GenerateMutableSource/GenerateMutableSourceMesh.cpp
@@ -434,7 +434,7 @@ TArray<TTuple<UPhysicsAsset*, int32>> GetPhysicsAssetsFromAnimInstance(const TSo
 		return Result;
 	}
 
-	UClass* AnimInstanceClass = AnimInstance.LoadSynchronous();
+	UClass* AnimInstanceClass = MutablePrivate::LoadClass(AnimInstance);
 	UAnimBlueprintGeneratedClass* AnimClass = Cast<UAnimBlueprintGeneratedClass>(AnimInstanceClass);
 
 	if (AnimClass)
@@ -709,7 +709,7 @@ TSharedPtr<mu::FMesh> ConvertSkeletalMeshToMutable(
 	const FString MeshName = Source.Mesh.GetLongPackageName().ToLower();
 	const uint32 MeshId = CityHash32(reinterpret_cast<const char*>(*MeshName), MeshName.Len() * sizeof(FString::ElementType));
 
-	const USkeletalMesh* SkeletalMesh = Cast<USkeletalMesh>( Source.Mesh.LoadSynchronous() );
+	const USkeletalMesh* SkeletalMesh = Cast<USkeletalMesh>( MutablePrivate::LoadObject(Source.Mesh) );
 	check(SkeletalMesh);
 	const FSkeletalMeshModel* ImportedModel = SkeletalMesh->GetImportedModel();
 	if (!ImportedModel)
@@ -1213,7 +1213,7 @@ TSharedPtr<mu::FMesh> ConvertSkeletalMeshToMutable(
 					{
 						check(MorphTarget);
 
-						if (Source.TableReferenceSkeletalMesh.LoadSynchronous()->FindMorphTarget(MorphTarget->GetFName()))
+						if (MutablePrivate::LoadObject(Source.TableReferenceSkeletalMesh)->FindMorphTarget(MorphTarget->GetFName()))
 						{
 							MorphTargetsNames.Add(MorphTarget->GetFName());
 						}
@@ -2585,7 +2585,7 @@ TSharedPtr<mu::FMesh> GenerateMutableMesh( TSoftObjectPtr<UStreamableRenderAsset
 	//MutableMesh = GenerateMeshConstant(Mesh, AnimInstance, bMeshMustExist, LODIndex, SectionIndex, SurfaceMetadataId, GenerationContext, bIsReference);
 
 	// Meanwhile:
-	UObject* LoadedMesh = Mesh.LoadSynchronous();
+	UObject* LoadedMesh = MutablePrivate::LoadObject(Mesh);
 	if (USkeletalMesh* SkeletalMesh = Cast<USkeletalMesh>(LoadedMesh))
 	{
 
@@ -2702,7 +2702,7 @@ TSharedPtr<mu::FMesh> BuildMorphedMutableMesh(const UEdGraphPin* BaseSourcePin,
 	{
 		int32 LayoutIndex;
 		TypedNodeSkeletalMesh->GetPinSection(*BaseSourcePin, LODIndexConnected, SectionIndexConnected, LayoutIndex);
-		SkeletalMesh = TypedNodeSkeletalMesh->SkeletalMesh.LoadSynchronous();
+		SkeletalMesh = MutablePrivate::LoadObject(TypedNodeSkeletalMesh->SkeletalMesh);
 	}
 
 	else if (const UCustomizableObjectNodeTable* TypedNodeTable = Cast<UCustomizableObjectNodeTable>(Node))
diff --git a/Engine/Plugins/Mutable/Source/CustomizableObjectEditor/Private/MuCOE/GenerateMutableSource/GenerateMutableSourceTable.cpp b/Engine/Plugins/Mutable/Source/CustomizableObjectEditor/Private/MuCOE/GenerateMutableSource/GenerateMutableSourceTable.cpp
--- a/Engine/Plugins/Mutable/Source/CustomizableObjectEditor/Private/MuCOE/GenerateMutableSource/GenerateMutableSourceTable.cpp
+++ b/Engine/Plugins/Mutable/Source/CustomizableObjectEditor/Private/MuCOE/GenerateMutableSource/GenerateMutableSourceTable.cpp
@@ -28,7 +28,6 @@
 #include "Rendering/SkeletalMeshModel.h"
 #include "MuCOE/CustomizableObjectVersionBridge.h"
 #include "ClothingAsset.h"
-#include "UObject/ICookInfo.h"
 
 #define LOCTEXT_NAMESPACE "CustomizableObjectEditor"
 
@@ -1258,9 +1257,7 @@ UDataTable* GenerateDataTableFromStruct(const UCustomizableObjectNodeTable* Tabl
 
 		return nullptr;
 	}
-
-	FCookLoadScope CookLoadScope(ECookLoadType::EditorOnly);
-
+	
 	TArray<FAssetData> DataTableAssets = TableNode->GetParentTables();
 
 	UCompositeDataTable* CompositeDataTable = NewObject<UCompositeDataTable>();
@@ -1272,7 +1269,7 @@ UDataTable* GenerateDataTableFromStruct(const UCustomizableObjectNodeTable* Tabl
 	{
 		if (DataTableAsset.IsValid())
 		{
-			if (UDataTable* DataTable = Cast<UDataTable>(DataTableAsset.GetAsset()))
+			if (UDataTable* DataTable = Cast<UDataTable>(MutablePrivate::LoadObject(DataTableAsset)))
 			{
 				ParentTables.Add(DataTable);
 			}
diff --git a/Engine/Plugins/Mutable/Source/CustomizableObjectEditor/Private/MuCOE/GraphTraversal.cpp b/Engine/Plugins/Mutable/Source/CustomizableObjectEditor/Private/MuCOE/GraphTraversal.cpp
--- a/Engine/Plugins/Mutable/Source/CustomizableObjectEditor/Private/MuCOE/GraphTraversal.cpp
+++ b/Engine/Plugins/Mutable/Source/CustomizableObjectEditor/Private/MuCOE/GraphTraversal.cpp
@@ -3,6 +3,7 @@
 #include "MuCOE/GraphTraversal.h"
 
 #include "AssetRegistry/AssetRegistryModule.h"
+#include "MuCO/LoadUtils.h"
 #include "MuCOE/CustomizableObjectPin.h"
 #include "MuCOE/EdGraphSchema_CustomizableObject.h"
 #include "MuCOE/Nodes/CustomizableObjectNodeEnumParameter.h"
@@ -25,7 +26,6 @@
 #include "MuCOE/Nodes/CustomizableObjectNodeObjectGroup.h"
 #include "MuCOE/Nodes/CustomizableObjectNodeReroute.h"
 #include "UObject/UObjectIterator.h"
-#include "UObject/ICookInfo.h"
 
 
 TArray<UEdGraphPin*> FollowPinArray(const UEdGraphPin& Pin, bool bIgnoreOrphan, bool* bOutCycleDetected)
@@ -653,9 +653,7 @@ void GetNodeGroupObjectNodeMappingImmersive(UCustomizableObject* Object, FAssetR
 			Filter.PackageNames.Add(ReferenceName);
 		}
 	}
-
-	FCookLoadScope CookLoadScope(ECookLoadType::EditorOnly);
-
+	
 	Filter.bIncludeOnlyOnDiskAssets = false;
 	
 	TArray<FAssetData> ArrayAssetData;
@@ -663,7 +661,7 @@ void GetNodeGroupObjectNodeMappingImmersive(UCustomizableObject* Object, FAssetR
 
 	for (FAssetData& AssetData : ArrayAssetData)
 	{
-		UCustomizableObject* ChildObject = Cast<UCustomizableObject>(AssetData.GetAsset());
+		UCustomizableObject* ChildObject = Cast<UCustomizableObject>(MutablePrivate::LoadObject(AssetData));
 		if (!ChildObject)
 		{
 			continue;
diff --git a/Engine/Plugins/Mutable/Source/CustomizableObjectEditor/Private/MuCOE/Nodes/CustomizableObjectNodeComponentPassthroughMesh.cpp b/Engine/Plugins/Mutable/Source/CustomizableObjectEditor/Private/MuCOE/Nodes/CustomizableObjectNodeComponentPassthroughMesh.cpp
--- a/Engine/Plugins/Mutable/Source/CustomizableObjectEditor/Private/MuCOE/Nodes/CustomizableObjectNodeComponentPassthroughMesh.cpp
+++ b/Engine/Plugins/Mutable/Source/CustomizableObjectEditor/Private/MuCOE/Nodes/CustomizableObjectNodeComponentPassthroughMesh.cpp
@@ -5,6 +5,7 @@
 #include "Engine/SkeletalMesh.h"
 #include "Modules/ModuleManager.h"
 #include "MuCO/CustomizableObjectCustomVersion.h"
+#include "MuCO/LoadUtils.h"
 #include "MuCOE/CustomizableObjectEditor.h"
 #include "MuCOE/CustomizableObjectEditorStyle.h"
 #include "MuCOE/CustomizableObjectLayout.h"
@@ -14,7 +15,6 @@
 #include "MuCOE/Nodes/CustomizableObjectNodeLayoutBlocks.h"
 #include "MuCOE/RemapPins/CustomizableObjectNodeRemapPinsByNameDefaultPin.h"
 #include "MuCOE/UnrealEditorPortabilityHelpers.h"
-#include "UObject/ICookInfo.h"
 
 class UCustomizableObjectNodeRemapPinsByName;
 class UObject;
@@ -88,11 +88,9 @@ void UCustomizableObjectNodeComponentPassthroughMesh::AllocateDefaultPins(UCusto
 	{
 		return;
 	}
-
-	FCookLoadScope CookLoadScope(ECookLoadType::EditorOnly);
-
+	
 	// Support SkeletalMeshes only, for now.
-	USkeletalMesh* SkeletalMesh = Cast<USkeletalMesh>( Mesh.TryLoad() );
+	USkeletalMesh* SkeletalMesh = Cast<USkeletalMesh>(MutablePrivate::LoadObject(Mesh));
 	if (!SkeletalMesh)
 	{
 		return;
@@ -186,9 +184,7 @@ void UCustomizableObjectNodeComponentPassthroughMesh::GetPinSection(const UEdGra
 
 UMaterialInterface* UCustomizableObjectNodeComponentPassthroughMesh::GetMaterialFor(const UEdGraphPin* Pin) const
 {
-	FCookLoadScope CookLoadScope(ECookLoadType::EditorOnly);
-
-	USkeletalMesh* SkeletalMesh = Cast<USkeletalMesh>(Mesh.TryLoad());
+	USkeletalMesh* SkeletalMesh = Cast<USkeletalMesh>(MutablePrivate::LoadObject(Mesh));
 	if (SkeletalMesh)
 	{
 		if (FSkeletalMaterial* SkeletalMaterial = GetSkeletalMaterialFor(*Pin))
@@ -239,9 +235,7 @@ UMaterialInterface* UCustomizableObjectNodeComponentPassthroughMesh::GetMaterial
 
 FSkeletalMaterial* UCustomizableObjectNodeComponentPassthroughMesh::GetSkeletalMaterialFor(const int32 LODIndex, const int32 SectionIndex, const FSkeletalMeshModel* ImportedModel) const
 {
-	FCookLoadScope CookLoadScope(ECookLoadType::EditorOnly);
-
-	USkeletalMesh* SkeletalMesh = Cast<USkeletalMesh>(Mesh.TryLoad());
+	USkeletalMesh* SkeletalMesh = Cast<USkeletalMesh>(MutablePrivate::LoadObject(Mesh));
 	if (!SkeletalMesh)
 	{
 		return nullptr;
diff --git a/Engine/Plugins/Mutable/Source/CustomizableObjectEditor/Private/MuCOE/Nodes/CustomizableObjectNodeExternalPinDetails.cpp b/Engine/Plugins/Mutable/Source/CustomizableObjectEditor/Private/MuCOE/Nodes/CustomizableObjectNodeExternalPinDetails.cpp
--- a/Engine/Plugins/Mutable/Source/CustomizableObjectEditor/Private/MuCOE/Nodes/CustomizableObjectNodeExternalPinDetails.cpp
+++ b/Engine/Plugins/Mutable/Source/CustomizableObjectEditor/Private/MuCOE/Nodes/CustomizableObjectNodeExternalPinDetails.cpp
@@ -11,6 +11,7 @@
 #include "PropertyCustomizationHelpers.h"
 #include "ScopedTransaction.h"
 #include "MuCO/CustomizableObjectPrivate.h"
+#include "MuCO/LoadUtils.h"
 #include "Widgets/Input/STextComboBox.h"
 
 
@@ -106,7 +107,7 @@ void FCustomizableObjectNodeExternalPinDetails::ParentObjectSelectionChanged(con
 {
 	const FScopedTransaction Transaction(LOCTEXT("ChangedImportPinCustomizableObjectTransaction", "Changed Import Pin Customizable Object"));
 	Node->Modify();
-	Node->ExternalObject = Cast<UCustomizableObject>(AssetData.GetAsset());
+	Node->ExternalObject = Cast<UCustomizableObject>(MutablePrivate::LoadObject(AssetData));
 	Node->SetExternalObjectNodeId(FGuid());
 
 	if (DetailBuilderPtr)
diff --git a/Engine/Plugins/Mutable/Source/CustomizableObjectEditor/Private/MuCOE/Nodes/CustomizableObjectNodeGroupProjectorParameter.cpp b/Engine/Plugins/Mutable/Source/CustomizableObjectEditor/Private/MuCOE/Nodes/CustomizableObjectNodeGroupProjectorParameter.cpp
--- a/Engine/Plugins/Mutable/Source/CustomizableObjectEditor/Private/MuCOE/Nodes/CustomizableObjectNodeGroupProjectorParameter.cpp
+++ b/Engine/Plugins/Mutable/Source/CustomizableObjectEditor/Private/MuCOE/Nodes/CustomizableObjectNodeGroupProjectorParameter.cpp
@@ -3,9 +3,9 @@
 #include "MuCOE/Nodes/CustomizableObjectNodeGroupProjectorParameter.h"
 
 #include "MuCO/CustomizableObjectCustomVersion.h"
+#include "MuCO/LoadUtils.h"
 #include "MuCOE/EdGraphSchema_CustomizableObject.h"
 #include "MuCOE/UnrealEditorPortabilityHelpers.h"
-#include "UObject/ICookInfo.h"
 
 
 class UCustomizableObjectNodeRemapPins;
@@ -34,14 +34,13 @@ TArray<FGroupProjectorParameterImage> UCustomizableObjectNodeGroupProjectorParam
 
 	int32 NameIndex = 0;
 
-	FCookLoadScope CookLoadScope(ECookLoadType::EditorOnly);
 	for (TMap<FName, uint8*>::TConstIterator RowIt = OptionTexturesDataTable->GetRowMap().CreateConstIterator(); RowIt; ++RowIt)
 	{
 		uint8* RowData = RowIt.Value();
 		FString PropertyValue(TEXT(""));
 		PropertyTexturePath->ExportText_InContainer(0, PropertyValue, RowData, RowData, nullptr, PPF_None);
-		UTexture2D* Texture = LoadObject<UTexture2D>(nullptr, *PropertyValue, nullptr);
-
+		UTexture2D* Texture = MutablePrivate::LoadObject<UTexture2D>(nullptr, *PropertyValue);
+		
 		if (Texture == nullptr)
 		{
 			UE_LOG(LogMutable, Warning, TEXT("WARNING: Unable to load texture %s"), *PropertyValue);
diff --git a/Engine/Plugins/Mutable/Source/CustomizableObjectEditor/Private/MuCOE/Nodes/CustomizableObjectNodeMeshMorphDetails.cpp b/Engine/Plugins/Mutable/Source/CustomizableObjectEditor/Private/MuCOE/Nodes/CustomizableObjectNodeMeshMorphDetails.cpp
--- a/Engine/Plugins/Mutable/Source/CustomizableObjectEditor/Private/MuCOE/Nodes/CustomizableObjectNodeMeshMorphDetails.cpp
+++ b/Engine/Plugins/Mutable/Source/CustomizableObjectEditor/Private/MuCOE/Nodes/CustomizableObjectNodeMeshMorphDetails.cpp
@@ -52,7 +52,7 @@ void FCustomizableObjectNodeMeshMorphDetails::CustomizeDetails( IDetailLayoutBui
 				USkeletalMesh* SkeletalMesh = nullptr;
 				if ( const UCustomizableObjectNodeSkeletalMesh* TypedNodeSkeletalMesh = Cast<UCustomizableObjectNodeSkeletalMesh>(BaseSourcePin->GetOwningNode()) )
 				{
-					SkeletalMesh = TypedNodeSkeletalMesh->SkeletalMesh.LoadSynchronous();
+					SkeletalMesh = MutablePrivate::LoadObject(TypedNodeSkeletalMesh->SkeletalMesh);
 				}
 				else if (const UCustomizableObjectNodeTable* TypedNodeTable = Cast<UCustomizableObjectNodeTable>(BaseSourcePin->GetOwningNode()))
 				{
diff --git a/Engine/Plugins/Mutable/Source/CustomizableObjectEditor/Private/MuCOE/Nodes/CustomizableObjectNodeMeshMorphStackApplication.cpp b/Engine/Plugins/Mutable/Source/CustomizableObjectEditor/Private/MuCOE/Nodes/CustomizableObjectNodeMeshMorphStackApplication.cpp
--- a/Engine/Plugins/Mutable/Source/CustomizableObjectEditor/Private/MuCOE/Nodes/CustomizableObjectNodeMeshMorphStackApplication.cpp
+++ b/Engine/Plugins/Mutable/Source/CustomizableObjectEditor/Private/MuCOE/Nodes/CustomizableObjectNodeMeshMorphStackApplication.cpp
@@ -82,7 +82,7 @@ TArray<FString> UCustomizableObjectNodeMeshMorphStackApplication::GetMorphList()
 
 		if (const UCustomizableObjectNodeSkeletalMesh* SkeletalMeshNode = Cast< UCustomizableObjectNodeSkeletalMesh >(MeshNodePin->GetOwningNode()))
 		{
-			SkeletalMesh = SkeletalMeshNode->SkeletalMesh.LoadSynchronous();
+			SkeletalMesh = MutablePrivate::LoadObject(SkeletalMeshNode->SkeletalMesh);
 		}
 		else if (const UCustomizableObjectNodeTable* TableNode = Cast< UCustomizableObjectNodeTable >(MeshNodePin->GetOwningNode()))
 		{
diff --git a/Engine/Plugins/Mutable/Source/CustomizableObjectEditor/Private/MuCOE/Nodes/CustomizableObjectNodeMeshMorphStackDefinition.cpp b/Engine/Plugins/Mutable/Source/CustomizableObjectEditor/Private/MuCOE/Nodes/CustomizableObjectNodeMeshMorphStackDefinition.cpp
--- a/Engine/Plugins/Mutable/Source/CustomizableObjectEditor/Private/MuCOE/Nodes/CustomizableObjectNodeMeshMorphStackDefinition.cpp
+++ b/Engine/Plugins/Mutable/Source/CustomizableObjectEditor/Private/MuCOE/Nodes/CustomizableObjectNodeMeshMorphStackDefinition.cpp
@@ -101,7 +101,7 @@ void UCustomizableObjectNodeMeshMorphStackDefinition::UpdateMorphList()
 
 		if (const UCustomizableObjectNodeSkeletalMesh* SkeletalMeshNode = Cast< UCustomizableObjectNodeSkeletalMesh >(MeshNodePin->GetOwningNode()))
 		{
-			SkeletalMesh = SkeletalMeshNode->SkeletalMesh.LoadSynchronous();
+			SkeletalMesh = MutablePrivate::LoadObject(SkeletalMeshNode->SkeletalMesh);
 		}
 		else if (const UCustomizableObjectNodeTable* TableNode = Cast< UCustomizableObjectNodeTable >(MeshNodePin->GetOwningNode()))
 		{
@@ -190,7 +190,7 @@ bool UCustomizableObjectNodeMeshMorphStackDefinition::IsNodeOutDatedAndNeedsRefr
 
 		if (const UCustomizableObjectNodeSkeletalMesh* SkeletalMeshNode = Cast< UCustomizableObjectNodeSkeletalMesh >(MeshNodePin->GetOwningNode()))
 		{
-			SkeletalMesh = SkeletalMeshNode->SkeletalMesh.LoadSynchronous();
+			SkeletalMesh = MutablePrivate::LoadObject(SkeletalMeshNode->SkeletalMesh);
 		}
 		else if (const UCustomizableObjectNodeTable* TableNode = Cast< UCustomizableObjectNodeTable >(MeshNodePin->GetOwningNode()))
 		{
diff --git a/Engine/Plugins/Mutable/Source/CustomizableObjectEditor/Private/MuCOE/Nodes/CustomizableObjectNodeMeshReshapeSelectionDetails.cpp b/Engine/Plugins/Mutable/Source/CustomizableObjectEditor/Private/MuCOE/Nodes/CustomizableObjectNodeMeshReshapeSelectionDetails.cpp
--- a/Engine/Plugins/Mutable/Source/CustomizableObjectEditor/Private/MuCOE/Nodes/CustomizableObjectNodeMeshReshapeSelectionDetails.cpp
+++ b/Engine/Plugins/Mutable/Source/CustomizableObjectEditor/Private/MuCOE/Nodes/CustomizableObjectNodeMeshReshapeSelectionDetails.cpp
@@ -90,7 +90,7 @@ USkeletalMesh* FMeshReshapeBonesReferenceCustomization::GetSkeletalMesh(TSharedP
 				{
 					if (const UCustomizableObjectNodeSkeletalMesh* SkeletalMeshNode = Cast<UCustomizableObjectNodeSkeletalMesh>(SourceMeshPin->GetOwningNode()))
 					{
-						return SkeletalMeshNode->SkeletalMesh.LoadSynchronous();
+						return MutablePrivate::LoadObject(SkeletalMeshNode->SkeletalMesh);
 					}
 					else if (UCustomizableObjectNodeTable* TableNode = Cast<UCustomizableObjectNodeTable>(SourceMeshPin->GetOwningNode()))
 					{
diff --git a/Engine/Plugins/Mutable/Source/CustomizableObjectEditor/Private/MuCOE/Nodes/CustomizableObjectNodeModifierMorphMeshSectionDetails.cpp b/Engine/Plugins/Mutable/Source/CustomizableObjectEditor/Private/MuCOE/Nodes/CustomizableObjectNodeModifierMorphMeshSectionDetails.cpp
--- a/Engine/Plugins/Mutable/Source/CustomizableObjectEditor/Private/MuCOE/Nodes/CustomizableObjectNodeModifierMorphMeshSectionDetails.cpp
+++ b/Engine/Plugins/Mutable/Source/CustomizableObjectEditor/Private/MuCOE/Nodes/CustomizableObjectNodeModifierMorphMeshSectionDetails.cpp
@@ -20,6 +20,7 @@
 #include "IDetailsView.h"
 #include "Engine/SkeletalMesh.h"
 #include "Misc/Attribute.h"
+#include "MuCO/LoadUtils.h"
 
 
 #define LOCTEXT_NAMESPACE "CustomizableObjectDetails"
@@ -277,7 +278,7 @@ void FCustomizableObjectNodeModifierMorphMeshSectionDetails::AddMorphsFromNode(
 				UCustomizableObjectNodeSkeletalMesh* SkeletalNode = Cast<UCustomizableObjectNodeSkeletalMesh>(SourceMeshPin->GetOwningNode());
 				if (SkeletalNode)
 				{
-					SkeletalMesh = SkeletalNode->SkeletalMesh.LoadSynchronous();
+					SkeletalMesh = MutablePrivate::LoadObject(SkeletalNode->SkeletalMesh);
 				}
 			}
 		}
@@ -293,7 +294,7 @@ void FCustomizableObjectNodeModifierMorphMeshSectionDetails::AddMorphsFromNode(
 				UCustomizableObjectNodeSkeletalMesh* SkeletalNode = Cast<UCustomizableObjectNodeSkeletalMesh>(SourceMeshPin->GetOwningNode());
 				if (SkeletalNode)
 				{
-					SkeletalMesh = SkeletalNode->SkeletalMesh.LoadSynchronous();
+					SkeletalMesh = MutablePrivate::LoadObject(SkeletalNode->SkeletalMesh);
 				}
 			}
 		}
diff --git a/Engine/Plugins/Mutable/Source/CustomizableObjectEditor/Private/MuCOE/Nodes/CustomizableObjectNodeObjectDetails.cpp b/Engine/Plugins/Mutable/Source/CustomizableObjectEditor/Private/MuCOE/Nodes/CustomizableObjectNodeObjectDetails.cpp
--- a/Engine/Plugins/Mutable/Source/CustomizableObjectEditor/Private/MuCOE/Nodes/CustomizableObjectNodeObjectDetails.cpp
+++ b/Engine/Plugins/Mutable/Source/CustomizableObjectEditor/Private/MuCOE/Nodes/CustomizableObjectNodeObjectDetails.cpp
@@ -308,7 +308,7 @@ void FCustomizableObjectNodeObjectDetails::ParentObjectSelectionChanged(const FA
 {
 	if (BaseObjectNode.IsValid())
 	{
-		UCustomizableObject* Parent = Cast<UCustomizableObject>(AssetData.GetAsset());
+		UCustomizableObject* Parent = Cast<UCustomizableObject>(MutablePrivate::LoadObject(AssetData));
 		BaseObjectNode->SetParentObject(Parent);
 
 		// If set the parent to nullt, invalidate also the reference GUID
diff --git a/Engine/Plugins/Mutable/Source/CustomizableObjectEditor/Private/MuCOE/Nodes/CustomizableObjectNodeSkeletalMesh.cpp b/Engine/Plugins/Mutable/Source/CustomizableObjectEditor/Private/MuCOE/Nodes/CustomizableObjectNodeSkeletalMesh.cpp
--- a/Engine/Plugins/Mutable/Source/CustomizableObjectEditor/Private/MuCOE/Nodes/CustomizableObjectNodeSkeletalMesh.cpp
+++ b/Engine/Plugins/Mutable/Source/CustomizableObjectEditor/Private/MuCOE/Nodes/CustomizableObjectNodeSkeletalMesh.cpp
@@ -7,6 +7,7 @@
 #include "ISinglePropertyView.h"
 #include "Modules/ModuleManager.h"
 #include "MuCO/CustomizableObjectCustomVersion.h"
+#include "MuCO/LoadUtils.h"
 #include "MuCOE/CustomizableObjectEditor.h"
 #include "MuCOE/CustomizableObjectEditorStyle.h"
 #include "MuCOE/CustomizableObjectEditorUtilities.h"
@@ -70,7 +71,7 @@ void UCustomizableObjectNodeSkeletalMesh::AllocateDefaultPins(UCustomizableObjec
 	const UEdGraphSchema_CustomizableObject* Schema = GetDefault<UEdGraphSchema_CustomizableObject>();
 
 	// Force the loading of the skeletal mesh so we can later access the internal WeakObjectPtr (may be null if no SKM is set)
-	SkeletalMesh.LoadSynchronous();
+	MutablePrivate::LoadObject(SkeletalMesh);
 	
 	if (!SkeletalMesh)
 	{
@@ -845,7 +846,7 @@ void SGraphNodeSkeletalMesh::Construct(const FArguments& InArgs, UEdGraphNode* I
 
 	// Thumbnail
 	AssetThumbnailPool = MakeShareable(new FAssetThumbnailPool(32));
-	USkeletalMesh* SkeletalMesh = NodeSkeletalMesh->SkeletalMesh.LoadSynchronous();
+	USkeletalMesh* SkeletalMesh = MutablePrivate::LoadObject(NodeSkeletalMesh->SkeletalMesh);
 	AssetThumbnail = MakeShareable(new FAssetThumbnail(SkeletalMesh, ThumbnailSize, ThumbnailSize, AssetThumbnailPool));
 
 	// Selector
diff --git a/Engine/Plugins/Mutable/Source/CustomizableObjectEditor/Private/MuCOE/Nodes/CustomizableObjectNodeTable.cpp b/Engine/Plugins/Mutable/Source/CustomizableObjectEditor/Private/MuCOE/Nodes/CustomizableObjectNodeTable.cpp
--- a/Engine/Plugins/Mutable/Source/CustomizableObjectEditor/Private/MuCOE/Nodes/CustomizableObjectNodeTable.cpp
+++ b/Engine/Plugins/Mutable/Source/CustomizableObjectEditor/Private/MuCOE/Nodes/CustomizableObjectNodeTable.cpp
@@ -407,7 +407,6 @@ void UCustomizableObjectNodeTable::AllocateDefaultPins(UCustomizableObjectNodeRe
 	TSet<FString> PinNames;
 	bool bHasRepeatedNames = false;
 	
-	FCookLoadScope CookLoadScope(ECookLoadType::EditorOnly);
 	for (TFieldIterator<FProperty> It(TableStruct); It; ++It)
 	{
 		FProperty* ColumnProperty = *It;
@@ -450,7 +449,7 @@ void UCustomizableObjectNodeTable::AllocateDefaultPins(UCustomizableObjectNodeRe
 
 				if (CellData)
 				{
-					Object = SoftObjectProperty->GetPropertyValue(CellData).LoadSynchronous();
+					Object = MutablePrivate::LoadObject(SoftObjectProperty->GetPropertyValue(CellData));
 				}
 
 				if (Object)
@@ -713,7 +712,6 @@ bool UCustomizableObjectNodeTable::IsNodeOutDatedAndNeedsRefresh()
 
 	bool bNeedsUpdate = false;
 
-	FCookLoadScope CookLoadScope(ECookLoadType::EditorOnly);
 	for (TFieldIterator<FProperty> It(TableStruct); It && !bNeedsUpdate; ++It)
 	{
 		FProperty* ColumnProperty = *It;
@@ -732,7 +730,7 @@ bool UCustomizableObjectNodeTable::IsNodeOutDatedAndNeedsRefresh()
 
 				if (CellData)
 				{
-					Object = SoftObjectProperty->GetPropertyValue(CellData).LoadSynchronous();
+					Object = MutablePrivate::LoadObject(SoftObjectProperty->GetPropertyValue(CellData));
 				}
 
 				if (!Object)
diff --git a/Engine/Plugins/Mutable/Source/CustomizableObjectEditor/Private/MuCOE/Nodes/CustomizableObjectNodeTable.h b/Engine/Plugins/Mutable/Source/CustomizableObjectEditor/Private/MuCOE/Nodes/CustomizableObjectNodeTable.h
--- a/Engine/Plugins/Mutable/Source/CustomizableObjectEditor/Private/MuCOE/Nodes/CustomizableObjectNodeTable.h
+++ b/Engine/Plugins/Mutable/Source/CustomizableObjectEditor/Private/MuCOE/Nodes/CustomizableObjectNodeTable.h
@@ -4,9 +4,9 @@
 
 #include "Engine/DataTable.h"
 #include "Engine/Texture2DArray.h"
+#include "MuCO/LoadUtils.h"
 #include "MuCOE/Nodes/CustomizableObjectNode.h"
 #include "MuCOE/RemapPins/CustomizableObjectNodeRemapPinsByName.h"
-#include "UObject/ICookInfo.h"
 
 #include "CustomizableObjectNodeTable.generated.h"
 
@@ -323,8 +323,7 @@ public:
 
 					if (CellData)
 					{
-						FCookLoadScope CookLoadScope(ECookLoadType::EditorOnly);
-						Object = SoftObjectProperty->GetPropertyValue(CellData).LoadSynchronous();
+						Object = MutablePrivate::LoadObject(SoftObjectProperty->GetPropertyValue(CellData));
 					}
 
 					if (Object)
diff --git a/Engine/Plugins/Mutable/Source/CustomizableObjectEditor/Private/MuCOE/SCustomizableObjectEditorAdvancedPreviewSettings.cpp b/Engine/Plugins/Mutable/Source/CustomizableObjectEditor/Private/MuCOE/SCustomizableObjectEditorAdvancedPreviewSettings.cpp
--- a/Engine/Plugins/Mutable/Source/CustomizableObjectEditor/Private/MuCOE/SCustomizableObjectEditorAdvancedPreviewSettings.cpp
+++ b/Engine/Plugins/Mutable/Source/CustomizableObjectEditor/Private/MuCOE/SCustomizableObjectEditorAdvancedPreviewSettings.cpp
@@ -7,6 +7,7 @@
 #include "IContentBrowserSingleton.h"
 #include "Misc/MessageDialog.h"
 #include "Misc/Paths.h"
+#include "MuCO/LoadUtils.h"
 #include "MuCOE/CustomizableObjectEditor.h"
 #include "MuCOE/CustomizableObjectEditorLevelProfile.h"
 #include "Widgets/Input/SButton.h"
@@ -481,7 +482,7 @@ void SCustomizableObjectEditorAdvancedPreviewSettings::LoadProfileEnvironment()
 		if (!DefaultSettings->Profiles[i].EnvironmentCubeMapPath.IsEmpty())
 		{
 			UObject* LoadedObject = nullptr;
-			LoadedObject = LoadObject<UObject>(nullptr, *(DefaultSettings->Profiles[i].EnvironmentCubeMapPath));
+			LoadedObject = MutablePrivate::LoadObject<UObject>(nullptr, *DefaultSettings->Profiles[i].EnvironmentCubeMapPath);
 			while (UObjectRedirector* Redirector = Cast<UObjectRedirector>(LoadedObject))
 			{
 				LoadedObject = Redirector->DestinationObject;
diff --git a/Engine/Plugins/Mutable/Source/CustomizableObjectEditor/Private/MuCOE/SCustomizableObjectEditorViewport.cpp b/Engine/Plugins/Mutable/Source/CustomizableObjectEditor/Private/MuCOE/SCustomizableObjectEditorViewport.cpp
--- a/Engine/Plugins/Mutable/Source/CustomizableObjectEditor/Private/MuCOE/SCustomizableObjectEditorViewport.cpp
+++ b/Engine/Plugins/Mutable/Source/CustomizableObjectEditor/Private/MuCOE/SCustomizableObjectEditorViewport.cpp
@@ -27,6 +27,7 @@
 #include "AssetRegistry/AssetRegistryModule.h"
 #include "MuCO/CustomizableObjectPrivate.h"
 #include "MuCO/CustomizableObjectSystem.h"
+#include "MuCO/LoadUtils.h"
 #include "Widgets/Input/STextComboBox.h"
 
 class UCustomizableObject;
@@ -1028,7 +1029,7 @@ FReply SCustomizableObjectEditorViewportTabBody::OnDrop(const FGeometry& MyGeome
 		if (DragDropOp->GetAssets().Num())
 		{
 			// This cast also includes UPoseAsset assets.
-			UAnimationAsset* AnimationAsset = Cast<UAnimationAsset>(DragDropOp->GetAssets()[0].GetAsset());
+			UAnimationAsset* AnimationAsset = Cast<UAnimationAsset>(MutablePrivate::LoadObject(DragDropOp->GetAssets()[0]));
 			if (AnimationAsset)
 			{
 				FObjectEditorUtils::SetPropertyValue(WeakEditor.Pin()->GetCustomSettings(), GET_MEMBER_NAME_CHECKED(UCustomSettings, Animation), AnimationAsset);
diff --git a/Engine/Plugins/Mutable/Source/CustomizableObjectEditor/Private/MuCOE/SCustomizableObjectLayoutGrid.cpp b/Engine/Plugins/Mutable/Source/CustomizableObjectEditor/Private/MuCOE/SCustomizableObjectLayoutGrid.cpp
--- a/Engine/Plugins/Mutable/Source/CustomizableObjectEditor/Private/MuCOE/SCustomizableObjectLayoutGrid.cpp
+++ b/Engine/Plugins/Mutable/Source/CustomizableObjectEditor/Private/MuCOE/SCustomizableObjectLayoutGrid.cpp
@@ -17,6 +17,7 @@
 #include "Engine/Texture2D.h"
 #include "TextureResource.h"
 #include "RenderGraphBuilder.h"
+#include "MuCO/LoadUtils.h"
 
 class FExtender;
 class FPaintArgs;
@@ -555,7 +556,7 @@ void SCustomizableObjectLayoutGrid::CloseMenu()
 
 void SCustomizableObjectLayoutGrid::OnMaskAssetSelected(const FAssetData& AssetData)
 {
-	UTexture2D* Mask = Cast<UTexture2D>(AssetData.GetAsset());
+	UTexture2D* Mask = Cast<UTexture2D>(MutablePrivate::LoadObject(AssetData));
 	if (SelectedBlocks.Num())
 	{
 		OnSetBlockMask.ExecuteIfBound(Mask);
diff --git a/Engine/Plugins/Mutable/Source/CustomizableObjectEditor/Private/MuCOE/SCustomizableObjectNodeSkeletalMeshRTMorphSelector.cpp b/Engine/Plugins/Mutable/Source/CustomizableObjectEditor/Private/MuCOE/SCustomizableObjectNodeSkeletalMeshRTMorphSelector.cpp
--- a/Engine/Plugins/Mutable/Source/CustomizableObjectEditor/Private/MuCOE/SCustomizableObjectNodeSkeletalMeshRTMorphSelector.cpp
+++ b/Engine/Plugins/Mutable/Source/CustomizableObjectEditor/Private/MuCOE/SCustomizableObjectNodeSkeletalMeshRTMorphSelector.cpp
@@ -6,6 +6,7 @@
 #include "Framework/Views/TableViewMetadata.h"
 #include "MuCOE/Nodes/CustomizableObjectNodeSkeletalMesh.h"
 #include "ScopedTransaction.h"
+#include "MuCO/LoadUtils.h"
 #include "Widgets/Input/SCheckBox.h"
 #include "Widgets/Input/SSearchBox.h"
 
@@ -188,8 +189,7 @@ void SCustomizableObjectNodeSkeletalMeshRTMorphSelector::GenerateMorphsInfoList(
 {
     Items.Reset();
 
-    const USkeletalMesh* SkeletalMesh = Node->SkeletalMesh.LoadSynchronous();
-
+    const USkeletalMesh* SkeletalMesh = MutablePrivate::LoadObject(Node->SkeletalMesh);
     if (!SkeletalMesh)
     {
         return;
diff --git a/Engine/Plugins/Mutable/Source/CustomizableObjectEditor/Private/MuCOE/SStandAloneAssetPicker.cpp b/Engine/Plugins/Mutable/Source/CustomizableObjectEditor/Private/MuCOE/SStandAloneAssetPicker.cpp
--- a/Engine/Plugins/Mutable/Source/CustomizableObjectEditor/Private/MuCOE/SStandAloneAssetPicker.cpp
+++ b/Engine/Plugins/Mutable/Source/CustomizableObjectEditor/Private/MuCOE/SStandAloneAssetPicker.cpp
@@ -13,6 +13,7 @@
 #include "Widgets/Input/SButton.h"
 #include "IContentBrowserSingleton.h"
 #include "ContentBrowserModule.h"
+#include "MuCO/LoadUtils.h"
 
 
 #define LOCTEXT_NAMESPACE "CustomizableObjectEditor"
@@ -153,7 +154,7 @@ void SStandAloneAssetPicker::OnAssetSelectedFromPicker(const struct FAssetData&
 {
 	AssetPickerAnchor->SetIsOpen(false);
 
-	CurrentAsset = AssetData.GetAsset();
+	CurrentAsset = MutablePrivate::LoadObject(AssetData);
 
 	OnAssetSelected.ExecuteIfBound(CurrentAsset.Get());
 
@@ -167,7 +168,7 @@ void SStandAloneAssetPicker::OnAssetEnterPressedFromPicker(const TArray<struct F
 	CurrentAsset = nullptr;
 	if (AssetData.Num() > 0)
 	{
-		CurrentAsset = AssetData[0].GetAsset();
+		CurrentAsset = MutablePrivate::LoadObject(AssetData[0]);
 	}
 
 	OnAssetSelected.ExecuteIfBound(CurrentAsset.Get());
diff --git a/Engine/Plugins/Mutable/Source/MutableValidation/Private/MuV/AssetValidator_ReferencedCustomizableObjects.cpp b/Engine/Plugins/Mutable/Source/MutableValidation/Private/MuV/AssetValidator_ReferencedCustomizableObjects.cpp
--- a/Engine/Plugins/Mutable/Source/MutableValidation/Private/MuV/AssetValidator_ReferencedCustomizableObjects.cpp
+++ b/Engine/Plugins/Mutable/Source/MutableValidation/Private/MuV/AssetValidator_ReferencedCustomizableObjects.cpp
@@ -11,6 +11,7 @@
 #include "Engine/Texture.h"
 #include "Materials/Material.h"
 #include "MuCO/CustomizableObject.h"
+#include "MuCO/LoadUtils.h"
 #include "MuV/AssetValidator_CustomizableObjects.h"
 #include "UObject/NameTypes.h"
 #include "UObject/Object.h"
@@ -123,7 +124,7 @@ TSet< UCustomizableObject*> UAssetValidator_ReferencedCustomizableObjects::FindC
         	// We have found a referenced CustomizableObject
             if (ReferencerAssetData.GetClass() == UCustomizableObject::StaticClass())
             {
-				UObject* ReferencedAsset = ReferencerAssetData.GetAsset();
+				UObject* ReferencedAsset = MutablePrivate::LoadObject(ReferencerAssetData);
             	if (ReferencedAsset)
             	{
             		UCustomizableObject* CastedCustomizableObject = Cast<UCustomizableObject>(ReferencedAsset);
diff --git a/Engine/Plugins/Mutable/Source/MutableValidation/Private/MuV/COIBakingTestCommandlet.cpp b/Engine/Plugins/Mutable/Source/MutableValidation/Private/MuV/COIBakingTestCommandlet.cpp
--- a/Engine/Plugins/Mutable/Source/MutableValidation/Private/MuV/COIBakingTestCommandlet.cpp
+++ b/Engine/Plugins/Mutable/Source/MutableValidation/Private/MuV/COIBakingTestCommandlet.cpp
@@ -14,6 +14,7 @@
 #include "Interfaces/ITargetPlatformManagerModule.h"
 #include "MuCO/CustomizableObjectPrivate.h"
 #include "MuCO/CustomizableObjectSystemPrivate.h"
+#include "MuCO/LoadUtils.h"
 #include "MuCOE/CustomizableObjectBenchmarkingUtils.h"
 
 
@@ -55,7 +56,7 @@ int32 UCOIBakingTestCommandlet::Main(const FString& Params)
 		}
 	
 		// Load the resource
-		UObject* FoundObject = FSoftObjectPath(CustomizableObjectInstanceAssetPath).TryLoad();
+		UObject* FoundObject = MutablePrivate::LoadObject(FSoftObjectPath(CustomizableObjectInstanceAssetPath));
 		if (!FoundObject)
 		{
 			UE_LOG(LogMutable,Error,TEXT("Failed to retrieve UObject from path %s"),*CustomizableObjectInstanceAssetPath);
diff --git a/Engine/Plugins/Mutable/Source/MutableValidation/Private/MuV/COIBulkUpdateTestCommandlet.cpp b/Engine/Plugins/Mutable/Source/MutableValidation/Private/MuV/COIBulkUpdateTestCommandlet.cpp
--- a/Engine/Plugins/Mutable/Source/MutableValidation/Private/MuV/COIBulkUpdateTestCommandlet.cpp
+++ b/Engine/Plugins/Mutable/Source/MutableValidation/Private/MuV/COIBulkUpdateTestCommandlet.cpp
@@ -13,6 +13,7 @@
 #include "MuCO/CustomizableObjectPrivate.h"
 #include "MuCO/CustomizableObjectSystem.h"
 #include "MuCO/CustomizableObjectSystemPrivate.h"
+#include "MuCO/LoadUtils.h"
 #include "MuCOE/CustomizableObjectBenchmarkingUtils.h"
 
 
@@ -74,7 +75,7 @@ int32 UCOIBulkUpdateTestCommandlet::Main(const FString& Params)
 	TMap<UCustomizableObject*, TArray<TStrongObjectPtr<UCustomizableObjectInstance>>>  MutableResources;
 	for (FAssetData& Data : FoundAssetData)
 	{
-		UObject* LoadedAsset =  Data.GetAsset();
+		UObject* LoadedAsset =  MutablePrivate::LoadObject(Data);
 		check (LoadedAsset);
 		UCustomizableObjectInstance* LoadedInstance =  Cast<UCustomizableObjectInstance>(LoadedAsset);
 		check (LoadedInstance)
diff --git a/Engine/Plugins/Mutable/Source/MutableValidation/Private/MuV/CustomizableObjectValidationCommandlet.cpp b/Engine/Plugins/Mutable/Source/MutableValidation/Private/MuV/CustomizableObjectValidationCommandlet.cpp
--- a/Engine/Plugins/Mutable/Source/MutableValidation/Private/MuV/CustomizableObjectValidationCommandlet.cpp
+++ b/Engine/Plugins/Mutable/Source/MutableValidation/Private/MuV/CustomizableObjectValidationCommandlet.cpp
@@ -17,6 +17,7 @@
 #include "Interfaces/ITargetPlatformManagerModule.h"
 #include "Interfaces/ITargetPlatform.h"
 #include "MuCO/CustomizableObjectSystemPrivate.h"
+#include "MuCO/LoadUtils.h"
 #include "MuCOE/CustomizableObjectBenchmarkingUtils.h"
 #include "Serialization/MemoryWriter.h"
 
@@ -51,7 +52,7 @@ int32 UCustomizableObjectValidationCommandlet::Main(const FString& Params)
 	}
 	
 	// Load the resource
-	UObject* FoundObject = FSoftObjectPath(CustomizableObjectAssetPath).TryLoad();
+	UObject* FoundObject = MutablePrivate::LoadObject(FSoftObjectPath(CustomizableObjectAssetPath));
 	if (!FoundObject)
 	{
 		UE_LOG(LogMutable,Error,TEXT("Failed to retrieve UObject from path %s"),*CustomizableObjectAssetPath);
