[mutable] Enable mesh conversion during core compilation stage instead of generation stage.
Ensure the core code generation step is run as tasks, so that mesh load requests don't lock the core thread.
Move "compilation state data" like LOD number from the CodeGenerator struct to the structs that are passed to each task.
Split mesh generation options struct into a "static" one that doesn't change during the generation recursion and a "dynamic" one that may change based on sibling op generation.
Remove core mesh generation cache for ops (because of concurrency). The constant mesh cache does most of the work anyway.
Remove morph support into NodeMeshConstants: now morphs are provided on demand instead of ahead of time during generation.
Generate switching parameters for tables by creating a switching node instead of directly creating the ASTOps, to avoid duplicating code.
#rnx

diff --git a/Engine/Plugins/Mutable/Source/CustomizableObject/Internal/MuCO/CustomizableObjectPrivate.h b/Engine/Plugins/Mutable/Source/CustomizableObject/Internal/MuCO/CustomizableObjectPrivate.h
--- a/Engine/Plugins/Mutable/Source/CustomizableObject/Internal/MuCO/CustomizableObjectPrivate.h
+++ b/Engine/Plugins/Mutable/Source/CustomizableObject/Internal/MuCO/CustomizableObjectPrivate.h
@@ -1176,7 +1176,9 @@ public:
 	UPROPERTY()
 	TArray<TSoftObjectPtr<const UTexture>> RuntimeReferencedTextures;
 	
-	/** Runtime referenced meshes used by the mu::FModel. */
+	/** Runtime referenced meshes used by the mu::FModel.
+	* TODO: Move to FMutableSourceMeshData when runtime is really implemented.
+	*/
 	UPROPERTY()
 	TArray<TSoftObjectPtr<const UStreamableRenderAsset>> RuntimeReferencedMeshes;
 #endif
@@ -1822,6 +1824,8 @@ public:
 
 		RemoveObsoletImageGradientOp,
 
+		MeshReferencesExtendedForCompilation,
+
 		// -----<new versions can be added above this line>--------
 		LastCustomizableObjectVersion
 	};
diff --git a/Engine/Plugins/Mutable/Source/CustomizableObject/Internal/MuCO/CustomizableObjectSystemPrivate.h b/Engine/Plugins/Mutable/Source/CustomizableObject/Internal/MuCO/CustomizableObjectSystemPrivate.h
--- a/Engine/Plugins/Mutable/Source/CustomizableObject/Internal/MuCO/CustomizableObjectSystemPrivate.h
+++ b/Engine/Plugins/Mutable/Source/CustomizableObject/Internal/MuCO/CustomizableObjectSystemPrivate.h
@@ -823,6 +823,9 @@ public:
 	 * to do things like check-if-up-to-date or compile.
 	 */
 	void AddPendingLoad(UCustomizableObject*);
+
+	CUSTOMIZABLEOBJECT_API void AddGameThreadCompileTask(TFunction<void()>&& Task);
+
 #endif
 	
 	/** Mutable TaskGraph system (Mutable Thread). */
@@ -833,10 +836,13 @@ public:
 	UPROPERTY(Transient)
 	TObjectPtr<UEditorImageProvider> EditorImageProvider = nullptr;
 
-	/** List of objects pending to complete loading. */
+	/** List of CustomizableObjects pending to complete loading. */
 	UPROPERTY(Transient)
 	TArray<TObjectPtr<UCustomizableObject>> ObjectsPendingLoad;
 
+	/** */
+	TQueue< TFunction<void()>, EQueueMode::Mpsc > PendingGameThreadCompileTasks;
+
 #endif
 
 	FLogBenchmarkUtil LogBenchmarkUtil;
diff --git a/Engine/Plugins/Mutable/Source/CustomizableObject/Internal/MuCO/UnrealToMutableTextureConversionUtils.h b/Engine/Plugins/Mutable/Source/CustomizableObject/Internal/MuCO/UnrealToMutableTextureConversionUtils.h
--- a/Engine/Plugins/Mutable/Source/CustomizableObject/Internal/MuCO/UnrealToMutableTextureConversionUtils.h
+++ b/Engine/Plugins/Mutable/Source/CustomizableObject/Internal/MuCO/UnrealToMutableTextureConversionUtils.h
@@ -63,7 +63,7 @@ ENUM_CLASS_FLAGS(EMutableMeshConversionFlags)
 struct CUSTOMIZABLEOBJECT_API FMutableSourceMeshData
 {
 	/** Assets involved in the conversion. */
-	TSoftObjectPtr<const UStreamableRenderAsset> Mesh;
+	TSoftObjectPtr<UStreamableRenderAsset> Mesh;
 	TSoftClassPtr<UAnimInstance> AnimInstance;
 	TSoftObjectPtr<const USkeletalMesh> TableReferenceSkeletalMesh;
 	FName Component = NAME_None;
@@ -83,9 +83,21 @@ struct CUSTOMIZABLEOBJECT_API FMutableSourceMeshData
 
 	/** Context for log messages. */
 	const UObject* MessageContext = nullptr;
+
+	bool operator==(const FMutableSourceMeshData&) const = default;
 };
 
 
+inline uint32 GetTypeHash(const FMutableSourceMeshData& Key)
+{
+	uint32 GuidHash = GetTypeHash(Key.Mesh);
+	GuidHash = HashCombineFast(GuidHash, GetTypeHash(Key.LODIndex));
+	GuidHash = HashCombineFast(GuidHash, GetTypeHash(Key.SectionIndex));
+	GuidHash = HashCombineFast(GuidHash, GetTypeHash(Key.Component));
+	GuidHash = HashCombineFast(GuidHash, GetTypeHash(Key.Flags));
+	return GuidHash;
+}
+
 enum class EUnrealToMutableConversionError 
 {
     Success,
diff --git a/Engine/Plugins/Mutable/Source/CustomizableObject/Private/MuCO/CustomizableObjectSystem.cpp b/Engine/Plugins/Mutable/Source/CustomizableObject/Private/MuCO/CustomizableObjectSystem.cpp
--- a/Engine/Plugins/Mutable/Source/CustomizableObject/Private/MuCO/CustomizableObjectSystem.cpp
+++ b/Engine/Plugins/Mutable/Source/CustomizableObject/Private/MuCO/CustomizableObjectSystem.cpp
@@ -1535,6 +1535,12 @@ void UCustomizableObjectSystemPrivate::AddPendingLoad(UCustomizableObject* CO)
 }
 
 
+void UCustomizableObjectSystemPrivate::AddGameThreadCompileTask(TFunction<void()>&& Task)
+{
+	PendingGameThreadCompileTasks.Enqueue(Task);
+}
+
+
 bool UCustomizableObjectSystem::LockObject(UCustomizableObject* InObject)
 {
 	check(InObject != nullptr);
@@ -3785,6 +3791,23 @@ int32 UCustomizableObjectSystem::TickInternal(const bool bBlocking)
 		return 0;
 	}
 
+	// See if there are compile-time tasks to run
+	constexpr double MaxSecondsPerFrame = 0.4;
+	double MaxTime = FPlatformTime::Seconds() + MaxSecondsPerFrame;
+
+	TFunction<void()> Task;
+	while (GetPrivate()->PendingGameThreadCompileTasks.Dequeue(Task))
+	{
+		Task();
+
+		// Simple time limit enforcement to avoid blocking the game thread if there are many requests.
+		double CurrentTime = FPlatformTime::Seconds();
+		if (CurrentTime >= MaxTime)
+		{
+			break;
+		}
+	}
+
 	// See if any COs pending to load can be completed.
 	for ( TArray<TObjectPtr<UCustomizableObject>>::TIterator It = GetPrivate()->ObjectsPendingLoad.CreateIterator(); It; ++It )
 	{
diff --git a/Engine/Plugins/Mutable/Source/CustomizableObjectEditor/Private/MuCOE/CustomizableObjectCompileRunnable.cpp b/Engine/Plugins/Mutable/Source/CustomizableObjectEditor/Private/MuCOE/CustomizableObjectCompileRunnable.cpp
--- a/Engine/Plugins/Mutable/Source/CustomizableObjectEditor/Private/MuCOE/CustomizableObjectCompileRunnable.cpp
+++ b/Engine/Plugins/Mutable/Source/CustomizableObjectEditor/Private/MuCOE/CustomizableObjectCompileRunnable.cpp
@@ -2,10 +2,13 @@
 
 #include "MuCOE/CustomizableObjectCompileRunnable.h"
 
+#include "MuCOE/GenerateMutableSource/GenerateMutableSource.h"
 #include "HAL/FileManager.h"
 #include "MuCO/UnrealMutableModelDiskStreamer.h"
 #include "MuCO/UnrealToMutableTextureConversionUtils.h"
 #include "MuCO/CustomizableObject.h"
+#include "MuCO/CustomizableObjectSystem.h"
+#include "MuCO/CustomizableObjectSystemPrivate.h"
 #include "MuR/Model.h"
 #include "MuT/Compiler.h"
 #include "MuT/ErrorLog.h"
@@ -16,6 +19,7 @@
 #include "Serialization/ObjectAndNameAsStringProxyArchive.h"
 #include "Trace/Trace.inl"
 #include "UObject/SoftObjectPtr.h"
+#include "UObject/ICookInfo.h"
 #include "Engine/AssetUserData.h"
 
 #include "DerivedDataCacheInterface.h"
@@ -48,8 +52,9 @@ TAutoConsoleVariable<bool> CVarMutableCompilerFastCompression(
 	ECVF_Default);
 
 
-FCustomizableObjectCompileRunnable::FCustomizableObjectCompileRunnable(mu::Ptr<mu::Node> Root)
+FCustomizableObjectCompileRunnable::FCustomizableObjectCompileRunnable(mu::Ptr<mu::Node> Root, TSharedPtr<FMutableCompilationContext> InCompilationContext)
 	: MutableRoot(Root)
+	, CompilationContext(InCompilationContext)
 	, bThreadCompleted(false)
 {
 	PrepareUnrealCompression();
@@ -87,44 +92,6 @@ TSharedPtr<mu::FImage> FCustomizableObjectCompileRunnable::LoadImageResourceRefe
 }
 
 
-TSharedPtr<mu::FMesh> FCustomizableObjectCompileRunnable::LoadMeshResourceReferenced(int32 ID)
-{
-	MUTABLE_CPUPROFILER_SCOPE(LoadResourceReferenced);
-
-	TSharedPtr<mu::FMesh> Mesh;
-
-	// TODO
-	check(false);
-
-	//if (!ReferencedTextures.IsValidIndex(ID))
-	//{
-	//	// The id is not valid for this CO
-	//	check(false);
-	//	return Mesh;
-	//}
-
-	//// Find the texture id
-	//FMutableSourceMeshData& MeshData = ReferencedMeshes[ID];
-	//check(IsInGameThread());
-
-	//// In the editor the src data can be directly accessed
-	////Mesh = MakeShared<mu::FMesh>();
-	////EUnrealToMutableConversionError Error = ConvertTextureUnrealSourceToMutable(Mesh.Get(), MeshData);
-
-	//Mesh = ConvertSkeletalMeshToMutable(MeshData, Context);
-
-
-	//if (Error != EUnrealToMutableConversionError::Success)
-	//{
-	//	// This could happen in the editor, because some source textures may have changed while there was a background compilation.
-	//	// We just show a warning and move on. This cannot happen during cooks, so it is fine.
-	//	UE_LOG(LogMutable, Warning, TEXT("Failed to load some source mesh data for mesh ID [%d]. Some meshes may be corrupted."), ID);
-	//}
-
-	return Mesh;
-}
-
-
 uint32 FCustomizableObjectCompileRunnable::Run()
 {
 	TRACE_BEGIN_REGION(UE_MUTABLE_CORE_REGION);
@@ -201,7 +168,7 @@ uint32 FCustomizableObjectCompileRunnable::Run()
 	CompilerOptions->SetReferencedResourceCallback(
 		[this](int32 ID, TSharedPtr<TSharedPtr<mu::FImage>> ResolvedImage, bool bRunImmediatlyIfPossible)
 		{
-			UE::Tasks::FTask LaunchTask = UE::Tasks::Launch(TEXT("ConstantGeneratorLaunchTasks"),
+			UE::Tasks::FTask LaunchTask = UE::Tasks::Launch(TEXT("LoadImageReferenceTasks"),
 				[ID,ResolvedImage,this]()
 				{
 					TSharedPtr<mu::FImage> Result = LoadImageResourceReferenced(ID);
@@ -213,18 +180,64 @@ uint32 FCustomizableObjectCompileRunnable::Run()
 			return LaunchTask;
 		},
 
-		[this](int32 ID, TSharedPtr<TSharedPtr<mu::FMesh>> ResolvedMesh, bool bRunImmediatlyIfPossible)
+		[this](int32 ID, const FString& MorphNameRef, TSharedPtr<TSharedPtr<mu::FMesh>> ResolvedMesh, bool bRunImmediatlyIfPossible) -> UE::Tasks::FTask
 		{
-			UE::Tasks::FTask LaunchTask = UE::Tasks::Launch(TEXT("ConstantGeneratorLaunchTasks"),
-			[ID, ResolvedMesh, this]()
+			MUTABLE_CPUPROFILER_SCOPE(LoadMeshReferenceTasks);
+
+			ResolvedMesh->Reset();
+
+			if (!ReferencedMeshes.IsValidIndex(ID))
+			{
+				// The id is not valid for this CO
+				check(false);
+				return UE::Tasks::MakeCompletedTask<void>();
+			}
+
+			UE::Tasks::FTaskEvent CompletionEvent(UE_SOURCE_LOCATION);
+
+			// Find the mesh conversion data
+			FMutableSourceMeshData& MeshData = ReferencedMeshes[ID];
+			FString MorphName = MorphNameRef;
+
+			// It would be great to be able to do this conversion in a worker thread, but the engine doesn't support it yet.
+			auto LoadMeshFunc = [MeshData, MorphName, ResolvedMesh, CompletionEvent, CompilationContext=CompilationContext]() mutable
 				{
-					TSharedPtr<mu::FMesh> Result = LoadMeshResourceReferenced(ID);
-					*ResolvedMesh = Result;
-				},
-				LowLevelTasks::ETaskPriority::BackgroundLow
-			);
+					MUTABLE_CPUPROFILER_SCOPE(LoadMeshFunc);
+					check(IsInGameThread());
+
+					// Ensure we don't pull unwanted data into the package when cooking.
+					FCookLoadScope CookLoadScope(ECookLoadType::EditorOnly);
+
+					if (MorphName.IsEmpty())
+					{
+						*ResolvedMesh = ConvertSkeletalMeshToMutable(MeshData, *CompilationContext);
+					}
+					else
+					{
+						const USkeletalMesh* SkeletalMesh = Cast<USkeletalMesh>(MeshData.Mesh.LoadSynchronous());
+						if (SkeletalMesh->FindMorphTarget(*MorphName))
+						{
+							TSharedPtr<mu::FMesh> Mesh;
+							Mesh = ConvertSkeletalMeshToMutable(MeshData, *CompilationContext);
+							Mesh = BuildMorphedMutableMeshFromMesh(Mesh, SkeletalMesh, *MorphName, MeshData.LODIndex, MeshData.SectionIndex);
+							*ResolvedMesh = Mesh;
+						}
+					}
+
+					CompletionEvent.Trigger();
+				};
+
+			if (IsInGameThread())
+			{
+				LoadMeshFunc();
+			}
+			else
+			{
+				check(!bRunImmediatlyIfPossible);
+				UCustomizableObjectSystem::GetInstance()->GetPrivate()->AddGameThreadCompileTask(MoveTemp(LoadMeshFunc));
+			}
 
-			return LaunchTask;
+			return CompletionEvent;
 		}
 
 	);
@@ -313,32 +326,6 @@ const TArray<FCustomizableObjectCompileRunnable::FError>& FCustomizableObjectCom
 }
 
 
-void FCustomizableObjectCompileRunnable::Tick()
-{
-	MUTABLE_CPUPROFILER_SCOPE(FCustomizableObjectCompileRunnable::Tick);
-
-	check(IsInGameThread());
-
-	constexpr double MaxSecondsPerFrame = 0.4;
-
-	double MaxTime = FPlatformTime::Seconds() + MaxSecondsPerFrame;
-
-	FReferenceResourceRequest Request;
-	while (PendingResourceReferenceRequests.Dequeue(Request))
-	{
-		*Request.ResolvedImage = LoadImageResourceReferenced(Request.ID);
-		Request.CompletionEvent->Trigger();
-
-		// Simple time limit enforcement to avoid blocking the game thread if there are many requests.
-		double CurrentTime = FPlatformTime::Seconds();
-		if (CurrentTime >= MaxTime)
-		{
-			break;
-		}
-	}
-}
-
-
 FCustomizableObjectSaveDDRunnable::FCustomizableObjectSaveDDRunnable(const TSharedPtr<FCompilationRequest>& InRequest,
 	TSharedPtr<mu::FModel> InModel,
 	UModelResources& InModelResources,
diff --git a/Engine/Plugins/Mutable/Source/CustomizableObjectEditor/Private/MuCOE/CustomizableObjectCompileRunnable.h b/Engine/Plugins/Mutable/Source/CustomizableObjectEditor/Private/MuCOE/CustomizableObjectCompileRunnable.h
--- a/Engine/Plugins/Mutable/Source/CustomizableObjectEditor/Private/MuCOE/CustomizableObjectCompileRunnable.h
+++ b/Engine/Plugins/Mutable/Source/CustomizableObjectEditor/Private/MuCOE/CustomizableObjectCompileRunnable.h
@@ -11,6 +11,7 @@
 #include "MuR/Ptr.h"
 #include "MuT/Node.h"
 #include "MuCOE/CustomizableObjectEditorLogger.h"
+#include "MuCOE/GenerateMutableSource/GenerateMutableSource.h"
 
 #include "DerivedDataCacheKey.h"
 #include "DerivedDataCachePolicy.h"
@@ -48,21 +49,11 @@ private:
 	mu::Ptr<mu::Node> MutableRoot;
 	TArray<FError> ArrayErrors;
 
-	/** */
-	struct FReferenceResourceRequest
-	{
-		int32 ID = -1;
-		TSharedPtr<TSharedPtr<mu::FImage>> ResolvedImage;
-		TSharedPtr<UE::Tasks::FTaskEvent> CompletionEvent;
-	};
-	TQueue<FReferenceResourceRequest, EQueueMode::Mpsc> PendingResourceReferenceRequests;
-
 	TSharedPtr<mu::FImage> LoadImageResourceReferenced(int32 ID);
-	TSharedPtr<mu::FMesh> LoadMeshResourceReferenced(int32 ID);
 
 public:
 
-	FCustomizableObjectCompileRunnable(mu::Ptr<mu::Node> Root);
+	FCustomizableObjectCompileRunnable(mu::Ptr<mu::Node> Root, TSharedPtr<FMutableCompilationContext> CompilationContext);
 
 	// FRunnable interface
 	uint32 Run() override;
@@ -75,12 +66,11 @@ public:
 	//
 	const TArray<FError>& GetArrayErrors() const;
 
-	void Tick();
-
 	TSharedPtr<mu::FModel, ESPMode::ThreadSafe> Model;
 
 	FCompilationOptions Options;
-	
+	TSharedPtr<FMutableCompilationContext> CompilationContext;
+
 	TArray<FMutableSourceTextureData> ReferencedTextures;
 	TArray<FMutableSourceMeshData> ReferencedMeshes;
 
diff --git a/Engine/Plugins/Mutable/Source/CustomizableObjectEditor/Private/MuCOE/CustomizableObjectCompiler.cpp b/Engine/Plugins/Mutable/Source/CustomizableObjectEditor/Private/MuCOE/CustomizableObjectCompiler.cpp
--- a/Engine/Plugins/Mutable/Source/CustomizableObjectEditor/Private/MuCOE/CustomizableObjectCompiler.cpp
+++ b/Engine/Plugins/Mutable/Source/CustomizableObjectEditor/Private/MuCOE/CustomizableObjectCompiler.cpp
@@ -63,8 +63,6 @@ bool FCustomizableObjectCompiler::Tick(bool bBlocking)
 	if (CompileTask)
 	{
 		bFinished = false;
-		CompileTask->Tick();
-
 		if (CompileTask->IsCompleted())
 		{
 			FinishCompilationTask();
@@ -723,8 +721,8 @@ void FCustomizableObjectCompiler::CompileInternal(bool bAsync)
 		return;
 	}
 
-	FMutableCompilationContext CompilationContext(CurrentObject, this, CurrentOptions);
-	FMutableGraphGenerationContext GenerationContext(CompilationContext);
+	CompilationContext = MakeShared<FMutableCompilationContext>(CurrentObject, this, CurrentOptions);
+	FMutableGraphGenerationContext GenerationContext(*CompilationContext);
 
 	// Perform a first participating objects pass
 	GenerationContext.ParticipatingObjects = ICustomizableObjectEditorModule::GetChecked().GetParticipatingObjects(CurrentObject, true, &CurrentOptions);
@@ -742,22 +740,36 @@ void FCustomizableObjectCompiler::CompileInternal(bool bAsync)
 	}
 	else
 	{
+		TArray<FMutableSourceTextureData> NewCompileTimeReferencedTextures;
+		{
+			for (const TPair<TSoftObjectPtr<const UTexture>, FMutableGraphGenerationContext::FGeneratedReferencedTexture>& Pair : GenerationContext.CompileTimeTextureMap)
+			{
+				check(Pair.Value.ID == NewCompileTimeReferencedTextures.Num());
+
+				FMutableSourceTextureData Tex(*MutablePrivate::LoadObject(Pair.Key));
+				NewCompileTimeReferencedTextures.Add(Tex);
+			}
+		}
+
+		TArray<FMutableSourceMeshData> NewCompileTimeReferencedMeshes;
+		for (const TPair<FMutableSourceMeshData, FMutableGraphGenerationContext::FGeneratedReferencedMesh>& Pair : GenerationContext.CompileTimeMeshMap)
+		{
+			check(Pair.Value.ID == NewCompileTimeReferencedMeshes.Num());
+			NewCompileTimeReferencedMeshes.Add(Pair.Key);
+		}
+
+		// The first part of model resources has to be filled from the GenerationContext.
+		// The second part will be filled when the core compilation finishes so that the CompilationContext is complete.
 		const FString ModelResourcesName = GetModelResourcesNameForPlatform(*CurrentObject, *CurrentOptions.TargetPlatform);
 		UObject* Outer = CurrentOptions.bIsCooking ? CurrentObject : StaticCast<UObject*>(GetTransientPackage());
 		ModelResources = NewObject<UModelResources>(Outer, FName(*ModelResourcesName), RF_Public);
 
 		ModelStreamableBulkData = MakeShared<FModelStreamableBulkData>();
-		
+
 		ModelResources->CodeVersion = UCustomizableObjectPrivate::CurrentSupportedVersion;
 
 		ModelResources->ReferenceSkeletalMeshesData = MoveTemp(GenerationContext.ReferenceSkeletalMeshesData);
 
-		ModelResources->Skeletons.Reserve(GenerationContext.CompilationContext.ReferencedSkeletons.Num());
-		for (const USkeleton* Skeleton : GenerationContext.CompilationContext.ReferencedSkeletons)
-		{
-			ModelResources->Skeletons.Emplace(const_cast<USkeleton*>(Skeleton));
-		}
-		
 		ModelResources->Materials.Reserve(GenerationContext.ReferencedMaterials.Num());
 		for (const UMaterialInterface* Material : GenerationContext.ReferencedMaterials)
 		{
@@ -770,10 +782,10 @@ void FCustomizableObjectCompiler::CompileInternal(bool bAsync)
 			ModelResources->PassThroughMeshes.Add(Pair.Key);
 		}
 
-		for (const TPair<TSoftObjectPtr<const UStreamableRenderAsset>, FMutableGraphGenerationContext::FGeneratedReferencedMesh>& Pair : GenerationContext.RuntimeReferencedMeshMap)
+		for (const TPair<FMutableSourceMeshData, FMutableGraphGenerationContext::FGeneratedReferencedMesh>& Pair : GenerationContext.RuntimeReferencedMeshMap)
 		{
 			check(Pair.Value.ID == ModelResources->RuntimeReferencedMeshes.Num());
-			ModelResources->RuntimeReferencedMeshes.Add(Pair.Key);
+			ModelResources->RuntimeReferencedMeshes.Add(Pair.Key.Mesh);
 		}
 
 		for (const TPair<TSoftObjectPtr<UTexture>, FMutableGraphGenerationContext::FGeneratedReferencedTexture>& Pair : GenerationContext.PassthroughTextureMap)
@@ -787,30 +799,12 @@ void FCustomizableObjectCompiler::CompileInternal(bool bAsync)
 			check(Pair.Value.ID == ModelResources->RuntimeReferencedTextures.Num());
 			ModelResources->RuntimeReferencedTextures.Add(Pair.Key);
 		}
-		
-		ModelResources->PhysicsAssets = MoveTemp(GenerationContext.CompilationContext.PhysicsAssets);
 
 		ModelResources->AnimBPs = MoveTemp(GenerationContext.AnimBPAssets);
-		ModelResources->AnimBpOverridePhysiscAssetsInfo = MoveTemp(GenerationContext.CompilationContext.AnimBpOverridePhysicsAssetsInfo);
 
 		ModelResources->MaterialSlotNames = MoveTemp(GenerationContext.ReferencedMaterialSlotNames);
 		ModelResources->SocketArray = MoveTemp(GenerationContext.SocketArray);
 
-		const int32 NumBones = GenerationContext.CompilationContext.UniqueBoneNames.Num() + GenerationContext.CompilationContext.RemappedBoneNames.Num();
-		ModelResources->BoneNamesMap.Reserve(NumBones);
-
-		for (auto& It : GenerationContext.CompilationContext.UniqueBoneNames)
-		{
-			ModelResources->BoneNamesMap.Add(It.Value, It.Key.Id);
-		}
-
-		for (auto& It : GenerationContext.CompilationContext.RemappedBoneNames)
-		{
-			ModelResources->BoneNamesMap.Add(It.Key, It.Value.Id);
-		}
-
-		ModelResources->SkinWeightProfilesInfo = MoveTemp(GenerationContext.CompilationContext.SkinWeightProfilesInfo);
-
 		TArray<FGeneratedImageProperties> ImageProperties;
 		GenerationContext.ImageProperties.GenerateValueArray(ImageProperties);
 		
@@ -837,119 +831,17 @@ void FCustomizableObjectCompiler::CompileInternal(bool bAsync)
 		ModelResources->StateUIDataMap = MoveTemp(GenerationContext.StateUIDataMap);
 		ModelResources->IntParameterOptionDataTable = MoveTemp(GenerationContext.IntParameterOptionDataTable);
 
-		// Create the RealTimeMorphsTargets Blocks from the per mesh Morph data.
-		uint64 RealTimeMorphDataSize = 0;
-		for (const TPair<uint32, FRealTimeMorphMeshData>& MeshData : GenerationContext.CompilationContext.RealTimeMorphTargetPerMeshData)
-		{
-			RealTimeMorphDataSize += MeshData.Value.Data.Num();
-		}
-		
-		ModelStreamableBulkData->RealTimeMorphStreamables.Empty(32);
-		ModelResources->EditorOnlyMorphTargetReconstructionData.Empty(RealTimeMorphDataSize);
-
-		uint64 RealTimeMorphDataOffsetInBytes = 0;
-		for (const TPair<uint32, FRealTimeMorphMeshData>& MeshData : GenerationContext.CompilationContext.RealTimeMorphTargetPerMeshData)
-		{
-			const uint32 DataSizeInBytes = (uint32)MeshData.Value.Data.Num()*sizeof(FMorphTargetVertexData); 
-			FRealTimeMorphStreamable& ResourceMeshData = ModelStreamableBulkData->RealTimeMorphStreamables.FindOrAdd(MeshData.Key);
-			
-			check(ResourceMeshData.NameResolutionMap.IsEmpty());
-			check(ResourceMeshData.Size == 0);
-
-			ResourceMeshData.NameResolutionMap = MeshData.Value.NameResolutionMap;
-			ResourceMeshData.Size = DataSizeInBytes;
-			EMutableFileFlags Flags = EMutableFileFlags::None;
-			ResourceMeshData.Block = FMutableStreamableBlock { uint32(0), uint32(Flags), RealTimeMorphDataOffsetInBytes };
-			ResourceMeshData.SourceId = MeshData.Value.SourceId;
-
-			RealTimeMorphDataOffsetInBytes += DataSizeInBytes;
-			ModelResources->EditorOnlyMorphTargetReconstructionData.Append(MeshData.Value.Data);
-		}
-	
-		// Create the Clothing Blocks from the per mesh Morph data.
-		uint64 ClothingDataNum = 0;
-		for (const TPair<uint32, FClothingMeshData>& MeshData : GenerationContext.CompilationContext.ClothingPerMeshData)
-		{
-			ClothingDataNum += MeshData.Value.Data.Num();
-		}
-		
-		ModelStreamableBulkData->ClothingStreamables.Empty(32);
-		ModelResources->EditorOnlyClothingMeshToMeshVertData.Empty(ClothingDataNum);
-
-		uint64 ClothingDataOffsetInBytes = 0;
-		for (const TPair<uint32, FClothingMeshData>& MeshData : GenerationContext.CompilationContext.ClothingPerMeshData)
-		{
-			const uint32 DataSizeInBytes = (uint32)MeshData.Value.Data.Num()*sizeof(FCustomizableObjectMeshToMeshVertData); 
-			FClothingStreamable& ResourceMeshData = ModelStreamableBulkData->ClothingStreamables.FindOrAdd(MeshData.Key);
-			
-			check(ResourceMeshData.ClothingAssetIndex == INDEX_NONE);
-			check(ResourceMeshData.ClothingAssetLOD == INDEX_NONE);
-			check(ResourceMeshData.Size == 0);
-
-			ResourceMeshData.ClothingAssetIndex = MeshData.Value.ClothingAssetIndex;
-			ResourceMeshData.ClothingAssetLOD = MeshData.Value.ClothingAssetLOD;
-			ResourceMeshData.PhysicsAssetIndex = MeshData.Value.PhysicsAssetIndex;
-			ResourceMeshData.Size = DataSizeInBytes;
-			EMutableFileFlags Flags = EMutableFileFlags::None;
-			ResourceMeshData.Block = FMutableStreamableBlock{ uint32(0), uint32(Flags), ClothingDataOffsetInBytes };
-			ResourceMeshData.SourceId = MeshData.Value.SourceId;
-
-			ClothingDataOffsetInBytes += DataSizeInBytes;
-			ModelResources->EditorOnlyClothingMeshToMeshVertData.Append(MeshData.Value.Data);
-		}
-
-		ModelResources->ClothingAssetsData = MoveTemp(GenerationContext.CompilationContext.ClothingAssetsData);
-
-		// A clothing backend, e.g. Chaos cloth, can use 2 config files, one owned by the asset, and another that is shared 
-		// among all assets in a SkeletalMesh. When merging different assets in a skeletalmesh we need to make sure only one of 
-		// the shared is used. In that case we will keep the first visited of a type and will be stored separated from the asset.
-		// TODO: Shared configs, which typically controls the quality of the simulation (iterations, etc), probably should be specified 
-		// somewhere else to give more control with which config ends up used. 
-		auto IsSharedConfigData = [](const FCustomizableObjectClothConfigData& ConfigData) -> bool
-		{
-			 const UClass* ConfigClass = FindObject<UClass>(nullptr, *ConfigData.ClassPath);
-			 return ConfigClass ? static_cast<bool>(Cast<UClothSharedConfigCommon>(ConfigClass->GetDefaultObject())) : false;
-		};
-		
-		// Find shared configs to be used (One of each type) 
-		for (FCustomizableObjectClothingAssetData& ClothingAssetData : ModelResources->ClothingAssetsData)
-		{
-			 for (FCustomizableObjectClothConfigData& ClothConfigData : ClothingAssetData.ConfigsData)
-			 {
-				  if (IsSharedConfigData(ClothConfigData))
-				  {
-					  FCustomizableObjectClothConfigData* FoundConfig = ModelResources->ClothSharedConfigsData.FindByPredicate(
-						   [Name = ClothConfigData.ConfigName](const FCustomizableObjectClothConfigData& Other)
-						   {
-							   return Name == Other.ConfigName;
-						   });
-
-					  if (!FoundConfig)
-					  {
-						  ModelResources->ClothSharedConfigsData.AddDefaulted_GetRef() = ClothConfigData;
-					  }
-				  }
-			 }
-		}
-		
-		// Remove shared configs
-		for (FCustomizableObjectClothingAssetData& ClothingAssetData : ModelResources->ClothingAssetsData)
-		{
-			 ClothingAssetData.ConfigsData.RemoveAllSwap(IsSharedConfigData);
-		}
-
-		ModelResources->MeshMetadata = MoveTemp(GenerationContext.CompilationContext.MeshMetadata);
 		ModelResources->SurfaceMetadata = MoveTemp(GenerationContext.SurfaceMetadata);
 
 		ModelResources->GroupNodeMap = GenerationContext.GroupNodeMap;
 
 		// If the optimization level is "none" disable texture streaming, because textures are all referenced
 		// unreal assets and progressive generation is not supported.
-		ModelResources->bIsTextureStreamingDisabled = GenerationContext.CompilationContext.Options.OptimizationLevel == 0;
-		
-		ModelResources->bIsCompiledWithOptimization = GenerationContext.CompilationContext.Options.OptimizationLevel == UE_MUTABLE_MAX_OPTIMIZATION;
+		ModelResources->bIsTextureStreamingDisabled = CurrentOptions.OptimizationLevel == 0;
 
-		ModelResources->bCompiledWithHDTextureCompression = GenerationContext.CompilationContext.Options.TextureCompression == ECustomizableObjectTextureCompression::HighQuality;
+		ModelResources->bIsCompiledWithOptimization = CurrentOptions.OptimizationLevel == UE_MUTABLE_MAX_OPTIMIZATION;
+
+		ModelResources->bCompiledWithHDTextureCompression = CurrentOptions.TextureCompression == ECustomizableObjectTextureCompression::HighQuality;
 
 		ModelResources->AlwaysLoadedExtensionData = MoveTemp(GenerationContext.AlwaysLoadedExtensionData);
 		ModelResources->StreamedExtensionDataEditor = MoveTemp(GenerationContext.StreamedExtensionData);
@@ -994,26 +886,7 @@ void FCustomizableObjectCompiler::CompileInternal(bool bAsync)
 
 		ModelResources->StreamedResourceDataEditor = MoveTemp(GenerationContext.StreamedResourceData);
 		
-		// Pass-through textures
-		TArray<FMutableSourceTextureData> NewCompileTimeReferencedTextures;
-		for (const TPair<TSoftObjectPtr<const UTexture>, FMutableGraphGenerationContext::FGeneratedReferencedTexture>& Pair : GenerationContext.CompileTimeTextureMap)
-		{
-			check(Pair.Value.ID == NewCompileTimeReferencedTextures.Num());
-			FMutableSourceTextureData Tex(*MutablePrivate::LoadObject(Pair.Key));
-			NewCompileTimeReferencedTextures.Add(Tex);
-		}
-
-		TArray<FMutableSourceMeshData> NewCompileTimeReferencedMeshes;
-		for (const TPair<TSoftObjectPtr<const UStreamableRenderAsset>, FMutableGraphGenerationContext::FGeneratedReferencedMesh>& Pair : GenerationContext.CompileTimeMeshMap)
-		{
-			check(Pair.Value.ID == NewCompileTimeReferencedMeshes.Num());
-
-			FMutableSourceMeshData Data;
-			Data.Mesh = Pair.Key;
-			NewCompileTimeReferencedMeshes.Add(Data);
-		}
-
-		CompileTask = MakeShareable(new FCustomizableObjectCompileRunnable(MutableRoot));
+		CompileTask = MakeShareable(new FCustomizableObjectCompileRunnable(MutableRoot, CompilationContext));
 		CompileTask->Options = CurrentOptions;
 		CompileTask->ReferencedTextures = NewCompileTimeReferencedTextures;
 		CompileTask->ReferencedMeshes = NewCompileTimeReferencedMeshes;
@@ -1225,7 +1098,7 @@ bool FCustomizableObjectCompiler::TryLoadCompiledDataFromDDC(UCustomizableObject
 mu::Ptr<mu::Node> FCustomizableObjectCompiler::Export(UCustomizableObject* Object, const FCompilationOptions& InCompilerOptions, 
 	TArray<TSoftObjectPtr<const UTexture>>& OutRuntimeReferencedTextures,
 	TArray<FMutableSourceTextureData>& OutCompilerReferencedTextures, 
-	TArray<TSoftObjectPtr<const UObject>>& OutRuntimeReferencedMeshes,
+	TArray<TSoftObjectPtr<const UStreamableRenderAsset>>& OutRuntimeReferencedMeshes,
 	TArray<FMutableSourceMeshData>& OutCompilerReferencedMeshes)
 {
 	UE_LOG(LogMutable, Log, TEXT("Started Customizable Object Export %s."), *Object->GetName());
@@ -1245,8 +1118,8 @@ mu::Ptr<mu::Node> FCustomizableObjectCompiler::Export(UCustomizableObject* Objec
 	CompilerOptions.bPhysicsAssetMergeEnabled = Object->GetPrivate()->IsPhysicsAssetMergeEnabled();
 	CompilerOptions.bAnimBpPhysicsManipulationEnabled = Object->GetPrivate()->IsEnabledAnimBpPhysicsAssetsManipulation();
 	
-	FMutableCompilationContext CompilationContext(Object, this, CompilerOptions);
-	FMutableGraphGenerationContext GenerationContext(CompilationContext);
+	CompilationContext = MakeShared<FMutableCompilationContext>(Object, this, CompilerOptions);
+	FMutableGraphGenerationContext GenerationContext(*CompilationContext);
 
 	GenerationContext.bSkipParticipatingObjectsPass = true;
 	
@@ -1277,19 +1150,18 @@ mu::Ptr<mu::Node> FCustomizableObjectCompiler::Export(UCustomizableObject* Objec
 
 	// Pass out the referenced meshes
 	OutRuntimeReferencedMeshes.Empty();
-	for (const TPair<TSoftObjectPtr<const UStreamableRenderAsset>, FMutableGraphGenerationContext::FGeneratedReferencedMesh>& Pair : GenerationContext.RuntimeReferencedMeshMap)
+	for (const TPair<FMutableSourceMeshData, FMutableGraphGenerationContext::FGeneratedReferencedMesh>& Pair : GenerationContext.RuntimeReferencedMeshMap)
 	{
 		check(Pair.Value.ID == OutRuntimeReferencedMeshes.Num());
-		OutRuntimeReferencedMeshes.Add(Pair.Key);
+		OutRuntimeReferencedMeshes.Add(Pair.Key.Mesh);
 	}
 
 	OutCompilerReferencedMeshes.Empty();
-	for (const TPair<TSoftObjectPtr<const UStreamableRenderAsset>, FMutableGraphGenerationContext::FGeneratedReferencedMesh>& Pair : GenerationContext.CompileTimeMeshMap)
+	for (const TPair<FMutableSourceMeshData, FMutableGraphGenerationContext::FGeneratedReferencedMesh>& Pair : GenerationContext.CompileTimeMeshMap)
 	{
 		check(Pair.Value.ID == OutCompilerReferencedMeshes.Num());
 
-		FMutableSourceMeshData Data;
-		Data.Mesh = Pair.Key;
+		FMutableSourceMeshData Data(Pair.Key);
 		OutCompilerReferencedMeshes.Add(Data);
 	}
 
@@ -1330,6 +1202,136 @@ void FCustomizableObjectCompiler::FinishCompilationTask()
 		}
 	}
 
+	// Second part of model resources can be filled now.
+	{
+		ModelResources->Skeletons.Reserve(CompilationContext->ReferencedSkeletons.Num());
+		for (const USkeleton* Skeleton : CompilationContext->ReferencedSkeletons)
+		{
+			ModelResources->Skeletons.Emplace(const_cast<USkeleton*>(Skeleton));
+		}
+
+		ModelResources->PhysicsAssets = MoveTemp(CompilationContext->PhysicsAssets);
+		ModelResources->AnimBpOverridePhysiscAssetsInfo = MoveTemp(CompilationContext->AnimBpOverridePhysicsAssetsInfo);
+
+		const int32 NumBones = CompilationContext->UniqueBoneNames.Num() + CompilationContext->RemappedBoneNames.Num();
+		ModelResources->BoneNamesMap.Reserve(NumBones);
+
+		for (auto& It : CompilationContext->UniqueBoneNames)
+		{
+			ModelResources->BoneNamesMap.Add(It.Value, It.Key.Id);
+		}
+
+		for (auto& It : CompilationContext->RemappedBoneNames)
+		{
+			ModelResources->BoneNamesMap.Add(It.Key, It.Value.Id);
+		}
+
+		ModelResources->SkinWeightProfilesInfo = MoveTemp(CompilationContext->SkinWeightProfilesInfo);
+
+		// Create the RealTimeMorphsTargets Blocks from the per mesh Morph data.
+		uint64 RealTimeMorphDataSize = 0;
+		for (const TPair<uint32, FRealTimeMorphMeshData>& MeshData : CompilationContext->RealTimeMorphTargetPerMeshData)
+		{
+			RealTimeMorphDataSize += MeshData.Value.Data.Num();
+		}
+
+		ModelStreamableBulkData->RealTimeMorphStreamables.Empty(32);
+		ModelResources->EditorOnlyMorphTargetReconstructionData.Empty(RealTimeMorphDataSize);
+
+		uint64 RealTimeMorphDataOffsetInBytes = 0;
+		for (const TPair<uint32, FRealTimeMorphMeshData>& MeshData : CompilationContext->RealTimeMorphTargetPerMeshData)
+		{
+			const uint32 DataSizeInBytes = (uint32)MeshData.Value.Data.Num() * sizeof(FMorphTargetVertexData);
+			FRealTimeMorphStreamable& ResourceMeshData = ModelStreamableBulkData->RealTimeMorphStreamables.FindOrAdd(MeshData.Key);
+
+			check(ResourceMeshData.NameResolutionMap.IsEmpty());
+			check(ResourceMeshData.Size == 0);
+
+			ResourceMeshData.NameResolutionMap = MeshData.Value.NameResolutionMap;
+			ResourceMeshData.Size = DataSizeInBytes;
+			EMutableFileFlags Flags = EMutableFileFlags::None;
+			ResourceMeshData.Block = FMutableStreamableBlock{ uint32(0), uint32(Flags), RealTimeMorphDataOffsetInBytes };
+			ResourceMeshData.SourceId = MeshData.Value.SourceId;
+
+			RealTimeMorphDataOffsetInBytes += DataSizeInBytes;
+			ModelResources->EditorOnlyMorphTargetReconstructionData.Append(MeshData.Value.Data);
+		}
+
+		// Create the Clothing Blocks from the per mesh Morph data.
+		uint64 ClothingDataNum = 0;
+		for (const TPair<uint32, FClothingMeshData>& MeshData : CompilationContext->ClothingPerMeshData)
+		{
+			ClothingDataNum += MeshData.Value.Data.Num();
+		}
+
+		ModelStreamableBulkData->ClothingStreamables.Empty(32);
+		ModelResources->EditorOnlyClothingMeshToMeshVertData.Empty(ClothingDataNum);
+
+		uint64 ClothingDataOffsetInBytes = 0;
+		for (const TPair<uint32, FClothingMeshData>& MeshData : CompilationContext->ClothingPerMeshData)
+		{
+			const uint32 DataSizeInBytes = (uint32)MeshData.Value.Data.Num() * sizeof(FCustomizableObjectMeshToMeshVertData);
+			FClothingStreamable& ResourceMeshData = ModelStreamableBulkData->ClothingStreamables.FindOrAdd(MeshData.Key);
+
+			check(ResourceMeshData.ClothingAssetIndex == INDEX_NONE);
+			check(ResourceMeshData.ClothingAssetLOD == INDEX_NONE);
+			check(ResourceMeshData.Size == 0);
+
+			ResourceMeshData.ClothingAssetIndex = MeshData.Value.ClothingAssetIndex;
+			ResourceMeshData.ClothingAssetLOD = MeshData.Value.ClothingAssetLOD;
+			ResourceMeshData.PhysicsAssetIndex = MeshData.Value.PhysicsAssetIndex;
+			ResourceMeshData.Size = DataSizeInBytes;
+			EMutableFileFlags Flags = EMutableFileFlags::None;
+			ResourceMeshData.Block = FMutableStreamableBlock{ uint32(0), uint32(Flags), ClothingDataOffsetInBytes };
+			ResourceMeshData.SourceId = MeshData.Value.SourceId;
+
+			ClothingDataOffsetInBytes += DataSizeInBytes;
+			ModelResources->EditorOnlyClothingMeshToMeshVertData.Append(MeshData.Value.Data);
+		}
+
+		ModelResources->ClothingAssetsData = MoveTemp(CompilationContext->ClothingAssetsData);
+
+		// A clothing backend, e.g. Chaos cloth, can use 2 config files, one owned by the asset, and another that is shared 
+		// among all assets in a SkeletalMesh. When merging different assets in a skeletalmesh we need to make sure only one of 
+		// the shared is used. In that case we will keep the first visited of a type and will be stored separated from the asset.
+		// TODO: Shared configs, which typically controls the quality of the simulation (iterations, etc), probably should be specified 
+		// somewhere else to give more control with which config ends up used. 
+		auto IsSharedConfigData = [](const FCustomizableObjectClothConfigData& ConfigData) -> bool
+			{
+				const UClass* ConfigClass = FindObject<UClass>(nullptr, *ConfigData.ClassPath);
+				return ConfigClass ? static_cast<bool>(Cast<UClothSharedConfigCommon>(ConfigClass->GetDefaultObject())) : false;
+			};
+
+		// Find shared configs to be used (One of each type) 
+		for (FCustomizableObjectClothingAssetData& ClothingAssetData : ModelResources->ClothingAssetsData)
+		{
+			for (FCustomizableObjectClothConfigData& ClothConfigData : ClothingAssetData.ConfigsData)
+			{
+				if (IsSharedConfigData(ClothConfigData))
+				{
+					FCustomizableObjectClothConfigData* FoundConfig = ModelResources->ClothSharedConfigsData.FindByPredicate(
+						[Name = ClothConfigData.ConfigName](const FCustomizableObjectClothConfigData& Other)
+						{
+							return Name == Other.ConfigName;
+						});
+
+					if (!FoundConfig)
+					{
+						ModelResources->ClothSharedConfigsData.AddDefaulted_GetRef() = ClothConfigData;
+					}
+				}
+			}
+		}
+
+		// Remove shared configs
+		for (FCustomizableObjectClothingAssetData& ClothingAssetData : ModelResources->ClothingAssetsData)
+		{
+			ClothingAssetData.ConfigsData.RemoveAllSwap(IsSharedConfigData);
+		}
+
+		ModelResources->MeshMetadata = MoveTemp(CompilationContext->MeshMetadata);
+	}
+
 	// Order matters
 	CompileThread.Reset();
 	CompileTask.Reset();
diff --git a/Engine/Plugins/Mutable/Source/CustomizableObjectEditor/Private/MuCOE/CustomizableObjectCompiler.h b/Engine/Plugins/Mutable/Source/CustomizableObjectEditor/Private/MuCOE/CustomizableObjectCompiler.h
--- a/Engine/Plugins/Mutable/Source/CustomizableObjectEditor/Private/MuCOE/CustomizableObjectCompiler.h
+++ b/Engine/Plugins/Mutable/Source/CustomizableObjectEditor/Private/MuCOE/CustomizableObjectCompiler.h
@@ -55,8 +55,10 @@ public:
 	
 	/** Generate the Mutable Graph from the Unreal Graph. */
 	mu::Ptr<mu::Node> Export(UCustomizableObject* Object, const FCompilationOptions& Options, 
-		TArray<TSoftObjectPtr<const UTexture>>& OutRuntimeReferencedTextures, TArray<FMutableSourceTextureData>& OutCompilerReferencedTextures,
-		TArray<TSoftObjectPtr<const UObject>>& OutRuntimeReferencedMeshes, TArray<FMutableSourceMeshData>& OutCompilerReferencedMeshes);
+		TArray<TSoftObjectPtr<const UTexture>>& OutRuntimeReferencedTextures, 
+		TArray<FMutableSourceTextureData>& OutCompilerReferencedTextures,
+		TArray<TSoftObjectPtr<const UStreamableRenderAsset>>& OutRuntimeReferencedMeshes,
+		TArray<FMutableSourceMeshData>& OutCompilerReferencedMeshes);
 
 	void CompilerLog(const FText& Message, const TArray<const UObject*>& UObject, const EMessageSeverity::Type MessageSeverity = EMessageSeverity::Warning, const bool bAddBaseObjectInfo = true, const ELoggerSpamBin SpamBin = ELoggerSpamBin::ShowAll);
 	void CompilerLog(const FText& Message, const UObject* Context = nullptr, const EMessageSeverity::Type MessageSeverity = EMessageSeverity::Warning, const bool bAddBaseObjectInfo = true, const ELoggerSpamBin SpamBin = ELoggerSpamBin::ShowAll);
@@ -85,6 +87,9 @@ public:
 	{
 		return TEXT("FCustomizableObjectCompiler");
 	}
+	
+	/** */
+	TSharedPtr<FMutableCompilationContext> CompilationContext;
 
 private:
 
diff --git a/Engine/Plugins/Mutable/Source/CustomizableObjectEditor/Private/MuCOE/GenerateMutableSource/GenerateMutableSource.cpp b/Engine/Plugins/Mutable/Source/CustomizableObjectEditor/Private/MuCOE/GenerateMutableSource/GenerateMutableSource.cpp
--- a/Engine/Plugins/Mutable/Source/CustomizableObjectEditor/Private/MuCOE/GenerateMutableSource/GenerateMutableSource.cpp
+++ b/Engine/Plugins/Mutable/Source/CustomizableObjectEditor/Private/MuCOE/GenerateMutableSource/GenerateMutableSource.cpp
@@ -1638,25 +1638,30 @@ TSharedPtr<mu::FImage> GenerateImageConstant(UTexture* Texture, FMutableGraphGen
 }
 
 
-TSharedPtr<mu::FMesh> GenerateMeshConstant(TSoftObjectPtr<UStreamableRenderAsset> Mesh, FMutableGraphGenerationContext& GenerationContext, bool bIsReference)
+TSharedPtr<mu::FMesh> GenerateMeshConstant(
+	const FMutableSourceMeshData& Source,
+	FMutableGraphGenerationContext& GenerationContext, 
+	bool bIsPassthrough)
 {
 	MUTABLE_CPUPROFILER_SCOPE(GenerateMeshConstant)
 
-	if (!Mesh)
+	if (!Source.Mesh)
 	{
 		return nullptr;
 	}
 
 	bool bForceLoad = false;
 	bool bIsCompileTime = false;
-	if (!bIsReference)
+	if (!bIsPassthrough)
 	{
 		bForceLoad = true;
-		if (GenerationContext.CompilationContext.Options.OptimizationLevel == 0)
-		{
-			bIsCompileTime = false;
-		}
-		else
+		// TODO: For now, even in none, we turn them into compile-time references, since runtime meshes are not
+		// supported yet.
+		//if (GenerationContext.CompilationContext.Options.OptimizationLevel == 0)
+		//{
+		//	bIsCompileTime = false;
+		//}
+		//else
 		{
 			bIsCompileTime = true;
 		}
@@ -1669,19 +1674,19 @@ TSharedPtr<mu::FMesh> GenerateMeshConstant(TSoftObjectPtr<UStreamableRenderAsset
 	FMutableGraphGenerationContext::FGeneratedReferencedMesh* Entry = nullptr;
 	int32 Num = 0;
 
-	if (bIsReference)
+	if (bIsPassthrough)
 	{
-		Entry = &GenerationContext.PassthroughMeshMap.FindOrAdd(Mesh, InvalidEntry);
+		Entry = &GenerationContext.PassthroughMeshMap.FindOrAdd(Source.Mesh, InvalidEntry);
 		Num = GenerationContext.PassthroughMeshMap.Num();
 	}
 	else if (bIsCompileTime)
 	{
-		Entry = &GenerationContext.CompileTimeMeshMap.FindOrAdd(Mesh, InvalidEntry);
+		Entry = &GenerationContext.CompileTimeMeshMap.FindOrAdd(Source, InvalidEntry);
 		Num = GenerationContext.CompileTimeMeshMap.Num();
 	}
 	else
 	{
-		Entry = &GenerationContext.RuntimeReferencedMeshMap.FindOrAdd(Mesh, InvalidEntry);
+		Entry = &GenerationContext.RuntimeReferencedMeshMap.FindOrAdd(Source, InvalidEntry);
 		Num = GenerationContext.RuntimeReferencedMeshMap.Num();
 	}
 
diff --git a/Engine/Plugins/Mutable/Source/CustomizableObjectEditor/Private/MuCOE/GenerateMutableSource/GenerateMutableSource.h b/Engine/Plugins/Mutable/Source/CustomizableObjectEditor/Private/MuCOE/GenerateMutableSource/GenerateMutableSource.h
--- a/Engine/Plugins/Mutable/Source/CustomizableObjectEditor/Private/MuCOE/GenerateMutableSource/GenerateMutableSource.h
+++ b/Engine/Plugins/Mutable/Source/CustomizableObjectEditor/Private/MuCOE/GenerateMutableSource/GenerateMutableSource.h
@@ -472,7 +472,8 @@ struct FMutableCompilationContext
 	TStrongObjectPtr<const UCustomizableObject> Object;
 		
 	/** Compilation options, including target platform. */
-	const FCompilationOptions& Options;
+	FCompilationOptions Options;
+
 	/** */
 	TArray<const USkeleton*> ReferencedSkeletons;
 
@@ -635,11 +636,11 @@ struct FMutableGraphGenerationContext
 
 	// Cache of runtime meshes and their IDs used in the core to identify them.
 	// These meshes will remain as external references even in optimized models.
-	TMap<TSoftObjectPtr<const UStreamableRenderAsset>, FGeneratedReferencedMesh> RuntimeReferencedMeshMap;
+	TMap<FMutableSourceMeshData, FGeneratedReferencedMesh> RuntimeReferencedMeshMap;
 
 	// Cache of runtime pass-through meshes and their IDs used in the core to identify them
 	// These meshes will become mutable images in the compiled model.
-	TMap<TSoftObjectPtr<const UStreamableRenderAsset>, FGeneratedReferencedMesh> CompileTimeMeshMap;
+	TMap<FMutableSourceMeshData, FGeneratedReferencedMesh> CompileTimeMeshMap;
 
 	// Cache of runtime pass-through images and their IDs used in the core to identify them.
 	// These textures will remain as external references even in optimized models.
@@ -1007,7 +1008,7 @@ int32 GetMaxTextureSize(const UTexture2D& ReferenceTexture, const UTextureLODSet
 int32 GetTextureSizeInGame(const UTexture2D& Texture, const UTextureLODSettings& LODSettings);
 
 TSharedPtr<mu::FImage> GenerateImageConstant(UTexture*, FMutableGraphGenerationContext&, bool bIsReference);
-TSharedPtr<mu::FMesh> GenerateMeshConstant(TSoftObjectPtr<UStreamableRenderAsset> Mesh, FMutableGraphGenerationContext&, bool bIsReference);
+TSharedPtr<mu::FMesh> GenerateMeshConstant(const FMutableSourceMeshData& Source, FMutableGraphGenerationContext&, bool bIsReference);
 
 /** Generates a mutable image descriptor from an unreal engine texture */
 mu::FImageDesc GenerateImageDescriptor(UTexture* Texture);
@@ -1018,3 +1019,8 @@ mu::FImageDesc GenerateImageDescriptor(UTexture* Texture);
  * return the unique id for the SurfaceMetadata in HashSurfaceMetadataSet.
  **/
 uint32 AddUniqueSurfaceMetadata(const FSkeletalMaterial* Material, const FSkelMeshSection* MeshSection, TMap<uint32, FMutableSurfaceMetadata>& InOutHashSurfaceMetadataSet);
+
+/** Convert a mesh to Mutable format.
+*/
+TSharedPtr<mu::FMesh> ConvertSkeletalMeshToMutable(const FMutableSourceMeshData& Source, FMutableCompilationContext& Context);
+TSharedPtr<mu::FMesh> BuildMorphedMutableMeshFromMesh(TSharedPtr<mu::FMesh> BaseSourceMesh, const USkeletalMesh* SkeletalMesh, const FString& MorphTargetName, int32 LODIndex, int32 SectionIndex);
diff --git a/Engine/Plugins/Mutable/Source/CustomizableObjectEditor/Private/MuCOE/GenerateMutableSource/GenerateMutableSourceComponent.cpp b/Engine/Plugins/Mutable/Source/CustomizableObjectEditor/Private/MuCOE/GenerateMutableSource/GenerateMutableSourceComponent.cpp
--- a/Engine/Plugins/Mutable/Source/CustomizableObjectEditor/Private/MuCOE/GenerateMutableSource/GenerateMutableSourceComponent.cpp
+++ b/Engine/Plugins/Mutable/Source/CustomizableObjectEditor/Private/MuCOE/GenerateMutableSource/GenerateMutableSourceComponent.cpp
@@ -225,10 +225,10 @@ mu::Ptr<mu::NodeComponent> GenerateMutableSourceComponent(const UEdGraphPin* Pin
 			MeshNode = new mu::NodeMeshConstant();
 
 			FString MeshUniqueTags;
-			constexpr bool bIsReference = true;
+			constexpr bool bIsPassthrough = true;
 			TSoftClassPtr<UAnimInstance> AnimInstance;
-			TSharedPtr<mu::FMesh> MutableMesh = GenerateMutableMesh(SkeletalMesh, AnimInstance, 0, 0, 0, 0, MeshUniqueTags, 0, 
-																GenerationContext, TypedComponentPassthroughMesh, nullptr, bIsReference);
+			TSharedPtr<mu::FMesh> MutableMesh = GenerateMutableSkeletalMesh(SkeletalMesh, AnimInstance, 0, 0, 0, 0, MeshUniqueTags, 0, 
+																GenerationContext, TypedComponentPassthroughMesh, nullptr, bIsPassthrough);
 
 			MeshNode->SetValue(MutableMesh);
 		}
diff --git a/Engine/Plugins/Mutable/Source/CustomizableObjectEditor/Private/MuCOE/GenerateMutableSource/GenerateMutableSourceLayout.cpp b/Engine/Plugins/Mutable/Source/CustomizableObjectEditor/Private/MuCOE/GenerateMutableSource/GenerateMutableSourceLayout.cpp
--- a/Engine/Plugins/Mutable/Source/CustomizableObjectEditor/Private/MuCOE/GenerateMutableSource/GenerateMutableSourceLayout.cpp
+++ b/Engine/Plugins/Mutable/Source/CustomizableObjectEditor/Private/MuCOE/GenerateMutableSource/GenerateMutableSourceLayout.cpp
@@ -85,7 +85,7 @@ mu::Ptr<mu::NodeLayout> CreateMutableLayoutNode(FMutableGraphGenerationContext&
 		// Convert the UE mesh in the layout into a Mutable mesh.
 		TSharedPtr<mu::FMesh> MutableMesh;
 
-		// TODO: Don't force load the mesh
+		// TODO: Don't force load the mesh by moving the creation of automatic blocks to the core compilation stage.
 		if (USkeletalMesh* SkeletalMesh = Cast<USkeletalMesh>(GenerationContext.LoadObject(UnrealLayout->GetMesh())))
 		{
 			// We don't need all the data to generate the blocks
@@ -96,7 +96,8 @@ mu::Ptr<mu::NodeLayout> CreateMutableLayoutNode(FMutableGraphGenerationContext&
 			GenerationContext.MeshGenerationFlags.Push(ShapeFlags);
 
 			bool bMeshMustExist = true;
-			MutableMesh = ConvertSkeletalMeshToMutable(SkeletalMesh, TSoftClassPtr<UAnimInstance>(), bMeshMustExist, UnrealLayout->GetLOD(), UnrealLayout->GetMaterial(), 0, GenerationContext, nullptr, nullptr);
+			bool bForceImmediateConversion = true;
+			MutableMesh = ConvertSkeletalMeshToMutable(SkeletalMesh, TSoftClassPtr<UAnimInstance>(), bMeshMustExist, UnrealLayout->GetLOD(), UnrealLayout->GetMaterial(), 0, GenerationContext, nullptr, nullptr, bForceImmediateConversion);
 
 			GenerationContext.MeshGenerationFlags.Pop();
 		}
diff --git a/Engine/Plugins/Mutable/Source/CustomizableObjectEditor/Private/MuCOE/GenerateMutableSource/GenerateMutableSourceMesh.cpp b/Engine/Plugins/Mutable/Source/CustomizableObjectEditor/Private/MuCOE/GenerateMutableSource/GenerateMutableSourceMesh.cpp
--- a/Engine/Plugins/Mutable/Source/CustomizableObjectEditor/Private/MuCOE/GenerateMutableSource/GenerateMutableSourceMesh.cpp
+++ b/Engine/Plugins/Mutable/Source/CustomizableObjectEditor/Private/MuCOE/GenerateMutableSource/GenerateMutableSourceMesh.cpp
@@ -49,7 +49,6 @@
 #include "PhysicsEngine/SkeletalBodySetup.h"
 #include "Engine/SkinnedAssetCommon.h"
 #include "AssetRegistry/AssetRegistryModule.h"
-
 #include "Hash/CityHash.h"
 
 #define LOCTEXT_NAMESPACE "CustomizableObjectEditor"
@@ -630,71 +629,76 @@ namespace MutablePrivate
 {
 
 	/** Copy mesh data from a single source buffer into a set of destination buffers assuming the format and components is the same. */
-	void CopyBufferClearingPadding(mu::FMeshBufferSet& BufferSet, const mu::FMeshBufferSet& SourceBufferFormat, const void* InSourceData)
+	void CopyBufferClearingPadding(UE::Tasks::FTaskEvent CompletionEvent, mu::FMeshBufferSet& BufferSet, const mu::FMeshBufferSet& SourceBufferFormat, const void* InSourceData)
 	{
 		MUTABLE_CPUPROFILER_SCOPE(CopyBufferClearingPadding);
 
-		int32 ElementCount = BufferSet.ElementCount;
-
 		for (int32 BufferIndex = 0; BufferIndex < BufferSet.GetBufferCount(); ++BufferIndex)
 		{
-			int32 ChannelCount = BufferSet.Buffers[BufferIndex].Channels.Num();
-			int32 ElementSize = BufferSet.Buffers[BufferIndex].ElementSize;
+			CompletionEvent.AddPrerequisites(
+				UE::Tasks::Launch(TEXT("BufferConversion"),
+				[BufferIndex, &BufferSet, &SourceBufferFormat, InSourceData]()
+				{
+						int32 ElementCount = BufferSet.ElementCount;
+						int32 ChannelCount = BufferSet.Buffers[BufferIndex].Channels.Num();
+						int32 ElementSize = BufferSet.Buffers[BufferIndex].ElementSize;
 
-			check(SourceBufferFormat.Buffers.Num()==1);
-			const mu::FMeshBuffer& SourceBuffer = SourceBufferFormat.Buffers[0];
+						check(SourceBufferFormat.Buffers.Num() == 1);
+						const mu::FMeshBuffer& SourceBuffer = SourceBufferFormat.Buffers[0];
 
-			uint8* TargetData = BufferSet.GetBufferData(BufferIndex);
-			TArray<const uint8*, TInlineAllocator<8>> SourceDataPerChannel;
-			SourceDataPerChannel.SetNumUninitialized(ChannelCount);
-			for (int32 ChannelIndex = 0; ChannelIndex < ChannelCount; ++ChannelIndex)
-			{
-				const mu::FMeshBufferChannel& DestinationChannel = BufferSet.Buffers[BufferIndex].Channels[ChannelIndex];
+						uint8* TargetData = BufferSet.GetBufferData(BufferIndex);
+						TArray<const uint8*, TInlineAllocator<8>> SourceDataPerChannel;
+						SourceDataPerChannel.SetNumUninitialized(ChannelCount);
+						for (int32 ChannelIndex = 0; ChannelIndex < ChannelCount; ++ChannelIndex)
+						{
+							const mu::FMeshBufferChannel& DestinationChannel = BufferSet.Buffers[BufferIndex].Channels[ChannelIndex];
 
-				int32 SourceBufferIndex = -1;
-				int32 SourceChannelIndex = -1;
-				SourceBufferFormat.FindChannel(DestinationChannel.Semantic, DestinationChannel.SemanticIndex, &SourceBufferIndex, &SourceChannelIndex);
-				check(SourceBufferIndex==0 && SourceChannelIndex >=0);
+							int32 SourceBufferIndex = -1;
+							int32 SourceChannelIndex = -1;
+							SourceBufferFormat.FindChannel(DestinationChannel.Semantic, DestinationChannel.SemanticIndex, &SourceBufferIndex, &SourceChannelIndex);
+							check(SourceBufferIndex == 0 && SourceChannelIndex >= 0);
 
-				const mu::FMeshBufferChannel& SourceChannel = SourceBuffer.Channels[SourceChannelIndex];
-				check(SourceChannel.Format == DestinationChannel.Format);
-				check(SourceChannel.ComponentCount == DestinationChannel.ComponentCount);
+							const mu::FMeshBufferChannel& SourceChannel = SourceBuffer.Channels[SourceChannelIndex];
+							check(SourceChannel.Format == DestinationChannel.Format);
+							check(SourceChannel.ComponentCount == DestinationChannel.ComponentCount);
 
-				SourceDataPerChannel[ChannelIndex] = reinterpret_cast<const uint8*>(InSourceData) + SourceChannel.Offset;
-			}
+							SourceDataPerChannel[ChannelIndex] = reinterpret_cast<const uint8*>(InSourceData) + SourceChannel.Offset;
+						}
 
-			for (int32 Element = 0; Element < ElementCount; ++Element)
-			{
-				int32 CurrentOffset = 0;
-				for (int32 ChannelIndex = 0; ChannelIndex < ChannelCount; ++ChannelIndex)
-				{
-					const mu::FMeshBufferChannel& Channel = BufferSet.Buffers[BufferIndex].Channels[ChannelIndex];
+						for (int32 Element = 0; Element < ElementCount; ++Element)
+						{
+							int32 CurrentOffset = 0;
+							for (int32 ChannelIndex = 0; ChannelIndex < ChannelCount; ++ChannelIndex)
+							{
+								const mu::FMeshBufferChannel& Channel = BufferSet.Buffers[BufferIndex].Channels[ChannelIndex];
 
-					int32 ChannelOffset = Channel.Offset;
+								int32 ChannelOffset = Channel.Offset;
 
-					int32 PreviousPadding = ChannelOffset - CurrentOffset;
-					if (PreviousPadding > 0)
-					{
-						FMemory::Memzero(TargetData, PreviousPadding);
-						TargetData += PreviousPadding;
-						CurrentOffset += PreviousPadding;
-					}
+								int32 PreviousPadding = ChannelOffset - CurrentOffset;
+								if (PreviousPadding > 0)
+								{
+									FMemory::Memzero(TargetData, PreviousPadding);
+									TargetData += PreviousPadding;
+									CurrentOffset += PreviousPadding;
+								}
 
-					int32 ChannelSize = Channel.ComponentCount * GetMeshFormatData(Channel.Format).SizeInBytes;
-					FMemory::Memcpy(TargetData, SourceDataPerChannel[ChannelIndex], ChannelSize);
-					TargetData += ChannelSize;
-					SourceDataPerChannel[ChannelIndex] += SourceBuffer.ElementSize;
-					CurrentOffset += ChannelSize;
-				}
+								int32 ChannelSize = Channel.ComponentCount * GetMeshFormatData(Channel.Format).SizeInBytes;
+								FMemory::Memcpy(TargetData, SourceDataPerChannel[ChannelIndex], ChannelSize);
+								TargetData += ChannelSize;
+								SourceDataPerChannel[ChannelIndex] += SourceBuffer.ElementSize;
+								CurrentOffset += ChannelSize;
+							}
 
-				// Padding at the end?
-				int32 FinalPadding = ElementSize - CurrentOffset;
-				if (FinalPadding > 0)
-				{
-					FMemory::Memzero(TargetData, FinalPadding);
-					TargetData += FinalPadding;
-				}
-			}
+							// Padding at the end?
+							int32 FinalPadding = ElementSize - CurrentOffset;
+							if (FinalPadding > 0)
+							{
+								FMemory::Memzero(TargetData, FinalPadding);
+								TargetData += FinalPadding;
+							}
+						}
+					})
+			);
 		}
 	}
 
@@ -1166,8 +1170,66 @@ TSharedPtr<mu::FMesh> ConvertSkeletalMeshToMutable(
 
 	}
 
-	MutablePrivate::CopyBufferClearingPadding(MutableMesh->GetVertexBuffers(), UnrealSourceVertexFormat, Vertices.GetData() + VertexStart);
+	UE::Tasks::FTaskEvent CompletionEvent(TEXT("MeshConversionCompletion"));
+
+	MutablePrivate::CopyBufferClearingPadding(CompletionEvent, MutableMesh->GetVertexBuffers(), UnrealSourceVertexFormat, Vertices.GetData() + VertexStart);
+
+	// Indices
+	CompletionEvent.AddPrerequisites(
+		UE::Tasks::Launch(TEXT("BufferConversion"),
+			[&MeshSection, MutableMesh, &LODModel, VertexStart, VertexCount]()
+			{
+				const uint32 IndexStart = MeshSection.BaseIndex;
+				const uint32 IndexCount = MeshSection.NumTriangles * 3;
+				MutableMesh->GetIndexBuffers().SetBufferCount(1);
+				MutableMesh->GetIndexBuffers().SetElementCount(IndexCount);
+
+				using namespace mu;
+
+				check(LODModel.IndexBuffer.IsValidIndex(IndexStart) && LODModel.IndexBuffer.IsValidIndex(IndexStart + IndexCount - 1));
+				const uint32* IndexDataPtr = &LODModel.IndexBuffer[IndexStart];
+
+				const int32 FinalElementSize = sizeof(uint32);
+				const int32 ChannelCount = 1;
+				const EMeshBufferSemantic Semantics[ChannelCount] = { EMeshBufferSemantic::VertexIndex };
+				const int32 SemanticIndices[ChannelCount] = { 0 };
+				// We force 32 bit indices, since merging meshes may create vertex buffers bigger than the initial mesh
+				// and for now the mutable runtime doesn't handle it.
+				// \TODO: go back to 16-bit indices when possible.
+				EMeshBufferFormat Formats[ChannelCount] = { EMeshBufferFormat::UInt32 };
+				const int32 Components[ChannelCount] = { 1 };
+				const int32 Offsets[ChannelCount] = { 0 };
+
+				MutableMesh->GetIndexBuffers().SetBuffer(0, FinalElementSize, ChannelCount, Semantics, SemanticIndices, Formats, Components, Offsets);
+
+				uint32* pDest = reinterpret_cast<uint32*>(MutableMesh->GetIndexBuffers().GetBufferData(0));
+
+				// 32-bit to 32-bit
+				uint32 VertexIndex = 0;
+				for (uint32 Index = 0; Index < IndexCount; ++Index)
+				{
+					VertexIndex = *IndexDataPtr - VertexStart;
+					if (VertexIndex < (uint32)VertexCount)
+					{
+						*pDest = VertexIndex;
+					}
+					else
+					{
+						// Malformed mesh?
+						ensure(false);
+						*pDest = 0;
+					}
+					++pDest;
+					++IndexDataPtr;
+				}
+			}
+		)
+	);
 
+	// TODO: For now we have to wait here because code belows touches the vertex buffers.
+	// Restructure so that this doesn't happen.
+	CompletionEvent.Trigger();
+	CompletionEvent.Wait();
 
 	// TODO: Add Mesh generation flags to not include RT Morph and clothing if not needed.
 	int32 NextBufferIndex = MutableMesh->VertexBuffers.Buffers.Num();
@@ -1903,53 +1965,6 @@ TSharedPtr<mu::FMesh> ConvertSkeletalMeshToMutable(
 	}
 
 
-	// Indices
-	{
-		const uint32 IndexStart = MeshSection.BaseIndex;
-		const uint32 IndexCount = MeshSection.NumTriangles * 3;
-		MutableMesh->GetIndexBuffers().SetBufferCount(1);
-		MutableMesh->GetIndexBuffers().SetElementCount(IndexCount);
-
-		using namespace mu;
-
-		check(LODModel.IndexBuffer.IsValidIndex(IndexStart) && LODModel.IndexBuffer.IsValidIndex(IndexStart + IndexCount - 1));
-		const uint32* IndexDataPtr = &LODModel.IndexBuffer[IndexStart];
-
-		const int32 FinalElementSize = sizeof(uint32);
-		const int32 ChannelCount = 1;
-		const EMeshBufferSemantic Semantics[ChannelCount] = { EMeshBufferSemantic::VertexIndex };
-		const int32 SemanticIndices[ChannelCount] = { 0 };
-		// We force 32 bit indices, since merging meshes may create vertex buffers bigger than the initial mesh
-		// and for now the mutable runtime doesn't handle it.
-		// \TODO: go back to 16-bit indices when possible.
-		EMeshBufferFormat Formats[ChannelCount] = { EMeshBufferFormat::UInt32 };
-		const int32 Components[ChannelCount] = { 1 };
-		const int32 Offsets[ChannelCount] = { 0 };
-
-		MutableMesh->GetIndexBuffers().SetBuffer(0, FinalElementSize, ChannelCount, Semantics, SemanticIndices, Formats, Components, Offsets);
-
-		uint32* pDest = reinterpret_cast<uint32*>(MutableMesh->GetIndexBuffers().GetBufferData(0));
-
-		// 32-bit to 32-bit
-		uint32 VertexIndex = 0;
-		for (uint32 Index = 0; Index < IndexCount; ++Index)
-		{
-			VertexIndex = *IndexDataPtr - VertexStart;
-			if (ensureMsgf(VertexIndex < (uint32)VertexCount, TEXT("Mutable: VertexIndex >= VertexCount. VI [%d], VC [%d], VS [%d]. SKM [%s] LOD [%d] Section [%d]."),
-				VertexIndex, VertexCount, VertexStart,
-				*GetNameSafe(SkeletalMesh), Source.LODIndex, Source.SectionIndex))
-			{
-				*pDest = VertexIndex;
-			}
-			else
-			{
-				*pDest = 0;
-			}
-			++pDest;
-			++IndexDataPtr;
-		}
-	}
-
 	if (!bIgnorePhysics && SkeletalMesh->GetPhysicsAsset() && MutableMesh->GetSkeleton() && Context.Options.bPhysicsAssetMergeEnabled)
 	{
 		// Find BodySetups with relevant bones.
@@ -2123,6 +2138,18 @@ TSharedPtr<mu::FMesh> ConvertSkeletalMeshToMutable(
 		}
 
 		MutableMesh->SetPhysicsBody(PhysicsBody);
+
+		// Add tags
+		if (SkeletalMesh->GetPhysicsAsset() && PhysicsBody && PhysicsBody->GetBodyCount())
+		{
+			TSoftObjectPtr<UPhysicsAsset> PhysicsAsset = SkeletalMesh->GetPhysicsAsset();
+
+			const int32 AssetIndex = Context.PhysicsAssets.AddUnique(PhysicsAsset);
+			FString PhysicsAssetTag = FString("__PA:") + FString::FromInt(AssetIndex);
+
+			AddTagToMutableMeshUnique(*MutableMesh, PhysicsAssetTag);
+		}
+
 	}
 
 	// Set Bone Parenting usages. This has to be done after all primary usages are set.
@@ -2242,12 +2269,13 @@ TSharedPtr<mu::FMesh> ConvertSkeletalMeshToMutable(
 
 
 TSharedPtr<mu::FMesh> ConvertSkeletalMeshToMutable(
-	TSoftObjectPtr<const USkeletalMesh> InSkeletalMesh, const TSoftClassPtr<UAnimInstance>& AnimBp,
+	TSoftObjectPtr<UStreamableRenderAsset> InSkeletalMesh, const TSoftClassPtr<UAnimInstance>& AnimBp,
 	bool bMeshMustExist, int32 LODIndex, int32 SectionIndex, 
 	uint32 SurfaceMetadataId,
 	FMutableGraphGenerationContext& GenerationContext, 
 	const UCustomizableObjectNode* CurrentNode, 
-	USkeletalMesh* TableReferenceSkeletalMesh)
+	USkeletalMesh* TableReferenceSkeletalMesh,
+	bool bForceImmediateConversion)
 {
 	MUTABLE_CPUPROFILER_SCOPE(ConvertSkeletalMeshToMutable);
 
@@ -2363,8 +2391,17 @@ TSharedPtr<mu::FMesh> ConvertSkeletalMeshToMutable(
 		}
 	}
 
-	// Do it. In the future this will happen in a callback from the mutable core compilation process.
-	TSharedPtr<mu::FMesh> MutableMesh = ConvertSkeletalMeshToMutable(Source, GenerationContext.CompilationContext );
+	bool bIsPassthrough = false;
+	TSharedPtr<mu::FMesh> MutableMesh;
+	if (!bForceImmediateConversion)
+	{
+		MutableMesh = GenerateMeshConstant(Source, GenerationContext, bIsPassthrough);
+	}
+	else
+	{
+		// At some point this shouldn't happen anymore because all mesh conversion could be moved to the core compilation stage.
+		MutableMesh = ConvertSkeletalMeshToMutable(Source, GenerationContext.CompilationContext);
+	}
 
 	return MutableMesh;
 }
@@ -2559,9 +2596,9 @@ TSharedPtr<mu::FMesh> ConvertStaticMeshToMutable(const UStaticMesh* StaticMesh,
 
 // Convert a Mesh constant to a mutable format. UniqueTags are the tags that make this Mesh unique that cannot be merged in the cache 
 // with the exact same Mesh with other tags
-TSharedPtr<mu::FMesh> GenerateMutableMesh( TSoftObjectPtr<UStreamableRenderAsset> Mesh, const TSoftClassPtr<UAnimInstance>& AnimInstance, int32 LODIndexConnected, int32 SectionIndexConnected,
+TSharedPtr<mu::FMesh> GenerateMutableSkeletalMesh( TSoftObjectPtr<UStreamableRenderAsset> Mesh, const TSoftClassPtr<UAnimInstance>& AnimInstance, int32 LODIndexConnected, int32 SectionIndexConnected,
 									  int32 LODIndex, int32 SectionIndex, const FString& UniqueTags, uint32 SurfaceMetadataId, FMutableGraphGenerationContext & GenerationContext, 
-									  const UCustomizableObjectNode* CurrentNode, USkeletalMesh* TableReferenceSkeletalMesh, bool bIsReference)
+									  const UCustomizableObjectNode* CurrentNode, USkeletalMesh* TableReferenceSkeletalMesh, bool bIsPassthrough)
 {
 	// Get the mesh generation flags to use
 	EMutableMeshConversionFlags CurrentFlags = GenerationContext.MeshGenerationFlags.Last();
@@ -2581,25 +2618,47 @@ TSharedPtr<mu::FMesh> GenerateMutableMesh( TSoftObjectPtr<UStreamableRenderAsset
 		// Are we generating the base LOD?
 		((LODIndex == LODIndexConnected) && (SectionIndex == SectionIndexConnected));
 
-	// When we want to defer the mesh conversion to the core compilation stage, we will do this:
-	//MutableMesh = GenerateMeshConstant(Mesh, AnimInstance, bMeshMustExist, LODIndex, SectionIndex, SurfaceMetadataId, GenerationContext, bIsReference);
+	// At some point we will want all meshes to be references at compile-time. For now, just create the actual pass-through meshes.
+	if (bIsPassthrough)
+	{
+		FMutableSourceMeshData Source;
+		Source.Mesh = Mesh;
+		MutableMesh = GenerateMeshConstant(Source, GenerationContext, bIsPassthrough);
+	}
+	else
+	{
+		MutableMesh = ConvertSkeletalMeshToMutable(Mesh, AnimInstance, bMeshMustExist, LODIndex, SectionIndex, SurfaceMetadataId, GenerationContext, CurrentNode, TableReferenceSkeletalMesh);
+	}
 
-	// Meanwhile:
-	UObject* LoadedMesh = MutablePrivate::LoadObject(Mesh);
-	if (USkeletalMesh* SkeletalMesh = Cast<USkeletalMesh>(LoadedMesh))
+	if (MutableMesh)
 	{
+		GenerationContext.GeneratedMeshes.Push({ Key, MutableMesh });
+	}
+	
+	return MutableMesh;
+}
 
-		// At some point we will want all meshes to be references at compile-time. For now, just create the actual pass-through meshes.
-		if (bIsReference)
-		{
-			MutableMesh = GenerateMeshConstant(SkeletalMesh, GenerationContext, bIsReference);
-		}
-		else
-		{
-			MutableMesh = ConvertSkeletalMeshToMutable(SkeletalMesh, AnimInstance, bMeshMustExist, LODIndex, SectionIndex, SurfaceMetadataId, GenerationContext, CurrentNode, TableReferenceSkeletalMesh);
-		}
+
+// Convert a Mesh constant to a mutable format. UniqueTags are the tags that make this Mesh unique that cannot be merged in the cache 
+// with the exact same Mesh with other tags
+TSharedPtr<mu::FMesh> GenerateMutableStaticMesh(TSoftObjectPtr<UStreamableRenderAsset> Mesh, const TSoftClassPtr<UAnimInstance>& AnimInstance,
+	int32 LODIndex, int32 SectionIndex, const FString& UniqueTags, uint32 SurfaceMetadataId, FMutableGraphGenerationContext& GenerationContext,
+	const UCustomizableObjectNode* CurrentNode, USkeletalMesh* TableReferenceSkeletalMesh, bool bIsPassthrough)
+{
+	// Get the mesh generation flags to use
+	EMutableMeshConversionFlags CurrentFlags = GenerationContext.MeshGenerationFlags.Last();
+
+	FMutableGraphGenerationContext::FGeneratedMeshData::FKey Key = { Mesh, LODIndex, GenerationContext.CurrentLOD, SectionIndex, CurrentFlags, UniqueTags, CurrentNode };
+	TSharedPtr<mu::FMesh> MutableMesh = GenerationContext.FindGeneratedMesh(Key);
+	if (MutableMesh)
+	{
+		return MutableMesh;
 	}
-	else if (const UStaticMesh* StaticMesh = Cast<UStaticMesh>(LoadedMesh))
+
+	// When we want to defer the mesh conversion to the core compilation stage, but it is not supported for static meshes yet.
+	// Meanwhile:
+	UObject* LoadedMesh = MutablePrivate::LoadObject(Mesh);
+	if (const UStaticMesh* StaticMesh = Cast<UStaticMesh>(LoadedMesh))
 	{
 		MutableMesh = ConvertStaticMeshToMutable(StaticMesh, LODIndex, SectionIndex, GenerationContext, CurrentNode);
 	}
@@ -2612,13 +2671,15 @@ TSharedPtr<mu::FMesh> GenerateMutableMesh( TSoftObjectPtr<UStreamableRenderAsset
 	{
 		GenerationContext.GeneratedMeshes.Push({ Key, MutableMesh });
 	}
-	
+
 	return MutableMesh;
 }
 
 
-TSharedPtr<mu::FMesh> BuildMorphedMutableMeshFromMesh(TSharedPtr<mu::FMesh> BaseSourceMesh, USkeletalMesh* SkeletalMesh, const FString& MorphTargetName, int32 LODIndex, int32 SectionIndex)
+TSharedPtr<mu::FMesh> BuildMorphedMutableMeshFromMesh(TSharedPtr<mu::FMesh> BaseSourceMesh, const USkeletalMesh* SkeletalMesh, const FString& MorphTargetName, int32 LODIndex, int32 SectionIndex)
 {
+	MUTABLE_CPUPROFILER_SCOPE(BuildMorphedMutableMeshFromMesh);
+
 	// Clone it (it will probably be shared)
 	TSharedPtr<mu::FMesh> MorphedSourceMesh = BaseSourceMesh->Clone();
 
@@ -2722,13 +2783,13 @@ TSharedPtr<mu::FMesh> BuildMorphedMutableMesh(const UEdGraphPin* BaseSourcePin,
 	{
 		GetLODAndSectionForAutomaticLODs(GenerationContext, *Node, *SkeletalMesh, LODIndexConnected, SectionIndexConnected, LODIndex, SectionIndex, bOnlyConnectedLOD);
 		// Get the base mesh
-		constexpr bool bIsReference = false;
-		TSharedPtr<mu::FMesh> BaseSourceMesh = GenerateMutableMesh(SkeletalMesh, TSoftClassPtr<UAnimInstance>(), LODIndexConnected, SectionIndexConnected,
-			LODIndex, SectionIndex, FString(), 0, GenerationContext, Node, nullptr, bIsReference);
-		if (BaseSourceMesh)
-		{
-			MorphedSourceMesh = BuildMorphedMutableMeshFromMesh(BaseSourceMesh, SkeletalMesh, MorphTargetName, LODIndex, SectionIndex);
-		}
+		constexpr bool bIsPassthrough = false;
+		TSharedPtr<mu::FMesh> BaseSourceMesh = GenerateMutableSkeletalMesh(SkeletalMesh, TSoftClassPtr<UAnimInstance>(), LODIndexConnected, SectionIndexConnected,
+			LODIndex, SectionIndex, FString(), 0, GenerationContext, Node, nullptr, bIsPassthrough);
+		check(BaseSourceMesh->IsReference());
+		// The mesh will be modified and it may come from a cache, so we need to clone it.
+		MorphedSourceMesh = BaseSourceMesh->Clone(); 
+		MorphedSourceMesh->SetReferencedMorph(MorphTargetName);
 	}
 
 	return MorphedSourceMesh;
@@ -3511,41 +3572,19 @@ mu::Ptr<mu::NodeMesh> GenerateMutableSourceMesh(const UEdGraphPin* Pin,
 
 			FSkeletalMeshModel* ImportedModel = TypedNodeSkel->SkeletalMesh->GetImportedModel();
 			
-			constexpr bool bIsReference = false;
-			TSharedPtr<mu::FMesh> MutableMesh = GenerateMutableMesh(TypedNodeSkel->SkeletalMesh, TypedNodeSkel->AnimInstance, LODIndexConnected, SectionIndexConnected, 
-																LODIndex, SectionIndex, MeshUniqueTags, SurfaceMetadataId, GenerationContext, TypedNodeSkel, nullptr, bIsReference);
+			constexpr bool bIsPassthrough = false;
+			TSharedPtr<mu::FMesh> MutableMesh = GenerateMutableSkeletalMesh(TypedNodeSkel->SkeletalMesh, TypedNodeSkel->AnimInstance, LODIndexConnected, SectionIndexConnected,
+																LODIndex, SectionIndex, MeshUniqueTags, SurfaceMetadataId, GenerationContext, TypedNodeSkel, nullptr, bIsPassthrough);
 			if (MutableMesh)
 			{
 				MeshNode->SetValue(MutableMesh);
 
-				// TODO
-				// Add the potentially required morphs. For now always add all morphs. A better implementation can narrow down with:
-				// - detecting morph names in all "ModifierMorphMeshSection" nodes
-				// - detecting morph names in "ModifierMorphMeshSection" nodes relevant for this current context.
-				{
-					MUTABLE_CPUPROFILER_SCOPE(GenerateMutableSourceMesh_AddAllMorphs);
-
-					const TArray<TObjectPtr<UMorphTarget>>& Morphs = TypedNodeSkel->SkeletalMesh->GetMorphTargets();
-					for (TObjectPtr<UMorphTarget> Morph : Morphs)
-					{
-						FString MorphTargetName = Morph->GetName();
-						TSharedPtr<mu::FMesh> MorphedMesh = BuildMorphedMutableMeshFromMesh(MutableMesh, TypedNodeSkel->SkeletalMesh.Get(), MorphTargetName, LODIndex, SectionIndex);
-						MeshNode->AddMorph(MorphTargetName,MorphedMesh);
-					}
-				}
-
 				if (TypedNodeSkel->SkeletalMesh->GetPhysicsAsset() && 
 					MutableMesh->GetPhysicsBody() && 
 					MutableMesh->GetPhysicsBody()->GetBodyCount())
 				{
 					TSoftObjectPtr<UPhysicsAsset> PhysicsAsset = TypedNodeSkel->SkeletalMesh->GetPhysicsAsset();
-
 					GenerationContext.AddParticipatingObject(*PhysicsAsset);
-
-					const int32 AssetIndex = GenerationContext.CompilationContext.PhysicsAssets.AddUnique(PhysicsAsset);
-					FString PhysicsAssetTag = FString("__PA:") + FString::FromInt(AssetIndex);
-
-					AddTagToMutableMeshUnique(*MutableMesh, PhysicsAssetTag);
 				}
 
 				if (GenerationContext.CompilationContext.Options.bSkinWeightProfilesEnabled && LODIndex>=0)
@@ -3732,9 +3771,9 @@ mu::Ptr<mu::NodeMesh> GenerateMutableSourceMesh(const UEdGraphPin* Pin,
 			
 			check(TypedNodeStatic->LODs.IsValidIndex(LODIndex) && SectionIndex < TypedNodeStatic->LODs[LODIndex].Materials.Num());
 
-			constexpr bool bIsReference = false;
-			TSharedPtr<mu::FMesh> MutableMesh = GenerateMutableMesh(TypedNodeStatic->StaticMesh, TSoftClassPtr<UAnimInstance>(), LODIndex, SectionIndex, LODIndex, SectionIndex, 
-														  FString(), 0, GenerationContext, TypedNodeStatic, nullptr, bIsReference);
+			constexpr bool bIsPassthrough = false;
+			TSharedPtr<mu::FMesh> MutableMesh = GenerateMutableStaticMesh(TypedNodeStatic->StaticMesh, TSoftClassPtr<UAnimInstance>(), LODIndex, SectionIndex,
+														  FString(), 0, GenerationContext, TypedNodeStatic, nullptr, bIsPassthrough);
 			if (MutableMesh)
 			{
 				MeshNode->SetValue(MutableMesh);
diff --git a/Engine/Plugins/Mutable/Source/CustomizableObjectEditor/Private/MuCOE/GenerateMutableSource/GenerateMutableSourceMesh.h b/Engine/Plugins/Mutable/Source/CustomizableObjectEditor/Private/MuCOE/GenerateMutableSource/GenerateMutableSourceMesh.h
--- a/Engine/Plugins/Mutable/Source/CustomizableObjectEditor/Private/MuCOE/GenerateMutableSource/GenerateMutableSourceMesh.h
+++ b/Engine/Plugins/Mutable/Source/CustomizableObjectEditor/Private/MuCOE/GenerateMutableSource/GenerateMutableSourceMesh.h
@@ -50,20 +50,20 @@ void GetLODAndSectionForAutomaticLODs(const FMutableGraphGenerationContext& Cont
  * @param CurrentNode The node being compiled, can be a skeletal mesh or a table node.
  * @param TableReferenceSkeletalMesh Only set if compiling a table node. It refers to the reference mesh of the skeletal mesh column being compiled.
  * @return Mutable Mesh. Nullptr if there has been an error. Empty mesh if the Skeletal Mesh does not contain the requested LOD + Section. */
-TSharedPtr<mu::FMesh> ConvertSkeletalMeshToMutable(TSoftObjectPtr<const USkeletalMesh> InSkeletalMesh, const TSoftClassPtr<UAnimInstance>& AnimBp,
+TSharedPtr<mu::FMesh> ConvertSkeletalMeshToMutable(TSoftObjectPtr<UStreamableRenderAsset> InSkeletalMesh, const TSoftClassPtr<UAnimInstance>& AnimBp,
                                          bool bMeshMustExist,
                                          int32 LODIndex, int32 SectionIndex, uint32 SurfaceMetadataId,
                                          FMutableGraphGenerationContext& GenerationContext,
                                          const UCustomizableObjectNode* CurrentNode, 
-										 USkeletalMesh* TableReferenceSkeletalMesh);
+										 USkeletalMesh* TableReferenceSkeletalMesh,
+										 bool bForceImmediateConversion = false );
 
 
 TSharedPtr<mu::FMesh> ConvertStaticMeshToMutable(const UStaticMesh* StaticMesh, int32 LODIndex, int32 SectionIndex,
 	FMutableGraphGenerationContext& GenerationContext, const UCustomizableObjectNode* CurrentNode);
 
 
-/** Compiler recursive function. Mutable Mesh.
- *
+/** 
  * @param Mesh Unreal Mesh.
  * @param AnimBp 
  * @param LODIndexConnected LOD which the pin is connected to.
@@ -75,12 +75,19 @@ TSharedPtr<mu::FMesh> ConvertStaticMeshToMutable(const UStaticMesh* StaticMesh,
  * @param CurrentNode The node being compiled, can be a skeletal mesh or a table node.
  * @param TableReferenceSkeletalMesh Only set if compiling a table node. It refers to the reference mesh of the skeletal mesh column being compiled.
  * @return Mutable Mesh. Nullptr if there has been an error. Empty mesh if the Skeletal Mesh does not contain the requested LOD + Section. */
-TSharedPtr<mu::FMesh> GenerateMutableMesh(TSoftObjectPtr<UStreamableRenderAsset> Mesh, const TSoftClassPtr<UAnimInstance>& AnimBp,
+TSharedPtr<mu::FMesh> GenerateMutableSkeletalMesh(TSoftObjectPtr<UStreamableRenderAsset> Mesh, const TSoftClassPtr<UAnimInstance>& AnimBp,
                                 int32 LODIndexConnected, int32 SectionIndexConnected,
                                 int32 LODIndex, int32 SectionIndex, const FString& MeshUniqueTags, uint32 SurfaceMetadataId, 
                                 FMutableGraphGenerationContext& GenerationContext, 
 								const UCustomizableObjectNode* CurrentNode, USkeletalMesh* TableReferenceSkeletalMesh,
-								bool bIsReference);
+								bool bIsPassthrough);
+
+/** */
+TSharedPtr<mu::FMesh> GenerateMutableStaticMesh(TSoftObjectPtr<UStreamableRenderAsset> Mesh, const TSoftClassPtr<UAnimInstance>& AnimBp,
+	int32 LODIndex, int32 SectionIndex, const FString& MeshUniqueTags, uint32 SurfaceMetadataId,
+	FMutableGraphGenerationContext& GenerationContext,
+	const UCustomizableObjectNode* CurrentNode, USkeletalMesh* TableReferenceSkeletalMesh,
+	bool bIsPassthrough);
 
 
 TSharedPtr<mu::FMesh> BuildMorphedMutableMesh(const UEdGraphPin* BaseSourcePin, const FString& MorphTargetName,
diff --git a/Engine/Plugins/Mutable/Source/CustomizableObjectEditor/Private/MuCOE/GenerateMutableSource/GenerateMutableSourceTable.cpp b/Engine/Plugins/Mutable/Source/CustomizableObjectEditor/Private/MuCOE/GenerateMutableSource/GenerateMutableSourceTable.cpp
--- a/Engine/Plugins/Mutable/Source/CustomizableObjectEditor/Private/MuCOE/GenerateMutableSource/GenerateMutableSourceTable.cpp
+++ b/Engine/Plugins/Mutable/Source/CustomizableObjectEditor/Private/MuCOE/GenerateMutableSource/GenerateMutableSourceTable.cpp
@@ -48,9 +48,10 @@ bool FillTableColumn(const UCustomizableObjectNodeTable* TableNode, mu::Ptr<mu::
 		{
 			GenerationContext.AddParticipatingObject(SoftObject);
 
+			// TODO: This is still forcing the load of the skeletal mesh asset. 
 			UObject* Object = GenerationContext.LoadObject(SoftObject, true);
-
 			USkeletalMesh* SkeletalMesh = Cast<USkeletalMesh>(Object);
+
 			if(!SkeletalMesh)
 			{
 				// Generating an Empty cell
@@ -278,22 +279,17 @@ bool FillTableColumn(const UCustomizableObjectNodeTable* TableNode, mu::Ptr<mu::
 			}
 
 			//TODO: Add AnimBp physics to Tables.
-			constexpr bool bIsReference = false;
-			TSharedPtr<mu::FMesh> MutableMesh = GenerateMutableMesh(SkeletalMesh, TSoftClassPtr<UAnimInstance>(), LODIndexConnected, SectionIndexConnected, 
-																LODIndex, SectionIndex, MeshUniqueTags, SectionMetadataId, GenerationContext, TableNode, ReferenceSkeletalMesh, bIsReference);
+			constexpr bool bIsPassthrough = false;
+			TSoftObjectPtr<UStreamableRenderAsset> SkeletalMeshSoftObjectPtr(SoftObject.ToSoftObjectPath());
+			TSharedPtr<mu::FMesh> MutableMesh = GenerateMutableSkeletalMesh(SkeletalMeshSoftObjectPtr, TSoftClassPtr<UAnimInstance>(), LODIndexConnected, SectionIndexConnected,
+																LODIndex, SectionIndex, MeshUniqueTags, SectionMetadataId, GenerationContext, TableNode, ReferenceSkeletalMesh, bIsPassthrough);
 
 			if (MutableMesh)
 			{
- 				if (SkeletalMesh->GetPhysicsAsset() && MutableMesh->GetPhysicsBody() && MutableMesh->GetPhysicsBody()->GetBodyCount())
+				if (SkeletalMesh->GetPhysicsAsset())
 				{
- 					TSoftObjectPtr<UPhysicsAsset> PhysicsAsset = SkeletalMesh->GetPhysicsAsset();
-
+					TSoftObjectPtr<UPhysicsAsset> PhysicsAsset = SkeletalMesh->GetPhysicsAsset();
 					GenerationContext.AddParticipatingObject(*PhysicsAsset);
-
-					const int32 AssetIndex = GenerationContext.CompilationContext.PhysicsAssets.AddUnique(PhysicsAsset);
-					FString PhysicsAssetTag = FString("__PA:") + FString::FromInt(AssetIndex);
-
-					AddTagToMutableMeshUnique(*MutableMesh, PhysicsAssetTag);
 				}
 
 				if (!AnimBPAssetTag.IsEmpty())
@@ -389,9 +385,9 @@ bool FillTableColumn(const UCustomizableObjectNodeTable* TableNode, mu::Ptr<mu::
 				CurrentColumn = MutableTable->AddColumn(MutableColumnName, mu::ETableColumnType::Mesh);
 			}
 
-			constexpr bool bIsReference = false;
-			TSharedPtr<mu::FMesh> MutableMesh = GenerateMutableMesh(StaticMesh, TSoftClassPtr<UAnimInstance>(), CurrentLOD, SectionIndex, CurrentLOD, SectionIndex, 
-														  FString(), 0, GenerationContext, TableNode, nullptr, bIsReference);
+			constexpr bool bIsPassthrough = false;
+			TSharedPtr<mu::FMesh> MutableMesh = GenerateMutableStaticMesh(StaticMesh, TSoftClassPtr<UAnimInstance>(), CurrentLOD, SectionIndex,
+														  FString(), 0, GenerationContext, TableNode, nullptr, bIsPassthrough);
 
 			if (MutableMesh)
 			{
@@ -935,7 +931,23 @@ bool GenerateTableColumn(const UCustomizableObjectNodeTable* TableNode, const UE
 	TArray<FName> RowNames = GetRowsToCompile(*DataTable, *TableNode, GenerationContext, RowIds);
 
 	// Pre-pass to request async loading of all data. This seems to be slightly faster because it avoids the sync after every separate load.
-	if (!GenerationContext.bParticipatingObjectsPass)
+	bool bPreload = !GenerationContext.bParticipatingObjectsPass;
+	//if (bPreload)
+	//{
+	//	if (const FSoftObjectProperty* SoftObjectProperty = CastField<FSoftObjectProperty>(ColumnProperty))
+	//	{
+	//		if (SoftObjectProperty->PropertyClass->IsChildOf(USkeletalMesh::StaticClass()))
+	//		{
+	//			bPreload = false;
+	//		}
+	//		else if (SoftObjectProperty->PropertyClass->IsChildOf(UStaticMesh::StaticClass()))
+	//		{
+	//			bPreload = false;
+	//		}
+	//	}
+	//}
+
+	if (bPreload)
 	{
 		FCookLoadScope CookLoadScope(ECookLoadType::EditorOnly);
 
diff --git a/Engine/Plugins/Mutable/Source/CustomizableObjectEditor/Private/MuCOE/SMutableCodeViewer.cpp b/Engine/Plugins/Mutable/Source/CustomizableObjectEditor/Private/MuCOE/SMutableCodeViewer.cpp
--- a/Engine/Plugins/Mutable/Source/CustomizableObjectEditor/Private/MuCOE/SMutableCodeViewer.cpp
+++ b/Engine/Plugins/Mutable/Source/CustomizableObjectEditor/Private/MuCOE/SMutableCodeViewer.cpp
@@ -296,7 +296,7 @@ void SMutableCodeViewer::ClearSelectedTreeRow() const
 
 void SMutableCodeViewer::SetCurrentModel(const TSharedPtr<mu::FModel, ESPMode::ThreadSafe>& InMutableModel, 
 	const TArray<TSoftObjectPtr<const UTexture>>& InReferencedTextures,
-	const TArray<TSoftObjectPtr<const UObject>>& InReferencedMeshes )
+	const TArray<TSoftObjectPtr<const UStreamableRenderAsset>>& InReferencedMeshes )
 {
 	MutableModel = InMutableModel;
 	ReferencedTextures = InReferencedTextures;
@@ -347,7 +347,7 @@ void SMutableCodeViewer::SetCurrentModel(const TSharedPtr<mu::FModel, ESPMode::T
 
 void SMutableCodeViewer::Construct(const FArguments& InArgs, const TSharedPtr<mu::FModel, ESPMode::ThreadSafe>& InMutableModel,
 	const TArray<TSoftObjectPtr<const UTexture>>& InReferencedTextures,
-	const TArray<TSoftObjectPtr<const UObject>>& InReferencedMeshes )
+	const TArray<TSoftObjectPtr<const UStreamableRenderAsset>>& InReferencedMeshes )
 {
 	// Min width allowed for the column. Needed to avoid having issues with the constants space being to small
 	// and then getting too tall on the y axis crashing the UI drawer.
@@ -2482,7 +2482,7 @@ namespace
 
 		/** */
 		TArray<TSoftObjectPtr<const UTexture>> ReferencedTextures;
-		TArray<TSoftObjectPtr<const UObject>> ReferencedMeshes;
+		TArray<TSoftObjectPtr<const UStreamableRenderAsset>> ReferencedMeshes;
 
 	public:
 
diff --git a/Engine/Plugins/Mutable/Source/CustomizableObjectEditor/Private/MuCOE/SMutableCodeViewer.h b/Engine/Plugins/Mutable/Source/CustomizableObjectEditor/Private/MuCOE/SMutableCodeViewer.h
--- a/Engine/Plugins/Mutable/Source/CustomizableObjectEditor/Private/MuCOE/SMutableCodeViewer.h
+++ b/Engine/Plugins/Mutable/Source/CustomizableObjectEditor/Private/MuCOE/SMutableCodeViewer.h
@@ -61,7 +61,7 @@ public:
 
 	void Construct(const FArguments& InArgs, const TSharedPtr<mu::FModel, ESPMode::ThreadSafe>& InMutableModel,
 		const TArray<TSoftObjectPtr<const UTexture>>& ReferencedTextures,
-		const TArray<TSoftObjectPtr<const UObject>>& ReferencedMeshes);
+		const TArray<TSoftObjectPtr<const UStreamableRenderAsset>>& ReferencedMeshes);
 
 	// SWidget interface
 	virtual void Tick(const FGeometry& AllottedGeometry, const double InCurrentTime, const float InDeltaTime) override;
@@ -80,7 +80,7 @@ private:
 	
 	/** Array of external referenced textures in MutableModel, indexed by id. */
 	TArray<TSoftObjectPtr<const UTexture>> ReferencedTextures;
-	TArray<TSoftObjectPtr<const UObject>> ReferencedMeshes;
+	TArray<TSoftObjectPtr<const UStreamableRenderAsset>> ReferencedMeshes;
 
 	/** Selected model operation for preview. */
 	mu::OP::ADDRESS SelectedOperationAddress = 0;
@@ -188,7 +188,7 @@ private:
 	TMap< mu::OP::ADDRESS, TSharedPtr<FMutableCodeTreeElement>> ExpandedElements;
 
 	/** Prepare the widget for the given model. */
-	void SetCurrentModel(const TSharedPtr<mu::FModel, ESPMode::ThreadSafe>&, const TArray<TSoftObjectPtr<const UTexture>>& ReferencedTextures, const TArray<TSoftObjectPtr<const UObject>>& ReferencedMeshes);
+	void SetCurrentModel(const TSharedPtr<mu::FModel, ESPMode::ThreadSafe>&, const TArray<TSoftObjectPtr<const UTexture>>& ReferencedTextures, const TArray<TSoftObjectPtr<const UStreamableRenderAsset>>& ReferencedMeshes);
 
 	/** Before any UI operation generate all the elements that may be navigable over the tree. No children of duplicated 
 	 * addresses will be generated.
diff --git a/Engine/Plugins/Mutable/Source/CustomizableObjectEditor/Private/MuCOE/SMutableGraphViewer.h b/Engine/Plugins/Mutable/Source/CustomizableObjectEditor/Private/MuCOE/SMutableGraphViewer.h
--- a/Engine/Plugins/Mutable/Source/CustomizableObjectEditor/Private/MuCOE/SMutableGraphViewer.h
+++ b/Engine/Plugins/Mutable/Source/CustomizableObjectEditor/Private/MuCOE/SMutableGraphViewer.h
@@ -37,7 +37,7 @@ public:
 
 	SLATE_ARGUMENT(TArray<TSoftObjectPtr<const UTexture>>, ReferencedRuntimeTextures)
 	SLATE_ARGUMENT(TArray<FMutableSourceTextureData>, ReferencedCompileTextures)
-	SLATE_ARGUMENT(TArray<TSoftObjectPtr<const UObject>>, ReferencedRuntimeMeshes)
+	SLATE_ARGUMENT(TArray<TSoftObjectPtr<const UStreamableRenderAsset>>, ReferencedRuntimeMeshes)
 	SLATE_ARGUMENT(TArray<FMutableSourceMeshData>, ReferencedCompileMeshes)
 
 	SLATE_END_ARGS()
diff --git a/Engine/Plugins/Mutable/Source/CustomizableObjectEditor/Private/MuCOE/SMutableObjectViewer.cpp b/Engine/Plugins/Mutable/Source/CustomizableObjectEditor/Private/MuCOE/SMutableObjectViewer.cpp
--- a/Engine/Plugins/Mutable/Source/CustomizableObjectEditor/Private/MuCOE/SMutableObjectViewer.cpp
+++ b/Engine/Plugins/Mutable/Source/CustomizableObjectEditor/Private/MuCOE/SMutableObjectViewer.cpp
@@ -110,7 +110,7 @@ void SMutableObjectViewer::GenerateMutableGraphPressed()
 	FCustomizableObjectCompiler Compiler;
 	TArray<TSoftObjectPtr<const UTexture>> RuntimeTextures;
 	TArray<FMutableSourceTextureData> CompilerTextures;
-	TArray<TSoftObjectPtr<const UObject>> RuntimeMeshes;
+	TArray<TSoftObjectPtr<const UStreamableRenderAsset>> RuntimeMeshes;
 	TArray<FMutableSourceMeshData> CompilerMeshes;
 	mu::Ptr<mu::Node> RootNode = Compiler.Export(CustomizableObject, CompileOptions, RuntimeTextures, CompilerTextures, RuntimeMeshes, CompilerMeshes);
 	if (!RootNode)
@@ -137,7 +137,7 @@ void SMutableObjectViewer::CompileMutableCodePressed()
 {
 	TArray<TSoftObjectPtr<const UTexture>> RuntimeTextures;
 	TArray<FMutableSourceTextureData> CompilerTextures;
-	TArray<TSoftObjectPtr<const UObject>> RuntimeMeshes;
+	TArray<TSoftObjectPtr<const UStreamableRenderAsset>> RuntimeMeshes;
 	TArray<FMutableSourceMeshData> CompilerMeshes;
 	if (CompileOptions.bForceLargeLODBias)
 	{
@@ -161,7 +161,7 @@ void SMutableObjectViewer::CompileMutableCodePressed()
 			}
 
 			// Do the compilation to Mutable Code synchronously.
-			TSharedPtr<FCustomizableObjectCompileRunnable> CompileTask = MakeShareable(new FCustomizableObjectCompileRunnable(RootNode));
+			TSharedPtr<FCustomizableObjectCompileRunnable> CompileTask = MakeShareable(new FCustomizableObjectCompileRunnable(RootNode, Compiler.CompilationContext));
 			CompileTask->Options = CompileOptions;
 			CompileTask->ReferencedTextures = CompilerTextures;
 			CompileTask->ReferencedMeshes = CompilerMeshes;
@@ -185,7 +185,7 @@ void SMutableObjectViewer::CompileMutableCodePressed()
 	}
 
 	// Do the compilation to Mutable Code synchronously.
-	TSharedPtr<FCustomizableObjectCompileRunnable> CompileTask = MakeShareable(new FCustomizableObjectCompileRunnable(RootNode));
+	TSharedPtr<FCustomizableObjectCompileRunnable> CompileTask = MakeShareable(new FCustomizableObjectCompileRunnable(RootNode, Compiler.CompilationContext));
 	CompileTask->Options = CompileOptions;
 	CompileTask->ReferencedTextures = CompilerTextures;
 	CompileTask->ReferencedMeshes = CompilerMeshes;
diff --git a/Engine/Plugins/Mutable/Source/MutableRuntime/Private/MuR/Mesh.cpp b/Engine/Plugins/Mutable/Source/MutableRuntime/Private/MuR/Mesh.cpp
--- a/Engine/Plugins/Mutable/Source/MutableRuntime/Private/MuR/Mesh.cpp
+++ b/Engine/Plugins/Mutable/Source/MutableRuntime/Private/MuR/Mesh.cpp
@@ -75,6 +75,18 @@ uint32 FMesh::GetReferencedMesh() const
 }
 
 
+void FMesh::SetReferencedMorph(const FString& MorphName)
+{
+	ReferencedMorph = MorphName;
+}
+
+
+const FString& FMesh::GetReferencedMorph() const
+{
+	return ReferencedMorph;
+}
+
+
 TSharedPtr<FMesh> FMesh::Clone() const
 {
     //MUTABLE_CPUPROFILER_SCOPE(MeshClone);
@@ -1031,7 +1043,12 @@ void FMesh::Serialise(FOutputArchive& Arch) const
 	Arch << AdditionalPhysicsBodies;
 
 	Arch << MeshIDPrefix;
-	Arch << ReferenceID;
+
+	if ( IsReference() )
+	{
+		Arch << ReferenceID;
+		Arch << ReferencedMorph;
+	}
 }
 
 
@@ -1062,7 +1079,12 @@ void FMesh::Unserialise(FInputArchive& Arch)
 	Arch >> AdditionalPhysicsBodies;
 
 	Arch >> MeshIDPrefix;
-	Arch >> ReferenceID;
+
+	if (IsReference())
+	{
+		Arch >> ReferenceID;
+		Arch >> ReferencedMorph;
+	}
 }
 
 
diff --git a/Engine/Plugins/Mutable/Source/MutableRuntime/Private/MuR/OpMeshClipWithMesh.cpp b/Engine/Plugins/Mutable/Source/MutableRuntime/Private/MuR/OpMeshClipWithMesh.cpp
--- a/Engine/Plugins/Mutable/Source/MutableRuntime/Private/MuR/OpMeshClipWithMesh.cpp
+++ b/Engine/Plugins/Mutable/Source/MutableRuntime/Private/MuR/OpMeshClipWithMesh.cpp
@@ -851,7 +851,12 @@ namespace mu
 		MUTABLE_CPUPROFILER_SCOPE(MakeMeshMaskFromLayout);
 
 		check(Result && Base && Mask && BaseForUVs);
-		check(Base->VertexBuffers.GetElementCount()== BaseForUVs->VertexBuffers.GetElementCount());
+		//check(Base->VertexBuffers.GetElementCount()== BaseForUVs->VertexBuffers.GetElementCount());
+		if (Base->VertexBuffers.GetElementCount() != BaseForUVs->VertexBuffers.GetElementCount())
+		{
+			ensure(false);
+			return;
+		}
 
 		bOutSuccess = true;
 
diff --git a/Engine/Plugins/Mutable/Source/MutableRuntime/Private/MuR/Operations.cpp b/Engine/Plugins/Mutable/Source/MutableRuntime/Private/MuR/Operations.cpp
--- a/Engine/Plugins/Mutable/Source/MutableRuntime/Private/MuR/Operations.cpp
+++ b/Engine/Plugins/Mutable/Source/MutableRuntime/Private/MuR/Operations.cpp
@@ -713,10 +713,11 @@ namespace mu
 
 		case EOpType::ME_MASKCLIPUVMASK:
 		{
-			OP::MeshMaskClipUVMaskArgs args = program.GetOpArgs<OP::MeshMaskClipUVMaskArgs>(at);
-			f(args.Source);
-			f(args.MaskImage);
-			f(args.MaskLayout);
+			OP::MeshMaskClipUVMaskArgs Args = program.GetOpArgs<OP::MeshMaskClipUVMaskArgs>(at);
+			f(Args.Source);
+			f(Args.UVSource);
+			f(Args.MaskImage);
+			f(Args.MaskLayout);
 			break;
 		}
 
diff --git a/Engine/Plugins/Mutable/Source/MutableRuntime/Public/MuR/Mesh.h b/Engine/Plugins/Mutable/Source/MutableRuntime/Public/MuR/Mesh.h
--- a/Engine/Plugins/Mutable/Source/MutableRuntime/Public/MuR/Mesh.h
+++ b/Engine/Plugins/Mutable/Source/MutableRuntime/Public/MuR/Mesh.h
@@ -210,6 +210,10 @@ namespace mu
 		/** Return the id of the engine referenced mesh. Only valid if IsReference. */
 		uint32 GetReferencedMesh() const;
 
+		/** Mesh references can actually reference a morphed mesh. In that case you can set and get the morph with these functions. */
+		void SetReferencedMorph(const FString& MorphName);
+		const FString& GetReferencedMorph() const;
+
         //! \name Buffers
         //! \{
 
@@ -379,6 +383,9 @@ namespace mu
 		/** Only valid if the right flags are set, this identifies a referenced mesh. */
 		uint32 ReferenceID = 0;
 
+		/** If the mesh is a reference the referenced morph name is stored here. Otherwise it is an empty string. */
+		FString ReferencedMorph;
+
 		/** 
 		 * Prefix for the unique IDs related to this mesh (vertices and layout blocks). Useful if the mesh stores them in an implicit, or relative way. 
 		 * See MeshVertexIdIterator for details.
diff --git a/Engine/Plugins/Mutable/Source/MutableTools/Internal/MuT/AST.h b/Engine/Plugins/Mutable/Source/MutableTools/Internal/MuT/AST.h
--- a/Engine/Plugins/Mutable/Source/MutableTools/Internal/MuT/AST.h
+++ b/Engine/Plugins/Mutable/Source/MutableTools/Internal/MuT/AST.h
@@ -540,6 +540,8 @@ namespace mu
         //! from this TArray.
 		TArray<ASTOp*, TInlineAllocator<4> > Parents;
 
+		UE::FMutex ParentsMutex;
+
     public:
 
         inline ASTOp()
@@ -1210,6 +1212,7 @@ namespace mu
 
     inline void ASTChild::AddParent()
     {
+		UE::TScopeLock Lock(Child->ParentsMutex);
         ParentIndexInChild = Child->Parents.Num();
         Child->Parents.Add(Parent);
     }
@@ -1217,7 +1220,9 @@ namespace mu
 
 	inline void ASTChild::ClearParent()
     {
-        check( ParentIndexInChild<Child->Parents.Num() );
+		UE::TScopeLock Lock(Child->ParentsMutex);
+		
+		check( ParentIndexInChild<Child->Parents.Num() );
 		// Can't do this, because the indices are stored in children.
 		//Child->m_parents.RemoveAtSwap(ParentIndexInChild);
 		Child->Parents[ParentIndexInChild] = nullptr;
diff --git a/Engine/Plugins/Mutable/Source/MutableTools/Internal/MuT/NodeMeshConstantPrivate.h b/Engine/Plugins/Mutable/Source/MutableTools/Internal/MuT/NodeMeshConstantPrivate.h
--- a/Engine/Plugins/Mutable/Source/MutableTools/Internal/MuT/NodeMeshConstantPrivate.h
+++ b/Engine/Plugins/Mutable/Source/MutableTools/Internal/MuT/NodeMeshConstantPrivate.h
@@ -20,15 +20,6 @@ namespace mu
 		TSharedPtr<FMesh> Value;
 
 		TArray<Ptr<NodeLayout>> Layouts;
-
-		struct FMorph
-		{
-			FString Name;
-			TSharedPtr<FMesh> MorphedMesh;
-		};
-
-		TArray<FMorph> Morphs;
-
 	};
 
 }
diff --git a/Engine/Plugins/Mutable/Source/MutableTools/Private/MuT/AST.cpp b/Engine/Plugins/Mutable/Source/MutableTools/Private/MuT/AST.cpp
--- a/Engine/Plugins/Mutable/Source/MutableTools/Private/MuT/AST.cpp
+++ b/Engine/Plugins/Mutable/Source/MutableTools/Private/MuT/AST.cpp
@@ -825,14 +825,18 @@ int32 ASTOp::GetParentCount() const
 }
 
 
-void ASTOp::Replace( const Ptr<ASTOp>& node, const Ptr<ASTOp>& other )
+void ASTOp::Replace( const Ptr<ASTOp>& Node, const Ptr<ASTOp>& Other)
 {
-    if(other==node)
+    if(Other == Node)
     {
         return;
     }
 
-	TArray<ASTOp*, TInlineAllocator<4> > ParentsCopy = node->Parents;
+	TArray<ASTOp*, TInlineAllocator<4> > ParentsCopy;
+	{
+		UE::TScopeLock Lock(Node->ParentsMutex);
+		ParentsCopy = Node->Parents;
+	}
 
     for(ASTOp* Parent: ParentsCopy)
     {
@@ -840,9 +844,9 @@ void ASTOp::Replace( const Ptr<ASTOp>& node, const Ptr<ASTOp>& other )
         {
 			Parent->ForEachChild( [=](ASTChild& c)
             {
-                if(c.Child==node)
+                if(c.Child==Node)
                 {
-                    c = other;
+                    c = Other;
                 }
             });
         }
diff --git a/Engine/Plugins/Mutable/Source/MutableTools/Private/MuT/ASTOpSwitch.cpp b/Engine/Plugins/Mutable/Source/MutableTools/Private/MuT/ASTOpSwitch.cpp
--- a/Engine/Plugins/Mutable/Source/MutableTools/Private/MuT/ASTOpSwitch.cpp
+++ b/Engine/Plugins/Mutable/Source/MutableTools/Private/MuT/ASTOpSwitch.cpp
@@ -426,7 +426,7 @@ namespace mu
 				if (cases[o].branch &&
 					typedCondition->op.args.IntConstant.value
 					==
-					(int)cases[o].condition)
+					cases[o].condition)
 				{
 					Branch = cases[o].branch.child();
 					break;
@@ -441,7 +441,10 @@ namespace mu
 			// If all the branches for the possible values are the same op remove the instruction
 			const ASTOpParameter* ParamOp = static_cast<const ASTOpParameter*>(variable.child().get());
 			check(ParamOp);
-			check(!ParamOp->Parameter.PossibleValues.IsEmpty());
+			if(ParamOp->Parameter.PossibleValues.IsEmpty())
+			{
+				return nullptr;
+			}
 
 			bool bFirstValue = true;
 			bool bAllSame = true;
diff --git a/Engine/Plugins/Mutable/Source/MutableTools/Private/MuT/CodeGenerator.cpp b/Engine/Plugins/Mutable/Source/MutableTools/Private/MuT/CodeGenerator.cpp
--- a/Engine/Plugins/Mutable/Source/MutableTools/Private/MuT/CodeGenerator.cpp
+++ b/Engine/Plugins/Mutable/Source/MutableTools/Private/MuT/CodeGenerator.cpp
@@ -48,6 +48,7 @@
 #include "MuT/ASTOpLayoutFromMesh.h"
 #include "MuT/ASTOpLayoutMerge.h"
 #include "MuT/ASTOpLayoutPack.h"
+#include "MuT/ASTOpReferenceResource.h"
 #include "MuT/CodeGenerator_SecondPass.h"
 #include "MuT/CodeOptimiser.h"
 #include "MuT/CompilerPrivate.h"
@@ -88,6 +89,7 @@
 #include "MuT/NodeRangeFromScalar.h"
 #include "MuT/NodeScalar.h"
 #include "MuT/NodeScalarConstant.h"
+#include "MuT/NodeScalarEnumParameter.h"
 #include "MuT/NodeSurface.h"
 #include "MuT/NodeSurfaceNew.h"
 #include "MuT/NodeSurfaceVariation.h"
@@ -103,6 +105,7 @@ namespace mu
 	//---------------------------------------------------------------------------------------------
 	//---------------------------------------------------------------------------------------------
 	CodeGenerator::CodeGenerator(CompilerOptions::Private* options)
+		: GenerateMeshConstantPipe(TEXT("GenerateMeshConstantPipe"))
 	{
 		CompilerOptions = options;
 
@@ -311,72 +314,87 @@ namespace mu
 	}
 
 
-	Ptr<ASTOp> CodeGenerator::GenerateTableVariable( Ptr<const Node> InNode, const FTableCacheKey& CacheKey, bool bAddNoneOption, const FString& DefaultRowName)
+	Ptr<NodeScalar> CodeGenerator::GenerateTableVariableNode(Ptr<const Node> InNode, const FTableCacheKey& CacheKey, bool bAddNoneOption, const FString& DefaultRowName)
 	{
-        FParameterDesc Param;
-		Param.Name = CacheKey.ParameterName;
+		Ptr<NodeScalarEnumParameter> Result = new NodeScalarEnumParameter;
 
-        if (Param.Name.Len() == 0)
-        {
-			Param.Name = CacheKey.Table->GetName();
-        }
-        
-		Param.Type = EParameterType::Int;
-		Param.DefaultValue.Set<FParamIntType>(0);
+		FString ParamName = CacheKey.ParameterName;
+		if (ParamName.Len() == 0)
+		{
+			ParamName = CacheKey.Table->GetName();
+		}
+		Result->SetName(ParamName);
+
+		Result->SetDefaultValueIndex(0);
+
+		int32 CurrentRow = 0;
+		int32 RowCount = CacheKey.Table->GetPrivate()->Rows.Num();
+		check(RowCount < MAX_int16) // max FIntValueDesc allows
+
+		if (bAddNoneOption)
+		{
+			Result->SetValueCount(RowCount);
+			Result->SetValue(CurrentRow, -1, "None");
+		}
+		else
+		{
+			Result->SetValueCount(RowCount-1);
+		}
 
 		// Add the possible values
 		{
 			// See if there is a string column. If there is one, we will use it as names for the
 			// options. Only the first string column will be used.
-			int32 NameColumn = -1;
-			int32 Columns = CacheKey.Table->GetPrivate()->Columns.Num();
-			for (int32 ColumnIndex = 0; ColumnIndex < Columns && NameColumn < 0; ++ColumnIndex)
+			int32 NameCol = -1;
+			int32 NumCols = CacheKey.Table->GetPrivate()->Columns.Num();
+			for (int32 ColumnIndex=0; ColumnIndex<NumCols && NameCol<0; ++ColumnIndex)
 			{
 				if (CacheKey.Table->GetPrivate()->Columns[ColumnIndex].Type == ETableColumnType::String)
 				{
-					NameColumn = ColumnIndex;
+					NameCol = ColumnIndex;
 				}
 			}
 
-			int32 RowCount = CacheKey.Table->GetPrivate()->Rows.Num();
-			check(RowCount < MAX_int16) // max FIntValueDesc allows
-
 			// Skip "None" option (first row) if it's not required
 			int32 StarCount = bAddNoneOption ? 0 : 1;
 
 			for (int32 RowIndex = StarCount; RowIndex < RowCount; ++RowIndex)
 			{
-				FParameterDesc::FIntValueDesc Value;
-				Value.Value = RowIndex;
-
-				if (NameColumn > -1)
+				FString ValueName;
+				if (NameCol > -1)
 				{
-					Value.Name = CacheKey.Table->GetPrivate()->Rows[RowIndex].Values[NameColumn].String;
+					ValueName = CacheKey.Table->GetPrivate()->Rows[RowIndex].Values[NameCol].String;
 				}
 
-				Param.PossibleValues.Add(Value);
+				Result->SetValue(CurrentRow, RowIndex, ValueName);
+
 
 				// Set the first row or the selected row as the default one.
-				if(RowIndex == StarCount || Value.Name == DefaultRowName)
+				if(RowIndex == StarCount || ValueName == DefaultRowName)
 				{
-					Param.DefaultValue.Set<FParamIntType>(Value.Value);
+					Result->SetDefaultValueIndex(RowIndex);
 				}
-            }
-        }
 
-		Ptr<ASTOpParameter> op = new ASTOpParameter();
-		op->Type = EOpType::NU_PARAMETER;
-		op->Parameter = Param;
+				// Set the selected row as default (if exists)
+				if (ValueName == DefaultRowName)
+				{
+					Result->SetDefaultValueIndex(RowIndex);
+				}
 
-		FirstPass.ParameterNodes.Add(InNode, op);
+				++CurrentRow;
+			}
+		}
 
-		return op;
+		return Result;
 	}
 
 
 	TSharedPtr<const FLayout> CodeGenerator::GenerateLayout(Ptr<const NodeLayout> SourceLayout, uint32 MeshIDPrefix)
 	{
-		TSharedPtr<const FLayout>* it = GeneratedLayouts.Find({ SourceLayout,MeshIDPrefix });
+		// This can run in any thread.
+		UE::TUniqueLock Lock(GenerateLayoutConstantState.Mutex);
+		
+		TSharedPtr<const FLayout>* it = GenerateLayoutConstantState.GeneratedLayouts.Find({ SourceLayout,MeshIDPrefix });
 
 		if (it)
 		{
@@ -407,7 +425,7 @@ namespace mu
 		}
 
 		check(GeneratedLayout->Blocks.IsEmpty() || GeneratedLayout->Blocks[0].Id != FLayoutBlock::InvalidBlockId);
-		GeneratedLayouts.Add({ SourceLayout,MeshIDPrefix }, GeneratedLayout);
+		GenerateLayoutConstantState.GeneratedLayouts.Add({ SourceLayout,MeshIDPrefix }, GeneratedLayout);
 
 		return GeneratedLayout;
 	}
@@ -454,7 +472,7 @@ namespace mu
 				Ptr<NodeImageConstant> pNode = new NodeImageConstant();
 				pNode->SetValue(PatchMask);
 
-				FImageGenerationOptions ConstantOptions(-1);
+				FImageGenerationOptions ConstantOptions(-1,-1);
 				FImageGenerationResult ConstantResult;
 				GenerateImage(ConstantOptions, ConstantResult, pNode);
 				RectConstantOp = ConstantResult.op;
@@ -899,20 +917,19 @@ namespace mu
 		GetModifiersFor(ComponentId, SurfaceNode->Tags, bModifiersForBeforeOperations, Modifiers);
 
 		// This pass on the modifiers is only to detect errors that cannot be detected at the point they are applied.
-		CheckModifiersForSurface(*SurfaceNode, Modifiers);
+		CheckModifiersForSurface(*SurfaceNode, Modifiers, Options.LODIndex);
 		
 		TBitArray<> LayoutFromExtension;
         //if (SurfaceNode->Mesh)
         {
             MUTABLE_CPUPROFILER_SCOPE(SurfaceMesh);
 
-            Ptr<ASTOp> LastMeshOp;
-
             // Generate the mesh
-			FMeshGenerationOptions MeshOptions(ComponentId);
-			MeshOptions.bLayouts = true;
-			MeshOptions.State = Options.State;
-			MeshOptions.ActiveTags = SurfaceNode->Tags;
+			FMeshGenerationStaticOptions MeshStaticOptions(ComponentId, Options.LODIndex);
+			MeshStaticOptions.ActiveTags = SurfaceNode->Tags;
+			MeshStaticOptions.State = Options.State;
+			FMeshGenerationDynamicOptions MeshDynamicOptions;
+			MeshDynamicOptions.bLayouts = true;
 
 			const FMeshGenerationResult* SharedMeshResults = nullptr;
 			if (bShareSurface)
@@ -922,23 +939,28 @@ namespace mu
 				check(SharedMeshResults);
 
 				// Override the layouts with the ones from the surface we share
-				MeshOptions.OverrideLayouts = SharedMeshResults->GeneratedLayouts;
+				MeshDynamicOptions.OverrideLayouts = SharedMeshResults->GeneratedLayouts;
 			}
 
 			// Normalize UVs if we're going to work with images and layouts.
 			// TODO: This should come from per-layout settings!
 			const bool bNormalizeUVs = false; // !SurfaceNode->Images.IsEmpty();
-			MeshOptions.bNormalizeUVs = bNormalizeUVs;
+			MeshDynamicOptions.bNormalizeUVs = bNormalizeUVs;
 
 			// Ensure UV islands remain within their main layout block on lower LODs to avoid unexpected reordering 
 			// of the layout blocks when reusing a surface between LODs. Used to fix small displacements on vertices
 			// that may cause them to fall on a different block.
-			MeshOptions.bClampUVIslands = bShareSurface;
+			MeshDynamicOptions.bClampUVIslands = bShareSurface;
 
-            GenerateMesh(MeshOptions, MeshResults, SurfaceNode->Mesh);
+            FMeshTask MeshTask = GenerateMesh(MeshStaticOptions, UE::Tasks::MakeCompletedTask<FMeshGenerationDynamicOptions>(MeshDynamicOptions), SurfaceNode->Mesh);
 
 			// Apply the modifier for the post-normal operations stage.
-			LastMeshOp = ApplyMeshModifiers(Modifiers, MeshOptions, MeshResults, SharedMeshResults, SurfaceNode->GetMessageContext(), nullptr);
+			MeshTask = ApplyMeshModifiers(Modifiers,
+				MeshStaticOptions, UE::Tasks::MakeCompletedTask<FMeshGenerationDynamicOptions>(MeshDynamicOptions),
+				MeshTask, SharedMeshResults, SurfaceNode->GetMessageContext(), nullptr);
+
+			// TODO
+			MeshResults = MeshTask.GetResult();
 
 			// Base mesh is allowed to be missing, aggregate all layouts and operations per layout indices in the
 			// generated mesh, base and extends.
@@ -1057,6 +1079,8 @@ namespace mu
 				}
 			}
 
+			Ptr<ASTOp> LastMeshOp = MeshResults.MeshOp;
+
 			check(SurfaceReferenceLayouts.Num() == SurfaceLayoutOps.Num());
             for (int32 LayoutIndex = 0; LayoutIndex < SurfaceReferenceLayouts.Num(); ++LayoutIndex)
 			{
@@ -1192,7 +1216,7 @@ namespace mu
 					if (ImageLayoutStrategy == CompilerOptions::TextureLayoutStrategy::None)
 					{
 						// Generate the image
-						FImageGenerationOptions ImageOptions(ComponentId);
+						FImageGenerationOptions ImageOptions(ComponentId, Options.LODIndex);
 						ImageOptions.State = Options.State;
 						ImageOptions.ImageLayoutStrategy = ImageLayoutStrategy;
 						ImageOptions.ActiveTags = SurfaceNode->Tags;
@@ -1314,7 +1338,7 @@ namespace mu
 								for (int32 BlockIndex = 0; BlockIndex < pLayout->GetBlockCount(); ++BlockIndex)
 								{
 									// Generate the image
-									FImageGenerationOptions ImageOptions(ComponentId);
+									FImageGenerationOptions ImageOptions(ComponentId, Options.LODIndex);
 									ImageOptions.State = Options.State;
 									ImageOptions.ImageLayoutStrategy = ImageLayoutStrategy;
 									ImageOptions.RectSize = { 0,0 };
@@ -1374,10 +1398,10 @@ namespace mu
 							}
 							check(imageAd);
 
-							FMeshGenerationOptions ModifierOptions(ComponentId);
+							FMeshGenerationStaticOptions ModifierOptions(ComponentId, Options.LODIndex);
 							ModifierOptions.State = Options.State;
 							ModifierOptions.ActiveTags = SurfaceNode->Tags;
-							imageAd = ApplyImageExtendModifiers( Modifiers, ModifierOptions, ComponentId, MeshResults, imageAd, ImageLayoutStrategy, 
+							imageAd = ApplyImageExtendModifiers( Modifiers, ModifierOptions, MeshResults, imageAd, ImageLayoutStrategy, 
 								LayoutIndex, ImageData, GridSize, LayoutBlockDesc, 
 								SurfaceNode->GetMessageContext());
 
@@ -1925,7 +1949,7 @@ namespace mu
 		int32 ComponentId,
 		const TArray<FString>& SurfaceTags,
 		bool bModifiersForBeforeOperations,
-		TArray<FirstPassGenerator::FModifier>& OutModifiers)
+		TArray<FirstPassGenerator::FModifier>& OutModifiers) const
 	{
         MUTABLE_CPUPROFILER_SCOPE(GetModifiersFor);
 
@@ -1934,28 +1958,28 @@ namespace mu
 			return;
 		}
 
-		for (const FirstPassGenerator::FModifier& m: FirstPass.Modifiers)
+		for (const FirstPassGenerator::FModifier& Modifier: FirstPass.Modifiers)
 		{
-			if (!m.Node)
+			if (!Modifier.Node)
 			{
 				continue;
 			}
 
 			// Correct stage?
-			if (m.Node->bApplyBeforeNormalOperations != bModifiersForBeforeOperations)
+			if (Modifier.Node->bApplyBeforeNormalOperations != bModifiersForBeforeOperations)
 			{
 				continue;
 			}
 
 			// Correct component?
-			if (m.Node->RequiredComponentId>=0 && m.Node->RequiredComponentId!=ComponentId)
+			if (Modifier.Node->RequiredComponentId>=0 && Modifier.Node->RequiredComponentId!=ComponentId)
 			{
 				continue;
 			}
 
 			// Already there?
 			bool bAlreadyAdded = 
-				OutModifiers.FindByPredicate( [&m](const FirstPassGenerator::FModifier& c) {return c.Node == m.Node; })
+				OutModifiers.FindByPredicate( [&Modifier](const FirstPassGenerator::FModifier& c) {return c.Node == Modifier.Node; })
 				!= 
 				nullptr;
 
@@ -1967,11 +1991,11 @@ namespace mu
 			// Matching tags?
 			bool bApply = false;
 
-			switch (m.Node->MultipleTagsPolicy)
+			switch (Modifier.Node->MultipleTagsPolicy)
 			{
 			case EMutableMultipleTagPolicy::OnlyOneRequired:
 			{
-				for (const FString& Tag: m.Node->RequiredTags)
+				for (const FString& Tag: Modifier.Node->RequiredTags)
 				{
 					if (SurfaceTags.Contains(Tag))
 					{
@@ -1985,7 +2009,7 @@ namespace mu
 			case EMutableMultipleTagPolicy::AllRequired:
 			{
 				bApply = true;
-				for (const FString& Tag : m.Node->RequiredTags)
+				for (const FString& Tag : Modifier.Node->RequiredTags)
 				{
 					if (!SurfaceTags.Contains(Tag))
 					{
@@ -1998,41 +2022,40 @@ namespace mu
 
 			if (bApply)
 			{
-				OutModifiers.Add(m);
+				OutModifiers.Add(Modifier);
 			}
 		}
 	}
 
 	//---------------------------------------------------------------------------------------------
-	Ptr<ASTOp> CodeGenerator::ApplyMeshModifiers(
+	CodeGenerator::FMeshTask CodeGenerator::ApplyMeshModifiers(
 		const TArray<FirstPassGenerator::FModifier>& Modifiers,
-		const FMeshGenerationOptions& Options, 
-		FMeshGenerationResult& BaseMeshResult,
+		const FMeshGenerationStaticOptions& StaticOptions, 
+		FMeshOptionsTask Options,
+		FMeshTask BaseTask,
 		const FMeshGenerationResult* SharedMeshResults,
 		const void* ErrorContext,
-		const NodeMeshConstant* OriginalMeshNode )
+		const NodeMeshConstant* OriginalMeshNode)
 	{
-		Ptr<ASTOp> LastMeshOp = BaseMeshResult.MeshOp;
+		FMeshTask LastMeshTask = BaseTask;
+		FMeshTask PreModifiersTask = BaseTask;
 
-		Ptr<ASTOp> PreModifiersMesh = LastMeshOp;
-
-		int32 CurrentLOD = CurrentParents.Last().Lod;
+		int32 CurrentLOD = StaticOptions.LODIndex;
+		check(CurrentLOD>=0);
 
 		// Process mesh extend modifiers (from edit modifiers)
 		int32 EditIndex = 0;
-		for (const FirstPassGenerator::FModifier& m : Modifiers)
+		for (const FirstPassGenerator::FModifier& Modifier : Modifiers)
 		{
-			if (ModifiersToIgnore.Contains(m))
+			if (StaticOptions.ModifiersToIgnore.Contains(Modifier))
 			{
 				// Prevent recursion.
 				continue;
 			}
 
-			if (m.Node->GetType() == NodeModifierSurfaceEdit::GetStaticType())
+			if (Modifier.Node->GetType() == NodeModifierSurfaceEdit::GetStaticType())
 			{
-				const NodeModifierSurfaceEdit* Edit = static_cast<const NodeModifierSurfaceEdit*>(m.Node);
-
-				BaseMeshResult.ExtraMeshLayouts.Emplace();
+				const NodeModifierSurfaceEdit* Edit = static_cast<const NodeModifierSurfaceEdit*>(Modifier.Node);
 
 				bool bAffectsCurrentLOD = Edit->LODs.IsValidIndex(CurrentLOD);
 				if (bAffectsCurrentLOD && Edit->LODs[CurrentLOD].MeshAdd)
@@ -2040,102 +2063,137 @@ namespace mu
 					Ptr<NodeMesh> pAdd = Edit->LODs[CurrentLOD].MeshAdd;
 
 					// Store the data necessary to apply modifiers for the pre-normal operations stage.
-					FMeshGenerationOptions MergedMeshOptions(Options);
-					MergedMeshOptions.ActiveTags = Edit->EnableTags; // TODO: Append to current?
-					MergedMeshOptions.bEnsureAllVerticesHaveLayoutBlock = false;
+					FMeshGenerationStaticOptions MergedMeshStaticOptions(StaticOptions);
+					MergedMeshStaticOptions.ActiveTags = Edit->EnableTags; // TODO: Append to current?
+					MergedMeshStaticOptions.ModifiersToIgnore.Add(Modifier);
 
-					if (SharedMeshResults)
-					{
-						check(SharedMeshResults->ExtraMeshLayouts.IsValidIndex(EditIndex));
-						MergedMeshOptions.OverrideLayouts = SharedMeshResults->ExtraMeshLayouts[EditIndex].GeneratedLayouts;
-					}
+					FMeshOptionsTask MergedMeshOptionsTask = UE::Tasks::Launch(TEXT("MutableMergedMeshOptions"),
+						[Options, SharedMeshResults, EditIndex]() mutable
+						{
+							FMeshGenerationDynamicOptions Result = Options.GetResult();
+							Result.bEnsureAllVerticesHaveLayoutBlock = false;
 
-					FMeshGenerationResult AddResults;
-					GenerateMesh(MergedMeshOptions, AddResults, pAdd);
+							if (SharedMeshResults)
+							{
+								check(SharedMeshResults->ExtraMeshLayouts.IsValidIndex(EditIndex));
+								Result.OverrideLayouts = SharedMeshResults->ExtraMeshLayouts[EditIndex].GeneratedLayouts;
+							}
 
-					// Warn about discrepancies on layout strategy between the added and the base
-					int32 LayoutIndexThatHasBlocks = -1;
-					{
-						if (BaseMeshResult.GeneratedLayouts.Num() != AddResults.GeneratedLayouts.Num())
-						{
-							// When extending a mesh section the added mesh sectio will use the layout strategy of the base one
-							FString Msg = FString::Printf(TEXT("Extended mesh section layout count is differenta than the mesh being extended."));
-							ErrorLog->Add(Msg, ELMT_INFO, Edit->GetMessageContext(), ErrorContext);
-						}
+							return Result;
+						},
+						Options);
+
+					FMeshTask AddBaseTask = GenerateMesh(MergedMeshStaticOptions, MergedMeshOptionsTask, pAdd);
 
-						for (int32 LayoutIndex = 0; LayoutIndex < BaseMeshResult.GeneratedLayouts.Num(); ++LayoutIndex)
+					FMeshTask AddTask = UE::Tasks::Launch(TEXT("MutableMergedMeshAdd"),
+						[AddBaseTask, LastMeshTask, ErrorLog=ErrorLog,Edit,ErrorContext]() mutable
 						{
-							if (BaseMeshResult.GeneratedLayouts[LayoutIndex].Layout
-								&&
-								BaseMeshResult.GeneratedLayouts[LayoutIndex].Layout->Strategy != EPackStrategy::Overlay)
+							FMeshGenerationResult AddResults = AddBaseTask.GetResult();
+							FMeshGenerationResult BaseMeshResult = LastMeshTask.GetResult();
+
+							// Warn about discrepancies on layout strategy between the added and the base
+							int32 LayoutIndexThatHasBlocks = -1;
 							{
-								LayoutIndexThatHasBlocks = LayoutIndex;
+								if (BaseMeshResult.GeneratedLayouts.Num() != AddResults.GeneratedLayouts.Num())
+								{
+									// When extending a mesh section the added mesh sectio will use the layout strategy of the base one
+									FString Msg = FString::Printf(TEXT("Extended mesh section layout count is differenta than the mesh being extended."));
+									ErrorLog->Add(Msg, ELMT_INFO, Edit->GetMessageContext(), ErrorContext);
+								}
+
+								for (int32 LayoutIndex = 0; LayoutIndex < BaseMeshResult.GeneratedLayouts.Num(); ++LayoutIndex)
+								{
+									if (BaseMeshResult.GeneratedLayouts[LayoutIndex].Layout
+										&&
+										BaseMeshResult.GeneratedLayouts[LayoutIndex].Layout->Strategy != EPackStrategy::Overlay)
+									{
+										LayoutIndexThatHasBlocks = LayoutIndex;
+									}
+
+									if (BaseMeshResult.GeneratedLayouts[LayoutIndex].Layout
+										&&
+										AddResults.GeneratedLayouts.IsValidIndex(LayoutIndex)
+										&&
+										AddResults.GeneratedLayouts[LayoutIndex].Layout
+										&&
+										BaseMeshResult.GeneratedLayouts[LayoutIndex].Layout->Strategy != AddResults.GeneratedLayouts[LayoutIndex].Layout->Strategy)
+									{
+										// When extending a mesh section the added mesh sectio will use the layout strategy of the base one
+										FString Msg = FString::Printf(TEXT("Extended mesh section layout [%d] is using a different strategy than the section being extended. The base strategy will be used."), LayoutIndex);
+										ErrorLog->Add(Msg, ELMT_INFO, Edit->GetMessageContext(), ErrorContext);
+									}
+								}
 							}
 
-							if (BaseMeshResult.GeneratedLayouts[LayoutIndex].Layout
-								&&
-								AddResults.GeneratedLayouts.IsValidIndex(LayoutIndex)
-								&&
-								AddResults.GeneratedLayouts[LayoutIndex].Layout
-								&&
-								BaseMeshResult.GeneratedLayouts[LayoutIndex].Layout->Strategy != AddResults.GeneratedLayouts[LayoutIndex].Layout->Strategy)
+							// Add the operation to extract the relevant layout blocks if necessary
+							// TODO: Handle multiple layouts defining blocks: what to extract?
+							if (LayoutIndexThatHasBlocks >= 0)
 							{
-								// When extending a mesh section the added mesh sectio will use the layout strategy of the base one
-								FString Msg = FString::Printf(TEXT("Extended mesh section layout [%d] is using a different strategy than the section being extended. The base strategy will be used."), LayoutIndex);
-								ErrorLog->Add(Msg, ELMT_INFO, Edit->GetMessageContext(), ErrorContext);
-							}
-						}
-					}
+								Ptr<ASTOpMeshExtractLayoutBlocks> ExtractOp = new ASTOpMeshExtractLayoutBlocks();
+								ExtractOp->Source = AddResults.MeshOp;
+								ExtractOp->LayoutIndex = LayoutIndexThatHasBlocks;
 
-					// Add the operation to extract the relevant layout blocks if necessary
-					// TODO: Handle multiple layouts defining blocks: what to extract?
-					if (LayoutIndexThatHasBlocks>=0)
-					{
-						Ptr<ASTOpMeshExtractLayoutBlocks> ExtractOp = new ASTOpMeshExtractLayoutBlocks();
-						ExtractOp->Source = AddResults.MeshOp;
-						ExtractOp->LayoutIndex = LayoutIndexThatHasBlocks;
+								AddResults.MeshOp = ExtractOp;
+							}
 
-						AddResults.MeshOp = ExtractOp;
-					}
+							return AddResults;
+						},
+						UE::Tasks::Prerequisites(AddBaseTask, LastMeshTask)
+						);
 
-					// Apply the modifier for the post-normal operations stage to the added mesh
-					FMeshGenerationOptions ModifierOptions(Options);
+					// Apply the modifiers for the post-normal operations stage to the added mesh
+					FMeshGenerationStaticOptions ModifierOptions(StaticOptions);
 					ModifierOptions.ActiveTags = Edit->EnableTags;
+					ModifierOptions.ModifiersToIgnore.AddUnique(Modifier);
 
 					TArray<FirstPassGenerator::FModifier> ChildModifiers;
 					constexpr bool bModifiersForBeforeOperations = false;
-					GetModifiersFor(Options.ComponentId, ModifierOptions.ActiveTags, bModifiersForBeforeOperations, ChildModifiers);
-
-					ModifiersToIgnore.Push(m);
-					Ptr<ASTOp> AddedMeshOp = ApplyMeshModifiers(ChildModifiers, ModifierOptions, AddResults, SharedMeshResults, ErrorContext, nullptr);
-					ModifiersToIgnore.Pop();
-
-					FMeshGenerationResult::FExtraLayouts data;
-					data.GeneratedLayouts = AddResults.GeneratedLayouts;
-					data.Condition = m.FinalCondition;
-					data.MeshFragment = AddedMeshOp;
-					BaseMeshResult.ExtraMeshLayouts[EditIndex] = data;
-
-					Ptr<ASTOpMeshMerge> MergeOp = new ASTOpMeshMerge();
-					MergeOp->Base = LastMeshOp;
-					MergeOp->Added = AddedMeshOp;
-					// will merge the meshes under the same surface
-					MergeOp->NewSurfaceID = 0;
-
-					// Condition to apply
-					if (m.FinalCondition)
-					{
-						Ptr<ASTOpConditional> conditionalAd = new ASTOpConditional();
-						conditionalAd->type = EOpType::ME_CONDITIONAL;
-						conditionalAd->no = LastMeshOp;
-						conditionalAd->yes = MergeOp;
-						conditionalAd->condition = m.FinalCondition;
-						LastMeshOp = conditionalAd;
-					}
-					else
-					{
-						LastMeshOp = MergeOp;
-					}
+					GetModifiersFor(StaticOptions.ComponentId, ModifierOptions.ActiveTags, bModifiersForBeforeOperations, ChildModifiers);
+
+					FMeshTask AddWithModifiersTask = ApplyMeshModifiers(ChildModifiers, ModifierOptions, Options, AddTask, SharedMeshResults, ErrorContext, nullptr);
+
+					LastMeshTask = UE::Tasks::Launch(TEXT("MutableMeshMergeModifier"),
+						[AddBaseTask, AddWithModifiersTask, LastMeshTask, Modifier, EditIndex]() mutable
+						{
+							FMeshGenerationResult AddResults = AddBaseTask.GetResult();
+							FMeshGenerationResult AddFinalResults = AddWithModifiersTask.GetResult();
+							FMeshGenerationResult LastMeshResults = LastMeshTask.GetResult();
+
+							FMeshGenerationResult::FExtraLayouts Data;
+							Data.GeneratedLayouts = AddResults.GeneratedLayouts;
+							Data.Condition = Modifier.FinalCondition;
+							Data.MeshFragment = AddFinalResults.MeshOp;
+							if (LastMeshResults.ExtraMeshLayouts.Num() <= EditIndex)
+							{
+								LastMeshResults.ExtraMeshLayouts.SetNum(EditIndex+1);
+							}
+							LastMeshResults.ExtraMeshLayouts[EditIndex] = Data;
+
+							Ptr<ASTOpMeshMerge> MergeOp = new ASTOpMeshMerge;
+							MergeOp->Base = LastMeshResults.MeshOp;
+							MergeOp->Added = AddFinalResults.MeshOp;
+							// will merge the meshes under the same surface
+							MergeOp->NewSurfaceID = 0;
+
+							// Condition to apply
+							if (Modifier.FinalCondition)
+							{
+								Ptr<ASTOpConditional> ConditionalOp = new ASTOpConditional();
+								ConditionalOp->type = EOpType::ME_CONDITIONAL;
+								ConditionalOp->no = LastMeshResults.MeshOp;
+								ConditionalOp->yes = MergeOp;
+								ConditionalOp->condition = Modifier.FinalCondition;
+								LastMeshResults.MeshOp = ConditionalOp;
+							}
+							else
+							{
+								LastMeshResults.MeshOp = MergeOp;
+							}
+
+							return LastMeshResults;
+						},						
+						UE::Tasks::Prerequisites(AddWithModifiersTask, LastMeshTask)
+					);
 				}
 
 				++EditIndex;
@@ -2143,15 +2201,19 @@ namespace mu
 
 		}
 
-		// "remove" operation to group all the removes
-		Ptr<ASTOpMeshRemoveMask> RemoveOp;
 
 		// Process mesh remove modifiers (from edit modifiers)
-		for (const FirstPassGenerator::FModifier& m : Modifiers)
+		for (const FirstPassGenerator::FModifier& Modifier : Modifiers)
 		{
-			if (m.Node->GetType() == NodeModifierSurfaceEdit::GetStaticType())
+			if (StaticOptions.ModifiersToIgnore.Contains(Modifier))
 			{
-				const NodeModifierSurfaceEdit* Edit = static_cast<const NodeModifierSurfaceEdit*>(m.Node);
+				// Prevent recursion.
+				continue;
+			}
+
+			if (Modifier.Node->GetType() == NodeModifierSurfaceEdit::GetStaticType())
+			{
+				const NodeModifierSurfaceEdit* Edit = static_cast<const NodeModifierSurfaceEdit*>(Modifier.Node);
 
 				bool bAffectsCurrentLOD = Edit->LODs.IsValidIndex(CurrentLOD);
 
@@ -2161,44 +2223,59 @@ namespace mu
 				// \TODO: Apply base removes first, and then "added meshes" removes here. It may have lower memory footprint during generation.
 				if (bAffectsCurrentLOD && Edit->LODs[CurrentLOD].MeshRemove)
 				{
-					Ptr<NodeMesh> pRemove = Edit->LODs[CurrentLOD].MeshRemove;
+					Ptr<NodeMesh> MeshRemove = Edit->LODs[CurrentLOD].MeshRemove;
 
-					FMeshGenerationResult removeResults;
-					FMeshGenerationOptions RemoveMeshOptions(Options.ComponentId);
-					RemoveMeshOptions.bLayouts = false;
-					RemoveMeshOptions.State = Options.State;
-					RemoveMeshOptions.ActiveTags = Edit->EnableTags;
+					FMeshGenerationStaticOptions RemoveMeshStaticOptions(StaticOptions);
+					RemoveMeshStaticOptions.ActiveTags = Edit->EnableTags;
+					RemoveMeshStaticOptions.ModifiersToIgnore.Add(Modifier);
 
-					GenerateMesh(RemoveMeshOptions, removeResults, pRemove);
+					FMeshGenerationDynamicOptions RemoveMeshDynamicOptions;
+					RemoveMeshDynamicOptions.bLayouts = false;
 
-					Ptr<ASTOpMeshMaskDiff> MaskOp = new ASTOpMeshMaskDiff();
+					FMeshTask RemoveMeshTask = GenerateMesh(RemoveMeshStaticOptions,UE::Tasks::MakeCompletedTask<FMeshGenerationDynamicOptions>(RemoveMeshDynamicOptions), MeshRemove);
 
-					// By default, remove from the base
-					Ptr<ASTOp> RemoveFrom = BaseMeshResult.BaseMeshOp;
-					MaskOp->Source = RemoveFrom;
-					MaskOp->Fragment = removeResults.MeshOp;
+					LastMeshTask = UE::Tasks::Launch(TEXT("MutableMeshMergeModifier"),
+						[RemoveMeshTask, LastMeshTask, PreModifiersTask, Edit, Modifier]() mutable
+						{
+							FMeshGenerationResult RemoveResults = RemoveMeshTask.GetResult();
+							FMeshGenerationResult LastMeshResults = LastMeshTask.GetResult();
+							FMeshGenerationResult BaseResults = PreModifiersTask.GetResult();
 
-					if (!RemoveOp)
-					{
-						RemoveOp = new ASTOpMeshRemoveMask();
-						RemoveOp->source = LastMeshOp;
-						RemoveOp->FaceCullStrategy = Edit->FaceCullStrategy;
-						LastMeshOp = RemoveOp;
-					}
+							Ptr<ASTOpMeshMaskDiff> MaskOp = new ASTOpMeshMaskDiff;
+
+							// By default, remove from the base
+							Ptr<ASTOp> RemoveFrom = BaseResults.BaseMeshOp;
+							MaskOp->Source = RemoveFrom;
+							MaskOp->Fragment = RemoveResults.MeshOp;
+
+							Ptr<ASTOpMeshRemoveMask> RemoveOp = new ASTOpMeshRemoveMask();
+							RemoveOp->source = LastMeshResults.MeshOp;
+							RemoveOp->FaceCullStrategy = Edit->FaceCullStrategy;
+							RemoveOp->AddRemove(Modifier.FinalCondition, MaskOp);
 
-					RemoveOp->AddRemove(m.FinalCondition, MaskOp);
+							LastMeshResults.MeshOp = RemoveOp;
+
+							return LastMeshResults;
+						},
+						UE::Tasks::Prerequisites(RemoveMeshTask, LastMeshTask)
+					);
 				}
 			}
 		}
 
 
 		// Process mesh morph modifiers (from edit modifiers)
-		for (const FirstPassGenerator::FModifier& m : Modifiers)
+		for (const FirstPassGenerator::FModifier& Modifier : Modifiers)
 		{
-			if (m.Node->GetType() == NodeModifierSurfaceEdit::GetStaticType())
+			if (StaticOptions.ModifiersToIgnore.Contains(Modifier))
 			{
-				const NodeModifierSurfaceEdit* Edit = static_cast<const NodeModifierSurfaceEdit*>(m.Node);
+				// Prevent recursion.
+				continue;
+			}
 
+			if (Modifier.Node->GetType() == NodeModifierSurfaceEdit::GetStaticType())
+			{
+				const NodeModifierSurfaceEdit* Edit = static_cast<const NodeModifierSurfaceEdit*>(Modifier.Node);
 				if (Edit->MeshMorph.IsEmpty())
 				{
 					continue;
@@ -2206,389 +2283,489 @@ namespace mu
 
 				check(OriginalMeshNode);
 
-				TSharedPtr<FMesh> TargetMesh = OriginalMeshNode->FindMorph(Edit->MeshMorph);
-				if (!TargetMesh)
+				// Request morphed skeletal mesh
+				TSharedPtr< TSharedPtr<FMesh> > ResolveMorphedMesh = MakeShared<TSharedPtr<FMesh>>();
+				UE::Tasks::FTask TargetMeshTask;
 				{
-					continue;
+					FMesh* OriginalMesh = OriginalMeshNode->GetValue().Get();
+					check(OriginalMesh->IsReference());
+					uint32 OriginalMeshID = OriginalMesh->GetReferencedMesh();
+
+					bool bRunImmediatlyIfPossible = IsInGameThread();
+					TargetMeshTask = CompilerOptions->OptimisationOptions.ReferencedMeshResourceProvider(
+						OriginalMeshID,
+						Edit->MeshMorph,
+						ResolveMorphedMesh, 
+						bRunImmediatlyIfPossible );
 				}
-				
+
+				// Factor
+				Ptr<ASTOp> FactorOp;
+				if (Edit->MorphFactor)
 				{
-					// Target mesh
-					Ptr<ASTOpConstantResource> TargetMeshOp = new ASTOpConstantResource;
-					TargetMeshOp->Type = EOpType::ME_CONSTANT;
-					TargetMeshOp->SetValue(TargetMesh->Clone(), CompilerOptions->OptimisationOptions.DiskCacheContext);
-					TargetMeshOp->SourceDataDescriptor = OriginalMeshNode->SourceDataDescriptor;
-
-					// Morph generation through mesh diff
-					Ptr<ASTOpMeshDifference> diffAd;
-					{
-						Ptr<ASTOpMeshDifference> op = new ASTOpMeshDifference();
-						op->Base = BaseMeshResult.BaseMeshOp;
-						op->Target = TargetMeshOp;
-
-						// Morphing tex coords here is not supported:
-						// Generating the homogoneous UVs is difficult since we don't have the base
-						// layout yet.                       
-						op->bIgnoreTextureCoords = true;
-						diffAd = op;
-					}
+					FScalarGenerationResult ChildResult;
+					GenerateScalar(ChildResult, StaticOptions, Edit->MorphFactor);
+					FactorOp = ChildResult.op;
+				}
+				else
+				{
+					Ptr<NodeScalarConstant> auxNode = new NodeScalarConstant();
+					auxNode->SetValue(1.0f);
 
-					// Morph operation
-					Ptr<ASTOp> morphAd;
+					FScalarGenerationResult ChildResult;
+					GenerateScalar(ChildResult, StaticOptions, auxNode);
+					FactorOp = ChildResult.op;
+				}
+
+				Ptr<const NodeMeshConstant> OriginalMeshCopy = OriginalMeshNode;
+				LastMeshTask = UE::Tasks::Launch(TEXT("MutableMeshMergeModifier"),
+					[TargetMeshTask, LastMeshTask, PreModifiersTask, Edit, Modifier, ResolveMorphedMesh, FactorOp, CompilerOptions=CompilerOptions, OriginalMeshCopy]() mutable
 					{
-						Ptr<ASTOpMeshMorph> op = new ASTOpMeshMorph();
+						FMeshGenerationResult LastMeshResults = LastMeshTask.GetResult();
+						FMeshGenerationResult BaseResults = PreModifiersTask.GetResult();
 
-						// Factor
-						if (Edit->MorphFactor)
+						TSharedPtr<FMesh> TargetMesh = *ResolveMorphedMesh;
+						if (!TargetMesh)
 						{
-							FScalarGenerationResult ChildResult;
-							GenerateScalar(ChildResult, Options, Edit->MorphFactor);
-							op->Factor = ChildResult.op;
+							return LastMeshResults;
 						}
-						else
-						{
-							Ptr<NodeScalarConstant> auxNode = new NodeScalarConstant();
-							auxNode->SetValue(1.0f);
 
-							FScalarGenerationResult ChildResult;
-							GenerateScalar(ChildResult, Options, auxNode);
-							op->Factor = ChildResult.op;
+						// Target mesh
+						Ptr<ASTOpConstantResource> TargetMeshOp = new ASTOpConstantResource;
+						TargetMeshOp->Type = EOpType::ME_CONSTANT;
+						TargetMeshOp->SetValue(TargetMesh->Clone(), CompilerOptions->OptimisationOptions.DiskCacheContext);
+						TargetMeshOp->SourceDataDescriptor = OriginalMeshCopy->SourceDataDescriptor;
+
+						// Morph generation through mesh diff
+						Ptr<ASTOpMeshDifference> DiffOp;
+						{
+							DiffOp = new ASTOpMeshDifference();
+							DiffOp->Base = BaseResults.BaseMeshOp;
+							DiffOp->Target = TargetMeshOp;
+
+							// Morphing tex coords here is not supported:
+							// Generating the homogoneous UVs is difficult since we don't have the base
+							// layout yet.                       
+							DiffOp->bIgnoreTextureCoords = true;
 						}
 
-						// Base		
-						op->Base = LastMeshOp;
+						// Morph operation
+						Ptr<ASTOpMeshMorph> MorphOp = new ASTOpMeshMorph();
+						{
+							MorphOp->Base = LastMeshResults.MeshOp;
+							MorphOp->Target = DiffOp;
+						}
 
-						// Targets
-						op->Target = diffAd;
-						morphAd = op;
-					}
+						// Condition to apply the morph
+						if (Modifier.FinalCondition)
+						{
+							Ptr<ASTOpConditional> ConditionalOp = new ASTOpConditional();
+							ConditionalOp->type = EOpType::ME_CONDITIONAL;
+							ConditionalOp->no = LastMeshResults.MeshOp;
+							ConditionalOp->yes = MorphOp;
+							ConditionalOp->condition = Modifier.FinalCondition;
+							LastMeshResults.MeshOp = ConditionalOp;
+						}
+						else
+						{
+							LastMeshResults.MeshOp = MorphOp;
+						}
 
-					// Condition to apply the morph
-					if (m.FinalCondition)
-					{
-						Ptr<ASTOpConditional> conditionalAd = new ASTOpConditional();
-						conditionalAd->type = EOpType::ME_CONDITIONAL;
-						conditionalAd->no = LastMeshOp;
-						conditionalAd->yes = morphAd;
-						conditionalAd->condition = m.FinalCondition;
-						LastMeshOp = conditionalAd;
-					}
-					else
-					{
-						LastMeshOp = morphAd;
-					}
-				}
+						return LastMeshResults;
+					},
+					UE::Tasks::Prerequisites(TargetMeshTask, LastMeshTask)
+				);
 			}
 		}
 
 
-
 		// Process clip-with-mesh modifiers
-		RemoveOp = nullptr;
-		for (const FirstPassGenerator::FModifier& m : Modifiers)
+		for (const FirstPassGenerator::FModifier& Modifier : Modifiers)
 		{
-			if (m.Node->GetType()== NodeModifierMeshClipWithMesh::GetStaticType())
+			if (StaticOptions.ModifiersToIgnore.Contains(Modifier))
+			{
+				// Prevent recursion.
+				continue;
+			}
+
+			if (Modifier.Node->GetType()== NodeModifierMeshClipWithMesh::GetStaticType())
 			{
-				const NodeModifierMeshClipWithMesh* TypedClipNode = static_cast<const NodeModifierMeshClipWithMesh*>(m.Node);
-				Ptr<ASTOpMeshMaskClipMesh> op = new ASTOpMeshMaskClipMesh();
-				op->source = PreModifiersMesh;
+				const NodeModifierMeshClipWithMesh* TypedClipNode = static_cast<const NodeModifierMeshClipWithMesh*>(Modifier.Node);
 
-				// FParameters
-				FMeshGenerationOptions ClipOptions( Options.ComponentId );
-				ClipOptions.bLayouts = false;
-				ClipOptions.State = Options.State;
+				FMeshGenerationDynamicOptions RemoveMeshDynamicOptions;
+				RemoveMeshDynamicOptions.bLayouts = false;
 
-				FMeshGenerationResult clipResult;
-				GenerateMesh(ClipOptions, clipResult, TypedClipNode->ClipMesh);
-				op->clip = clipResult.MeshOp;
+				FMeshGenerationStaticOptions ClipStaticOptions(StaticOptions);
+				ClipStaticOptions.ModifiersToIgnore.AddUnique(Modifier);
+				ClipStaticOptions.ActiveTags.Empty();
+				FMeshTask ClipMeshTask = GenerateMesh(ClipStaticOptions, UE::Tasks::MakeCompletedTask<FMeshGenerationDynamicOptions>(RemoveMeshDynamicOptions), TypedClipNode->ClipMesh);
 
-				if (!op->clip)
-				{
-					ErrorLog->Add("Clip mesh has not been generated", ELMT_ERROR, ErrorContext);
-					continue;
-				}
+				LastMeshTask = UE::Tasks::Launch(TEXT("MutableMeshMergeModifier"),
+					[ClipMeshTask, LastMeshTask, PreModifiersTask, Modifier, TypedClipNode, ErrorLog=ErrorLog, ErrorContext]() mutable
+					{
+						FMeshGenerationResult ClipResults = ClipMeshTask.GetResult();
+						FMeshGenerationResult LastMeshResults = LastMeshTask.GetResult();
+						FMeshGenerationResult BaseResults = PreModifiersTask.GetResult();
 
-				Ptr<ASTOp> maskAt = op;
+						if (!ClipResults.MeshOp)
+						{
+							ErrorLog->Add("Clip mesh has not been generated", ELMT_ERROR, ErrorContext);
+							return LastMeshResults;
+						}
 
-				if (!RemoveOp)
-				{
-					RemoveOp = new ASTOpMeshRemoveMask();
-					RemoveOp->source = LastMeshOp;
-					RemoveOp->FaceCullStrategy = TypedClipNode->FaceCullStrategy;
-					LastMeshOp = RemoveOp;
-				}
+						Ptr<ASTOpMeshMaskClipMesh> MaskOp = new ASTOpMeshMaskClipMesh();
+						MaskOp->source = BaseResults.MeshOp;
+						MaskOp->clip = ClipResults.MeshOp;
+
+						Ptr<ASTOpMeshRemoveMask> RemoveOp = new ASTOpMeshRemoveMask();
+						RemoveOp->source = LastMeshResults.MeshOp;
+						RemoveOp->FaceCullStrategy = TypedClipNode->FaceCullStrategy;
+						RemoveOp->AddRemove(Modifier.FinalCondition, MaskOp);
 
-				Ptr<ASTOp> fullCondition = m.FinalCondition;
+						LastMeshResults.MeshOp = RemoveOp;
 
-				RemoveOp->AddRemove(fullCondition, maskAt);
+						return LastMeshResults;
+					},
+					UE::Tasks::Prerequisites(ClipMeshTask, LastMeshTask)
+				);
 			}
 		}
 
 		// Process clip-with-mask modifiers
-		for (const FirstPassGenerator::FModifier& m : Modifiers)
+		for (const FirstPassGenerator::FModifier& Modifier : Modifiers)
 		{
-			if (m.Node->GetType() == NodeModifierMeshClipWithUVMask::GetStaticType())
+			if (StaticOptions.ModifiersToIgnore.Contains(Modifier))
+			{
+				// Prevent recursion.
+				continue;
+			}
+
+			if (Modifier.Node->GetType() == NodeModifierMeshClipWithUVMask::GetStaticType())
 			{
 				// Create a constant mesh with the original UVs required by this modifier.
 				// TODO: Optimize, by caching.
 				// TODO: Optimize by formatting and keeping only UVs
 				check(OriginalMeshNode);
-				const FMesh* OriginalMesh = OriginalMeshNode->GetPrivate()->Value.Get();
-				Ptr<ASTOpConstantResource> UVMeshOp = new ASTOpConstantResource();
-				UVMeshOp->Type = EOpType::ME_CONSTANT;
-				UVMeshOp->SetValue(OriginalMesh->Clone(), CompilerOptions->OptimisationOptions.DiskCacheContext);
-				UVMeshOp->SourceDataDescriptor = OriginalMeshNode->SourceDataDescriptor;
-
-				const NodeModifierMeshClipWithUVMask* TypedClipNode = static_cast<const NodeModifierMeshClipWithUVMask*>(m.Node);
 
-				Ptr<ASTOp> MeshMaskAt;
+				TSharedPtr< TSharedPtr<FMesh> > ResolveOriginalMesh = MakeShared<TSharedPtr<FMesh>>();
+				UE::Tasks::FTask TargetMeshTask;
+				{
+					FMesh* OriginalMesh = OriginalMeshNode->GetValue().Get();
+					check(OriginalMesh->IsReference());
+					uint32 OriginalMeshID = OriginalMesh->GetReferencedMesh();
+
+					bool bRunImmediatlyIfPossible = IsInGameThread();
+					FString NoMorph;
+					TargetMeshTask = CompilerOptions->OptimisationOptions.ReferencedMeshResourceProvider(
+						OriginalMeshID,
+						NoMorph,
+						ResolveOriginalMesh,
+						bRunImmediatlyIfPossible);
+				}
 
-				Ptr<ASTOpMeshMaskClipUVMask> op = new ASTOpMeshMaskClipUVMask();
-				MeshMaskAt = op;
-				op->Source = BaseMeshResult.BaseMeshOp; 
-				op->UVSource = UVMeshOp; 
-				op->LayoutIndex = TypedClipNode->LayoutIndex;
+				const NodeModifierMeshClipWithUVMask* TypedClipNode = static_cast<const NodeModifierMeshClipWithUVMask*>(Modifier.Node);
 
+				Ptr<ASTOp> MaskImage;
+				TSharedPtr<const FLayout> Layout;
 				if (TypedClipNode->ClipMask)
 				{
-					// FParameters to generate the mask image
-					FImageGenerationOptions ClipOptions(Options.ComponentId);
+					FImageGenerationOptions ClipOptions(StaticOptions.ComponentId, StaticOptions.LODIndex);
 					ClipOptions.ImageLayoutStrategy = CompilerOptions::TextureLayoutStrategy::None;
 					ClipOptions.LayoutBlockId = FLayoutBlock::InvalidBlockId;
-					ClipOptions.State = Options.State;
+					ClipOptions.State = StaticOptions.State;
 
 					FImageGenerationResult ClipMaskResult;
 					GenerateImage(ClipOptions, ClipMaskResult, TypedClipNode->ClipMask);
 
 					// It could be IF_L_UBIT, but since this should be optimized out at compile time, leave the most cpu efficient.
-					op->MaskImage = GenerateImageFormat(ClipMaskResult.op, mu::EImageFormat::L_UByte);
-
-					if (!op->MaskImage)
-					{
-						ErrorLog->Add("Clip UV mask has not been generated", ELMT_ERROR, ErrorContext);
-						continue;
-					}
+					MaskImage = GenerateImageFormat(ClipMaskResult.op, mu::EImageFormat::L_UByte);
 				}
-
 				else if (TypedClipNode->ClipLayout)
 				{
 					// Generate the layout with blocks to extract
-					TSharedPtr<const FLayout> Layout = GenerateLayout(TypedClipNode->ClipLayout, 0);
-
-					Ptr<ASTOpConstantResource> LayoutOp = new ASTOpConstantResource();
-					LayoutOp->Type = EOpType::LA_CONSTANT;
-					LayoutOp->SetValue(Layout, CompilerOptions->OptimisationOptions.DiskCacheContext);
-					op->MaskLayout = LayoutOp;
-				}
-
-				else
-				{
-					// No mask or layout specified to clip. Don't clip anything.
+					Layout = GenerateLayout(TypedClipNode->ClipLayout, 0);
 				}
 
-				if (MeshMaskAt)
-				{
-					if (!RemoveOp)
+				Ptr<const NodeMeshConstant> OriginalMeshCopy = OriginalMeshNode;
+				LastMeshTask = UE::Tasks::Launch(TEXT("MutableModifier"),
+					[Layout, MaskImage, ResolveOriginalMesh, LastMeshTask, PreModifiersTask, Modifier, TypedClipNode, CompilerOptions = CompilerOptions, ErrorLog = ErrorLog, ErrorContext, OriginalMeshCopy]() mutable
 					{
-						RemoveOp = new ASTOpMeshRemoveMask();
-						RemoveOp->source = LastMeshOp;
-						RemoveOp->FaceCullStrategy = TypedClipNode->FaceCullStrategy;
-						LastMeshOp = RemoveOp;
-					}
+						TSharedPtr<FMesh> TargetMesh = *ResolveOriginalMesh;
+						FMeshGenerationResult LastMeshResults = LastMeshTask.GetResult();
+						FMeshGenerationResult BaseResults = PreModifiersTask.GetResult();
 
-					Ptr<ASTOp> fullCondition = m.FinalCondition;
-					RemoveOp->AddRemove(fullCondition, MeshMaskAt);
-				}
+						Ptr<ASTOpConstantResource> UVMeshOp = new ASTOpConstantResource;
+						UVMeshOp->Type = EOpType::ME_CONSTANT;
+						UVMeshOp->SetValue(TargetMesh->Clone(), CompilerOptions->OptimisationOptions.DiskCacheContext);
+						UVMeshOp->SourceDataDescriptor = OriginalMeshCopy->SourceDataDescriptor;
+
+						Ptr<ASTOpMeshMaskClipUVMask> MeshMaskAt = new ASTOpMeshMaskClipUVMask();
+						MeshMaskAt->Source = BaseResults.BaseMeshOp;
+						MeshMaskAt->UVSource = UVMeshOp;
+						MeshMaskAt->LayoutIndex = TypedClipNode->LayoutIndex;
+
+						if (TypedClipNode->ClipMask)
+						{
+							MeshMaskAt->MaskImage = MaskImage;
+							if (!MeshMaskAt->MaskImage)
+							{
+								ErrorLog->Add("Clip UV mask has not been generated", ELMT_ERROR, ErrorContext);
+								return LastMeshResults;
+							}
+						}
+
+						else if (TypedClipNode->ClipLayout)
+						{
+							Ptr<ASTOpConstantResource> LayoutOp = new ASTOpConstantResource();
+							LayoutOp->Type = EOpType::LA_CONSTANT;
+							LayoutOp->SetValue(Layout, CompilerOptions->OptimisationOptions.DiskCacheContext);
+							MeshMaskAt->MaskLayout = LayoutOp;
+						}
+
+						else
+						{
+							// No mask or layout specified to clip. Don't clip anything.
+							return LastMeshResults;
+						}
+
+						if (MeshMaskAt)
+						{
+							Ptr<ASTOpMeshRemoveMask> RemoveOp = new ASTOpMeshRemoveMask();
+							RemoveOp->source = LastMeshResults.MeshOp;
+							RemoveOp->FaceCullStrategy = TypedClipNode->FaceCullStrategy;
+							RemoveOp->AddRemove(Modifier.FinalCondition, MeshMaskAt);
+
+							LastMeshResults.MeshOp = RemoveOp;
+						}
+
+						return LastMeshResults;
+					},
+					UE::Tasks::Prerequisites(TargetMeshTask, LastMeshTask)
+				);
 			}
 		}
 
 		// Process clip-morph-plane modifiers
-		for (const FirstPassGenerator::FModifier& m : Modifiers)
+		for (const FirstPassGenerator::FModifier& Modifier : Modifiers)
 		{
-			Ptr<ASTOp> modifiedMeshOp;
+			if (StaticOptions.ModifiersToIgnore.Contains(Modifier))
+			{
+				// Prevent recursion.
+				continue;
+			}
 
-			if (m.Node->GetType() == NodeModifierMeshClipMorphPlane::GetStaticType())
+			if (Modifier.Node->GetType() == NodeModifierMeshClipMorphPlane::GetStaticType())
 			{
-				const NodeModifierMeshClipMorphPlane* TypedNode = static_cast<const NodeModifierMeshClipMorphPlane*>(m.Node);
-				Ptr<ASTOpMeshClipMorphPlane> op = new ASTOpMeshClipMorphPlane();
-				op->Source = LastMeshOp;
-				op->FaceCullStrategy = TypedNode->Parameters.FaceCullStrategy;
+				const NodeModifierMeshClipMorphPlane* TypedNode = static_cast<const NodeModifierMeshClipMorphPlane*>(Modifier.Node);
 
-				// Morph to an ellipse
-				{
-					FShape morphShape;
-					morphShape.type = (uint8_t)FShape::Type::Ellipse;
-					morphShape.position = TypedNode->Parameters.Origin;
-					morphShape.up = TypedNode->Parameters.Normal;
-					// TODO: Move rotation to ellipse rotation reference base instead of passing it directly
-					morphShape.size = FVector3f(TypedNode->Parameters.Radius1, TypedNode->Parameters.Radius2, TypedNode->Parameters.Rotation);
-
-					// Generate a "side" vector.
-					// \todo: make generic and move to the vector class
+				LastMeshTask = UE::Tasks::Launch(TEXT("MutableModifier"),
+					[LastMeshTask, PreModifiersTask, Modifier, TypedNode]() mutable
 					{
-						// Generate vector perpendicular to normal for ellipse rotation reference base
-						FVector3f aux_base(0.f, 1.f, 0.f);
+						FMeshGenerationResult LastMeshResults = LastMeshTask.GetResult();
+						FMeshGenerationResult BaseResults = PreModifiersTask.GetResult();
+
+						Ptr<ASTOpMeshClipMorphPlane> ClipOp = new ASTOpMeshClipMorphPlane();
+						ClipOp->Source = LastMeshResults.MeshOp;
+						ClipOp->FaceCullStrategy = TypedNode->Parameters.FaceCullStrategy;
 
-						if (FMath::Abs(FVector3f::DotProduct(TypedNode->Parameters.Normal, aux_base)) > 0.95f)
+						// Morph to an ellipse
 						{
-							aux_base = FVector3f(0.f, 0.f, 1.f);
+							FShape MorphShape;
+							MorphShape.type = (uint8_t)FShape::Type::Ellipse;
+							MorphShape.position = TypedNode->Parameters.Origin;
+							MorphShape.up = TypedNode->Parameters.Normal;
+							// TODO: Move rotation to ellipse rotation reference base instead of passing it directly
+							MorphShape.size = FVector3f(TypedNode->Parameters.Radius1, TypedNode->Parameters.Radius2, TypedNode->Parameters.Rotation);
+
+							// Generate a "side" vector.
+							// \todo: make generic and move to the vector class
+							{
+								// Generate vector perpendicular to normal for ellipse rotation reference base
+								FVector3f AuxBase(0.f, 1.f, 0.f);
+
+								if (FMath::Abs(FVector3f::DotProduct(TypedNode->Parameters.Normal, AuxBase)) > 0.95f)
+								{
+									AuxBase = FVector3f(0.f, 0.f, 1.f);
+								}
+
+								MorphShape.side = FVector3f::CrossProduct(TypedNode->Parameters.Normal, AuxBase);
+							}
+							ClipOp->MorphShape = MorphShape;
 						}
 
-						morphShape.side = FVector3f::CrossProduct(TypedNode->Parameters.Normal, aux_base);
-					}
-					op->MorphShape = morphShape;
-				}
+						// Selection box
+						ClipOp->VertexSelectionType = TypedNode->Parameters.VertexSelectionType;
+						if (ClipOp->VertexSelectionType == EClipVertexSelectionType::Shape)
+						{
+							FShape SelectionShape;
+							SelectionShape.type = (uint8)FShape::Type::AABox;
+							SelectionShape.position = TypedNode->Parameters.SelectionBoxOrigin;
+							SelectionShape.size = TypedNode->Parameters.SelectionBoxRadius;
+							ClipOp->SelectionShape = SelectionShape;
+						}
+						else if (ClipOp->VertexSelectionType == EClipVertexSelectionType::BoneHierarchy)
+						{
+							ClipOp->VertexSelectionBone = TypedNode->Parameters.VertexSelectionBone;
+							ClipOp->VertexSelectionBoneMaxRadius = TypedNode->Parameters.MaxEffectRadius;
+						}
 
-				// Selection box
-				op->VertexSelectionType = TypedNode->Parameters.VertexSelectionType;
-				if (op->VertexSelectionType == EClipVertexSelectionType::Shape)
-				{
-					FShape selectionShape;
-					selectionShape.type = (uint8)FShape::Type::AABox;
-					selectionShape.position = TypedNode->Parameters.SelectionBoxOrigin;
-					selectionShape.size = TypedNode->Parameters.SelectionBoxRadius;
-					op->SelectionShape = selectionShape;
-				}
-				else if (op->VertexSelectionType == EClipVertexSelectionType::BoneHierarchy)
-				{
-					op->VertexSelectionBone = TypedNode->Parameters.VertexSelectionBone;
-					op->VertexSelectionBoneMaxRadius = TypedNode->Parameters.MaxEffectRadius;
-				}
+						ClipOp->Dist = TypedNode->Parameters.DistanceToPlane;
+						ClipOp->Factor = TypedNode->Parameters.LinearityFactor;
 
-				// FParameters
-				op->Dist = TypedNode->Parameters.DistanceToPlane;
-				op->Factor = TypedNode->Parameters.LinearityFactor;
+						Ptr<ASTOpConditional> ConditionalOp = new ASTOpConditional();
+						ConditionalOp->type = EOpType::ME_CONDITIONAL;
+						ConditionalOp->no = LastMeshResults.MeshOp;
+						ConditionalOp->yes = ClipOp;
+						ConditionalOp->condition = Modifier.FinalCondition;
 
-				modifiedMeshOp = op;
+						LastMeshResults.MeshOp = ConditionalOp;
 
-				Ptr<ASTOp> fullCondition = m.FinalCondition;
+						return LastMeshResults;
+					},
+					UE::Tasks::Prerequisites(LastMeshTask)
+				);
 
-				Ptr<ASTOpConditional> ConditionalOp = new ASTOpConditional();
-				ConditionalOp->type = EOpType::ME_CONDITIONAL;
-				ConditionalOp->no = LastMeshOp;
-				ConditionalOp->yes = modifiedMeshOp;
-				ConditionalOp->condition = fullCondition;
-				LastMeshOp = ConditionalOp;
 			}
 		}
 
     	// Process clip deform modifiers.
-		for (const FirstPassGenerator::FModifier& M : Modifiers)
+		for (const FirstPassGenerator::FModifier& Modifier : Modifiers)
 		{
-			Ptr<ASTOp> ModifiedMeshOp;
+			if (StaticOptions.ModifiersToIgnore.Contains(Modifier))
+			{
+				// Prevent recursion.
+				continue;
+			}
 
-			if (M.Node->GetType()==NodeModifierMeshClipDeform::GetStaticType())
+			if (Modifier.Node->GetType() == NodeModifierMeshClipDeform::GetStaticType())
 			{
-				const NodeModifierMeshClipDeform* TypedClipNode = static_cast<const NodeModifierMeshClipDeform*>(M.Node);
-				Ptr<ASTOpMeshBindShape>  BindOp = new ASTOpMeshBindShape();
-				Ptr<ASTOpMeshClipDeform> ClipOp = new ASTOpMeshClipDeform();
+				const NodeModifierMeshClipDeform* TypedClipNode = static_cast<const NodeModifierMeshClipDeform*>(Modifier.Node);
 
-				ClipOp->FaceCullStrategy = TypedClipNode->FaceCullStrategy;
+				FMeshGenerationStaticOptions ClipStaticOptions(StaticOptions);
+				ClipStaticOptions.ActiveTags.Empty();
+				ClipStaticOptions.ModifiersToIgnore.AddUnique(Modifier);
+				FMeshGenerationDynamicOptions ClipMeshDynamicOptions;
+				ClipMeshDynamicOptions.bLayouts = false;
+				FMeshTask ClipMeshTask = GenerateMesh(ClipStaticOptions, UE::Tasks::MakeCompletedTask<FMeshGenerationDynamicOptions>(ClipMeshDynamicOptions), TypedClipNode->ClipMesh);
 
-				FMeshGenerationOptions ClipOptions(Options.ComponentId);
-				ClipOptions.bLayouts = false;
-				ClipOptions.State = Options.State;				
+				LastMeshTask = UE::Tasks::Launch(TEXT("MutableMeshMergeModifier"),
+					[ClipMeshTask, LastMeshTask, Modifier, TypedClipNode, ErrorLog = ErrorLog, ErrorContext]() mutable
+					{
+						FMeshGenerationResult ClipResults = ClipMeshTask.GetResult();
+						FMeshGenerationResult LastMeshResults = LastMeshTask.GetResult();
 
-				FMeshGenerationResult ClipShapeResult;
-				GenerateMesh(ClipOptions, ClipShapeResult, TypedClipNode->ClipMesh);
-				ClipOp->ClipShape = ClipShapeResult.MeshOp;
-				
-				BindOp->Mesh = LastMeshOp;
-				BindOp->Shape = ClipShapeResult.MeshOp; 
-				BindOp->BindingMethod = static_cast<uint32>(TypedClipNode->BindingMethod);
-	
-				ClipOp->Mesh = BindOp;
+						Ptr<ASTOpMeshBindShape>  BindOp = new ASTOpMeshBindShape();
+						BindOp->Mesh = LastMeshResults.MeshOp;
+						BindOp->Shape = ClipResults.MeshOp;
+						BindOp->BindingMethod = static_cast<uint32>(TypedClipNode->BindingMethod);
 
-				if (!ClipOp->ClipShape)
-				{
-					ErrorLog->Add("Clip shape mesh has not been generated", ELMT_ERROR, ErrorContext);
-				}
-				else
-				{
-					ModifiedMeshOp = ClipOp;
-				}
-			}
-			
-			if (ModifiedMeshOp)
-			{
-				Ptr<ASTOp> FullCondition = M.FinalCondition;
+						Ptr<ASTOpMeshClipDeform> ClipOp = new ASTOpMeshClipDeform();
+						ClipOp->FaceCullStrategy = TypedClipNode->FaceCullStrategy;
+						ClipOp->ClipShape = ClipResults.MeshOp;
+						ClipOp->Mesh = BindOp;
+
+						if (!ClipOp->ClipShape)
+						{
+							ErrorLog->Add("Clip shape mesh has not been generated", ELMT_ERROR, ErrorContext);
+							return LastMeshResults;
+						}
+
+						Ptr<ASTOpConditional> Op = new ASTOpConditional();
+						Op->type = EOpType::ME_CONDITIONAL;
+						Op->no = LastMeshResults.MeshOp;
+						Op->yes = ClipOp;
+						Op->condition = Modifier.FinalCondition;
 
-				Ptr<ASTOpConditional> Op = new ASTOpConditional();
-				Op->type = EOpType::ME_CONDITIONAL;
-				Op->no = LastMeshOp;
-				Op->yes = ModifiedMeshOp;
-				Op->condition = FullCondition;
-				LastMeshOp = Op;
+						LastMeshResults.MeshOp = Op;
+
+						return LastMeshResults;
+					},
+					UE::Tasks::Prerequisites(ClipMeshTask, LastMeshTask)
+				);
 			}
 		}
 		
 		// Process transform mesh within mesh modifiers.
-		for (const FirstPassGenerator::FModifier& m : Modifiers)
+		for (const FirstPassGenerator::FModifier& Modifier : Modifiers)
 		{
-			if (m.Node->GetType()== NodeModifierMeshTransformInMesh::GetStaticType())
+			if (StaticOptions.ModifiersToIgnore.Contains(Modifier))
+			{
+				// Prevent recursion.
+				continue;
+			}
+
+			if (Modifier.Node->GetType()== NodeModifierMeshTransformInMesh::GetStaticType())
 			{
-				const NodeModifierMeshTransformInMesh* TypedTransformNode = static_cast<const NodeModifierMeshTransformInMesh*>(m.Node);
+				const NodeModifierMeshTransformInMesh* TypedTransformNode = static_cast<const NodeModifierMeshTransformInMesh*>(Modifier.Node);
 
 				// If a matrix node is not connected, the op won't do anything, so let's not create it at all.
 				if (TypedTransformNode->MatrixNode)
 				{
-					Ptr<ASTOpMeshTransformWithBoundingMesh> TransformOp = new ASTOpMeshTransformWithBoundingMesh();
-					TransformOp->source = LastMeshOp;
-
-					// Transform matrix.
-					if (TypedTransformNode->MatrixNode)
+					if (!TypedTransformNode->MatrixNode || !TypedTransformNode->BoundingMesh)
 					{
-						FMatrixGenerationResult ChildResult;
-						GenerateMatrix(ChildResult, Options, TypedTransformNode->MatrixNode);
-						TransformOp->matrix = ChildResult.op;
+						ErrorLog->Add("Bounding mesh or matrix have not been generated", ELMT_ERROR, ErrorContext);
+						continue;
 					}
-					
-					if (TypedTransformNode->BoundingMesh)
-					{
-						// FParameters
-						FMeshGenerationOptions MeshOptions(Options.ComponentId);
-						MeshOptions.bLayouts = false;
-						MeshOptions.State = Options.State;
-
-						FMeshGenerationResult BoundingMeshResult;
-						GenerateMesh(MeshOptions, BoundingMeshResult, TypedTransformNode->BoundingMesh);
-						TransformOp->boundingMesh = BoundingMeshResult.MeshOp;
 
-						if (!TransformOp->boundingMesh)
-						{
-							ErrorLog->Add("Bounding mesh has not been generated", ELMT_ERROR, ErrorContext);
-							continue;
-						}
-						else
+					// Transform matrix.
+					Ptr<ASTOp> MatrixOp;
+					FMatrixGenerationResult ChildResult;
+					GenerateMatrix(ChildResult, StaticOptions, TypedTransformNode->MatrixNode);
+					MatrixOp = ChildResult.op;
+
+					// Bounding mesh
+					FMeshGenerationStaticOptions ClipStaticOptions(StaticOptions);
+					ClipStaticOptions.ActiveTags.Empty();
+					ClipStaticOptions.ModifiersToIgnore.AddUnique(Modifier);
+					FMeshGenerationDynamicOptions BoundingMeshDynamicOptions;
+					BoundingMeshDynamicOptions.bLayouts = false;
+					FMeshTask BoundingMeshTask = GenerateMesh(ClipStaticOptions, UE::Tasks::MakeCompletedTask<FMeshGenerationDynamicOptions>(BoundingMeshDynamicOptions), TypedTransformNode->BoundingMesh);
+
+					LastMeshTask = UE::Tasks::Launch(TEXT("MutableMeshMergeModifier"),
+						[BoundingMeshTask, LastMeshTask, Modifier, MatrixOp, ErrorLog = ErrorLog, ErrorContext]() mutable
 						{
-							ASTOp::EClosedMeshTest IsClosed = TransformOp->boundingMesh->IsClosedMesh();
-							if (IsClosed==ASTOp::EClosedMeshTest::No)
+							FMeshGenerationResult BoundingResults = BoundingMeshTask.GetResult();
+							FMeshGenerationResult LastMeshResults = LastMeshTask.GetResult();
+
+							Ptr<ASTOpMeshTransformWithBoundingMesh> TransformOp = new ASTOpMeshTransformWithBoundingMesh();
+							TransformOp->source = LastMeshResults.MeshOp;
+							TransformOp->matrix = MatrixOp;
+							TransformOp->boundingMesh = BoundingResults.MeshOp;
+
+							if (TransformOp->boundingMesh)
 							{
-								ErrorLog->Add("Mesh used for clipping is not closed.", ELMT_WARNING, ErrorContext);
+								ASTOp::EClosedMeshTest IsClosed = TransformOp->boundingMesh->IsClosedMesh();
+								if (IsClosed == ASTOp::EClosedMeshTest::No)
+								{
+									ErrorLog->Add("Mesh used for clipping is not closed.", ELMT_WARNING, ErrorContext);
+								}
 							}
-						}
-					}
 
-					// Condition to apply the transform op
-					if (m.FinalCondition)
-					{
-						Ptr<ASTOpConditional> conditionalAd = new ASTOpConditional();
-						conditionalAd->type = EOpType::ME_CONDITIONAL;
-						conditionalAd->no = LastMeshOp;
-						conditionalAd->yes = TransformOp;
-						conditionalAd->condition = m.FinalCondition;
-						LastMeshOp = conditionalAd;
-					}
-					else
-					{
-						LastMeshOp = TransformOp;
-					}
+							// Condition to apply the transform op
+							if (Modifier.FinalCondition)
+							{
+								Ptr<ASTOpConditional> ConditionalOp = new ASTOpConditional();
+								ConditionalOp->type = EOpType::ME_CONDITIONAL;
+								ConditionalOp->no = LastMeshResults.MeshOp;
+								ConditionalOp->yes = TransformOp;
+								ConditionalOp->condition = Modifier.FinalCondition;
+								LastMeshResults.MeshOp = ConditionalOp;
+							}
+							else
+							{
+								LastMeshResults.MeshOp = TransformOp;
+							}
+
+							return LastMeshResults;
+						},
+						UE::Tasks::Prerequisites(BoundingMeshTask, LastMeshTask)
+					);
 				}
 			}
 		}
 		
 
-		return LastMeshOp;
+		return LastMeshTask;
 	}
 
 
@@ -2603,14 +2780,15 @@ namespace mu
 	{
 		Ptr<ASTOp> LastImageOp = BaseImageOp;
 
-		int32 CurrentLOD = CurrentParents.Last().Lod;
+		int32 CurrentLOD = Options.LODIndex;
+		check(CurrentLOD >= 0);
 
 		// Process patch image modifiers (from edit modifiers)
-		for (const FirstPassGenerator::FModifier& m : Modifiers)
+		for (const FirstPassGenerator::FModifier& Modifier : Modifiers)
 		{
-			if (m.Node->GetType() == NodeModifierSurfaceEdit::GetStaticType())
+			if (Modifier.Node->GetType() == NodeModifierSurfaceEdit::GetStaticType())
 			{
-				const NodeModifierSurfaceEdit* Edit = static_cast<const NodeModifierSurfaceEdit*>(m.Node);
+				const NodeModifierSurfaceEdit* Edit = static_cast<const NodeModifierSurfaceEdit*>(Modifier.Node);
 
 				bool bAffectsCurrentLOD = Edit->LODs.IsValidIndex(CurrentLOD);
 
@@ -2637,7 +2815,7 @@ namespace mu
 
 					if (PatchMask)
 					{
-						LastImageOp = GenerateImageBlockPatch(LastImageOp, *MatchingEdit, PatchMask, m.FinalCondition, Options);
+						LastImageOp = GenerateImageBlockPatch(LastImageOp, *MatchingEdit, PatchMask, Modifier.FinalCondition, Options);
 					}
 				}
 			}
@@ -2671,8 +2849,7 @@ namespace mu
 
 	Ptr<ASTOp> CodeGenerator::ApplyImageExtendModifiers(
 		const TArray<FirstPassGenerator::FModifier>& Modifiers,
-		const FGenericGenerationOptions& Options,
-		int32 ComponentId,
+		const FMeshGenerationStaticOptions& Options,
 		const FMeshGenerationResult& BaseMeshResults,
 		Ptr<ASTOp> BaseImageOp, 
 		CompilerOptions::TextureLayoutStrategy ImageLayoutStrategy,
@@ -2684,16 +2861,16 @@ namespace mu
 	{
 		Ptr<ASTOp> LastImageOp = BaseImageOp;
 
-		int32 CurrentLOD = CurrentParents.Last().Lod;
-
+		int32 CurrentLOD = Options.LODIndex;
+		check(CurrentLOD >= 0);
 
 		// Process mesh extend modifiers (from edit modifiers)
 		int32 EditIndex = 0;
-		for (const FirstPassGenerator::FModifier& m : Modifiers)
+		for (const FirstPassGenerator::FModifier& Modifier : Modifiers)
 		{
-			if (m.Node->GetType() == NodeModifierSurfaceEdit::GetStaticType())
+			if (Modifier.Node->GetType() == NodeModifierSurfaceEdit::GetStaticType())
 			{
-				const NodeModifierSurfaceEdit* Edit = static_cast<const NodeModifierSurfaceEdit*>(m.Node);
+				const NodeModifierSurfaceEdit* Edit = static_cast<const NodeModifierSurfaceEdit*>(Modifier.Node);
 
 				int32 ThisEditIndex = EditIndex;
 				++EditIndex;
@@ -2722,84 +2899,86 @@ namespace mu
 					continue;
 				}
 
+				if (!BaseMeshResults.ExtraMeshLayouts.IsValidIndex(ThisEditIndex))
+				{
+					ErrorLog->Add(TEXT("Trying to extend a layout that doesn't exist."), ELMT_WARNING, Edit->GetMessageContext(), ModifiedNodeErrorContext);
+					continue;
+				}
+
 				const TArray<FGeneratedLayout>& ExtraLayouts = BaseMeshResults.ExtraMeshLayouts[ThisEditIndex].GeneratedLayouts;
 
 				if (LayoutIndex >= ExtraLayouts.Num() || !ExtraLayouts[LayoutIndex].Layout)
 				{
 					ErrorLog->Add(TEXT("Trying to extend a layout that doesn't exist."), ELMT_WARNING, Edit->GetMessageContext(), ModifiedNodeErrorContext);
+					continue;
 				}
-				else
-				{
-					TSharedPtr<const FLayout> pExtendLayout = ExtraLayouts[LayoutIndex].Layout;
 
-					Ptr<ASTOp> lastBase = LastImageOp;
+				TSharedPtr<const FLayout> ExtendLayout = ExtraLayouts[LayoutIndex].Layout;
 
-					for (int32 b = 0; b < pExtendLayout->GetBlockCount(); ++b)
-					{
-						// Generate the image block
-						FImageGenerationOptions ImageOptions(ComponentId);
-						ImageOptions.State = Options.State;
-						ImageOptions.ImageLayoutStrategy = ImageLayoutStrategy;
-						ImageOptions.ActiveTags = Edit->EnableTags; // TODO: Merge with current tags?
-						ImageOptions.RectSize = { 0,0 };
-						ImageOptions.LayoutToApply = pExtendLayout;
-						ImageOptions.LayoutBlockId = pExtendLayout->Blocks[b].Id;
-						FImageGenerationResult ExtendResult;
-						GenerateImage(ImageOptions, ExtendResult, MatchingEdit->Extend);
-						Ptr<ASTOp> fragmentAd = ExtendResult.op;
-
-						// Block in layout grid units
-						box< FIntVector2 > rectInCells;
-						rectInCells.min = pExtendLayout->Blocks[b].Min;
-						rectInCells.size = pExtendLayout->Blocks[b].Size;
-
-						FImageDesc ExtendDesc = fragmentAd->GetImageDesc();
-
-						// If we don't know the size of a layout block in pixels, calculate it
-						UpdateLayoutBlockDesc(InOutLayoutBlockDesc, ExtendDesc, rectInCells.size);
-
-						// Adjust the format and size of the block to be added
-						// Actually don't do it, it will be propagated from the top format operation.
-						//fragmentAd = GenerateImageFormat(fragmentAd, FinalImageFormat);
-
-						UE::Math::TIntVector2<int32> expectedSize;
-						expectedSize[0] = InOutLayoutBlockDesc.BlockPixelsX * rectInCells.size[0];
-						expectedSize[1] = InOutLayoutBlockDesc.BlockPixelsY * rectInCells.size[1];
-						fragmentAd = GenerateImageSize(fragmentAd, expectedSize);
-
-						// Apply tiling to avoid generating chunks of image that are too big.
-						fragmentAd = ApplyTiling(fragmentAd, expectedSize, InOutLayoutBlockDesc.FinalFormat);
-
-						// Compose operation
-						Ptr<ASTOpImageCompose> composeOp = new ASTOpImageCompose();
-						composeOp->Layout = BaseMeshResults.LayoutOps[LayoutIndex];
-						composeOp->Base = lastBase;
-						composeOp->BlockImage = fragmentAd;
-
-						// Set the absolute block index.
-						check(pExtendLayout->Blocks[b].Id != FLayoutBlock::InvalidBlockId);
-						composeOp->BlockId = pExtendLayout->Blocks[b].Id;
-
-						lastBase = composeOp;
-					}
+				Ptr<ASTOp> LastBase = LastImageOp;
 
-					// Condition to enable this image extension
-					if (m.FinalCondition)
-					{
-						Ptr<ASTOp> conditionalAd;
-						Ptr<ASTOpConditional> cop = new ASTOpConditional();
-						cop->type = EOpType::IM_CONDITIONAL;
-						cop->no = LastImageOp;
-						cop->yes = lastBase;
-						cop->condition = m.FinalCondition;
-						conditionalAd = cop;
-						LastImageOp = conditionalAd;
-					}
-					else
-					{
-						LastImageOp = lastBase;
-					}
+				for (int32 BlockIndex = 0; BlockIndex < ExtendLayout->GetBlockCount(); ++BlockIndex)
+				{
+					// Generate the image block
+					FImageGenerationOptions ImageOptions(Options.ComponentId,Options.LODIndex);
+					ImageOptions.State = Options.State;
+					ImageOptions.ImageLayoutStrategy = ImageLayoutStrategy;
+					ImageOptions.ActiveTags = Edit->EnableTags; // TODO: Merge with current tags?
+					ImageOptions.RectSize = { 0,0 };
+					ImageOptions.LayoutToApply = ExtendLayout;
+					ImageOptions.LayoutBlockId = ExtendLayout->Blocks[BlockIndex].Id;
+					FImageGenerationResult ExtendResult;
+					GenerateImage(ImageOptions, ExtendResult, MatchingEdit->Extend);
+					Ptr<ASTOp> FragmentOp = ExtendResult.op;
+
+					// Block in layout grid units
+					box< FIntVector2 > RectInCells;
+					RectInCells.min = ExtendLayout->Blocks[BlockIndex].Min;
+					RectInCells.size = ExtendLayout->Blocks[BlockIndex].Size;
+
+					FImageDesc ExtendDesc = FragmentOp->GetImageDesc();
+
+					// If we don't know the size of a layout block in pixels, calculate it
+					UpdateLayoutBlockDesc(InOutLayoutBlockDesc, ExtendDesc, RectInCells.size);
+
+					// Adjust the format and size of the block to be added
+					// Actually don't do it, it will be propagated from the top format operation.
+					//FragmentOp = GenerateImageFormat(FragmentOp, FinalImageFormat);
+
+					UE::Math::TIntVector2<int32> ExpectedSize;
+					ExpectedSize[0] = InOutLayoutBlockDesc.BlockPixelsX * RectInCells.size[0];
+					ExpectedSize[1] = InOutLayoutBlockDesc.BlockPixelsY * RectInCells.size[1];
+					FragmentOp = GenerateImageSize(FragmentOp, ExpectedSize);
+
+					// Apply tiling to avoid generating chunks of image that are too big.
+					FragmentOp = ApplyTiling(FragmentOp, ExpectedSize, InOutLayoutBlockDesc.FinalFormat);
+
+					// Compose operation
+					Ptr<ASTOpImageCompose> ComposeOp = new ASTOpImageCompose();
+					ComposeOp->Layout = BaseMeshResults.LayoutOps[LayoutIndex];
+					ComposeOp->Base = LastBase;
+					ComposeOp->BlockImage = FragmentOp;
+
+					// Set the absolute block index.
+					check(ExtendLayout->Blocks[BlockIndex].Id != FLayoutBlock::InvalidBlockId);
+					ComposeOp->BlockId = ExtendLayout->Blocks[BlockIndex].Id;
+
+					LastBase = ComposeOp;
+				}
 
+				// Condition to enable this image extension
+				if (Modifier.FinalCondition)
+				{
+					Ptr<ASTOpConditional> Op = new ASTOpConditional();
+					Op->type = EOpType::IM_CONDITIONAL;
+					Op->no = LastImageOp;
+					Op->yes = LastBase;
+					Op->condition = Modifier.FinalCondition;
+					LastImageOp = Op;
+				}
+				else
+				{
+					LastImageOp = LastBase;
 				}
 			}
 		}
@@ -2808,9 +2987,10 @@ namespace mu
 	}
 
 
-	void CodeGenerator::CheckModifiersForSurface(const NodeSurfaceNew& Node, const TArray<FirstPassGenerator::FModifier>& Modifiers )
+	void CodeGenerator::CheckModifiersForSurface(const NodeSurfaceNew& Node, const TArray<FirstPassGenerator::FModifier>& Modifiers, int32 LODIndex )
 	{
-		int32 CurrentLOD = CurrentParents.Last().Lod;
+		int32 CurrentLOD = LODIndex;
+		check(CurrentLOD >= 0);
 
 		for (const FirstPassGenerator::FModifier& Mod : Modifiers)
 		{
diff --git a/Engine/Plugins/Mutable/Source/MutableTools/Private/MuT/CodeGenerator.h b/Engine/Plugins/Mutable/Source/MutableTools/Private/MuT/CodeGenerator.h
--- a/Engine/Plugins/Mutable/Source/MutableTools/Private/MuT/CodeGenerator.h
+++ b/Engine/Plugins/Mutable/Source/MutableTools/Private/MuT/CodeGenerator.h
@@ -38,6 +38,7 @@
 #include "MuT/Table.h"
 #include "MuT/TablePrivate.h"
 #include "Templates/TypeHash.h"
+#include "Tasks/Pipe.h"
 
 
 namespace mu
@@ -103,11 +104,7 @@ namespace mu
 	struct FObjectState;
 	struct FProgram;
 
-
-
-    //---------------------------------------------------------------------------------------------
-    //! Code generator
-    //---------------------------------------------------------------------------------------------
+    /** Converts mu::Node graphs into ASTOp graphs. */
     class CodeGenerator
     {
 		
@@ -237,29 +234,37 @@ namespace mu
 
     private:
 
-        /** Container of meshes generated to be able to reuse them. 
-		* They are sorted by a cheap hash to speed up searches.
+		/** This struct contains the generation state that can only be accessed when generating NodeMeshConstants.
+		* It should only be accessed from tasks running in the specific FPipe below.
 		*/
-		struct FGeneratedConstantMesh
+		struct FGenerateMeshConstantState
 		{
-			TSharedPtr<FMesh> Mesh;
-			Ptr<ASTOp> LastMeshOp;
-		};
-		TMap<uint64,TArray<FGeneratedConstantMesh>> GeneratedConstantMeshes;
-
-        /** List of image resources for every image formats that have been generated so far as palceholders for missing images. */
-        TSharedPtr<FImage> MissingImage[size_t(EImageFormat::Count)];
+			/** Container of meshes generated to be able to reuse them.
+			* They are sorted by a cheap hash to speed up searches.
+			*/
+			struct FGeneratedConstantMesh
+			{
+				TSharedPtr<FMesh> Mesh;
+				Ptr<ASTOp> LastMeshOp;
+			};
+			TMap<uint64, TArray<FGeneratedConstantMesh>> GeneratedConstantMeshes;
 
+			/** List of already used vertex ID groups that must be unique. */
+			TSet<uint32> UniqueVertexIDGroups;
+		};
+		FGenerateMeshConstantState GenerateMeshConstantState;
+		UE::Tasks::FPipe GenerateMeshConstantPipe;
 
-        //! List of already used vertex ID groups that must be unique.
-        TSet<uint32> UniqueVertexIDGroups;
+		/** List of image resources for every image formats that have been generated so far as palceholders for missing images. */
+		TSharedPtr<FImage> MissingImage[size_t(EImageFormat::Count)];
 
-        struct FParentKey
-        {
-            const NodeObjectNew* ObjectNode = nullptr;
-            int32 Lod = -1;
-         };
+		struct FParentKey
+		{
+			const NodeObjectNew* ObjectNode = nullptr;
+			int32 Lod = -1;
+		};
 
+		// TODO: complete move of this top-down state to method arguments.
 		TArray< FParentKey > CurrentParents;
 
         /** List of additional components to add to an object that come from child objects.
@@ -323,7 +328,7 @@ namespace mu
 				return true;
 			}
 		};
-		TMap< FTableCacheKey, Ptr<ASTOp> > GeneratedTables;
+		TMap< FTableCacheKey, Ptr<NodeScalar> > GeneratedTableNodes;
 
 		struct FConditionalExtensionDataOp
 		{
@@ -361,9 +366,8 @@ namespace mu
 		void GenerateComponent_Switch(const FComponentGenerationOptions&, FGenericGenerationResult&, const NodeComponentSwitch*);
 		void GenerateComponent_Variation(const FComponentGenerationOptions&, FGenericGenerationResult&, const NodeComponentVariation*);
 
-        //-----------------------------------------------------------------------------------------
-        //!
-        Ptr<ASTOp> GenerateTableVariable(Ptr<const Node>, const FTableCacheKey&, bool bAddNoneOption, const FString& DefaultRowName);
+        /** */
+		Ptr<NodeScalar> GenerateTableVariableNode(Ptr<const Node>, const FTableCacheKey&, bool bAddNoneOption, const FString& DefaultRowName);
 
         //!
         Ptr<ASTOp> GenerateMissingBoolCode(const TCHAR* strWhere, bool value, const void* errorContext );
@@ -372,6 +376,9 @@ namespace mu
 		template<class NODE_TABLE, ETableColumnType TYPE, EOpType OPTYPE, typename F>
 		Ptr<ASTOp> GenerateTableSwitch( const NODE_TABLE& node, F&& GenerateOption );
 
+		template<class NODE_TABLE, ETableColumnType TYPE, class NODE_TYPE, class NODE_SWITCH_TYPE, typename F>
+		void GenerateTableSwitchNode(const NODE_TABLE& node, Ptr<NODE_SWITCH_TYPE>& OutResult, F&& GenerateOption);
+
 
 		//-----------------------------------------------------------------------------------------
 		// Images
@@ -379,8 +386,9 @@ namespace mu
 		/** Options that affect the generation of images. It is like list of what required data we want while parsing down the image node graph. */
 		struct FImageGenerationOptions : public FGenericGenerationOptions
 		{
-			FImageGenerationOptions(int32 InComponentId)
+			FImageGenerationOptions(int32 InComponentId, int32 InLODIndex)
 				: ComponentId(InComponentId)
+				, LODIndex(InLODIndex)
 			{
 				bIsImage = true;
 			}
@@ -388,6 +396,9 @@ namespace mu
 			/** The id of the component that we are currently generating. */
 			int32 ComponentId = -1;
 
+			/** The LOD being generated. */
+			int32 LODIndex = -1;
+
 			/** */
 			CompilerOptions::TextureLayoutStrategy ImageLayoutStrategy = CompilerOptions::TextureLayoutStrategy::None;
 
@@ -402,6 +413,7 @@ namespace mu
 			{
 				uint32 KeyHash = GetTypeHash((FGenericGenerationOptions&)InKey);
 				KeyHash = HashCombineFast(KeyHash, ::GetTypeHash(InKey.ComponentId));
+				KeyHash = HashCombineFast(KeyHash, ::GetTypeHash(InKey.LODIndex));
 				KeyHash = HashCombineFast(KeyHash, ::GetTypeHash(InKey.ImageLayoutStrategy));
 				KeyHash = HashCombineFast(KeyHash, GetTypeHash(InKey.RectSize));
 				KeyHash = HashCombineFast(KeyHash, ::GetTypeHash(InKey.LayoutBlockId));
@@ -410,7 +422,6 @@ namespace mu
 			}
 
 			FORCEINLINE bool operator==(const FImageGenerationOptions& Other) const = default;
-
 		};
 
 		/** */
@@ -516,22 +527,39 @@ namespace mu
 		/** Options that affect the generation of meshes. It is like list of what required data we want
 		* while parsing down the mesh node graph.
 		*/
-		struct FMeshGenerationOptions : public FGenericGenerationOptions
+		struct FMeshGenerationStaticOptions : public FGenericGenerationOptions
 		{
-			FMeshGenerationOptions(int32 InComponentId)
+			FMeshGenerationStaticOptions(int32 InComponentId, int32 InLODIndex)
 				: ComponentId(InComponentId)
+				, LODIndex(InLODIndex)
 			{
 			}
 
 			/** The id of the component that we are currently generating. */
 			int32 ComponentId = -1;
 
+			/** The LOD being generated. */
+			int32 LODIndex = -1;
+			
+			/** List of midifiers that shouldn't be applied in the processing of the current subgraph. 
+			* This is usually because they have been processed already in the generation so far.
+			*/
+			TArray<FirstPassGenerator::FModifier> ModifiersToIgnore;
+		};
+
+
+		/** Options that affect the generation of meshes. It is like list of what required data we want
+		* while parsing down the mesh node graph.
+		*/
+		struct FMeshGenerationDynamicOptions
+		{
 			/** The meshes at the leaves will need their own layout block data. */
 			bool bLayouts = false;
 
-			/** If true, Ensure UV Islands are not split between two or more blocks. UVs shared between multiple 
+			/** If true, Ensure UV Islands are not split between two or more blocks. UVs shared between multiple
 			* layout blocks will be clamped to fit the one with more vertices belonging to the UV island.
-			* Mainly used to keep consistent layouts when reusing textures between LODs. */
+			* Mainly used to keep consistent layouts when reusing textures between LODs. 
+			*/
 			bool bClampUVIslands = false;
 
 			/** If true, UVs will be normalized. Normalize UVs should be done in cases where we operate with Images and Layouts */
@@ -549,17 +577,6 @@ namespace mu
 			/** Optional context to use instead of the node error context.
 			 * Be careful since it is not used everywhere. Check usages before assigning a value to it. */
 			TOptional<const void*> OverrideContext;
-			
-			friend FORCEINLINE uint32 GetTypeHash(const FMeshGenerationOptions& InKey)
-			{
-				uint32 KeyHash = GetTypeHash((FGenericGenerationOptions&)InKey);
-				KeyHash = HashCombineFast(KeyHash, ::GetTypeHash(InKey.ComponentId));
-				KeyHash = HashCombineFast(KeyHash, ::GetTypeHash(InKey.bLayouts));
-				KeyHash = HashCombineFast(KeyHash, ::GetTypeHash(InKey.OverrideLayouts.Num()));
-				return KeyHash;
-			}
-
-			FORCEINLINE bool operator==(const FMeshGenerationOptions& Other) const = default;
 		};
 
 		//! Store the results of the code generation of a mesh.
@@ -586,73 +603,65 @@ namespace mu
 			TArray< FExtraLayouts > ExtraMeshLayouts;
 		};
 		
-		struct FGeneratedMeshCacheKey
-		{
-			NodePtrConst Node;
-			FMeshGenerationOptions Options;
-
-			friend FORCEINLINE uint32 GetTypeHash(const FGeneratedMeshCacheKey& InKey)
-			{
-				uint32 KeyHash = 0;
-				KeyHash = HashCombineFast(KeyHash, ::GetTypeHash(InKey.Node.get()));
-				KeyHash = HashCombineFast(KeyHash, GetTypeHash(InKey.Options));
-				return KeyHash;
-			}
-
-			FORCEINLINE bool operator==(const FGeneratedMeshCacheKey& Other) const
-			{
-				return Node == Other.Node && Options == Other.Options;
-			}
-		};
 
-        typedef TMap<FGeneratedMeshCacheKey,FMeshGenerationResult> GeneratedMeshMap;
-        GeneratedMeshMap GeneratedMeshes;
+		using FMeshTask = UE::Tasks::TTask<CodeGenerator::FMeshGenerationResult>;
+		using FMeshOptionsTask = UE::Tasks::TTask<CodeGenerator::FMeshGenerationDynamicOptions>;
 
 		/** Store the mesh generation data for surfaces that we intend to share across LODs. 
 		* The key is the SharedSurfaceId.
 		*/
 		TMap<int32, FMeshGenerationResult> SharedMeshOptionsMap;
 
-		//! Map of layouts found in the code already generated. The map is from the source layout
-		//! node to the generated layout.
-		struct FGeneratedLayoutKey
+		/** This struct contains the generation state that can only be accessed when generating constant layouts.
+		* It should only be accessed after locking the included mutex.
+		*/
+		struct FGenerateLayoutConstantState
 		{
-			Ptr<const NodeLayout> SourceLayout;
-			uint32 MeshIdPrefix;
+			UE::FMutex Mutex;
 
-			friend FORCEINLINE uint32 GetTypeHash(const FGeneratedLayoutKey& InKey)
+			//! Map of layouts found in the code already generated. The map is from the source layout
+			//! node to the generated layout.
+			struct FGeneratedLayoutKey
 			{
-				uint32 KeyHash = ::GetTypeHash(InKey.SourceLayout.get());
-				KeyHash = HashCombineFast(KeyHash, GetTypeHash(InKey.MeshIdPrefix));
-				return KeyHash;
-			}
+				Ptr<const NodeLayout> SourceLayout;
+				uint32 MeshIdPrefix;
+
+				friend FORCEINLINE uint32 GetTypeHash(const FGeneratedLayoutKey& InKey)
+				{
+					uint32 KeyHash = ::GetTypeHash(InKey.SourceLayout.get());
+					KeyHash = HashCombineFast(KeyHash, GetTypeHash(InKey.MeshIdPrefix));
+					return KeyHash;
+				}
 
-			FORCEINLINE bool operator==(const FGeneratedLayoutKey& Other) const = default;
+				FORCEINLINE bool operator==(const FGeneratedLayoutKey& Other) const = default;
+			};
+			TMap<FGeneratedLayoutKey, TSharedPtr<const FLayout>> GeneratedLayouts;
 		};
-		TMap<FGeneratedLayoutKey, TSharedPtr<const FLayout>> GeneratedLayouts;
-
-        void GenerateMesh(const FMeshGenerationOptions&, FMeshGenerationResult& result, const NodeMeshPtrConst&);
-        void GenerateMesh_Constant(const FMeshGenerationOptions&, FMeshGenerationResult&, const NodeMeshConstant* );
-        void GenerateMesh_Format(const FMeshGenerationOptions&, FMeshGenerationResult&, const NodeMeshFormat* );
-        void GenerateMesh_Morph(const FMeshGenerationOptions&, FMeshGenerationResult&, const NodeMeshMorph* );
-        void GenerateMesh_MakeMorph(const FMeshGenerationOptions&, FMeshGenerationResult&, const NodeMeshMakeMorph* );
-        void GenerateMesh_Fragment(const FMeshGenerationOptions&, FMeshGenerationResult&, const NodeMeshFragment* );
-        void GenerateMesh_Switch(const FMeshGenerationOptions&, FMeshGenerationResult&, const NodeMeshSwitch* );
-        void GenerateMesh_Transform(const FMeshGenerationOptions&, FMeshGenerationResult&, const NodeMeshTransform* );
-        void GenerateMesh_ClipMorphPlane(const FMeshGenerationOptions&, FMeshGenerationResult&, const NodeMeshClipMorphPlane* );
-        void GenerateMesh_ClipWithMesh(const FMeshGenerationOptions&, FMeshGenerationResult&, const NodeMeshClipWithMesh* );
-        void GenerateMesh_ApplyPose(const FMeshGenerationOptions&, FMeshGenerationResult&, const NodeMeshApplyPose* );
-        void GenerateMesh_Variation(const FMeshGenerationOptions&, FMeshGenerationResult&, const NodeMeshVariation* );
-		void GenerateMesh_Table(const FMeshGenerationOptions&, FMeshGenerationResult&, const NodeMeshTable*);
-		void GenerateMesh_Reshape(const FMeshGenerationOptions&, FMeshGenerationResult&, const NodeMeshReshape*);
-		void GenerateMesh_ClipDeform(const FMeshGenerationOptions&, FMeshGenerationResult&, const NodeMeshClipDeform*);
+		FGenerateLayoutConstantState GenerateLayoutConstantState;
+
+		FMeshTask GenerateMesh(const FMeshGenerationStaticOptions& StaticOptions, FMeshOptionsTask, const Ptr<const NodeMesh>&);
+		FMeshTask GenerateMesh_Constant(const FMeshGenerationStaticOptions& StaticOptions, FMeshOptionsTask, const NodeMeshConstant* );
+		FMeshTask GenerateMesh_Format(const FMeshGenerationStaticOptions& StaticOptions, FMeshOptionsTask, const NodeMeshFormat* );
+		FMeshTask GenerateMesh_Morph(const FMeshGenerationStaticOptions& StaticOptions, FMeshOptionsTask, const NodeMeshMorph* );
+		FMeshTask GenerateMesh_MakeMorph(const FMeshGenerationStaticOptions& StaticOptions, FMeshOptionsTask, const NodeMeshMakeMorph* );
+        FMeshTask GenerateMesh_Fragment(const FMeshGenerationStaticOptions& StaticOptions, FMeshOptionsTask, const NodeMeshFragment* );
+        FMeshTask GenerateMesh_Switch(const FMeshGenerationStaticOptions& StaticOptions, FMeshOptionsTask, const NodeMeshSwitch* );
+        FMeshTask GenerateMesh_Transform(const FMeshGenerationStaticOptions& StaticOptions, FMeshOptionsTask, const NodeMeshTransform* );
+        FMeshTask GenerateMesh_ClipMorphPlane(const FMeshGenerationStaticOptions& StaticOptions, FMeshOptionsTask, const NodeMeshClipMorphPlane* );
+        FMeshTask GenerateMesh_ClipWithMesh(const FMeshGenerationStaticOptions& StaticOptions, FMeshOptionsTask, const NodeMeshClipWithMesh* );
+        FMeshTask GenerateMesh_ApplyPose(const FMeshGenerationStaticOptions& StaticOptions, FMeshOptionsTask, const NodeMeshApplyPose* );
+        FMeshTask GenerateMesh_Variation(const FMeshGenerationStaticOptions& StaticOptions, FMeshOptionsTask, const NodeMeshVariation* );
+		FMeshTask GenerateMesh_Table(const FMeshGenerationStaticOptions& StaticOptions, FMeshOptionsTask, const NodeMeshTable*);
+		FMeshTask GenerateMesh_Reshape(const FMeshGenerationStaticOptions& StaticOptions, FMeshOptionsTask, const NodeMeshReshape*);
+		FMeshTask GenerateMesh_ClipDeform(const FMeshGenerationStaticOptions& StaticOptions, FMeshOptionsTask, const NodeMeshClipDeform*);
 
 		//-----------------------------------------------------------------------------------------
 		void PrepareMeshForLayout(const FGeneratedLayout&, TSharedPtr<FMesh>,
 			int32 currentLayoutChannel,
 			const void* errorContext,
-			const FMeshGenerationOptions&,
-			bool bUseAbsoluteBlockIds);
+			const FMeshGenerationStaticOptions&,
+			const FMeshGenerationDynamicOptions&,
+			bool bUseAbsoluteBlockIds) const;
 
 		//!
 		TSharedPtr<const FLayout> GenerateLayout(Ptr<const NodeLayout> SourceLayout, uint32 MeshIDPrefix);
@@ -815,16 +824,14 @@ namespace mu
 		void UpdateLayoutBlockDesc(FLayoutBlockDesc& Out, FImageDesc BlockDesc, FIntVector2 LayoutCellSize);
 
 		// Get the modifiers that have to be applied to elements with a specific tag.
-		void GetModifiersFor(int32 ComponentId, const TArray<FString>& SurfaceTags, bool bModifiersForBeforeOperations, TArray<FirstPassGenerator::FModifier>& OutModifiers);
-
-		// Used to avoid recursion when generating modifiers.
-		TArray<FirstPassGenerator::FModifier> ModifiersToIgnore;
+		void GetModifiersFor(int32 ComponentId, const TArray<FString>& SurfaceTags, bool bModifiersForBeforeOperations, TArray<FirstPassGenerator::FModifier>& OutModifiers) const;
 
 		// Apply the required mesh modifiers to the given operation.
-		Ptr<ASTOp> ApplyMeshModifiers(
+		FMeshTask ApplyMeshModifiers(
 			const TArray<FirstPassGenerator::FModifier>&,
-			const FMeshGenerationOptions&, 
-			FMeshGenerationResult& BaseResults, 
+			const FMeshGenerationStaticOptions&,
+			FMeshOptionsTask DynamicOptionsTask,
+			FMeshTask BaseTask,
 			const FMeshGenerationResult* SharedSurfaceResults, 
 			const void* ErrorContext,
 			const NodeMeshConstant* OriginalMeshNode);
@@ -841,8 +848,7 @@ namespace mu
 
 		Ptr<ASTOp> ApplyImageExtendModifiers(
 			const TArray<FirstPassGenerator::FModifier>&,
-			const FGenericGenerationOptions& Options, 
-			int32 ComponentId,
+			const FMeshGenerationStaticOptions& Options,
 			const FMeshGenerationResult& BaseMeshResults,
 			Ptr<ASTOp> ImageAd, 
 			CompilerOptions::TextureLayoutStrategy ImageLayoutStrategy, 
@@ -852,7 +858,7 @@ namespace mu
 			CodeGenerator::FLayoutBlockDesc& InOutLayoutBlockDesc,
 			const void* ModifiedNodeErrorContext);
 
-		void CheckModifiersForSurface(const NodeSurfaceNew&, const TArray<FirstPassGenerator::FModifier>&);
+		void CheckModifiersForSurface(const NodeSurfaceNew&, const TArray<FirstPassGenerator::FModifier>&, int32 LODIndex);
 
     };
 
@@ -863,21 +869,26 @@ namespace mu
     {
         Ptr<const FTable> NodeTable = node.Table;
         Ptr<ASTOp> Variable;
+		Ptr<NodeScalar> VariableNode;
 
 		FTableCacheKey CacheKey = FTableCacheKey{ node.Table, node.ParameterName };
-        Ptr<ASTOp>* it = GeneratedTables.Find( CacheKey );
-        if ( it )
-        {
-            Variable = *it;
-        }
+		Ptr<NodeScalar>* it = GeneratedTableNodes.Find(CacheKey);
+		if (it)
+		{
+			VariableNode = *it;
+		}
 
-        if ( !Variable)
+        if (!VariableNode)
         {
             // Create the table variable expression
-            Variable = GenerateTableVariable( &node, CacheKey, node.bNoneOption, node.DefaultRowName);
+			VariableNode = GenerateTableVariableNode(&node, CacheKey, node.bNoneOption, node.DefaultRowName);
+			GeneratedTableNodes.Add(CacheKey, VariableNode);
+		}
 
-            GeneratedTables.Add(CacheKey, Variable );
-        }
+		FGenericGenerationOptions ScalarOptions;
+		FScalarGenerationResult ScalarResult;
+		GenerateScalar(ScalarResult, ScalarOptions, VariableNode);
+		Variable = ScalarResult.op;
 
 		int32 NumRows = NodeTable->GetPrivate()->Rows.Num();
 
@@ -923,4 +934,68 @@ namespace mu
 
         return SwitchOp;
     }
+
+	
+    //---------------------------------------------------------------------------------------------
+    template<class NODE_TABLE, ETableColumnType TYPE, class NODE_TYPE, class NODE_SWITCH_TYPE, typename F>
+    void CodeGenerator::GenerateTableSwitchNode( const NODE_TABLE& node, Ptr<NODE_SWITCH_TYPE>& OutResult, F&& GenerateOption )
+    {
+		OutResult = nullptr;
+
+        Ptr<const FTable> NodeTable = node.Table;
+        Ptr<NodeScalar> Variable;
+
+		FTableCacheKey CacheKey = FTableCacheKey{ node.Table, node.ParameterName };
+        Ptr<NodeScalar>* it = GeneratedTableNodes.Find( CacheKey );
+        if ( it )
+        {
+            Variable = *it;
+        }
+
+        if (!Variable)
+        {
+            // Create the table variable expression
+            Variable = GenerateTableVariableNode( &node, CacheKey, node.bNoneOption, node.DefaultRowName);
+
+            GeneratedTableNodes.Add(CacheKey, Variable );
+        }
+
+		int32 NumRows = NodeTable->GetPrivate()->Rows.Num();
+
+        // Verify that the table column is the right type
+        int32 ColIndex = NodeTable->FindColumn( node.ColumnName );
+
+		if (NumRows == 0)
+		{
+			ErrorLog->Add("The table has no rows.", ELMT_ERROR, node.GetMessageContext());
+			return;
+		}
+        else if (ColIndex < 0)
+        {
+            ErrorLog->Add("Table column not found.", ELMT_ERROR, node.GetMessageContext());
+            return;
+        }
+
+        if (NodeTable->GetPrivate()->Columns[ ColIndex ].Type != TYPE )
+        {
+            ErrorLog->Add("Table column type is not the right type.", ELMT_ERROR, node.GetMessageContext());
+            return;
+        }
+
+        // Create the switch node
+        Ptr<NODE_SWITCH_TYPE> SwitchNode = new NODE_SWITCH_TYPE;
+		SwitchNode->SetParameter( Variable );
+		SwitchNode->SetOptionCount(NumRows);
+
+		for (int32 RowIndex = 0; RowIndex < NumRows; ++RowIndex)
+        {
+            check(RowIndex <= 0xFFFF);
+            Ptr<NODE_TYPE> BranchNode = GenerateOption( node, ColIndex, RowIndex, ErrorLog.Get() );
+			SwitchNode->SetOption(RowIndex, BranchNode);
+        }
+
+		OutResult = SwitchNode;
+    }
+
+
 }
diff --git a/Engine/Plugins/Mutable/Source/MutableTools/Private/MuT/CodeGenerator_Color.cpp b/Engine/Plugins/Mutable/Source/MutableTools/Private/MuT/CodeGenerator_Color.cpp
--- a/Engine/Plugins/Mutable/Source/MutableTools/Private/MuT/CodeGenerator_Color.cpp
+++ b/Engine/Plugins/Mutable/Source/MutableTools/Private/MuT/CodeGenerator_Color.cpp
@@ -300,7 +300,8 @@ namespace mu
 
 		// Source image
 		int32 ComponentId = -1; // TODO.
-		FImageGenerationOptions ImageOptions(ComponentId);
+		int32 LODIndex = 0; // TODO.
+		FImageGenerationOptions ImageOptions(ComponentId, LODIndex);
 		ImageOptions.State = Options.State;
 		ImageOptions.ActiveTags = Options.ActiveTags;
 
diff --git a/Engine/Plugins/Mutable/Source/MutableTools/Private/MuT/CodeGenerator_Image.cpp b/Engine/Plugins/Mutable/Source/MutableTools/Private/MuT/CodeGenerator_Image.cpp
--- a/Engine/Plugins/Mutable/Source/MutableTools/Private/MuT/CodeGenerator_Image.cpp
+++ b/Engine/Plugins/Mutable/Source/MutableTools/Private/MuT/CodeGenerator_Image.cpp
@@ -1368,14 +1368,16 @@ namespace mu
         {
 			// TODO: This will probably Result in a duplicated mesh subgraph, with the original mesh but new layout block ids.
 			// See if it can be optimized and try to reuse the existing layout block ids instead of generating new ones.
-			FMeshGenerationOptions MeshOptions(Options.ComponentId);
-			MeshOptions.State = Options.State;
-			MeshOptions.ActiveTags = Options.ActiveTags;
+			FMeshGenerationStaticOptions MeshStaticOptions(Options.ComponentId, Options.LODIndex);
+			MeshStaticOptions.State = Options.State;
+			MeshStaticOptions.ActiveTags = Options.ActiveTags;
+			FMeshGenerationDynamicOptions MeshOptions;
 			MeshOptions.bLayouts = true;			// We need the layout that we will use to render
 			MeshOptions.bNormalizeUVs = true;		// We need normalized UVs for the projection
-			
-			FMeshGenerationResult MeshResult;
-			GenerateMesh( MeshOptions, MeshResult, node.m_pMesh );
+
+			FMeshTask MeshTask = GenerateMesh(MeshStaticOptions, UE::Tasks::MakeCompletedTask<FMeshGenerationDynamicOptions>(MeshOptions), node.m_pMesh );
+			// This forces a wait to sync here. When images are also generated in tasks, it can be turned into a prerequisite instead.
+			FMeshGenerationResult MeshResult = MeshTask.GetResult();
 
 			// Match the block id of the block we are generating with the id that resulted in the generated mesh			
 			TSharedPtr<const FLayout> Layout = MeshResult.GeneratedLayouts.IsValidIndex(node.m_layout) ? MeshResult.GeneratedLayouts[node.m_layout].Layout : nullptr;
diff --git a/Engine/Plugins/Mutable/Source/MutableTools/Private/MuT/CodeGenerator_Mesh.cpp b/Engine/Plugins/Mutable/Source/MutableTools/Private/MuT/CodeGenerator_Mesh.cpp
--- a/Engine/Plugins/Mutable/Source/MutableTools/Private/MuT/CodeGenerator_Mesh.cpp
+++ b/Engine/Plugins/Mutable/Source/MutableTools/Private/MuT/CodeGenerator_Mesh.cpp
@@ -106,10 +106,13 @@ namespace mu
 		TSharedPtr<FMesh> Mesh,
 		int32 LayoutChannel,
 		const void* ErrorContext,
-		const FMeshGenerationOptions& MeshOptions,
+		const FMeshGenerationStaticOptions& StaticMeshOptions,
+		const FMeshGenerationDynamicOptions& DynamicMeshOptions,
 		bool bUseAbsoluteBlockIds
-		)
+		) const
 	{
+		// This can run in any thread.
+
 		MUTABLE_CPUPROFILER_SCOPE(PrepareMeshForLayout);
 
 		if (Mesh->GetVertexCount() == 0)
@@ -185,9 +188,7 @@ namespace mu
 		// Notify Overlapping layout blocks
 		if (!OverlappingBlocks.IsEmpty())
 		{
-			FString Msg = FString::Printf(TEXT("Source mesh has %d layout block overlapping in LOD %d"),
-				OverlappingBlocks.Num() + 1, CurrentParents.Last().Lod
-			);
+			FString Msg = FString::Printf(TEXT("Source mesh has %d layout block overlapping in LOD %d"), OverlappingBlocks.Num() + 1, StaticMeshOptions.LODIndex );
 			ErrorLog->Add(Msg, ELMT_WARNING, ErrorContext);
 		}
 
@@ -229,7 +230,7 @@ namespace mu
 				}
 
 				// Check that UVs are normalized. If not, clamp the values and throw a warning.
-				if (MeshOptions.bNormalizeUVs && !bIsOverlayLayout
+				if (DynamicMeshOptions.bNormalizeUVs && !bIsOverlayLayout
 					&& (UV[0] < 0.f || UV[0] > 1.f || UV[1] < 0.f || UV[1] > 1.f))
 				{
 					UV[0] = FMath::Clamp(UV[0], 0.f, 1.f);
@@ -243,7 +244,7 @@ namespace mu
 			// Mutable does not support non-normalized UVs
 			if (bNonNormalizedUVs && !bIsOverlayLayout)
 			{
-				FString Msg = FString::Printf(TEXT("Source mesh has non-normalized UVs in LOD %d"), CurrentParents.Last().Lod );
+				FString Msg = FString::Printf(TEXT("Source mesh has non-normalized UVs in LOD %d"), StaticMeshOptions.LODIndex);
 				ErrorLog->Add(Msg, ELMT_WARNING, ErrorContext);
 			}
 		}
@@ -259,7 +260,7 @@ namespace mu
 		TMultiMap<int32, uint32> VertexToFaceMap;
 
 		// Find Unique Vertices
-		if (MeshOptions.bClampUVIslands)
+		if (DynamicMeshOptions.bClampUVIslands)
 		{
 			VertexToFaceMap.Reserve(NumVertices);
 			Triangles.SetNumUninitialized(NumTriangles);
@@ -269,8 +270,8 @@ namespace mu
 
 		TArray<int32> ConflictiveTriangles;
 
-		const uint32 MaxGridX = MeshOptions.bNormalizeUVs ? MAX_uint32 : WorkingGrid.X - 1;
-		const uint32 MaxGridY = MeshOptions.bNormalizeUVs ? MAX_uint32 : WorkingGrid.Y - 1;
+		const uint32 MaxGridX = DynamicMeshOptions.bNormalizeUVs ? MAX_uint32 : WorkingGrid.X - 1;
+		const uint32 MaxGridY = DynamicMeshOptions.bNormalizeUVs ? MAX_uint32 : WorkingGrid.Y - 1;
 
 		// Allocate the per-vertex layout block data
 		TArray<uint16> LayoutData;
@@ -350,7 +351,7 @@ namespace mu
 			const uint16 BlockIndexV1 = AssignOneVertex(Index1);
 			const uint16 BlockIndexV2 = AssignOneVertex(Index2);
 
-			if (MeshOptions.bClampUVIslands)
+			if (DynamicMeshOptions.bClampUVIslands)
 			{
 				if (BlockIndexV0 != BlockIndexV1 || BlockIndexV0 != BlockIndexV2)
 				{
@@ -456,8 +457,9 @@ namespace mu
 		}
 
 		// Warn about vertices without a block id
+		check(StaticMeshOptions.LODIndex>=0);
 		int32 FirstLODToIgnoreWarnings = GeneratedLayout.Source->FirstLODToIgnoreWarnings;
-		if (FirstLODToIgnoreWarnings == -1 || CurrentParents.Last().Lod < FirstLODToIgnoreWarnings)
+		if (FirstLODToIgnoreWarnings == -1 || StaticMeshOptions.LODIndex < FirstLODToIgnoreWarnings)
 		{
 			TArray<float> UnassignedUVs;
 			UnassignedUVs.Reserve(NumVertices / 100);
@@ -474,7 +476,7 @@ namespace mu
 
 			if (!UnassignedUVs.IsEmpty())
 			{
-				FString Msg = FString::Printf(TEXT("Source mesh has %d vertices not assigned to any layout block in LOD %d"), UnassignedUVs.Num(), CurrentParents.Last().Lod);
+				FString Msg = FString::Printf(TEXT("Source mesh has %d vertices not assigned to any layout block in LOD %d"), UnassignedUVs.Num(), StaticMeshOptions.LODIndex);
 
 				ErrorLogMessageAttachedDataView AttachedDataView;
 				AttachedDataView.UnassignedUVs = UnassignedUVs.GetData();
@@ -544,12 +546,12 @@ namespace mu
 					if (bUseAbsoluteBlockIds)
 					{
 						uint64* Ptr = reinterpret_cast<uint64*>(LayoutBufferPtr) + VertexIndex;
-						*Ptr = MeshOptions.bEnsureAllVerticesHaveLayoutBlock ? 0 : std::numeric_limits<uint64>::max();
+						*Ptr = DynamicMeshOptions.bEnsureAllVerticesHaveLayoutBlock ? 0 : std::numeric_limits<uint64>::max();
 					}
 					else
 					{
 						uint16* Ptr = reinterpret_cast<uint16*>(LayoutBufferPtr) + VertexIndex;
-						*Ptr = MeshOptions.bEnsureAllVerticesHaveLayoutBlock ? 0 : std::numeric_limits<uint16>::max();
+						*Ptr = DynamicMeshOptions.bEnsureAllVerticesHaveLayoutBlock ? 0 : std::numeric_limits<uint16>::max();
 					}
 				}
 
@@ -573,246 +575,310 @@ namespace mu
 	
 
     //---------------------------------------------------------------------------------------------
-    void CodeGenerator::GenerateMesh( const FMeshGenerationOptions& InOptions, FMeshGenerationResult& OutResult, const NodeMeshPtrConst& InUntypedNode)
+	CodeGenerator::FMeshTask CodeGenerator::GenerateMesh(const FMeshGenerationStaticOptions& StaticOptions, FMeshOptionsTask InOptions, const Ptr<const NodeMesh>& InUntypedNode)
     {
         if (!InUntypedNode)
         {
-            OutResult = FMeshGenerationResult();
-            return;
-        }
-
-        // See if it was already generated
-		FGeneratedMeshCacheKey Key(FGeneratedMeshCacheKey{ InUntypedNode, InOptions});
-        GeneratedMeshMap::ValueType* it = GeneratedMeshes.Find(Key);
-        if ( it )
-        {
-			OutResult = *it;
-            return;
+            return UE::Tasks::MakeCompletedTask<FMeshGenerationResult>(FMeshGenerationResult());
         }
 
 		const NodeMesh* Node = InUntypedNode.get();
 
         // Generate for each different type of node
+		FMeshTask Result;
 		switch (Node->GetType()->Type)
 		{
-		case Node::EType::MeshConstant: GenerateMesh_Constant(InOptions, OutResult, static_cast<const NodeMeshConstant*>(Node)); break;
-		case Node::EType::MeshFormat: GenerateMesh_Format(InOptions, OutResult, static_cast<const NodeMeshFormat*>(Node)); break;
-		case Node::EType::MeshMorph: GenerateMesh_Morph(InOptions, OutResult, static_cast<const NodeMeshMorph*>(Node)); break;
-		case Node::EType::MeshMakeMorph: GenerateMesh_MakeMorph(InOptions, OutResult, static_cast<const NodeMeshMakeMorph*>(Node)); break;
-		case Node::EType::MeshFragment: GenerateMesh_Fragment(InOptions, OutResult, static_cast<const NodeMeshFragment*>(Node)); break;
-		case Node::EType::MeshSwitch: GenerateMesh_Switch(InOptions, OutResult, static_cast<const NodeMeshSwitch*>(Node)); break;
-		case Node::EType::MeshTransform: GenerateMesh_Transform(InOptions, OutResult, static_cast<const NodeMeshTransform*>(Node)); break;
-		case Node::EType::MeshClipMorphPlane: GenerateMesh_ClipMorphPlane(InOptions, OutResult, static_cast<const NodeMeshClipMorphPlane*>(Node)); break;
-		case Node::EType::MeshClipWithMesh: GenerateMesh_ClipWithMesh(InOptions, OutResult, static_cast<const NodeMeshClipWithMesh*>(Node)); break;
-		case Node::EType::MeshApplyPose: GenerateMesh_ApplyPose(InOptions, OutResult, static_cast<const NodeMeshApplyPose*>(Node)); break;
-		case Node::EType::MeshVariation: GenerateMesh_Variation(InOptions, OutResult, static_cast<const NodeMeshVariation*>(Node)); break;
-		case Node::EType::MeshTable: GenerateMesh_Table(InOptions, OutResult, static_cast<const NodeMeshTable*>(Node)); break;
-		case Node::EType::MeshReshape: GenerateMesh_Reshape(InOptions, OutResult, static_cast<const NodeMeshReshape*>(Node)); break;
-		case Node::EType::MeshClipDeform: GenerateMesh_ClipDeform(InOptions, OutResult, static_cast<const NodeMeshClipDeform*>(Node)); break;
+		case Node::EType::MeshConstant: Result = GenerateMesh_Constant(StaticOptions, InOptions, static_cast<const NodeMeshConstant*>(Node)); break;
+		case Node::EType::MeshFormat: Result = GenerateMesh_Format(StaticOptions, InOptions, static_cast<const NodeMeshFormat*>(Node)); break;
+		case Node::EType::MeshMorph: Result = GenerateMesh_Morph(StaticOptions, InOptions, static_cast<const NodeMeshMorph*>(Node)); break;
+		case Node::EType::MeshMakeMorph: Result = GenerateMesh_MakeMorph(StaticOptions, InOptions, static_cast<const NodeMeshMakeMorph*>(Node)); break;
+		case Node::EType::MeshFragment: Result = GenerateMesh_Fragment(StaticOptions, InOptions, static_cast<const NodeMeshFragment*>(Node)); break;
+		case Node::EType::MeshSwitch: Result = GenerateMesh_Switch(StaticOptions, InOptions, static_cast<const NodeMeshSwitch*>(Node)); break;
+		case Node::EType::MeshTransform: Result = GenerateMesh_Transform(StaticOptions, InOptions, static_cast<const NodeMeshTransform*>(Node)); break;
+		case Node::EType::MeshClipMorphPlane: Result = GenerateMesh_ClipMorphPlane(StaticOptions, InOptions, static_cast<const NodeMeshClipMorphPlane*>(Node)); break;
+		case Node::EType::MeshClipWithMesh: Result = GenerateMesh_ClipWithMesh(StaticOptions, InOptions, static_cast<const NodeMeshClipWithMesh*>(Node)); break;
+		case Node::EType::MeshApplyPose: Result = GenerateMesh_ApplyPose(StaticOptions, InOptions, static_cast<const NodeMeshApplyPose*>(Node)); break;
+		case Node::EType::MeshVariation: Result = GenerateMesh_Variation(StaticOptions, InOptions, static_cast<const NodeMeshVariation*>(Node)); break;
+		case Node::EType::MeshTable: Result = GenerateMesh_Table(StaticOptions, InOptions, static_cast<const NodeMeshTable*>(Node)); break;
+		case Node::EType::MeshReshape: Result = GenerateMesh_Reshape(StaticOptions, InOptions, static_cast<const NodeMeshReshape*>(Node)); break;
+		case Node::EType::MeshClipDeform: Result = GenerateMesh_ClipDeform(StaticOptions, InOptions, static_cast<const NodeMeshClipDeform*>(Node)); break;
 		default: check(false);
 		}
 
-        // Cache the result
-        GeneratedMeshes.Add( Key, OutResult);
+		return Result;
     }
 
 
-    //---------------------------------------------------------------------------------------------
-    void CodeGenerator::GenerateMesh_Morph(
-		const FMeshGenerationOptions& InOptions, 
-		FMeshGenerationResult& OutResult, 
-		const NodeMeshMorph* InMorphNode 
-	)
+	CodeGenerator::FMeshTask CodeGenerator::GenerateMesh_Morph( const FMeshGenerationStaticOptions& StaticOptions, FMeshOptionsTask InOptions, const NodeMeshMorph* InMorphNode )
     {
         NodeMeshMorph::Private& node = *InMorphNode->GetPrivate();
 
-        Ptr<ASTOpMeshMorph> OpMorph = new ASTOpMeshMorph();
+		Ptr<ASTOp> FactorOp;
 
         // Factor
         if ( node.Factor )
         {
-            OpMorph->Factor = Generate_Generic( node.Factor.get(), InOptions );
+			FactorOp = Generate_Generic( node.Factor.get(), StaticOptions );
         }
         else
         {
             // This argument is required
-            OpMorph->Factor = GenerateMissingScalarCode(TEXT("Morph factor"), 0.5f, InMorphNode->GetMessageContext());
+			FactorOp = GenerateMissingScalarCode(TEXT("Morph factor"), 0.5f, InMorphNode->GetMessageContext());
         }
 
+		TArray<UE::Tasks::FTask, TInlineAllocator<4>> Requisites;
+
         // Base
-        FMeshGenerationResult BaseResult;
+		FMeshTask BaseTask;
         if ( node.Base )
         {
-            GenerateMesh(InOptions,BaseResult, node.Base );
-            OpMorph->Base = BaseResult.MeshOp;
+			BaseTask = GenerateMesh(StaticOptions, InOptions, node.Base );
+			Requisites.Add(BaseTask);
         }
         else
         {
             // This argument is required
             ErrorLog->Add( "Mesh morph base node is not set.", ELMT_ERROR, InMorphNode->GetMessageContext());
+			return UE::Tasks::MakeCompletedTask<FMeshGenerationResult>();
         }        
 
-		if (node.Morph)
-        {
-            FMeshGenerationResult TargetResult;
-			FMeshGenerationOptions TargetOptions = InOptions;
-			TargetOptions.bLayouts = false;
-			// We need to override the layouts with the layouts that were generated for the base to make
-			// sure that we get the correct mesh when generating the target
-			TargetOptions.OverrideLayouts = BaseResult.GeneratedLayouts;
-			TargetOptions.ActiveTags.Empty();
-            GenerateMesh(TargetOptions, TargetResult, node.Morph);
-
-            // TODO: Make sure that the target is a mesh with the morph format
-            Ptr<ASTOp> target = TargetResult.MeshOp;
-
-            OpMorph->Target = target;
-        }
- 
-        const bool bReshapeEnabled = node.bReshapeSkeleton || node.bReshapePhysicsVolumes;
-        
-        Ptr<ASTOpMeshMorphReshape> OpMorphReshape;
-        if ( bReshapeEnabled )
+		if (!node.Morph)
+		{
+			return BaseTask;
+		}
+
+		FMeshTask TargetTask;
         {
-		    Ptr<ASTOpMeshBindShape> OpBind = new ASTOpMeshBindShape();
-		    Ptr<ASTOpMeshApplyShape> OpApply = new ASTOpMeshApplyShape();
-
-			// Setting bReshapeVertices to false the bind op will remove all mesh members except 
-			// PhysicsBodies and the Skeleton.
-            OpBind->bReshapeVertices = false;
-            OpBind->bApplyLaplacian = false;
-            OpBind->bRecomputeNormals = false;
-		    OpBind->bReshapeSkeleton = node.bReshapeSkeleton;
-		    OpBind->BonesToDeform = node.BonesToDeform;
-    	    OpBind->bReshapePhysicsVolumes = node.bReshapePhysicsVolumes; 
-			OpBind->PhysicsToDeform = node.PhysicsToDeform;
-			OpBind->BindingMethod = static_cast<uint32>(EShapeBindingMethod::ReshapeClosestProject);
-            
-			OpBind->Mesh = BaseResult.MeshOp;
-            OpBind->Shape = BaseResult.MeshOp;
-           
-			OpApply->bReshapeVertices = OpBind->bReshapeVertices;
-			OpApply->bRecomputeNormals = OpBind->bRecomputeNormals;
-		    OpApply->bReshapeSkeleton = OpBind->bReshapeSkeleton;
-		    OpApply->bReshapePhysicsVolumes = OpBind->bReshapePhysicsVolumes;
-
-			OpApply->Mesh = OpBind;
-            OpApply->Shape = OpMorph;
-
-            OpMorphReshape = new ASTOpMeshMorphReshape();
-            OpMorphReshape->Morph = OpMorph;
-            OpMorphReshape->Reshape = OpApply;
+			FMeshGenerationStaticOptions TargetStaticOptions = StaticOptions;
+			TargetStaticOptions.ActiveTags.Empty();
+			FMeshOptionsTask TargetOptionsTask = UE::Tasks::Launch( TEXT("MutableMorphOptions"),
+				[BaseTask, InOptions]() mutable
+				{
+					FMeshGenerationDynamicOptions TargetOptions = InOptions.GetResult();
+					TargetOptions.bLayouts = false;
+
+					// We need to override the layouts with the layouts that were generated for the base to make
+					// sure that we get the correct mesh when generating the target
+					FMeshGenerationResult BaseResult = BaseTask.GetResult();
+					TargetOptions.OverrideLayouts = BaseResult.GeneratedLayouts;
+
+					return TargetOptions;
+				},
+				UE::Tasks::Prerequisites( BaseTask, InOptions )
+				);
+
+			TargetTask = GenerateMesh(TargetStaticOptions, TargetOptionsTask, node.Morph);
         }
 
- 		if (OpMorphReshape)
-		{
-			OutResult.MeshOp = OpMorphReshape;
-		}
-		else
-		{
-			OutResult.MeshOp = OpMorph;
-		}
+		return UE::Tasks::Launch(TEXT("MutableMorph"),
+			[BaseTask, TargetTask, node, FactorOp]() mutable
+			{
+				FMeshGenerationResult BaseResult = BaseTask.GetResult();
+				FMeshGenerationResult TargetResult = TargetTask.GetResult();
 
-        OutResult.BaseMeshOp = BaseResult.BaseMeshOp;
-		OutResult.GeneratedLayouts = BaseResult.GeneratedLayouts;
-}
+				Ptr<ASTOpMeshMorph> OpMorph = new ASTOpMeshMorph();
+				OpMorph->Factor = FactorOp;
+				OpMorph->Base = BaseTask.GetResult().MeshOp;
+				OpMorph->Target = TargetTask.GetResult().MeshOp;
 
+				const bool bReshapeEnabled = node.bReshapeSkeleton || node.bReshapePhysicsVolumes;
 
-    //---------------------------------------------------------------------------------------------
-    void CodeGenerator::GenerateMesh_MakeMorph(const FMeshGenerationOptions& InOptions, FMeshGenerationResult& OutResult, const NodeMeshMakeMorph* InMakeMorphNode )
+				Ptr<ASTOpMeshMorphReshape> OpMorphReshape;
+				if (bReshapeEnabled)
+				{
+					Ptr<ASTOpMeshBindShape> OpBind = new ASTOpMeshBindShape();
+					Ptr<ASTOpMeshApplyShape> OpApply = new ASTOpMeshApplyShape();
+
+					// Setting bReshapeVertices to false the bind op will remove all mesh members except 
+					// PhysicsBodies and the Skeleton.
+					OpBind->bReshapeVertices = false;
+					OpBind->bApplyLaplacian = false;
+					OpBind->bRecomputeNormals = false;
+					OpBind->bReshapeSkeleton = node.bReshapeSkeleton;
+					OpBind->BonesToDeform = node.BonesToDeform;
+					OpBind->bReshapePhysicsVolumes = node.bReshapePhysicsVolumes;
+					OpBind->PhysicsToDeform = node.PhysicsToDeform;
+					OpBind->BindingMethod = static_cast<uint32>(EShapeBindingMethod::ReshapeClosestProject);
+
+					OpBind->Mesh = BaseResult.MeshOp;
+					OpBind->Shape = BaseResult.MeshOp;
+
+					OpApply->bReshapeVertices = OpBind->bReshapeVertices;
+					OpApply->bRecomputeNormals = OpBind->bRecomputeNormals;
+					OpApply->bReshapeSkeleton = OpBind->bReshapeSkeleton;
+					OpApply->bReshapePhysicsVolumes = OpBind->bReshapePhysicsVolumes;
+
+					OpApply->Mesh = OpBind;
+					OpApply->Shape = OpMorph;
+
+					OpMorphReshape = new ASTOpMeshMorphReshape();
+					OpMorphReshape->Morph = OpMorph;
+					OpMorphReshape->Reshape = OpApply;
+				}
+
+				FMeshGenerationResult Result;
+				if (OpMorphReshape)
+				{
+					Result.MeshOp = OpMorphReshape;
+				}
+				else
+				{
+					Result.MeshOp = OpMorph;
+				}
+
+				Result.BaseMeshOp = BaseResult.BaseMeshOp;
+				Result.GeneratedLayouts = BaseResult.GeneratedLayouts;
+
+				return Result;
+			},
+			UE::Tasks::Prerequisites(BaseTask, TargetTask)
+		);
+	}
+
+
+	CodeGenerator::FMeshTask CodeGenerator::GenerateMesh_MakeMorph(const FMeshGenerationStaticOptions& StaticOptions, FMeshOptionsTask InOptions, const NodeMeshMakeMorph* InMakeMorphNode )
     {
         NodeMeshMakeMorph::Private& node = *InMakeMorphNode->GetPrivate();
 
-        Ptr<ASTOpMeshDifference> op = new ASTOpMeshDifference();
+		if (!node.m_pBase)
+		{
+			// This argument is required
+			ErrorLog->Add("Mesh make morph base node is not set.", ELMT_ERROR, InMakeMorphNode->GetMessageContext());
+			return UE::Tasks::MakeCompletedTask<FMeshGenerationResult>();
+		}
 
-        // \todo Texcoords are broken?
-        op->bIgnoreTextureCoords = true;
-	
-		// UE only has position and normal morph data, optimize for this case if indicated. 
-		if (node.bOnlyPositionAndNormal)
+		if (!node.m_pTarget)
 		{
-			op->Channels = { {static_cast<uint8>(EMeshBufferSemantic::Position), 0}, {static_cast<uint8>(EMeshBufferSemantic::Normal), 0} };
+			// This argument is required
+			ErrorLog->Add("Mesh make morph target node is not set.", ELMT_ERROR, InMakeMorphNode->GetMessageContext());
+			return UE::Tasks::MakeCompletedTask<FMeshGenerationResult>();
 		}
 
         // Base
-        FMeshGenerationResult BaseResult;
-        if ( node.m_pBase )
+		FMeshTask BaseTask;
         {
-			FMeshGenerationOptions BaseOptions = InOptions;
-			BaseOptions.bLayouts = false;
-			GenerateMesh(BaseOptions, BaseResult, node.m_pBase );
+			FMeshOptionsTask BaseOptionsTask = UE::Tasks::Launch(TEXT("MutableMorphBaseOptions"),
+				[InOptions]() mutable
+				{
+					FMeshGenerationDynamicOptions Result = InOptions.GetResult();
+					Result.bLayouts = false;
+					return Result;
+				},
+				InOptions
+			);
 
-            op->Base = BaseResult.MeshOp;
-        }
-        else
-        {
-            // This argument is required
-            ErrorLog->Add( "Mesh make morph base node is not set.", ELMT_ERROR, InMakeMorphNode->GetMessageContext());
+			BaseTask = GenerateMesh(StaticOptions, BaseOptionsTask, node.m_pBase );
         }
 
         // Target
-		if ( node.m_pTarget )
-        {
-			FMeshGenerationOptions TargetOptions = InOptions;
-			TargetOptions.bLayouts = false;
-			TargetOptions.OverrideLayouts.Empty();
-			TargetOptions.ActiveTags.Empty();
-			FMeshGenerationResult TargetResult;
-            GenerateMesh( TargetOptions, TargetResult, node.m_pTarget );
-
-            op->Target = TargetResult.MeshOp;
-        }
-        else
-        {
-            // This argument is required
-            ErrorLog->Add( "Mesh make morph target node is not set.", ELMT_ERROR, InMakeMorphNode->GetMessageContext());
+		FMeshTask TargetTask;
+		{
+			FMeshGenerationStaticOptions TargetStaticOptions = StaticOptions;
+			TargetStaticOptions.ActiveTags.Empty();
+
+			FMeshOptionsTask TargetOptionsTask = UE::Tasks::Launch(TEXT("MutableMorphTargetOptions"),
+				[InOptions]() mutable
+				{
+					FMeshGenerationDynamicOptions Result = InOptions.GetResult();
+					Result.bLayouts = false;
+					Result.OverrideLayouts.Empty();
+					return Result;
+				},
+				InOptions
+			);
+
+			TargetTask = GenerateMesh(TargetStaticOptions, TargetOptionsTask, node.m_pTarget );
         }
 
-        OutResult.MeshOp = op;
-        OutResult.BaseMeshOp = BaseResult.BaseMeshOp;
-		OutResult.GeneratedLayouts = BaseResult.GeneratedLayouts;
+		return UE::Tasks::Launch(TEXT("MutableMakeMorph"),
+			[BaseTask, TargetTask, node]() mutable
+			{
+				Ptr<ASTOpMeshDifference> op = new ASTOpMeshDifference();
+
+				// \todo Texcoords are broken?
+				op->bIgnoreTextureCoords = true;
+
+				// UE only has position and normal morph data, optimize for this case if indicated. 
+				if (node.bOnlyPositionAndNormal)
+				{
+					op->Channels = { {static_cast<uint8>(EMeshBufferSemantic::Position), 0}, {static_cast<uint8>(EMeshBufferSemantic::Normal), 0} };
+				}
+
+				FMeshGenerationResult BaseResult = BaseTask.GetResult();
+				op->Base = BaseResult.MeshOp;
+
+				FMeshGenerationResult TargetResult = TargetTask.GetResult();
+				op->Target = TargetResult.MeshOp;
+
+				FMeshGenerationResult Result;
+				Result.MeshOp = op;
+				Result.BaseMeshOp = BaseResult.BaseMeshOp;
+				Result.GeneratedLayouts = BaseResult.GeneratedLayouts;
+				return Result;
+			},
+			UE::Tasks::Prerequisites(BaseTask, TargetTask)
+			);
 	}
 
-    //---------------------------------------------------------------------------------------------
-    void CodeGenerator::GenerateMesh_Fragment(const FMeshGenerationOptions& InOptions, FMeshGenerationResult& OutResult, const NodeMeshFragment* Node )
-    {
-        FMeshGenerationResult BaseResult;
-        if ( Node->SourceMesh )
-        {
-			Ptr<ASTOpMeshExtractLayoutBlocks> op = new ASTOpMeshExtractLayoutBlocks();
-			OutResult.MeshOp = op;
 
-			op->LayoutIndex = (uint16)Node->LayoutIndex;
+	CodeGenerator::FMeshTask CodeGenerator::GenerateMesh_Fragment(const FMeshGenerationStaticOptions& StaticOptions, FMeshOptionsTask InOptions, const NodeMeshFragment* Node )
+    {
+		if (!Node->SourceMesh)
+		{
+			// This argument is required
+			ErrorLog->Add("Mesh fragment source is not set.", ELMT_ERROR, Node->GetMessageContext());
+			return UE::Tasks::MakeCompletedTask<FMeshGenerationResult>();
+		}
 
-			// Generate the source mesh
-			FMeshGenerationOptions BaseOptions = InOptions;
-			BaseOptions.bLayouts = true;
-			BaseOptions.bEnsureAllVerticesHaveLayoutBlock = false;
+		TSharedPtr<const FLayout> Layout;
+		if (Node->Layout)
+		{
+			// Generate the layout with blocks to extract
+			Layout = GenerateLayout(Node->Layout, 0);
+		}
 
-			if (Node->Layout)
-			{
-				// Generate the layout with blocks to extract
-				TSharedPtr<const FLayout> Layout = GenerateLayout(Node->Layout, 0);
-				BaseOptions.OverrideLayouts.Empty();
-				BaseOptions.OverrideLayouts.Add({ Layout, Node->Layout });
-			}
+		FMeshTask SourceTask;
+		{
+			FMeshOptionsTask SourceOptionsTask = UE::Tasks::Launch(TEXT("MutableMeshFragmentOptions"),
+				[InOptions, Layout, Node]() mutable
+				{
+					FMeshGenerationDynamicOptions Result = InOptions.GetResult();
+					Result.bLayouts = true;
+					Result.bEnsureAllVerticesHaveLayoutBlock = false;
+					if (Layout)
+					{
+						Result.OverrideLayouts.Empty();
+						Result.OverrideLayouts.Add({ Layout, Node->Layout });
+					}
+					return Result;
+				},
+				InOptions
+			);
 
-            GenerateMesh( BaseOptions, BaseResult, Node->SourceMesh);
-            op->Source = BaseResult.MeshOp;
-        }
-        else
-        {
-            // This argument is required
-            ErrorLog->Add( "Mesh fragment source is not set.", ELMT_ERROR, Node->GetMessageContext());
-        }
+			SourceTask = GenerateMesh(StaticOptions, SourceOptionsTask, Node->SourceMesh);
+		}
 
-        OutResult.BaseMeshOp = BaseResult.BaseMeshOp;
-		OutResult.GeneratedLayouts = BaseResult.GeneratedLayouts;
+		return UE::Tasks::Launch(TEXT("MutableMeshFragment"),
+			[SourceTask, Node]() mutable
+			{
+				FMeshGenerationResult SourceResult = SourceTask.GetResult();
+
+				Ptr<ASTOpMeshExtractLayoutBlocks> op = new ASTOpMeshExtractLayoutBlocks();
+				op->LayoutIndex = (uint16)Node->LayoutIndex;
+				op->Source = SourceResult.MeshOp;
+
+				FMeshGenerationResult Result;
+				Result.MeshOp = op;
+				Result.BaseMeshOp = SourceResult.BaseMeshOp;
+				Result.GeneratedLayouts = SourceResult.GeneratedLayouts;
+				return Result;
+			},
+			UE::Tasks::Prerequisites(SourceTask)
+		);
 	}
 
 
-    //---------------------------------------------------------------------------------------------
-    void CodeGenerator::GenerateMesh_Switch(const FMeshGenerationOptions& InOptions, FMeshGenerationResult& OutResult, const NodeMeshSwitch* SwitchNode )
+	CodeGenerator::FMeshTask CodeGenerator::GenerateMesh_Switch(const FMeshGenerationStaticOptions& StaticOptions, FMeshOptionsTask InOptions, const NodeMeshSwitch* SwitchNode )
     {
         NodeMeshSwitch::Private& node = *SwitchNode->GetPrivate();
 
@@ -820,146 +886,161 @@ namespace mu
         {
             // No options in the switch!
             // TODO
-            OutResult = FMeshGenerationResult();
-			return;
-        }
-
-        Ptr<ASTOpSwitch> op = new ASTOpSwitch();
-        op->type = EOpType::ME_SWITCH;
+			return UE::Tasks::MakeCompletedTask<FMeshGenerationResult>();
+		}
 
-        // Factor
+		Ptr<ASTOp> Variable;
         if ( node.m_pParameter )
         {
-            op->variable = Generate_Generic( node.m_pParameter.get(), InOptions);
+			Variable = Generate_Generic( node.m_pParameter.get(), StaticOptions);
         }
         else
         {
             // This argument is required
-            op->variable = GenerateMissingScalarCode(TEXT("Switch variable"), 0.0f, SwitchNode->GetMessageContext());
+			Variable = GenerateMissingScalarCode(TEXT("Switch variable"), 0.0f, SwitchNode->GetMessageContext());
         }
 
+		Ptr<ASTOpSwitch> SwitchOp = new ASTOpSwitch();
+		SwitchOp->type = EOpType::ME_SWITCH;
+		SwitchOp->variable = Variable;
+
         // Options
 		bool bFirstValidConnectionFound = false;
-        for ( int32 t=0; t< node.Options.Num(); ++t )
+		FMeshOptionsTask CurrentTargetOptions = InOptions;
+		TArray<FMeshTask> OptionTasks;
+		TArray<int16> OptionIndices;
+		OptionTasks.Reserve(node.Options.Num());
+		OptionIndices.Reserve(node.Options.Num());
+		for ( int32 OptionIndex=0; OptionIndex<node.Options.Num(); ++OptionIndex)
         {
-			FMeshGenerationOptions TargetOptions = InOptions;
+            if ( node.Options[OptionIndex] )
+            {								
+				FMeshTask BranchTask = GenerateMesh(StaticOptions, CurrentTargetOptions, node.Options[OptionIndex] );
+				OptionTasks.Emplace(BranchTask);
+				OptionIndices.Emplace(int16(OptionIndex));
 
-            if ( node.Options[t] )
-            {
 				// Take the layouts from the first non-null connection.
-				// \TODO: Take them from the first connection that actually returns layouts?
-				if (bFirstValidConnectionFound)
-				{
-					TargetOptions.OverrideLayouts = OutResult.GeneratedLayouts;
-				}
-				
-				FMeshGenerationResult BranchResults;
-                GenerateMesh(TargetOptions, BranchResults, node.Options[t] );
-
-                Ptr<ASTOp> branch = BranchResults.MeshOp;
-                op->cases.Emplace((int16)t,op,branch);
-
 				if (!bFirstValidConnectionFound)
 				{
 					bFirstValidConnectionFound = true;
-                    OutResult = BranchResults;
-                }
-            }
-        }
-
-        OutResult.MeshOp = op;
-    }
 
+					CurrentTargetOptions = UE::Tasks::Launch(TEXT("MutableMeshSwitchOptions"),
+						[InOptions, BranchTask]() mutable
+						{
+							FMeshGenerationDynamicOptions Result = InOptions.GetResult();
+							Result.OverrideLayouts = BranchTask.GetResult().GeneratedLayouts;
+							return Result;
+						},
+						UE::Tasks::Prerequisites(InOptions, BranchTask)
+					);
 
-	//---------------------------------------------------------------------------------------------
-	void CodeGenerator::GenerateMesh_Table(const FMeshGenerationOptions& InOptions, FMeshGenerationResult& OutResult, const NodeMeshTable* TableNode)
-	{
-		//
-		FMeshGenerationResult NewResult = OutResult;
-		int32 t = 0;
-		bool bFirstRowGenerated = false;
+				}
+			}
+        }
 
-		Ptr<ASTOp> Op = GenerateTableSwitch<NodeMeshTable, ETableColumnType::Mesh, EOpType::ME_SWITCH>(*TableNode,
-			[this, &NewResult, &bFirstRowGenerated, &InOptions] (const NodeMeshTable& node, int32 colIndex, int32 row, FErrorLog*)
+		return UE::Tasks::Launch(TEXT("MutableMeshSwitch"),
+			[SwitchOp, OptionTasks, OptionIndices]() mutable
 			{
-				TSharedPtr<mu::FMesh> pMesh = node.Table->GetPrivate()->Rows[row].Values[colIndex].Mesh;
-				FMeshGenerationResult BranchResults;
+				FMeshGenerationResult Result;
 
-				if (pMesh)
+				bool bFirstValidConnectionFound = false;
+				for (int32 ValidOptionIndex = 0; ValidOptionIndex < OptionIndices.Num(); ++ValidOptionIndex)
 				{
-					NodeMeshConstantPtr pCell = new NodeMeshConstant();
-					pCell->SetValue(pMesh);
-
-					// TODO Take into account layout strategy
-					int32 numLayouts = node.Layouts.Num();
-					pCell->SetLayoutCount(numLayouts);
-					for (int32 i = 0; i < numLayouts; ++i)
+					FMeshGenerationResult BranchResult = OptionTasks[ValidOptionIndex].GetResult();
+					if (!bFirstValidConnectionFound)
 					{
-						pCell->SetLayout(i, node.Layouts[i]);
+						bFirstValidConnectionFound = true;
+						Result = BranchResult;
 					}
 
-					FMeshGenerationOptions TargetOptions = InOptions;
-
-					if (bFirstRowGenerated)
-					{
-						TargetOptions.OverrideLayouts = NewResult.GeneratedLayouts;
-					}
+					Ptr<ASTOp> Branch = BranchResult.MeshOp;
+					SwitchOp->cases.Emplace(OptionIndices[ValidOptionIndex], SwitchOp, Branch);
+				}
 
-					TargetOptions.OverrideContext = node.Table->GetPrivate()->Rows[row].Values[colIndex].ErrorContext;
+				Result.MeshOp = SwitchOp;
+				return Result;
+			},
+			OptionTasks
+		);
+    }
 
-					pCell->SourceDataDescriptor = node.SourceDataDescriptor;
 
-					// Combine the SourceId of the node with the RowId to generate one shared between all resources from this row.
-					// Hash collisions are allowed, since it is used to group resources, not to differentiate them.
-					const uint32 RowId = node.Table->GetPrivate()->Rows[row].Id;
-					pCell->SourceDataDescriptor.SourceId = HashCombine(node.SourceDataDescriptor.SourceId, RowId);
+	CodeGenerator::FMeshTask CodeGenerator::GenerateMesh_Table(const FMeshGenerationStaticOptions& StaticOptions, FMeshOptionsTask InOptions, const NodeMeshTable* TableNode)
+	{
+		Ptr<NodeMeshSwitch> SwitchNode;
+		GenerateTableSwitchNode<NodeMeshTable, ETableColumnType::Mesh, NodeMesh, NodeMeshSwitch>(*TableNode, SwitchNode,
+			[this] (const NodeMeshTable& node, int32 colIndex, int32 row, FErrorLog* pErrorLog) -> Ptr<NodeMesh>
+			{
+				const FTableValue& Cell = node.Table->GetPrivate()->Rows[row].Values[colIndex];
+				TSharedPtr<mu::FMesh> pMesh = Cell.Mesh;
+				
+				if (!pMesh)
+				{
+					return nullptr;
+				}
 
-					GenerateMesh(TargetOptions, BranchResults, pCell);
+				Ptr<NodeMeshConstant> CellNode = new NodeMeshConstant();
+				CellNode->SetValue(pMesh);
 
-					if (!bFirstRowGenerated)
-					{
-						NewResult = BranchResults;
-						bFirstRowGenerated = true;
-					}
+				// TODO Take into account layout strategy
+				int32 numLayouts = node.Layouts.Num();
+				CellNode->SetLayoutCount(numLayouts);
+				for (int32 i = 0; i < numLayouts; ++i)
+				{
+					CellNode->SetLayout(i, node.Layouts[i]);
 				}
 
-				return BranchResults.MeshOp;
-			});
+				CellNode->SetMessageContext(Cell.ErrorContext);
+				CellNode->SourceDataDescriptor = node.SourceDataDescriptor;
 
-		NewResult.MeshOp = Op;
+				// Combine the SourceId of the node with the RowId to generate one shared between all resources from this row.
+				// Hash collisions are allowed, since it is used to group resources, not to differentiate them.
+				const uint32 RowId = node.Table->GetPrivate()->Rows[row].Id;
+				CellNode->SourceDataDescriptor.SourceId = HashCombine(node.SourceDataDescriptor.SourceId, RowId);
 
-		OutResult = NewResult;
+				return CellNode;
+			});
+
+		return GenerateMesh(StaticOptions, InOptions, SwitchNode);
 	}
 
 
-    //---------------------------------------------------------------------------------------------
-    void CodeGenerator::GenerateMesh_Variation(const FMeshGenerationOptions& InOptions, FMeshGenerationResult& OutResult,
-		const NodeMeshVariation* VariationNode )
+	CodeGenerator::FMeshTask CodeGenerator::GenerateMesh_Variation(const FMeshGenerationStaticOptions& StaticOptions, FMeshOptionsTask InOptions, const NodeMeshVariation* VariationNode )
     {
         NodeMeshVariation::Private& node = *VariationNode->GetPrivate();
 
-        FMeshGenerationResult currentResult;
-        Ptr<ASTOp> currentMeshOp;
+		bool bFirstValidConnectionFound = false;
+		FMeshOptionsTask VariationOptionsTask = InOptions;
 
-        bool firstOptionProcessed = false;
+		FMeshTask DefaultTask = UE::Tasks::MakeCompletedTask<FMeshGenerationResult>();
 
         // Default case
         if ( node.m_defaultMesh )
         {
-            FMeshGenerationResult BranchResults;
-			FMeshGenerationOptions DefaultOptions = InOptions;
+			DefaultTask = GenerateMesh(StaticOptions, InOptions, node.m_defaultMesh );
+			bFirstValidConnectionFound = true;
 
-			GenerateMesh(DefaultOptions, BranchResults, node.m_defaultMesh );
-            currentMeshOp = BranchResults.MeshOp;
-            currentResult = BranchResults;
-            firstOptionProcessed = true;
+			VariationOptionsTask = UE::Tasks::Launch(TEXT("MutableMeshVariationDefaultOptions"),
+				[InOptions, DefaultTask]() mutable
+				{
+					FMeshGenerationDynamicOptions Result = InOptions.GetResult();
+					Result.OverrideLayouts = DefaultTask.GetResult().GeneratedLayouts;
+					return Result;
+				},
+				UE::Tasks::Prerequisites(InOptions, DefaultTask)
+			);
         }
 
         // Process variations in reverse order, since conditionals are built bottom-up.
-        for ( int32 t = node.m_variations.Num()-1; t >= 0; --t )
+		TArray<FMeshTask> ReverseVariations;
+		TArray<int32> ReverseVariationsIndices;
+		ReverseVariations.Reserve(node.m_variations.Num());
+		ReverseVariationsIndices.Reserve(node.m_variations.Num());
+
+        for ( int32 VariationIndex = node.m_variations.Num()-1; VariationIndex >= 0; --VariationIndex)
         {
             int32 tagIndex = -1;
-            const FString& tag = node.m_variations[t].m_tag;
+            const FString& tag = node.m_variations[VariationIndex].m_tag;
             for ( int32 i = 0; i < FirstPass.Tags.Num(); ++i )
             {
                 if ( FirstPass.Tags[i].Tag==tag)
@@ -980,59 +1061,115 @@ namespace mu
             }
 
             Ptr<ASTOp> variationMeshOp;
-            if ( node.m_variations[t].m_mesh )
-            {
-				FMeshGenerationOptions VariationOptions = InOptions;
-
-                if (firstOptionProcessed)
-                {
-					VariationOptions.OverrideLayouts = currentResult.GeneratedLayouts;
-                }
-         
-                FMeshGenerationResult BranchResults;
-				GenerateMesh(VariationOptions, BranchResults, node.m_variations[t].m_mesh );
+            if ( node.m_variations[VariationIndex].m_mesh )
+            {         
+				FMeshTask BranchTask = GenerateMesh(StaticOptions, VariationOptionsTask, node.m_variations[VariationIndex].m_mesh );
 
-                variationMeshOp = BranchResults.MeshOp;
+				ReverseVariations.Add(BranchTask);
 
-                if ( !firstOptionProcessed )
+                if ( !bFirstValidConnectionFound)
                 {
-                    firstOptionProcessed = true;                   
-                    currentResult = BranchResults;
-                }
+					bFirstValidConnectionFound = true;
+
+					VariationOptionsTask = UE::Tasks::Launch(TEXT("MutableMeshVariationOptions"),
+						[InOptions, BranchTask]() mutable
+						{
+							FMeshGenerationDynamicOptions Result = InOptions.GetResult();
+							Result.OverrideLayouts = BranchTask.GetResult().GeneratedLayouts;
+							return Result;
+						},
+						UE::Tasks::Prerequisites(VariationOptionsTask, BranchTask)
+					);
+				}
             }
+			else
+			{
 
-            Ptr<ASTOpConditional> conditional = new ASTOpConditional;
-            conditional->type = EOpType::ME_CONDITIONAL;
-            conditional->no = currentMeshOp;
-            conditional->yes = variationMeshOp;            
-            conditional->condition = FirstPass.Tags[tagIndex].GenericCondition;
+				ReverseVariations.Add(UE::Tasks::MakeCompletedTask<FMeshGenerationResult>());
+			}
+		
+			ReverseVariationsIndices.Add(tagIndex);
+		}
+		
+		TArray<UE::Tasks::FTask> Requisites;
+		Requisites.Reserve(ReverseVariations.Num()+2);
+		Requisites.Add(DefaultTask);
+		Requisites.Append(ReverseVariations);
+
+		return UE::Tasks::Launch(TEXT("MutableMeshVariation"),
+			[this, node, DefaultTask, ReverseVariations, ReverseVariationsIndices]() mutable
+			{
+				FMeshGenerationResult Result;
+				Ptr<ASTOp> CurrentMeshOp;
 
-            currentMeshOp = conditional;
-        }
+				bool bFirstValidConnectionFound = false;
+
+				// Default case
+				if (node.m_defaultMesh)
+				{
+					Result = DefaultTask.GetResult();
+					CurrentMeshOp = Result.MeshOp;
+					bFirstValidConnectionFound = true;
+				}
+
+				for (int32 ReverseVariationIndex = 0; ReverseVariationIndex<ReverseVariations.Num(); ++ReverseVariationIndex)
+				{
+					int32 tagIndex = ReverseVariationsIndices[ReverseVariationIndex];
+
+					FMeshGenerationResult VariationResult = ReverseVariations[ReverseVariationIndex].GetResult();
+
+					Ptr<ASTOpConditional> Conditional = new ASTOpConditional;
+					Conditional->type = EOpType::ME_CONDITIONAL;
+					Conditional->no = CurrentMeshOp;
+					Conditional->yes = VariationResult.MeshOp;
+					Conditional->condition = FirstPass.Tags[tagIndex].GenericCondition;
+
+					CurrentMeshOp = Conditional;
+				}
+
+				Result.MeshOp = CurrentMeshOp;
+				return Result;
+			},
+			Requisites
+		);
 
-        OutResult = currentResult;
-        OutResult.MeshOp = currentMeshOp;
     }
 
 
-    //---------------------------------------------------------------------------------------------
-    void CodeGenerator::GenerateMesh_Constant(const FMeshGenerationOptions& InOptions, FMeshGenerationResult& OutResult, const NodeMeshConstant* InNode )
+	CodeGenerator::FMeshTask CodeGenerator::GenerateMesh_Constant(const FMeshGenerationStaticOptions& StaticOptions, FMeshOptionsTask InOptions, const NodeMeshConstant* InNode )
     {
 		MUTABLE_CPUPROFILER_SCOPE(GenerateMesh_Constant);
 
         NodeMeshConstant::Private& Node = *InNode->GetPrivate();
 
+		TSharedPtr<FMesh> Mesh = Node.Value;
+
+		// True passthrough?
+		if (Mesh && Mesh->IsReference() && !Mesh->IsForceLoad())
+		{
+			Ptr<ASTOpReferenceResource> ReferenceOp = new ASTOpReferenceResource();
+			ReferenceOp->Type = EOpType::ME_REFERENCE;
+			ReferenceOp->ID = Mesh->GetReferencedMesh();
+			ReferenceOp->bForceLoad = false;
+
+			FMeshGenerationResult Result;
+			Result.BaseMeshOp = ReferenceOp;
+			Result.MeshOp = ReferenceOp;
+
+			// We won't be able to do anything with the passthrough mesh at compile time.
+			//check(InOptions.OverrideLayouts.IsEmpty());
+			return UE::Tasks::MakeCompletedTask<FMeshGenerationResult>(Result);
+		}
+
         Ptr<ASTOpConstantResource> ConstantOp = new ASTOpConstantResource();
 		ConstantOp->Type = EOpType::ME_CONSTANT;
 		ConstantOp->SourceDataDescriptor = InNode->SourceDataDescriptor;
-		OutResult.BaseMeshOp = ConstantOp;
-		OutResult.MeshOp = ConstantOp;
-		OutResult.GeneratedLayouts.Empty();
 
-		bool bIsOverridingLayouts = !InOptions.OverrideLayouts.IsEmpty();
+		FMeshGenerationResult Result;
+		Result.BaseMeshOp = ConstantOp;
+		Result.MeshOp = ConstantOp;
 
-        TSharedPtr<FMesh> pMesh = Node.Value;
-		if (!pMesh)
+		if (!Mesh)
 		{
 			// This data is required
 			TSharedPtr<FMesh> EmptyMesh = MakeShared<FMesh>();
@@ -1042,525 +1179,644 @@ namespace mu
 			// Log an error message
 			ErrorLog->Add("Constant mesh not set.", ELMT_WARNING, InNode->GetMessageContext());
 
-			return;
+			return UE::Tasks::MakeCompletedTask<FMeshGenerationResult>(Result);
 		}
 
-		if (pMesh->IsReference())
-		{
-			Ptr<ASTOpReferenceResource> ReferenceOp = new ASTOpReferenceResource();
-			ReferenceOp->Type = EOpType::ME_REFERENCE;
-			ReferenceOp->ID = pMesh->GetReferencedMesh();
-			ReferenceOp->bForceLoad = pMesh->IsForceLoad();
+		TArray<UE::Tasks::FTask, TInlineAllocator<4>> Requisites;
+		Requisites.Emplace(InOptions);
 
-			OutResult.BaseMeshOp = ReferenceOp;
-			OutResult.MeshOp = ReferenceOp;
+		TSharedPtr< TSharedPtr<FMesh> > ResolveMesh = MakeShared<TSharedPtr<FMesh>>();
 
-			return;
-		}
-
-		// Separate the tags from the mesh
-		TArray<FString> Tags = pMesh->Tags;
-		if (Tags.Num())
+		// Compile-time reference? Add task to resolve as requisite.
+		if (Mesh->IsReference() && Mesh->IsForceLoad())
 		{
-			TSharedPtr<FMesh> TaglessMesh = CloneOrTakeOver(pMesh);
-			TaglessMesh->Tags.SetNum(0, EAllowShrinking::No);
-			pMesh = TaglessMesh;
-		}
+			uint32 MeshID = Mesh->GetReferencedMesh();
 
-		// Find out if we can (or have to) reuse a mesh that we have already generated.
-		FGeneratedConstantMesh DuplicateOf;
-		uint32 ThisMeshHash = HashCombineFast( GetTypeHash(pMesh->GetVertexCount()), GetTypeHash(pMesh->GetIndexCount()) );
-		TArray<FGeneratedConstantMesh>& CachedCandidates = GeneratedConstantMeshes.FindOrAdd(ThisMeshHash,{});
-		for (const FGeneratedConstantMesh& Candidate : CachedCandidates)
+			bool bRunImmediatlyIfPossible = IsInGameThread();
+			FString Morph = Mesh->GetReferencedMorph();
+			UE::Tasks::FTask ReferenceCompletion = CompilerOptions->OptimisationOptions.ReferencedMeshResourceProvider(MeshID, Morph, ResolveMesh, bRunImmediatlyIfPossible);
+			Requisites.Emplace(ReferenceCompletion);
+		}
+		else
 		{
-			bool bCompareLayouts = InOptions.bLayouts && !bIsOverridingLayouts;
+			*ResolveMesh = Mesh;
+		}
 
-			if (Candidate.Mesh->IsSimilar(*pMesh, bCompareLayouts))
+		// This task does the necessary processing of the mesh when it is available
+		Ptr<const NodeMeshConstant> NodeCopy = InNode;
+		FMeshTask MeshProcessingTask = GenerateMeshConstantPipe.Launch( TEXT("MutableGenerateMeshConstant"),
+			[ReferenceMesh=Mesh,ResolveMesh,Result,ConstantOp, StaticOptions, InOptions, NodeCopy,
+			// This is read-only
+			CompilerOptions=this->CompilerOptions,
+			this,
+			// This is read-write so it is protected with a mutex.
+			&GenerateMeshConstantState = this->GenerateMeshConstantState
+			]() mutable
 			{
-				// If it is similar and we are overriding the layouts, we must compare the layouts of the candidate with the ones
-				// we are using to override.
-				if (bIsOverridingLayouts)
+				TSharedPtr<FMesh> Mesh = *ResolveMesh;
+
+				if (!Mesh)
 				{
-					if (Candidate.Mesh->GetLayoutCount() != InOptions.OverrideLayouts.Num())
-					{
-						continue;
-					}
+					return FMeshGenerationResult();
+				}
 
-					bool bLayoutsAreEqual = true;
-					for (int32 l = 0; l < Candidate.Mesh->GetLayoutCount(); ++l)
+				// Separate the tags from the mesh
+				TArray<FString> Tags = ReferenceMesh->Tags;
+				if (!Mesh->Tags.IsEmpty())
+				{
+					Tags.Append(Mesh->Tags);
+					TSharedPtr<FMesh> TaglessMesh = CloneOrTakeOver(Mesh);
+					TaglessMesh->Tags.SetNum(0, EAllowShrinking::No);
+					Mesh = TaglessMesh;
+				}
+
+				FMeshGenerationDynamicOptions Options = InOptions.GetResult();
+				bool bIsOverridingLayouts = !Options.OverrideLayouts.IsEmpty();
+
+				// Find out if we can (or have to) reuse a mesh that we have already generated.
+				FGenerateMeshConstantState::FGeneratedConstantMesh DuplicateOf;
+				uint32 ThisMeshHash = HashCombineFast(GetTypeHash(Mesh->GetVertexCount()), GetTypeHash(Mesh->GetIndexCount()));
+				TSharedPtr<FMesh> ClonedMesh;
+				{
+					TArray<FGenerateMeshConstantState::FGeneratedConstantMesh>& CachedCandidates = GenerateMeshConstantState.GeneratedConstantMeshes.FindOrAdd(ThisMeshHash, {});
+					for (const FGenerateMeshConstantState::FGeneratedConstantMesh& Candidate : CachedCandidates)
 					{
-						bLayoutsAreEqual = (*Candidate.Mesh->GetLayout(l) == *InOptions.OverrideLayouts[l].Layout);
-						if ( !bLayoutsAreEqual )
+						bool bCompareLayouts = Options.bLayouts && !bIsOverridingLayouts;
+
+						if (Candidate.Mesh->IsSimilar(*Mesh, bCompareLayouts))
 						{
+							// If it is similar and we are overriding the layouts, we must compare the layouts of the candidate with the ones
+							// we are using to override.
+							if (bIsOverridingLayouts)
+							{
+								if (Candidate.Mesh->GetLayoutCount() != Options.OverrideLayouts.Num())
+								{
+									continue;
+								}
+
+								bool bLayoutsAreEqual = true;
+								for (int32 l = 0; l < Candidate.Mesh->GetLayoutCount(); ++l)
+								{
+									bLayoutsAreEqual = (*Candidate.Mesh->GetLayout(l) == *Options.OverrideLayouts[l].Layout);
+									if (!bLayoutsAreEqual)
+									{
+										break;
+									}
+								}
+
+								if (!bLayoutsAreEqual)
+								{
+									continue;
+								}
+							}
+
+							DuplicateOf = Candidate;
 							break;
 						}
 					}
 
-					if (!bLayoutsAreEqual)
+					if (!DuplicateOf.Mesh)
 					{
-						continue;
-					}
-				}
+						// We need to clone the mesh in the node because we will modify it.
+						ClonedMesh = Mesh->Clone();
+						ClonedMesh->EnsureSurfaceData();
 
-				DuplicateOf = Candidate;
-				break;
-			}
-		}
+						ConstantOp->SetValue(ClonedMesh, CompilerOptions->OptimisationOptions.DiskCacheContext);
 
-		Ptr<ASTOp> LastMeshOp = ConstantOp;
+						// Add the unique vertex ID prefix in all cases, since it is free memory-wise
+						uint32 MeshIDPrefix = uint32(ConstantOp->GetValueHash());
+						{
+							// Ensure the ID group is unique
+							bool bValid = false;
+							do
+							{
+								bool bAlreadyPresent = false;
+								GenerateMeshConstantState.UniqueVertexIDGroups.FindOrAdd(MeshIDPrefix, &bAlreadyPresent);
+								bValid = !bAlreadyPresent && MeshIDPrefix != 0;
+								if (!bValid)
+								{
+									++MeshIDPrefix;
+								}
+							} while (bValid);
 
-		if (DuplicateOf.Mesh)
-		{
-			// Make sure the source layouts of the mesh are mapped to the layouts of the duplicated mesh.
-			if (InOptions.bLayouts)
-			{
-				if (bIsOverridingLayouts)
-				{
-					OutResult.GeneratedLayouts = InOptions.OverrideLayouts;
-				}
-				else
-				{
-					for (int32 l = 0; l < DuplicateOf.Mesh->GetLayoutCount(); ++l)
-					{
-						TSharedPtr<const FLayout> DuplicatedLayout = DuplicateOf.Mesh->GetLayout(l);
-						OutResult.GeneratedLayouts.Add({ DuplicatedLayout,0 });
+							ClonedMesh->MeshIDPrefix = MeshIDPrefix;
+						}
+
+						// Add the constant data
+						FGenerateMeshConstantState::FGeneratedConstantMesh MeshEntry;
+						MeshEntry.Mesh = ClonedMesh;
+						MeshEntry.LastMeshOp = ConstantOp;
+						CachedCandidates.Add(MeshEntry);
 					}
 				}
-			}
-
-			LastMeshOp = DuplicateOf.LastMeshOp;
-			ConstantOp = nullptr;
-		}
-		else
-		{
-			// We need to clone the mesh in the node because we will modify it.
-			TSharedPtr<FMesh> ClonedMesh = pMesh->Clone();
-			ClonedMesh->EnsureSurfaceData();
 
-			ConstantOp->SetValue(ClonedMesh, CompilerOptions->OptimisationOptions.DiskCacheContext);
+				Ptr<ASTOp> LastMeshOp = ConstantOp;
 
-			// Add the unique vertex ID prefix in all cases, since it is free memory-wise
-			uint32 MeshIDPrefix = uint32(ConstantOp->GetValueHash());
-			{
-				// Ensure the ID group is unique
-				bool bValid = false;
-				do
+				if (DuplicateOf.Mesh)
 				{
-					bool bAlreadyPresent = false;
-					UniqueVertexIDGroups.FindOrAdd(MeshIDPrefix, &bAlreadyPresent);
-					bValid = !bAlreadyPresent && MeshIDPrefix != 0;
-					if (!bValid)
+					// Make sure the source layouts of the mesh are mapped to the layouts of the duplicated mesh.
+					if (Options.bLayouts)
 					{
-						++MeshIDPrefix;
+						if (bIsOverridingLayouts)
+						{
+							Result.GeneratedLayouts = Options.OverrideLayouts;
+						}
+						else
+						{
+							for (int32 l = 0; l < DuplicateOf.Mesh->GetLayoutCount(); ++l)
+							{
+								TSharedPtr<const FLayout> DuplicatedLayout = DuplicateOf.Mesh->GetLayout(l);
+								Result.GeneratedLayouts.Add({ DuplicatedLayout,0 });
+							}
+						}
 					}
-				} while (bValid);
-
-				ClonedMesh->MeshIDPrefix = MeshIDPrefix;
-			}
 
-			// Add the constant data
-			FGeneratedConstantMesh MeshEntry;
-			MeshEntry.Mesh = ClonedMesh;
-			MeshEntry.LastMeshOp = LastMeshOp;
-			CachedCandidates.Add(MeshEntry);
-
-			if (InOptions.bLayouts)
-			{
-				if (!bIsOverridingLayouts)
+					LastMeshOp = DuplicateOf.LastMeshOp;
+					ConstantOp = nullptr;
+				}
+				else
 				{
-					// Apply whatever transform is necessary for every layout
-					for (int32 LayoutIndex = 0; LayoutIndex < Node.Layouts.Num(); ++LayoutIndex)
+					if (Options.bLayouts)
 					{
-						Ptr<NodeLayout> LayoutNode = Node.Layouts[LayoutIndex];
-						if (!LayoutNode)
+						if (!bIsOverridingLayouts)
 						{
-							continue;
-						}
+							NodeMeshConstant::Private& Node = *NodeCopy->GetPrivate();
+
+							// Apply whatever transform is necessary for every layout
+							for (int32 LayoutIndex = 0; LayoutIndex < Node.Layouts.Num(); ++LayoutIndex)
+							{
+								Ptr<NodeLayout> LayoutNode = Node.Layouts[LayoutIndex];
+								if (!LayoutNode)
+								{
+									continue;
+								}
 
-						FGeneratedLayout GeneratedData;
-						GeneratedData.Source = LayoutNode;
-						GeneratedData.Layout = GenerateLayout(LayoutNode, MeshIDPrefix);
-						const void* Context = InOptions.OverrideContext.Get(InNode->GetMessageContext());
+								FGeneratedLayout GeneratedData;
+								GeneratedData.Source = LayoutNode;
+								GeneratedData.Layout = this->GenerateLayout(LayoutNode, ClonedMesh->MeshIDPrefix);
+								const void* Context = Options.OverrideContext.Get(NodeCopy->GetMessageContext());
 
-						bool bUseAbsoluteBlockIds = false;
-						PrepareMeshForLayout(GeneratedData, ClonedMesh, LayoutIndex, Context, InOptions, bUseAbsoluteBlockIds);
+								bool bUseAbsoluteBlockIds = false;
+								this->PrepareMeshForLayout(GeneratedData, ClonedMesh, LayoutIndex, Context, StaticOptions, Options, bUseAbsoluteBlockIds);
 
-						OutResult.GeneratedLayouts.Add(GeneratedData);
+								Result.GeneratedLayouts.Add(GeneratedData);
+							}
+						}
+						else
+						{
+							// We need to apply the transform of the layouts used to override
+							for (int32 LayoutIndex = 0; LayoutIndex < Options.OverrideLayouts.Num(); ++LayoutIndex)
+							{
+								const FGeneratedLayout& OverrideData = Options.OverrideLayouts[LayoutIndex];
+								TSharedPtr<const FLayout> GeneratedLayout = OverrideData.Layout;
+								const void* Context = Options.OverrideContext.Get(NodeCopy->GetMessageContext());
+
+								// In this case we need the layout block ids to use the ids in the parent layout, and not be prefixed with
+								// the current mesh id prefix. For this reason we need them to be absolute.
+								bool bUseAbsoluteBlockIds = true;
+								this->PrepareMeshForLayout(OverrideData, ClonedMesh, LayoutIndex, Context, StaticOptions, Options, bUseAbsoluteBlockIds);
+
+								Result.GeneratedLayouts.Add(OverrideData);
+							}
+						}
 					}
 				}
-				else
-				{
-					// We need to apply the transform of the layouts used to override
-					for (int32 LayoutIndex = 0; LayoutIndex < InOptions.OverrideLayouts.Num(); ++LayoutIndex)
-					{
-						const FGeneratedLayout& OverrideData = InOptions.OverrideLayouts[LayoutIndex];
-						TSharedPtr<const FLayout> GeneratedLayout = OverrideData.Layout;
-						const void* Context = InOptions.OverrideContext.Get(InNode->GetMessageContext());
 
-						// In this case we need the layout block ids to use the ids in the parent layout, and not be prefixed with
-						// the current mesh id prefix. For this reason we need them to be absolute.
-						bool bUseAbsoluteBlockIds = true;
-						PrepareMeshForLayout(OverrideData, ClonedMesh, LayoutIndex, Context, InOptions, bUseAbsoluteBlockIds);
+				Result.BaseMeshOp = LastMeshOp;
 
-						OutResult.GeneratedLayouts.Add(OverrideData);
-					}
+				// Add the tags operation
+				if (Tags.Num())
+				{
+					Ptr<ASTOpMeshAddTags> AddTagsOp = new ASTOpMeshAddTags;
+					AddTagsOp->Source = LastMeshOp;
+					AddTagsOp->Tags = Tags;
+					LastMeshOp = AddTagsOp;
 				}
-			}
-		}
 
-		OutResult.BaseMeshOp = LastMeshOp;
+				Result.MeshOp = LastMeshOp;
 
-		// Add the tags operation
-		if (Tags.Num())
-		{
-			Ptr<ASTOpMeshAddTags> AddTagsOp = new ASTOpMeshAddTags;
-			AddTagsOp->Source = LastMeshOp;
-			AddTagsOp->Tags = Tags;
-			LastMeshOp = AddTagsOp;
-		}
+				return Result;
+			},
+			Requisites);
 
-		OutResult.MeshOp = LastMeshOp;
 
 		// Apply the modifier for the pre-normal operations stage.
 		TArray<FirstPassGenerator::FModifier> Modifiers;
 		constexpr bool bModifiersForBeforeOperations = true;
-		GetModifiersFor(InOptions.ComponentId, InOptions.ActiveTags, bModifiersForBeforeOperations, Modifiers);
+		this->GetModifiersFor(StaticOptions.ComponentId, StaticOptions.ActiveTags, bModifiersForBeforeOperations, Modifiers);
+
+		// This task does the necessary processing of the mesh when it is available
+		FMeshTask ApplyModifiersTask = ApplyMeshModifiers(Modifiers, StaticOptions, InOptions, MeshProcessingTask, nullptr, InNode->GetMessageContext(), InNode);
 
-		OutResult.MeshOp = ApplyMeshModifiers(Modifiers, InOptions, OutResult, nullptr, InNode->GetMessageContext(), InNode);
+		return ApplyModifiersTask;
     }
 
 
-    //---------------------------------------------------------------------------------------------
-    void CodeGenerator::GenerateMesh_Format(const FMeshGenerationOptions& InOptions, FMeshGenerationResult& OutResult,
-		const NodeMeshFormat* format )
+	CodeGenerator::FMeshTask CodeGenerator::GenerateMesh_Format(const FMeshGenerationStaticOptions& StaticOptions, FMeshOptionsTask InOptions, const NodeMeshFormat* format )
     {
         NodeMeshFormat::Private& node = *format->GetPrivate();
 
-        if ( node.Source )
-        {
-			FMeshGenerationOptions Options = InOptions;
+		if (!node.Source)
+		{
+			// Put something there
+			return GenerateMesh(StaticOptions, InOptions, new NodeMeshConstant());
+		}
 
-			FMeshGenerationResult baseResult;
-			GenerateMesh(Options,baseResult, node.Source);
-            Ptr<ASTOpMeshFormat> op = new ASTOpMeshFormat();
-            op->Source = baseResult.MeshOp;
-            op->Flags = 0;
+ 		FMeshTask BaseTask = GenerateMesh(StaticOptions, InOptions, node.Source);
 
-            TSharedPtr<FMesh> FormatMesh = MakeShared<FMesh>();
+		return UE::Tasks::Launch(TEXT("MutableMeshFormat"),
+			[node, BaseTask, CompilerOptions=CompilerOptions]() mutable
+			{
+				Ptr<ASTOpMeshFormat> op = new ASTOpMeshFormat();
 
-            if (node.VertexBuffers.GetBufferCount())
-            {
-                op->Flags |= OP::MeshFormatArgs::Vertex;
-                FormatMesh->VertexBuffers = node.VertexBuffers;
-            }
+				FMeshGenerationResult baseResult = BaseTask.GetResult();
+				op->Source = baseResult.MeshOp;
+				op->Flags = 0;
 
-            if (node.IndexBuffers.GetBufferCount())
-            {
-				op->Flags |= OP::MeshFormatArgs::Index;
-                FormatMesh->IndexBuffers = node.IndexBuffers;
-            }
+				TSharedPtr<FMesh> FormatMesh = MakeShared<FMesh>();
 
-			if (node.bOptimizeBuffers)
-			{
-				op->Flags |= OP::MeshFormatArgs::OptimizeBuffers;
-			}
+				if (node.VertexBuffers.GetBufferCount())
+				{
+					op->Flags |= OP::MeshFormatArgs::Vertex;
+					FormatMesh->VertexBuffers = node.VertexBuffers;
+				}
 
-            Ptr<ASTOpConstantResource> cop = new ASTOpConstantResource();
-            cop->Type = EOpType::ME_CONSTANT;
-            cop->SetValue( FormatMesh, CompilerOptions->OptimisationOptions.DiskCacheContext );
-			if (baseResult.BaseMeshOp)
-			{
-				cop->SourceDataDescriptor = baseResult.BaseMeshOp->GetSourceDataDescriptor();
-			}
-            op->Format = cop;
+				if (node.IndexBuffers.GetBufferCount())
+				{
+					op->Flags |= OP::MeshFormatArgs::Index;
+					FormatMesh->IndexBuffers = node.IndexBuffers;
+				}
 
-            OutResult.MeshOp = op;
-            OutResult.BaseMeshOp = baseResult.BaseMeshOp;
-			OutResult.GeneratedLayouts = baseResult.GeneratedLayouts;
-		}
-        else
-        {
-            // Put something there
-            GenerateMesh(InOptions, OutResult, new NodeMeshConstant() );
-        }
+				if (node.bOptimizeBuffers)
+				{
+					op->Flags |= OP::MeshFormatArgs::OptimizeBuffers;
+				}
+
+				Ptr<ASTOpConstantResource> cop = new ASTOpConstantResource();
+				cop->Type = EOpType::ME_CONSTANT;
+				cop->SetValue(FormatMesh, CompilerOptions->OptimisationOptions.DiskCacheContext);
+				if (baseResult.BaseMeshOp)
+				{
+					cop->SourceDataDescriptor = baseResult.BaseMeshOp->GetSourceDataDescriptor();
+				}
+				op->Format = cop;
+
+				FMeshGenerationResult Result;
+				Result.MeshOp = op;
+				Result.BaseMeshOp = baseResult.BaseMeshOp;
+				Result.GeneratedLayouts = baseResult.GeneratedLayouts;
+
+				return Result;
+			},
+			BaseTask
+		);
     }
 
 
-    //---------------------------------------------------------------------------------------------
-    void CodeGenerator::GenerateMesh_Transform(const FMeshGenerationOptions& InOptions, FMeshGenerationResult& OutResult,
-                                                const NodeMeshTransform* TransformNode )
+	CodeGenerator::FMeshTask CodeGenerator::GenerateMesh_Transform(const FMeshGenerationStaticOptions& StaticOptions, FMeshOptionsTask InOptions, const NodeMeshTransform* TransformNode )
     {
-        const auto& node = *TransformNode->GetPrivate();
+        const NodeMeshTransform::Private& node = *TransformNode->GetPrivate();
 
-        Ptr<ASTOpMeshTransform> op = new ASTOpMeshTransform();
+		if (!node.Source)
+		{
+			// This argument is required
+			ErrorLog->Add("Mesh transform base node is not set.", ELMT_ERROR, TransformNode->GetMessageContext());
+			return UE::Tasks::MakeCompletedTask<FMeshGenerationResult>();
+		}
 
-        // Base
-        if (node.Source)
-        {
-            GenerateMesh(InOptions, OutResult, node.Source);
-            op->source = OutResult.MeshOp;
-        }
-        else
-        {
-            // This argument is required
-            ErrorLog->Add("Mesh transform base node is not set.", ELMT_ERROR, TransformNode->GetMessageContext() );
-        }
+		FMeshTask SourceTask = GenerateMesh(StaticOptions, InOptions, node.Source);
 
-        op->matrix = node.Transform;
+		Ptr<const NodeMeshTransform> NodeCopy = TransformNode;
+		return UE::Tasks::Launch(TEXT("MutableMeshTransform"),
+			[NodeCopy, SourceTask, CompilerOptions = CompilerOptions]() mutable
+			{
+				FMeshGenerationResult SourceResult = SourceTask.GetResult();
 
-        OutResult.MeshOp = op;
-    }
+				Ptr<ASTOpMeshTransform> TransformOp = new ASTOpMeshTransform();
+				TransformOp->source = SourceResult.MeshOp;
+				TransformOp->matrix = NodeCopy->GetPrivate()->Transform;
 
+				FMeshGenerationResult Result = SourceResult;
+				Result.MeshOp = TransformOp;
+				return Result;
+			},
+			SourceTask
+		);
+     }
 
-    //---------------------------------------------------------------------------------------------
-    void CodeGenerator::GenerateMesh_ClipMorphPlane(const FMeshGenerationOptions& InOptions, FMeshGenerationResult& OutResult,
-                                                     const NodeMeshClipMorphPlane* ClipNode )
+
+	CodeGenerator::FMeshTask CodeGenerator::GenerateMesh_ClipMorphPlane(const FMeshGenerationStaticOptions& StaticOptions, FMeshOptionsTask InOptions, const NodeMeshClipMorphPlane* ClipNode )
     {
-        Ptr<ASTOpMeshClipMorphPlane> op = new ASTOpMeshClipMorphPlane();
+		if (!ClipNode->Source)
+		{
+			// This argument is required
+			ErrorLog->Add("Mesh transform base node is not set.", ELMT_ERROR, ClipNode->GetMessageContext());
+			return UE::Tasks::MakeCompletedTask<FMeshGenerationResult>();
+		}
 
-		op->FaceCullStrategy = ClipNode->Parameters.FaceCullStrategy;
+		FMeshTask SourceTask = GenerateMesh(StaticOptions, InOptions, ClipNode->Source);
 
-        // Base
-        if (ClipNode->Source)
-        {
-			FMeshGenerationOptions BaseOptions = InOptions;
-            GenerateMesh(BaseOptions, OutResult, ClipNode->Source);
-            op->Source = OutResult.MeshOp;
-        }
-        else
-        {
-            // This argument is required
-            ErrorLog->Add("Mesh clip-morph-plane source node is not set.", ELMT_ERROR, ClipNode->GetMessageContext());
-        }
+		return UE::Tasks::Launch(TEXT("MutableMeshClipMorphPlane"),
+			[ClipNode, SourceTask, CompilerOptions = CompilerOptions]() mutable
+			{
+				FMeshGenerationResult SourceResult = SourceTask.GetResult();
 
-        // Morph to an ellipse
-        {
-            op->MorphShape.type = (uint8_t)FShape::Type::Ellipse;
-            op->MorphShape.position = ClipNode->Parameters.Origin;
-            op->MorphShape.up = ClipNode->Parameters.Normal;
-            op->MorphShape.size = FVector3f(ClipNode->Parameters.Radius1, ClipNode->Parameters.Radius2, ClipNode->Parameters.Rotation); // TODO: Move rotation to ellipse rotation reference base instead of passing it directly
+				Ptr<ASTOpMeshClipMorphPlane> op = new ASTOpMeshClipMorphPlane();
 
-            // Generate a "side" vector.
-            // \todo: make generic and move to the vector class
-            {
-                // Generate vector perpendicular to normal for ellipse rotation reference base
-				FVector3f aux_base(0.f, 1.f, 0.f);
+				op->FaceCullStrategy = ClipNode->Parameters.FaceCullStrategy;
+				op->Source = SourceResult.MeshOp;
 
-                if (fabs(FVector3f::DotProduct(ClipNode->Parameters.Normal, aux_base)) > 0.95f)
-                {
-                    aux_base = FVector3f(0.f, 0.f, 1.f);
-                }
+				// Morph to an ellipse
+				{
+					op->MorphShape.type = uint8(FShape::Type::Ellipse);
+					op->MorphShape.position = ClipNode->Parameters.Origin;
+					op->MorphShape.up = ClipNode->Parameters.Normal;
+					op->MorphShape.size = FVector3f(ClipNode->Parameters.Radius1, ClipNode->Parameters.Radius2, ClipNode->Parameters.Rotation); // TODO: Move rotation to ellipse rotation reference base instead of passing it directly
 
-                op->MorphShape.side = FVector3f::CrossProduct(ClipNode->Parameters.Normal, aux_base);
-            }
-        }
+					// Generate a "side" vector.
+					// \todo: make generic and move to the vector class
+					{
+						// Generate vector perpendicular to normal for ellipse rotation reference base
+						FVector3f aux_base(0.f, 1.f, 0.f);
 
-        // Selection by shape
-		op->VertexSelectionType = ClipNode->Parameters.VertexSelectionType;
-        if (op->VertexSelectionType == EClipVertexSelectionType::Shape)
-        {
-            op->SelectionShape.type = (uint8_t)FShape::Type::AABox;
-            op->SelectionShape.position = ClipNode->Parameters.SelectionBoxOrigin;
-            op->SelectionShape.size = ClipNode->Parameters.SelectionBoxRadius;
-        }
-        else if (op->VertexSelectionType == EClipVertexSelectionType::BoneHierarchy)
-        {
-            // Selection by bone hierarchy?
-            op->VertexSelectionBone = ClipNode->Parameters.VertexSelectionBone;
-			op->VertexSelectionBoneMaxRadius = ClipNode->Parameters.MaxEffectRadius;
-        }
- 
-        // FParameters
-        op->Dist = ClipNode->Parameters.DistanceToPlane;
-        op->Factor = ClipNode->Parameters.LinearityFactor;
+						if (fabs(FVector3f::DotProduct(ClipNode->Parameters.Normal, aux_base)) > 0.95f)
+						{
+							aux_base = FVector3f(0.f, 0.f, 1.f);
+						}
+
+						op->MorphShape.side = FVector3f::CrossProduct(ClipNode->Parameters.Normal, aux_base);
+					}
+				}
 
-        OutResult.MeshOp = op;
+				// Selection by shape
+				op->VertexSelectionType = ClipNode->Parameters.VertexSelectionType;
+				if (op->VertexSelectionType == EClipVertexSelectionType::Shape)
+				{
+					op->SelectionShape.type = uint8(FShape::Type::AABox);
+					op->SelectionShape.position = ClipNode->Parameters.SelectionBoxOrigin;
+					op->SelectionShape.size = ClipNode->Parameters.SelectionBoxRadius;
+				}
+				else if (op->VertexSelectionType == EClipVertexSelectionType::BoneHierarchy)
+				{
+					// Selection by bone hierarchy?
+					op->VertexSelectionBone = ClipNode->Parameters.VertexSelectionBone;
+					op->VertexSelectionBoneMaxRadius = ClipNode->Parameters.MaxEffectRadius;
+				}
+
+				// FParameters
+				op->Dist = ClipNode->Parameters.DistanceToPlane;
+				op->Factor = ClipNode->Parameters.LinearityFactor;
+
+				FMeshGenerationResult Result = SourceResult;
+				Result.MeshOp = op;
+
+				return Result;
+			},
+			SourceTask
+		);
     }
 
 
-    //---------------------------------------------------------------------------------------------
-    void CodeGenerator::GenerateMesh_ClipWithMesh(const FMeshGenerationOptions& InOptions, FMeshGenerationResult& OutResult,
-                                                   const NodeMeshClipWithMesh* ClipNode)
+	CodeGenerator::FMeshTask CodeGenerator::GenerateMesh_ClipWithMesh(const FMeshGenerationStaticOptions& StaticOptions, FMeshOptionsTask InOptions, const NodeMeshClipWithMesh* ClipNode)
     {
-        Ptr<ASTOpMeshClipWithMesh> Op = new ASTOpMeshClipWithMesh();
+		if (!ClipNode->Source)
+		{
+			// This argument is required
+			ErrorLog->Add("Mesh clip-with-mesh source node is not set.", ELMT_ERROR, ClipNode->GetMessageContext());
+			return UE::Tasks::MakeCompletedTask<FMeshGenerationResult>();
+		}
 
-        // Base
-        if (ClipNode->Source)
-        {
-            GenerateMesh(InOptions, OutResult, ClipNode->Source );
-			Op->Source = OutResult.MeshOp;
-        }
-        else
-        {
-            // This argument is required
-            ErrorLog->Add("Mesh clip-with-mesh source node is not set.", ELMT_ERROR, ClipNode->GetMessageContext());
-        }
+		FMeshTask SourceTask = GenerateMesh(StaticOptions, InOptions, ClipNode->Source);
 
-        // Clipping mesh
-        if (ClipNode->ClipMesh)
-        {
-			FMeshGenerationOptions ClipOptions = InOptions;
-			ClipOptions.bLayouts = false;
-			ClipOptions.OverrideLayouts.Empty();
-			ClipOptions.ActiveTags.Empty();
-
-            FMeshGenerationResult clipResult;
-            GenerateMesh(ClipOptions, clipResult, ClipNode->ClipMesh);
-            Op->ClipMesh = clipResult.MeshOp;
+		if (!ClipNode->ClipMesh)
+		{
+			// This argument is required
+			ErrorLog->Add("Mesh clip-with-mesh clipping mesh node is not set.", ELMT_ERROR, ClipNode->GetMessageContext());
+			return SourceTask;
 		}
-        else
-        {
-            // This argument is required
-            ErrorLog->Add("Mesh clip-with-mesh clipping mesh node is not set.", ELMT_ERROR, ClipNode->GetMessageContext());
-        }
 
-        OutResult.MeshOp = Op;
+		FMeshGenerationStaticOptions StaticClipOptions = StaticOptions;
+		StaticClipOptions.ActiveTags.Empty();
+		FMeshOptionsTask ClipOptionsTask = UE::Tasks::Launch(TEXT("MutableMeshClipWithMeshClipOptions"),
+			[InOptions]() mutable
+			{
+				FMeshGenerationDynamicOptions Result = InOptions.GetResult();
+				Result.bLayouts = false;
+				Result.OverrideLayouts.Empty();
+				return Result;
+			},
+			InOptions
+		);
+		FMeshTask ClipTask = GenerateMesh(StaticClipOptions, ClipOptionsTask, ClipNode->ClipMesh);
+
+		return UE::Tasks::Launch(TEXT("MutableMeshClipWithMesh"),
+			[ClipNode, SourceTask, ClipTask]() mutable
+			{
+				FMeshGenerationResult SourceResult = SourceTask.GetResult();
+				FMeshGenerationResult ClipResult = ClipTask.GetResult();
+
+				Ptr<ASTOpMeshClipWithMesh> op = new ASTOpMeshClipWithMesh;
+
+				// Base
+				op->Source = SourceResult.MeshOp;
+				op->ClipMesh = ClipResult.MeshOp;
+
+				FMeshGenerationResult Result = SourceResult;
+				Result.MeshOp = op;
+
+				return Result;
+			},
+			UE::Tasks::Prerequisites( SourceTask, ClipTask )
+		);
     }
 
-	//---------------------------------------------------------------------------------------------
-	void CodeGenerator::GenerateMesh_ClipDeform(const FMeshGenerationOptions& InOptions, FMeshGenerationResult& Result, const NodeMeshClipDeform* ClipDeform)
-	{
-		const Ptr<ASTOpMeshBindShape> OpBind = new ASTOpMeshBindShape();
-		const Ptr<ASTOpMeshClipDeform> OpClipDeform = new ASTOpMeshClipDeform();
 
-		// Base Mesh
-		if (ClipDeform->BaseMesh)
-		{
-			GenerateMesh(InOptions, Result, ClipDeform->BaseMesh);
-			OpBind->Mesh = Result.MeshOp;
-		}
-		else
+	CodeGenerator::FMeshTask CodeGenerator::GenerateMesh_ClipDeform(const FMeshGenerationStaticOptions& StaticOptions, FMeshOptionsTask InOptions, const NodeMeshClipDeform* ClipDeform)
+	{
+		if (!ClipDeform->BaseMesh)
 		{
 			// This argument is required
 			ErrorLog->Add("Mesh Clip Deform base mesh node is not set.", ELMT_ERROR, ClipDeform->GetMessageContext());
+			return UE::Tasks::MakeCompletedTask<FMeshGenerationResult>();
 		}
 
-		// Base Shape
-		if (ClipDeform->ClipShape)
+		// Base Mesh
+		FMeshTask BaseTask = GenerateMesh(StaticOptions, InOptions, ClipDeform->BaseMesh);
+
+		if (!ClipDeform->ClipShape)
 		{
-			FMeshGenerationOptions ClipOptions = InOptions;
-			ClipOptions.bLayouts = false;
-			ClipOptions.OverrideLayouts.Empty();
-			ClipOptions.ActiveTags.Empty();
-
-			FMeshGenerationResult baseResult;
-			GenerateMesh(ClipOptions, baseResult, ClipDeform->ClipShape);
-			OpBind->Shape = baseResult.MeshOp;
-			OpClipDeform->ClipShape = baseResult.MeshOp;
+			return BaseTask;
 		}
 
-		OpClipDeform->Mesh = OpBind;
-
-		Result.MeshOp = OpClipDeform;
+		FMeshGenerationStaticOptions StaticClipOptions = StaticOptions;
+		StaticClipOptions.ActiveTags.Empty();
+		FMeshOptionsTask ClipOptionsTask = UE::Tasks::Launch(TEXT("MutableMeshClipDeformOptions"),
+			[InOptions]() mutable
+			{
+				FMeshGenerationDynamicOptions Result = InOptions.GetResult();
+				Result.bLayouts = false;
+				Result.OverrideLayouts.Empty();
+				return Result;
+			},
+			InOptions
+		);
+		FMeshTask ShapeTask = GenerateMesh(StaticClipOptions, ClipOptionsTask, ClipDeform->ClipShape);
+
+		return UE::Tasks::Launch(TEXT("MutableMeshClipDeform"),
+			[BaseTask, ShapeTask]() mutable
+			{
+				FMeshGenerationResult BaseResult = BaseTask.GetResult();
+				FMeshGenerationResult ShapeResult = ShapeTask.GetResult();
+
+				const Ptr<ASTOpMeshBindShape> OpBind = new ASTOpMeshBindShape();
+				const Ptr<ASTOpMeshClipDeform> OpClipDeform = new ASTOpMeshClipDeform();
+				OpBind->Mesh = BaseResult.MeshOp;
+				OpBind->Shape = ShapeResult.MeshOp;
+				OpClipDeform->ClipShape = ShapeResult.MeshOp;
+				OpClipDeform->Mesh = OpBind;
+
+				FMeshGenerationResult Result = BaseResult;
+				Result.MeshOp = OpClipDeform;
+
+				return Result;
+			},
+			UE::Tasks::Prerequisites(BaseTask, ShapeTask)
+		);
 	}
 
-    //---------------------------------------------------------------------------------------------
-    void CodeGenerator::GenerateMesh_ApplyPose(const FMeshGenerationOptions& InOptions, FMeshGenerationResult& OutResult,
-                                                const NodeMeshApplyPose* PoseNode )
+	CodeGenerator::FMeshTask CodeGenerator::GenerateMesh_ApplyPose(const FMeshGenerationStaticOptions& StaticOptions, FMeshOptionsTask InOptions, const NodeMeshApplyPose* PoseNode )
     {
-        const auto& node = *PoseNode->GetPrivate();
+        const NodeMeshApplyPose::Private& node = *PoseNode->GetPrivate();
 
-        Ptr<ASTOpMeshApplyPose> op = new ASTOpMeshApplyPose();
+		if (!node.m_pBase)
+		{
+			// This argument is required
+			ErrorLog->Add("Mesh apply-pose base node is not set.", ELMT_ERROR, PoseNode->GetMessageContext());
+			return UE::Tasks::MakeCompletedTask<FMeshGenerationResult>();
+		}
 
-        // Base
-        if (node.m_pBase)
-        {
-            GenerateMesh(InOptions, OutResult, node.m_pBase );
-            op->Base = OutResult.MeshOp;
-        }
-        else
-        {
-            // This argument is required
-            ErrorLog->Add("Mesh apply-pose base node is not set.", ELMT_ERROR, PoseNode->GetMessageContext());
-        }
+		// Base Mesh
+		FMeshTask BaseTask = GenerateMesh(StaticOptions, InOptions, node.m_pBase);
 
-        // Pose mesh
-        if (node.m_pPose)
-        {
-			FMeshGenerationOptions PoseOptions = InOptions;
-			PoseOptions.bLayouts = false;
-			PoseOptions.OverrideLayouts.Empty();
-			PoseOptions.ActiveTags.Empty();
-
-            FMeshGenerationResult PoseResult;
-            GenerateMesh(PoseOptions, PoseResult, node.m_pPose );
-            op->Pose = PoseResult.MeshOp;
+		if (!node.m_pPose)
+		{
+			ErrorLog->Add("Mesh apply-pose pose node is not set.", ELMT_ERROR, PoseNode->GetMessageContext());
+			return BaseTask;
 		}
-        else
-        {
-            // This argument is required
-            ErrorLog->Add("Mesh apply-pose pose node is not set.", ELMT_ERROR, PoseNode->GetMessageContext());
-        }
 
-        OutResult.MeshOp = op;
+		FMeshGenerationStaticOptions StaticPoseOptions = StaticOptions;
+		StaticPoseOptions.ActiveTags.Empty();
+		FMeshOptionsTask PoseOptionsTask = UE::Tasks::Launch(TEXT("MutableMeshApplyPoseOptions"),
+			[InOptions]() mutable
+			{
+				FMeshGenerationDynamicOptions Result = InOptions.GetResult();
+				Result.bLayouts = false;
+				Result.OverrideLayouts.Empty();
+				return Result;
+			},
+			InOptions
+		);
+		FMeshTask PoseTask = GenerateMesh(StaticPoseOptions, PoseOptionsTask, node.m_pPose);
+
+		return UE::Tasks::Launch(TEXT("MutableMeshApplyPose"),
+			[BaseTask, PoseTask]() mutable
+			{
+				FMeshGenerationResult BaseResult = BaseTask.GetResult();
+				FMeshGenerationResult PoseResult = PoseTask.GetResult();
+
+				Ptr<ASTOpMeshApplyPose> op = new ASTOpMeshApplyPose();
+				op->Base = BaseResult.MeshOp;
+				op->Pose = PoseResult.MeshOp;
+
+				FMeshGenerationResult Result = BaseResult;
+				Result.MeshOp = op;
+
+				return Result;
+			},
+			UE::Tasks::Prerequisites(BaseTask, PoseTask)
+		);
     }
 
 
-	//---------------------------------------------------------------------------------------------
-	void CodeGenerator::GenerateMesh_Reshape(const FMeshGenerationOptions& InOptions, FMeshGenerationResult& OutResult, const NodeMeshReshape* Reshape)
+	CodeGenerator::FMeshTask CodeGenerator::GenerateMesh_Reshape(const FMeshGenerationStaticOptions& StaticOptions, FMeshOptionsTask InOptions, const NodeMeshReshape* Reshape)
 	{
 		const NodeMeshReshape::Private& Node = *Reshape->GetPrivate();
 
-		Ptr<ASTOpMeshBindShape> OpBind = new ASTOpMeshBindShape();
-		Ptr<ASTOpMeshApplyShape> OpApply = new ASTOpMeshApplyShape();
-
-		OpBind->bReshapeSkeleton = Node.bReshapeSkeleton;	
-		OpBind->BonesToDeform = Node.BonesToDeform;
-    	OpBind->bReshapePhysicsVolumes = Node.bReshapePhysicsVolumes;
-		OpBind->PhysicsToDeform = Node.PhysicsToDeform;
-		OpBind->bReshapeVertices = Node.bReshapeVertices;
-		OpBind->bRecomputeNormals = Node.bRecomputeNormals;	
-		OpBind->bApplyLaplacian = Node.bApplyLaplacian;
-		OpBind->BindingMethod = static_cast<uint32>(EShapeBindingMethod::ReshapeClosestProject);
-
-		OpBind->RChannelUsage = Node.ColorRChannelUsage;
-		OpBind->GChannelUsage = Node.ColorGChannelUsage;
-		OpBind->BChannelUsage = Node.ColorBChannelUsage;
-		OpBind->AChannelUsage = Node.ColorAChannelUsage;
-
-		OpApply->bReshapeVertices = OpBind->bReshapeVertices;
-		OpApply->bRecomputeNormals = OpBind->bRecomputeNormals;
-		OpApply->bReshapeSkeleton = OpBind->bReshapeSkeleton;
-		OpApply->bApplyLaplacian = OpBind->bApplyLaplacian;
-		OpApply->bReshapePhysicsVolumes = OpBind->bReshapePhysicsVolumes;
-
-		// Base Mesh
-		if (Node.BaseMesh)
-		{
-			GenerateMesh(InOptions, OutResult, Node.BaseMesh);
-			OpBind->Mesh = OutResult.MeshOp;
-		}
-		else
+		if (!Node.BaseMesh)
 		{
 			// This argument is required
 			ErrorLog->Add("Mesh reshape base node is not set.", ELMT_ERROR, Reshape->GetMessageContext());
+			return UE::Tasks::MakeCompletedTask<FMeshGenerationResult>();
 		}
 
-		// Base and target shapes shouldn't have layouts or modifiers.
-		FMeshGenerationOptions ShapeOptions = InOptions;
-		ShapeOptions.bLayouts = false;
-		ShapeOptions.OverrideLayouts.Empty();
-		ShapeOptions.ActiveTags.Empty();
+		// Base Mesh
+		FMeshTask BaseTask = GenerateMesh(StaticOptions, InOptions, Node.BaseMesh);
 
-		// Base Shape
-		if (Node.BaseShape)
+		if (!Node.BaseShape || !Node.TargetShape)
 		{
-			FMeshGenerationResult baseResult;
-			GenerateMesh(ShapeOptions, baseResult, Node.BaseShape);
-			OpBind->Shape = baseResult.MeshOp;
+			return BaseTask;
 		}
 
-		OpApply->Mesh = OpBind;
-
-		// Target Shape
-		if (Node.TargetShape)
-		{
-			FMeshGenerationResult targetResult;
-			GenerateMesh(ShapeOptions, targetResult, Node.TargetShape);
-			OpApply->Shape = targetResult.MeshOp;
-		}
 
-		OutResult.MeshOp = OpApply;
+		FMeshGenerationStaticOptions StaticShapeOptions = StaticOptions;
+		StaticShapeOptions.ActiveTags.Empty();
+		FMeshOptionsTask ShapeOptionsTask = UE::Tasks::Launch(TEXT("MutableMeshApplyPoseOptions"),
+			[InOptions]() mutable
+			{
+				FMeshGenerationDynamicOptions Result = InOptions.GetResult();
+				Result.bLayouts = false;
+				Result.OverrideLayouts.Empty();
+				return Result;
+			},
+			InOptions
+		);
+		FMeshTask BaseShapeTask = GenerateMesh(StaticShapeOptions, ShapeOptionsTask, Node.BaseShape);
+		FMeshTask TargetShapeTask = GenerateMesh(StaticShapeOptions, ShapeOptionsTask, Node.TargetShape);
+
+		return UE::Tasks::Launch(TEXT("MutableMeshApplyPose"),
+			[BaseTask, BaseShapeTask, TargetShapeTask, &Node]() mutable
+			{
+				FMeshGenerationResult BaseResult = BaseTask.GetResult();
+				FMeshGenerationResult BaseShapeResult = BaseShapeTask.GetResult();
+				FMeshGenerationResult TargetShapeResult = TargetShapeTask.GetResult();
+
+				Ptr<ASTOpMeshBindShape> OpBind = new ASTOpMeshBindShape();
+				Ptr<ASTOpMeshApplyShape> OpApply = new ASTOpMeshApplyShape();
+
+				OpBind->bReshapeSkeleton = Node.bReshapeSkeleton;
+				OpBind->BonesToDeform = Node.BonesToDeform;
+				OpBind->bReshapePhysicsVolumes = Node.bReshapePhysicsVolumes;
+				OpBind->PhysicsToDeform = Node.PhysicsToDeform;
+				OpBind->bReshapeVertices = Node.bReshapeVertices;
+				OpBind->bRecomputeNormals = Node.bRecomputeNormals;
+				OpBind->bApplyLaplacian = Node.bApplyLaplacian;
+				OpBind->BindingMethod = static_cast<uint32>(EShapeBindingMethod::ReshapeClosestProject);
+
+				OpBind->RChannelUsage = Node.ColorRChannelUsage;
+				OpBind->GChannelUsage = Node.ColorGChannelUsage;
+				OpBind->BChannelUsage = Node.ColorBChannelUsage;
+				OpBind->AChannelUsage = Node.ColorAChannelUsage;
+
+				OpApply->bReshapeVertices = OpBind->bReshapeVertices;
+				OpApply->bRecomputeNormals = OpBind->bRecomputeNormals;
+				OpApply->bReshapeSkeleton = OpBind->bReshapeSkeleton;
+				OpApply->bApplyLaplacian = OpBind->bApplyLaplacian;
+				OpApply->bReshapePhysicsVolumes = OpBind->bReshapePhysicsVolumes;
+
+				// Base Mesh
+				OpBind->Mesh = BaseResult.MeshOp;
+				OpBind->Shape = BaseShapeResult.MeshOp;
+
+				OpApply->Mesh = OpBind;
+				OpApply->Shape = TargetShapeResult.MeshOp;
+
+				FMeshGenerationResult Result = BaseResult;
+				Result.MeshOp = OpApply;
+
+				return Result;
+			},
+			UE::Tasks::Prerequisites(BaseTask, BaseShapeTask, TargetShapeTask)
+		);
 	}
 
 }
diff --git a/Engine/Plugins/Mutable/Source/MutableTools/Private/MuT/CodeOptimiser.cpp b/Engine/Plugins/Mutable/Source/MutableTools/Private/MuT/CodeOptimiser.cpp
--- a/Engine/Plugins/Mutable/Source/MutableTools/Private/MuT/CodeOptimiser.cpp
+++ b/Engine/Plugins/Mutable/Source/MutableTools/Private/MuT/CodeOptimiser.cpp
@@ -798,6 +798,7 @@ namespace mu
 		};
 
 		bool bUseConcurrency = InOptions->bUseConcurrency;
+		//bUseConcurrency = false;
 
 		if (bUseConcurrency)
 		{
@@ -818,89 +819,125 @@ namespace mu
 						Ptr<ASTOp> SubgraphRoot = ConstantSubgraphs[Index].Root;
 						UE::Tasks::FTaskEvent& SubgraphCompletionEvent = ConstantSubgraphs[Index].CompletedEvent;
 
-						// Referenced images and meshes are resolved in its own task to prevent requesting them twice (and loading them twice)
-						bool bIsCompileTimeReferenceImage = (SubgraphRoot->GetOpType() == EOpType::IM_REFERENCE);
+						bool bIsReference = false;
+						EOpType SubgraphType = SubgraphRoot->GetOpType();
 
-						// Launch the task with its dependencies
-						switch (SubgraphRoot->GetOpType())
+						if (SubgraphType == EOpType::IM_REFERENCE
+							||
+							SubgraphType == EOpType::IM_CONSTANT)
 						{
-						case EOpType::IM_REFERENCE:
-						{
-							// Instead of generating the constant we resolve the reference, which also replaces the ASTOp.
-							const ASTOpReferenceResource* Typed = static_cast<const ASTOpReferenceResource*>(SubgraphRoot.get());
-							uint32 ImageID = Typed->ID;
+							uint32 ImageID = 0;
+							if (SubgraphType == EOpType::IM_REFERENCE)							
+							{
+								bIsReference = true;
+								const ASTOpReferenceResource* Typed = static_cast<const ASTOpReferenceResource*>(SubgraphRoot.get());
+								ImageID = Typed->ID;
+							}
+							else if(SubgraphType == EOpType::IM_CONSTANT)
+							{
+								const ASTOpConstantResource* Typed = static_cast<const ASTOpConstantResource*>(SubgraphRoot.get());
+								const FImage* Value = static_cast<const FImage*>( Typed->GetValue().Get() );
+								bIsReference = Value->IsReference();
+								if (bIsReference)
+								{
+									ImageID = Value->GetReferencedTexture();
+								}
+							}
 
-							TSharedPtr< TSharedPtr<FImage> > ResolveImage = MakeShared<TSharedPtr<FImage>>();
+							// Instead of generating the constant we resolve the reference, which also replaces the ASTOp.
+							if (bIsReference)
+							{
+								TSharedPtr< TSharedPtr<FImage> > ResolveImage = MakeShared<TSharedPtr<FImage>>();
 
-							constexpr bool bRunImmediatlyIfPossible = false;
-							UE::Tasks::FTask ReferenceCompletion = InOptions->OptimisationOptions.ReferencedImageResourceProvider(ImageID, ResolveImage, bRunImmediatlyIfPossible);
+								constexpr bool bRunImmediatlyIfPossible = false;
+								UE::Tasks::FTask ReferenceCompletion = InOptions->OptimisationOptions.ReferencedImageResourceProvider(ImageID, ResolveImage, bRunImmediatlyIfPossible);
 
-							UE::Tasks::FTask CompleteTask = UE::Tasks::Launch(TEXT("MutableResolveComplete"),
-								[SubgraphRoot, InOptions, ResolveImage]()
-								{
-									Ptr<ASTOpConstantResource> ConstantOp;
+								UE::Tasks::FTask CompleteTask = UE::Tasks::Launch(TEXT("MutableResolveComplete"),
+									[SubgraphRoot, InOptions, ResolveImage]()
 									{
-										MUTABLE_CPUPROFILER_SCOPE(MutableResolveComplete_CreateConstant);
-										ConstantOp = new ASTOpConstantResource;
-										ConstantOp->Type = EOpType::IM_CONSTANT;
+										Ptr<ASTOpConstantResource> ConstantOp;
 										{
-											MUTABLE_CPUPROFILER_SCOPE(GetSourceDataDescriptor);
-											ConstantOp->SourceDataDescriptor = SubgraphRoot->GetSourceDataDescriptor();
+											MUTABLE_CPUPROFILER_SCOPE(MutableResolveComplete_CreateConstant);
+											ConstantOp = new ASTOpConstantResource;
+											ConstantOp->Type = EOpType::IM_CONSTANT;
+											{
+												MUTABLE_CPUPROFILER_SCOPE(GetSourceDataDescriptor);
+												ConstantOp->SourceDataDescriptor = SubgraphRoot->GetSourceDataDescriptor();
+											}
+											ConstantOp->SetValue(*ResolveImage, InOptions->OptimisationOptions.DiskCacheContext);
 										}
-										ConstantOp->SetValue(*ResolveImage, InOptions->OptimisationOptions.DiskCacheContext);
-									}
-									{
-										MUTABLE_CPUPROFILER_SCOPE(MutableResolveComplete_Replace);
-										ASTOp::Replace(SubgraphRoot, ConstantOp);
-									}
-								},
-								ReferenceCompletion,
-								LowLevelTasks::ETaskPriority::BackgroundNormal);
+										{
+											MUTABLE_CPUPROFILER_SCOPE(MutableResolveComplete_Replace);
+											ASTOp::Replace(SubgraphRoot, ConstantOp);
+										}
+									},
+									ReferenceCompletion,
+									LowLevelTasks::ETaskPriority::BackgroundNormal);
 
-							SubgraphCompletionEvent.AddPrerequisites(CompleteTask);
+								SubgraphCompletionEvent.AddPrerequisites(CompleteTask);
+							}
 
-							break;
 						}
 
-						case EOpType::ME_REFERENCE:
+						if (SubgraphType == EOpType::ME_REFERENCE
+							||
+							SubgraphType == EOpType::ME_CONSTANT)
 						{
-							// Instead of generating the constant we resolve the reference, which also replaces the ASTOp.
-							const ASTOpReferenceResource* Typed = static_cast<const ASTOpReferenceResource*>(SubgraphRoot.get());
-							uint32 MeshID = Typed->ID;
+							uint32 MeshID = 0;
+							FString MeshMorph;
+							if (SubgraphType == EOpType::ME_REFERENCE)
+							{
+								bIsReference = true;
+								const ASTOpReferenceResource* Typed = static_cast<const ASTOpReferenceResource*>(SubgraphRoot.get());
+								MeshID = Typed->ID;
+							}
+							else if (SubgraphType == EOpType::ME_CONSTANT)
+							{
+								const ASTOpConstantResource* Typed = static_cast<const ASTOpConstantResource*>(SubgraphRoot.get());
+								const FMesh* Value = static_cast<const FMesh*>(Typed->GetValue().Get());
+								bIsReference = Value->IsReference();
+								if (bIsReference)
+								{
+									MeshID = Value->GetReferencedMesh();
+									MeshMorph = Value->GetReferencedMorph();
+								}
+							}
 
-							TSharedPtr< TSharedPtr<FMesh> > ResolveMesh = MakeShared<TSharedPtr<FMesh>>();
+							// Instead of generating the constant we resolve the reference, which also replaces the ASTOp.
+							if (bIsReference)
+							{
+								TSharedPtr< TSharedPtr<FMesh> > ResolveMesh = MakeShared<TSharedPtr<FMesh>>();
 
-							constexpr bool bRunImmediatlyIfPossible = false;
-							UE::Tasks::FTask ReferenceCompletion = InOptions->OptimisationOptions.ReferencedMeshResourceProvider(MeshID, ResolveMesh, bRunImmediatlyIfPossible);
+								constexpr bool bRunImmediatlyIfPossible = false;
+								UE::Tasks::FTask ReferenceCompletion = InOptions->OptimisationOptions.ReferencedMeshResourceProvider(MeshID, MeshMorph, ResolveMesh, bRunImmediatlyIfPossible);
 
-							UE::Tasks::FTask CompleteTask = UE::Tasks::Launch(TEXT("MutableResolveComplete"),
-								[SubgraphRoot, InOptions, ResolveMesh]()
-								{
-									Ptr<ASTOpConstantResource> ConstantOp;
+								UE::Tasks::FTask CompleteTask = UE::Tasks::Launch(TEXT("MutableResolveComplete"),
+									[SubgraphRoot, InOptions, ResolveMesh]()
 									{
-										MUTABLE_CPUPROFILER_SCOPE(MutableResolveComplete_CreateConstant);
-										ConstantOp = new ASTOpConstantResource;
-										ConstantOp->Type = EOpType::IM_CONSTANT;
+										Ptr<ASTOpConstantResource> ConstantOp;
 										{
-											MUTABLE_CPUPROFILER_SCOPE(GetSourceDataDescriptor);
-											ConstantOp->SourceDataDescriptor = SubgraphRoot->GetSourceDataDescriptor();
+											MUTABLE_CPUPROFILER_SCOPE(MutableResolveComplete_CreateConstant);
+											ConstantOp = new ASTOpConstantResource;
+											ConstantOp->Type = EOpType::ME_CONSTANT;
+											{
+												MUTABLE_CPUPROFILER_SCOPE(GetSourceDataDescriptor);
+												ConstantOp->SourceDataDescriptor = SubgraphRoot->GetSourceDataDescriptor();
+											}
+											ConstantOp->SetValue(*ResolveMesh, InOptions->OptimisationOptions.DiskCacheContext);
 										}
-										ConstantOp->SetValue(*ResolveMesh, InOptions->OptimisationOptions.DiskCacheContext);
-									}
-									{
-										MUTABLE_CPUPROFILER_SCOPE(MutableResolveComplete_Replace);
-										ASTOp::Replace(SubgraphRoot, ConstantOp);
-									}
-								},
-								ReferenceCompletion,
-								LowLevelTasks::ETaskPriority::BackgroundNormal);
-
-							SubgraphCompletionEvent.AddPrerequisites(CompleteTask);
+										{
+											MUTABLE_CPUPROFILER_SCOPE(MutableResolveComplete_Replace);
+											ASTOp::Replace(SubgraphRoot, ConstantOp);
+										}
+									},
+									ReferenceCompletion,
+									LowLevelTasks::ETaskPriority::BackgroundNormal);
 
-							break;
+								SubgraphCompletionEvent.AddPrerequisites(CompleteTask);
+							}
 						}
 
-						default:
+						if (!bIsReference)
 						{
 							// Scan for requisites
 							TArray< UE::Tasks::FTask, TInlineAllocator<8> > Requisites;
@@ -928,8 +965,6 @@ namespace mu
 							SubgraphCompletionEvent.AddPrerequisites(CompleteTask);
 						}
 
-						}
-
 						ConstantSubgraphs[Index].Root = nullptr;
 						SubgraphCompletionEvent.Trigger();
 						
@@ -962,32 +997,95 @@ namespace mu
 					Ptr<ASTOp> SubgraphRoot = ConstantSubgraphs[Index].Root;
 					UE::Tasks::FTaskEvent& SubgraphCompletionEvent = ConstantSubgraphs[Index].CompletedEvent;
 
-					// Referenced images are resolved in its own task to prevent requesting them twice (and loading them twice)
-					bool bIsCompileTimeReferenceImage = (SubgraphRoot->GetOpType() == EOpType::IM_REFERENCE);
+					bool bIsReference = false;
+					EOpType SubgraphType = SubgraphRoot->GetOpType();
 
-					// Launch the task with its dependencies
-					if (bIsCompileTimeReferenceImage)
+					if (SubgraphType == EOpType::IM_REFERENCE
+						||
+						SubgraphType == EOpType::IM_CONSTANT)
 					{
+						uint32 ImageID = 0;
+						if (SubgraphType == EOpType::IM_REFERENCE)
+						{
+							bIsReference = true;
+							const ASTOpReferenceResource* Typed = static_cast<const ASTOpReferenceResource*>(SubgraphRoot.get());
+							ImageID = Typed->ID;
+						}
+						else if (SubgraphType == EOpType::IM_CONSTANT)
+						{
+							const ASTOpConstantResource* Typed = static_cast<const ASTOpConstantResource*>(SubgraphRoot.get());
+							const FImage* Value = static_cast<const FImage*>(Typed->GetValue().Get());
+							bIsReference = Value->IsReference();
+							if (bIsReference)
+							{
+								ImageID = Value->GetReferencedTexture();
+							}
+						}
+
 						// Instead of generating the constant we resolve the reference, which also replaces the ASTOp.
-						const ASTOpReferenceResource* Typed = static_cast<const ASTOpReferenceResource*>(SubgraphRoot.get());
-						uint32 ImageID = Typed->ID;
+						if (bIsReference)
+						{
+
+							auto ImmediateCompleteFunc = [SubgraphRoot, InOptions, ImageID]()
+								{
+									TSharedPtr< TSharedPtr<FImage> > ResolveImage = MakeShared<TSharedPtr<FImage>>();
+									constexpr bool bRunImmediatlyIfPossible = true;
+									UE::Tasks::FTask ReferenceCompletion = InOptions->OptimisationOptions.ReferencedImageResourceProvider(ImageID, ResolveImage, bRunImmediatlyIfPossible);
+									ReferenceCompletion.Wait();
+									Ptr<ASTOpConstantResource> ConstantOp = new ASTOpConstantResource;
+									ConstantOp->SourceDataDescriptor = SubgraphRoot->GetSourceDataDescriptor();
+									ConstantOp->Type = EOpType::IM_CONSTANT;
+									ConstantOp->SetValue(*ResolveImage, InOptions->OptimisationOptions.DiskCacheContext);
+									ASTOp::Replace(SubgraphRoot, ConstantOp);
+								};
+							ConstantSubgraphs[Index].NonConcurrentTask = ImmediateCompleteFunc;
+						}
+					}
 
-						auto ImmediateCompleteFunc = [SubgraphRoot, InOptions, ImageID]()
+					if (SubgraphType == EOpType::ME_REFERENCE
+						||
+						SubgraphType == EOpType::ME_CONSTANT)
+					{
+						uint32 MeshID = 0;
+						FString MeshMorph;
+						if (SubgraphType == EOpType::ME_REFERENCE)
+						{
+							bIsReference = true;
+							const ASTOpReferenceResource* Typed = static_cast<const ASTOpReferenceResource*>(SubgraphRoot.get());
+							MeshID = Typed->ID;
+						}
+						else if (SubgraphType == EOpType::ME_CONSTANT)
+						{
+							const ASTOpConstantResource* Typed = static_cast<const ASTOpConstantResource*>(SubgraphRoot.get());
+							const FMesh* Value = static_cast<const FMesh*>(Typed->GetValue().Get());
+							bIsReference = Value->IsReference();
+							if (bIsReference)
 							{
-								TSharedPtr< TSharedPtr<FImage> > ResolveImage = MakeShared<TSharedPtr<FImage>>();
-								constexpr bool bRunImmediatlyIfPossible = true;
-								UE::Tasks::FTask ReferenceCompletion = InOptions->OptimisationOptions.ReferencedImageResourceProvider(ImageID, ResolveImage, bRunImmediatlyIfPossible);
-								ReferenceCompletion.Wait();
-								Ptr<ASTOpConstantResource> ConstantOp = new ASTOpConstantResource;
-								ConstantOp->SourceDataDescriptor = SubgraphRoot->GetSourceDataDescriptor();
-								ConstantOp->Type = EOpType::IM_CONSTANT;
-								ConstantOp->SetValue(*ResolveImage, InOptions->OptimisationOptions.DiskCacheContext);
-								ASTOp::Replace(SubgraphRoot, ConstantOp);
-							};
-						ConstantSubgraphs[Index].NonConcurrentTask = ImmediateCompleteFunc;
+								MeshID = Value->GetReferencedMesh();
+								MeshMorph = Value->GetReferencedMorph();
+							}
+						}
+
+						// Instead of generating the constant we resolve the reference, which also replaces the ASTOp.
+						if (bIsReference)
+						{
+							auto ImmediateCompleteFunc = [SubgraphRoot, InOptions, MeshID, MeshMorph]()
+								{
+									TSharedPtr< TSharedPtr<FMesh> > ResolveMesh = MakeShared<TSharedPtr<FMesh>>();
+									constexpr bool bRunImmediatlyIfPossible = true;
+									UE::Tasks::FTask ReferenceCompletion = InOptions->OptimisationOptions.ReferencedMeshResourceProvider(MeshID, MeshMorph, ResolveMesh, bRunImmediatlyIfPossible);
+									ReferenceCompletion.Wait();
+									Ptr<ASTOpConstantResource> ConstantOp = new ASTOpConstantResource;
+									ConstantOp->SourceDataDescriptor = SubgraphRoot->GetSourceDataDescriptor();
+									ConstantOp->Type = EOpType::ME_CONSTANT;
+									ConstantOp->SetValue(*ResolveMesh, InOptions->OptimisationOptions.DiskCacheContext);
+									ASTOp::Replace(SubgraphRoot, ConstantOp);
+								};
+							ConstantSubgraphs[Index].NonConcurrentTask = ImmediateCompleteFunc;
+						}
 					}
 
-					else
+					if (!bIsReference)
 					{
 						// Scan for requisites
 						TArray< UE::Tasks::FTask, TInlineAllocator<8> > Requisites;
@@ -1005,7 +1103,7 @@ namespace mu
 								ASTOp::Replace(TaskPtr->Source, TaskPtr->Result);
 							};
 					}
-
+					
 					ConstantSubgraphs[Index].Root = nullptr;
 				}
 			}
diff --git a/Engine/Plugins/Mutable/Source/MutableTools/Private/MuT/NodeMeshConstant.cpp b/Engine/Plugins/Mutable/Source/MutableTools/Private/MuT/NodeMeshConstant.cpp
--- a/Engine/Plugins/Mutable/Source/MutableTools/Private/MuT/NodeMeshConstant.cpp
+++ b/Engine/Plugins/Mutable/Source/MutableTools/Private/MuT/NodeMeshConstant.cpp
@@ -36,25 +36,6 @@ namespace mu
     }
 
 
-	void NodeMeshConstant::AddMorph(const FString& Name, TSharedPtr<FMesh> Morphed)
-	{
-		m_pD->Morphs.Add({Name,Morphed});
-	}
-
-
-	TSharedPtr<FMesh> NodeMeshConstant::FindMorph(const FString& Name) const
-	{
-		for (Private::FMorph& Morph : m_pD->Morphs )
-		{
-			if (Morph.Name == Name)
-			{
-				return Morph.MorphedMesh;
-			}
-		}
-		return nullptr;
-	}
-
-
 	int32 NodeMeshConstant::GetLayoutCount() const
 	{
 		return m_pD->Layouts.Num();
diff --git a/Engine/Plugins/Mutable/Source/MutableTools/Public/MuT/Compiler.h b/Engine/Plugins/Mutable/Source/MutableTools/Public/MuT/Compiler.h
--- a/Engine/Plugins/Mutable/Source/MutableTools/Public/MuT/Compiler.h
+++ b/Engine/Plugins/Mutable/Source/MutableTools/Public/MuT/Compiler.h
@@ -24,8 +24,11 @@ namespace mu
     class Node;
     class NodeTransformedObject;
 
+	/** Arguments are Texture ID, output generated texture and a "run immediately" flag. */
 	typedef TFunction<UE::Tasks::FTask(int32, TSharedPtr<TSharedPtr<FImage>>, bool)> FReferencedImageResourceFunc;
-	typedef TFunction<UE::Tasks::FTask(int32, TSharedPtr<TSharedPtr<FMesh>>, bool)> FReferencedMeshResourceFunc;
+
+	/** Arguments are Mesh ID, optional morph name, output generated mesh and a "run immediately" flag. */
+	typedef TFunction<UE::Tasks::FTask(int32, const FString&, TSharedPtr<TSharedPtr<FMesh>>, bool)> FReferencedMeshResourceFunc;
 
     //! \brief Options used to compile the models with a compiler.
     class MUTABLETOOLS_API CompilerOptions : public RefCounted
diff --git a/Engine/Plugins/Mutable/Source/MutableTools/Public/MuT/NodeMeshConstant.h b/Engine/Plugins/Mutable/Source/MutableTools/Public/MuT/NodeMeshConstant.h
--- a/Engine/Plugins/Mutable/Source/MutableTools/Public/MuT/NodeMeshConstant.h
+++ b/Engine/Plugins/Mutable/Source/MutableTools/Public/MuT/NodeMeshConstant.h
@@ -49,12 +49,6 @@ namespace mu
 		//! Set the constant mesh that will be returned.
 		void SetValue(TSharedPtr<FMesh>);
 
-		/** */
-		void AddMorph(const FString& Name, TSharedPtr<FMesh>);
-
-		/** */
-		TSharedPtr<FMesh> FindMorph(const FString& Name) const;
-
 		//! Get the number of layouts defined in this mesh.
 		int32 GetLayoutCount() const;
 		void SetLayoutCount( int32 );
