[mutable] Fix compilation warning due to shadowing.
#rmx

diff --git a/Engine/Plugins/Mutable/Source/MutableTools/Private/MuT/CodeGenerator_Mesh.cpp b/Engine/Plugins/Mutable/Source/MutableTools/Private/MuT/CodeGenerator_Mesh.cpp
--- a/Engine/Plugins/Mutable/Source/MutableTools/Private/MuT/CodeGenerator_Mesh.cpp
+++ b/Engine/Plugins/Mutable/Source/MutableTools/Private/MuT/CodeGenerator_Mesh.cpp
@@ -1207,10 +1207,10 @@ namespace mu
 		FMeshTask MeshProcessingTask = GenerateMeshConstantPipe.Launch( TEXT("MutableGenerateMeshConstant"),
 			[ReferenceMesh=Mesh,ResolveMesh,Result,ConstantOp, StaticOptions, InOptions, NodeCopy,
 			// This is read-only
-			CompilerOptions=this->CompilerOptions,
-			this,
+			//CompilerOptions=this->CompilerOptions,
+			this
 			// This is read-write so it is protected with a mutex.
-			&GenerateMeshConstantState = this->GenerateMeshConstantState
+			//&GenerateMeshConstantState = this->GenerateMeshConstantState
 			]() mutable
 			{
 				TSharedPtr<FMesh> Mesh = *ResolveMesh;
