[mutable] Exposed empty/optional operations in Mutable Code Viewer.
- Removed "0 OP" conditional guards in Operations.cpp "ForEachReference(...)" since it was required for the debugger to be able to show some branches that would be hidden otherwise.
- Removed unnecessary includes
- Updated string structure for the CodeViewer tree rows

#rnx

diff --git a/Engine/Plugins/Mutable/Source/CustomizableObjectEditor/Private/MuCOE/SMutableCodeViewer.cpp b/Engine/Plugins/Mutable/Source/CustomizableObjectEditor/Private/MuCOE/SMutableCodeViewer.cpp
--- a/Engine/Plugins/Mutable/Source/CustomizableObjectEditor/Private/MuCOE/SMutableCodeViewer.cpp
+++ b/Engine/Plugins/Mutable/Source/CustomizableObjectEditor/Private/MuCOE/SMutableCodeViewer.cpp
@@ -2,9 +2,6 @@
 
 #include "MuCOE/SMutableCodeViewer.h"
 
-#include "DesktopPlatformModule.h"
-#include "EditorDirectories.h"
-#include "IDesktopPlatform.h"
 #include "Framework/MultiBox/MultiBoxBuilder.h"
 #include "Framework/Views/TableViewMetadata.h"
 #include "Misc/Paths.h"
@@ -1367,12 +1364,9 @@ void SMutableCodeViewer::GenerateElementRecursive(const int32& InStateIndex, mu:
 	uint32 ChildIndex = 0;
 	auto AddOpFunc = [this, InParentAddress, &InProgram, &ChildIndex, &InStateIndex](mu::OP::ADDRESS ChildAddress, const FString& Caption)
 		{
-			if (ChildAddress)
 			{
 				const FItemCacheKey Key = { InParentAddress, ChildAddress, ChildIndex };
 				const TSharedPtr<FMutableCodeTreeElement>* CachedItem = ItemCache.Find(Key);
-
-				// Will this ever really hit?
 				check(!CachedItem);
 
 				if (!CachedItem)
@@ -1380,6 +1374,7 @@ void SMutableCodeViewer::GenerateElementRecursive(const int32& InStateIndex, mu:
 					// Locate the "original" tree element
 					const TSharedPtr<FMutableCodeTreeElement>* MainItemPtr = MainItemPerOp.Find(ChildAddress);
 
+					// Provide the color this element should be using for the displayed text
 					const FSlateColor LabelColor = ColorPerComputationalCost[StaticCast<uint8>(GetOperationTypeComputationalCost(InProgram.GetOpType(ChildAddress)))];
 
 					const TSharedPtr<FMutableCodeTreeElement> Item = MakeShareable(new FMutableCodeTreeElement(ItemCache.Num(),InStateIndex, MutableModel, ChildAddress, Caption, LabelColor, MainItemPtr));
@@ -1455,13 +1450,9 @@ void SMutableCodeViewer::GenerateElementRecursive(const int32& InStateIndex, mu:
 			mu::OP::ADDRESS At;
 			FMemory::Memcpy(&At, OpData, sizeof(mu::OP::ADDRESS));
 			OpData += sizeof(mu::OP::ADDRESS);
-
-			// This conditional is necessary to exactly match the generic op behaviour (see ForEachReference in Operations.cpp)
-			if (At)
-			{
-				FString Caption = FString::Printf(TEXT("case %d "), Condition);
-				AddOpFunc(At, Caption);
-			}
+			
+			FString Caption = FString::Printf(TEXT("case %d "), Condition);
+			AddOpFunc(At, Caption);
 		}
 
 		break;
@@ -1493,10 +1484,7 @@ void SMutableCodeViewer::GenerateElementRecursive(const int32& InStateIndex, mu:
 	{
 		mu::OP::ImageLayerArgs Args = InProgram.GetOpArgs<mu::OP::ImageLayerArgs>(InParentAddress);
 		AddOpFunc(Args.base, TEXT("base "));
-		if (Args.mask)
-		{
-			AddOpFunc(Args.mask, TEXT("mask "));
-		}
+		AddOpFunc(Args.mask, TEXT("mask "));
 		AddOpFunc(Args.blended, TEXT("blended "));
 		break;
 	}
@@ -1505,10 +1493,7 @@ void SMutableCodeViewer::GenerateElementRecursive(const int32& InStateIndex, mu:
 	{
 		mu::OP::ImageLayerColourArgs Args = InProgram.GetOpArgs<mu::OP::ImageLayerColourArgs>(InParentAddress);
 		AddOpFunc(Args.base, TEXT("base "));
-		if (Args.mask)
-		{
-			AddOpFunc(Args.mask, TEXT("mask "));
-		}
+		AddOpFunc(Args.mask, TEXT("mask "));
 		AddOpFunc(Args.colour, TEXT("colour "));
 		break;
 	}
@@ -1518,10 +1503,7 @@ void SMutableCodeViewer::GenerateElementRecursive(const int32& InStateIndex, mu:
 		mu::OP::ImageMultiLayerArgs Args = InProgram.GetOpArgs<mu::OP::ImageMultiLayerArgs>(InParentAddress);
 		AddOpFunc(Args.rangeSize, TEXT("range "));
 		AddOpFunc(Args.base, TEXT("base "));
-		if (Args.mask)
-		{
-			AddOpFunc(Args.mask, TEXT("mask "));
-		}
+		AddOpFunc(Args.mask, TEXT("mask "));
 		AddOpFunc(Args.blended, TEXT("blended "));
 		break;
 	}
@@ -1543,19 +1525,210 @@ void SMutableCodeViewer::GenerateElementRecursive(const int32& InStateIndex, mu:
 
 		break;
 	}
+		
+	case mu::EOpType::ME_APPLYLAYOUT:
+	{
+		mu::OP::MeshApplyLayoutArgs Args = InProgram.GetOpArgs<mu::OP::MeshApplyLayoutArgs>(InParentAddress);
+		AddOpFunc(Args.Layout, TEXT("layout "));
+		AddOpFunc(Args.Mesh, TEXT("mesh "));
+		break;
+	}
+		
+	case mu::EOpType::ME_DIFFERENCE:
+	{
+		const uint8* data = InProgram.GetOpArgsPointer(InParentAddress);
+
+		mu::OP::ADDRESS BaseAt = 0;
+		FMemory::Memcpy(&BaseAt, data, sizeof(mu::OP::ADDRESS));
+		AddOpFunc(BaseAt, TEXT("base "));
+		data += sizeof(mu::OP::ADDRESS);
+
+		mu::OP::ADDRESS TargetAt = 0;
+		FMemory::Memcpy(&TargetAt, data, sizeof(mu::OP::ADDRESS));
+		AddOpFunc(TargetAt, TEXT("target "));
+		data += sizeof(mu::OP::ADDRESS);
+		
+		break;
+	}
+		
+	case mu::EOpType::ME_MORPH:
+	{
+		const uint8* Data = InProgram.GetOpArgsPointer(InParentAddress);
+
+		mu::OP::ADDRESS FactorAt = 0;
+		FMemory::Memcpy(&FactorAt, Data, sizeof(mu::OP::ADDRESS));
+		AddOpFunc(FactorAt, TEXT("factor "));
+		Data += sizeof(mu::OP::ADDRESS);
+		
+		mu::OP::ADDRESS BaseAt = 0;
+		FMemory::Memcpy(&BaseAt, Data, sizeof(mu::OP::ADDRESS));
+		AddOpFunc(BaseAt, TEXT("base "));
+		Data += sizeof(mu::OP::ADDRESS);
+
+		mu::OP::ADDRESS TargetAt = 0;
+		FMemory::Memcpy(&TargetAt, Data, sizeof(mu::OP::ADDRESS));
+		AddOpFunc(TargetAt, TEXT("target "));
+		Data += sizeof(mu::OP::ADDRESS);
+		
+		break;
+	}
 
-	default:
+	case mu::EOpType::SC_MULTIPLYADD:
+	{
+		mu::OP::ScalarMultiplyAddArgs Args = InProgram.GetOpArgs<mu::OP::ScalarMultiplyAddArgs>(InParentAddress);
+		AddOpFunc(Args.factor0, TEXT("factor0 "));
+		AddOpFunc(Args.factor1, TEXT("factor1 "));
+		AddOpFunc(Args.add, TEXT("add "));
+		break;
+	}
+		
+	case mu::EOpType::SC_CURVE:
+	{
+		mu::OP::ScalarCurveArgs Args = InProgram.GetOpArgs<mu::OP::ScalarCurveArgs>(InParentAddress);
+		AddOpFunc(Args.time, TEXT("time "));
+		AddOpFunc(Args.curve, TEXT("curve "));
+		break;
+	}
+		
+	case mu::EOpType::CO_SAMPLEIMAGE:
+	{
+		mu::OP::ColourSampleImageArgs Args = InProgram.GetOpArgs<mu::OP::ColourSampleImageArgs>(InParentAddress);
+		AddOpFunc(Args.image, TEXT("image "));
+		AddOpFunc(Args.x, TEXT("x "));
+		AddOpFunc(Args.y, TEXT("y "));
+		break;
+	}
+		
+	case mu::EOpType::IM_RESIZELIKE:
+	{
+		mu::OP::ImageResizeLikeArgs Args = InProgram.GetOpArgs<mu::OP::ImageResizeLikeArgs>(InParentAddress);
+		AddOpFunc(Args.Source, TEXT("src "));
+		AddOpFunc(Args.SizeSource, TEXT("sizeSrc "));
+		break;
+	}
+
+	case mu::EOpType::IN_ADDMESH:
+	case mu::EOpType::IN_ADDIMAGE:
+	case mu::EOpType::IN_ADDVECTOR:
+	case mu::EOpType::IN_ADDSCALAR:
+	case mu::EOpType::IN_ADDSTRING:
+	case mu::EOpType::IN_ADDCOMPONENT:
+	case mu::EOpType::IN_ADDSURFACE:
+	{
+		mu::OP::InstanceAddArgs Args = InProgram.GetOpArgs<mu::OP::InstanceAddArgs>(InParentAddress);
+		AddOpFunc(Args.instance, TEXT("instance "));
+		AddOpFunc(Args.value, TEXT("value "));
+		break;
+	}
+		
+	case mu::EOpType::IM_COMPOSE:
+	{
+		mu::OP::ImageComposeArgs Args = InProgram.GetOpArgs<mu::OP::ImageComposeArgs>(InParentAddress);
+		AddOpFunc(Args.layout, TEXT("layout "));
+		AddOpFunc(Args.base, TEXT("base "));
+		AddOpFunc(Args.blockImage, TEXT("blockImage "));
+		AddOpFunc(Args.mask, TEXT("mask "));
+		break;
+	}
+		
+	case mu::EOpType::IM_INTERPOLATE:
+	{
+		mu::OP::ImageInterpolateArgs Args = InProgram.GetOpArgs<mu::OP::ImageInterpolateArgs>(InParentAddress);
+		AddOpFunc(Args.Factor, TEXT("factor "));
+			
+		int32 TargetIndex = 0;
+		for (const mu::OP::ADDRESS& Operation : Args.Targets)
+		{
+			AddOpFunc(Operation, FString::Printf(TEXT("target %i "), TargetIndex++));
+		}
+		break;
+	}
+		
+	case mu::EOpType::IM_SATURATE:
+	{
+		mu::OP::ImageSaturateArgs Args = InProgram.GetOpArgs<mu::OP::ImageSaturateArgs>(InParentAddress);
+		AddOpFunc(Args.Base, TEXT("base "));
+		AddOpFunc(Args.Factor, TEXT("factor "));
+		break;
+	}
+		
+	case mu::EOpType::IM_COLOURMAP:
+	{
+		mu::OP::ImageColourMapArgs Args = InProgram.GetOpArgs<mu::OP::ImageColourMapArgs>(InParentAddress);
+		AddOpFunc(Args.Base, TEXT("base "));
+		AddOpFunc(Args.Mask, TEXT("mask "));
+		AddOpFunc(Args.Map, TEXT("map "));
+		break;
+	}
+		
+	case mu::EOpType::IM_BINARISE:
+	{
+		mu::OP::ImageBinariseArgs Args = InProgram.GetOpArgs<mu::OP::ImageBinariseArgs>(InParentAddress);
+		AddOpFunc(Args.Base, TEXT("base "));
+		AddOpFunc(Args.Threshold, TEXT("threshold "));
+		break;
+	}
+		
+	case mu::EOpType::IM_PATCH:
+	{
+		mu::OP::ImagePatchArgs Args = InProgram.GetOpArgs<mu::OP::ImagePatchArgs>(InParentAddress);
+		AddOpFunc(Args.base, TEXT("base "));
+		AddOpFunc(Args.patch, TEXT("patch "));
+		break;
+	}
+		
+	case mu::EOpType::IM_RASTERMESH:
+	{
+		mu::OP::ImageRasterMeshArgs Args = InProgram.GetOpArgs<mu::OP::ImageRasterMeshArgs>(InParentAddress);
+		AddOpFunc(Args.mesh, TEXT("mesh "));
+		AddOpFunc(Args.image, TEXT("image "));
+		AddOpFunc(Args.mask, TEXT("mask "));
+		AddOpFunc(Args.angleFadeProperties, TEXT("angleFadeProperties "));
+		AddOpFunc(Args.projector, TEXT("projector "));
+		break;
+	}
+		
+	case mu::EOpType::IM_DISPLACE:
+	{
+		mu::OP::ImageDisplaceArgs Args = InProgram.GetOpArgs<mu::OP::ImageDisplaceArgs>(InParentAddress);
+		AddOpFunc(Args.Source, TEXT("src "));
+		AddOpFunc(Args.DisplacementMap, TEXT("displacementMap "));
+		break;
+	}
+		
+	case mu::EOpType::IM_NORMALCOMPOSITE:
 	{
-		// Generic list of child operations
-		bUseGeneric = true;
+		mu::OP::ImageNormalCompositeArgs Args = InProgram.GetOpArgs<mu::OP::ImageNormalCompositeArgs>(InParentAddress);
+		AddOpFunc(Args.base, TEXT("base "));
+		AddOpFunc(Args.normal, TEXT("normal "));
+		break;
+	}
+		
+	case mu::EOpType::IM_TRANSFORM:
+	{
+		mu::OP::ImageTransformArgs Args = InProgram.GetOpArgs<mu::OP::ImageTransformArgs>(InParentAddress);
+		AddOpFunc(Args.Base, TEXT("base "));
+		AddOpFunc(Args.OffsetX, TEXT("offsetX "));
+		AddOpFunc(Args.OffsetY, TEXT("offsetY "));
+		AddOpFunc(Args.ScaleX, TEXT("scaleX "));
+		AddOpFunc(Args.ScaleY, TEXT("scaleY "));
+		AddOpFunc(Args.Rotation, TEXT("rotation "));
 		break;
 	}
 
+	// Add here more operation types to define how they are exposed in the tree (set a Caption)
+		
+	default:
+		{
+			// Generic list of child operations
+			bUseGeneric = true;
+			break;
+		}
 	}
 
 	if (bUseGeneric)
 	{
-		// Find children of the provided element
+		// Find children of the provided element without adding any extra string to better identify what each of them represents
 		mu::ForEachReference(InProgram, InParentAddress, [this, &InProgram, AddOpFunc](mu::OP::ADDRESS ChildAddress)
 			{
 				AddOpFunc(ChildAddress,TEXT(""));
@@ -1566,21 +1739,22 @@ void SMutableCodeViewer::GenerateElementRecursive(const int32& InStateIndex, mu:
 		// Validate in case there is a mismatch in the custom processing of children and the generic one, which would cause problems.
 		ChildIndex = 0;
 
-		auto ValidateOpFunc = [this, InParentAddress, &ChildIndex](mu::OP::ADDRESS ChildAddress)
+		auto ValidateOpFunc = [this, InParentAddress, ParentOperationType, &ChildIndex](mu::OP::ADDRESS ChildAddress)
 		{
-			if (ChildAddress)
-			{
-				const FItemCacheKey Key = { InParentAddress, ChildAddress, ChildIndex };
-				const TSharedPtr<FMutableCodeTreeElement>* CachedItem = ItemCache.Find(Key);
-				check(CachedItem);
-			}
+			const FItemCacheKey Key = { InParentAddress, ChildAddress, ChildIndex };
+			const TSharedPtr<FMutableCodeTreeElement>* CachedItem = ItemCache.Find(Key);
+
+			// If this check fails could mean that in the switch above one type of OP processes its children in one
+			// order but in "ForEachReference" that same operation processes the same children in another order.
+			check(CachedItem);
+				
 			++ChildIndex;
 		};
 
-		mu::ForEachReference(InProgram, InParentAddress, [this, &InProgram, ValidateOpFunc](mu::OP::ADDRESS ChildAddress)
-			{
-				ValidateOpFunc(ChildAddress);
-			});
+		mu::ForEachReference(InProgram, InParentAddress, [this, ValidateOpFunc](mu::OP::ADDRESS ChildAddress)
+		{
+			ValidateOpFunc(ChildAddress);
+		});
 	}
 }
 
@@ -1614,12 +1788,14 @@ TSharedRef<ITableRow> SMutableCodeViewer::GenerateRowForNodeTree(TSharedPtr<FMut
 	
 	// Generate a row element
 	TSharedRef<SMutableCodeTreeRow> Row = SNew(SMutableCodeTreeRow, InOwnerTable, InTreeNode);
-	
+
 	// Determine if a row should be painted as highlighted based on the selected item
 	if (TreeView->GetNumItemsSelected())
 	{
 		const TSharedPtr<FMutableCodeTreeElement> SelectedElement = TreeView->GetSelectedItems()[0];
-		if (InTreeNode->MutableOperation == SelectedElement->MutableOperation)
+		if (SelectedElement != InTreeNode &&
+			SelectedElement->MutableOperation > 0 &&
+			InTreeNode->MutableOperation == SelectedElement->MutableOperation)
 		{
 			Row->Highlight();
 		}
@@ -1645,7 +1821,6 @@ void SMutableCodeViewer::GetChildrenForInfo(TSharedPtr<FMutableCodeTreeElement>
 	uint32 ChildIndex = 0;
 	mu::ForEachReference(Program, InInfo->MutableOperation, [this, ParentAddress, &ChildIndex, &OutChildren](mu::OP::ADDRESS ChildAddress)
 	{
-		if (ChildAddress)
 		{
 			const FItemCacheKey Key = { ParentAddress, ChildAddress, ChildIndex };
 			const TSharedPtr<FMutableCodeTreeElement>* CachedItem = ItemCache.Find(Key);
@@ -1721,7 +1896,6 @@ void SMutableCodeViewer::GetVisibleChildren(TSharedPtr<FMutableCodeTreeElement>
 		uint32 ChildIndex = 0;
 		mu::ForEachReference(MutableProgram, ParentAddress, [this, ParentAddress, &ChildIndex, &OutChildren, &ToSearchForChildren ](mu::OP::ADDRESS ChildAddress)
 		{
-			if (ChildAddress)
 			{
 				const FItemCacheKey Key = { ParentAddress, ChildAddress, ChildIndex };
 				const TSharedPtr<FMutableCodeTreeElement> CachedItem = *ItemCache.Find(Key);
@@ -1891,8 +2065,9 @@ void SMutableCodeViewer::OnSelectionChanged(TSharedPtr<FMutableCodeTreeElement>
 TSharedPtr<SWidget> SMutableCodeViewer::OnTreeContextMenuOpening()
 {
 	FMenuBuilder MenuBuilder(true, nullptr);
-	
-	if (TreeView->GetSelectedItems().Num())
+
+	// Only show the Ui for operations different from "None" or "0"
+	if (TreeView->GetSelectedItems().Num() && TreeView->GetSelectedItems()[0]->MutableOperation > 0)
 	{
 		if (TreeView->GetSelectedItems().Num() == 1)
 		{
@@ -1953,6 +2128,7 @@ void SMutableCodeViewer::OnRowReleased(const TSharedRef<ITableRow>& InTreeRow)
 }
 #pragma endregion
 
+
 #pragma region Highlight Methods
 void SMutableCodeViewer::HighlightDuplicatesOfEntry(const TSharedPtr<FMutableCodeTreeElement>& InTargetEntry)
 {
@@ -1961,6 +2137,12 @@ void SMutableCodeViewer::HighlightDuplicatesOfEntry(const TSharedPtr<FMutableCod
 		ClearHighlightedItems();
 	}
 
+	// Do not highlight empty entries
+	if (InTargetEntry->MutableOperation == 0)
+	{
+		return;
+	}
+	
 	// Highlight the elements related to the currently selected item of the tree
 	HighlightedOperation = InTargetEntry->MutableOperation;
 
@@ -1977,6 +2159,7 @@ void SMutableCodeViewer::HighlightDuplicatesOfEntry(const TSharedPtr<FMutableCod
 	bIsElementHighlighted = true;
 }
 
+
 void SMutableCodeViewer::ClearHighlightedItems()
 {
 	// Clear the previously highlighted elements
@@ -1991,7 +2174,6 @@ void SMutableCodeViewer::ClearHighlightedItems()
 				SMutableCodeTreeRow* MutableRow = static_cast<SMutableCodeTreeRow*>(TableRow.Get());
 				MutableRow->ResetHighlight();
 			}
-
 		}
 	}
 
@@ -2235,7 +2417,7 @@ void SMutableCodeViewer::CacheRootNodeAddresses()
 	}
 	else
 	{
-		UE_LOG(LogTemp,Error,TEXT("The provided constant index does not seem to be used anywere : Make sure the index is valid and that IsConstantResourceUsedByOperation() switch is up to date"));
+		UE_LOG(LogTemp,Error,TEXT("The provided constant index does not seem to be used anywhere : Make sure the index is valid and that IsConstantResourceUsedByOperation() switch is up to date"));
 	}
 }
 
@@ -2909,6 +3091,13 @@ FMutableCodeTreeElement::FMutableCodeTreeElement(int32 InIndexOnTree, const int3
 	MutableOperation = InOperation;
 	Caption = InCaption;
 	LabelColor = InLabelColor;
+	
+	// Use a special color to denote "none" entries instead of the provided one
+	if (MutableOperation == 0)
+	{
+		LabelColor = FColor(FColorList::DimGrey);
+	}
+	
 	IndexOnTree = InIndexOnTree;
 	if (InDuplicatedOf)
 	{
@@ -2963,12 +3152,12 @@ void FMutableCodeTreeElement::SetElementCurrentState(const int32& InStateIndex)
 void FMutableCodeTreeElement::GenerateLabelText()
 {
 	const mu::FProgram& Program = MutableModel->GetPrivate()->Program;
-	const mu::EOpType Type = Program.GetOpType(MutableOperation);
-	FString OpName = mu::s_opNames[static_cast<int32>(Type)];
+	const mu::EOpType OperationType = Program.GetOpType(MutableOperation);
+	FString OpName = mu::s_opNames[static_cast<int32>(OperationType)];
 	OpName.TrimEndInline();
 
 	// See if the operation type accepts additional information in the label
-	switch (Type)
+	switch (OperationType)
 	{
 	case mu::EOpType::BO_PARAMETER:
 	case mu::EOpType::NU_PARAMETER:
@@ -3086,8 +3275,19 @@ void FMutableCodeTreeElement::GenerateLabelText()
 		break;
 	}
 
+	OpName = OpName.TrimStartAndEnd();
+	
 	// Prepare the text shown on the UI side of the operation tree
-	MainLabel = FString::Printf(TEXT("%s %d : %s"), *Caption, int32(MutableOperation), *OpName);
+	if (!Caption.IsEmpty())
+	{
+		Caption = Caption.TrimStartAndEnd();
+		MainLabel = FString::Printf(TEXT("%d : [ %s ] %s"), int32(MutableOperation), *Caption, *OpName);
+	}
+	else
+	{
+		MainLabel = FString::Printf(TEXT("%d : %s"), int32(MutableOperation), *OpName);
+	}
+	
 
 	// DEBUG : 
 	// FString IndexOnTree = FString::FromInt(IndexOnTree);
@@ -3099,7 +3299,8 @@ void FMutableCodeTreeElement::GenerateLabelText()
 	// RowStateIndex.Append(TEXT(" st "));
 	// MainLabel.InsertAt(0,RowStateIndex);
 
-	if (DuplicatedOf)
+	// Ignore the special case of operations of type "None"
+	if (MutableOperation > 0 && DuplicatedOf)
 	{
 		MainLabel.Append(TEXT(" (duplicated)"));
 	}
diff --git a/Engine/Plugins/Mutable/Source/CustomizableObjectEditor/Private/MuCOE/SMutableCodeViewer.h b/Engine/Plugins/Mutable/Source/CustomizableObjectEditor/Private/MuCOE/SMutableCodeViewer.h
--- a/Engine/Plugins/Mutable/Source/CustomizableObjectEditor/Private/MuCOE/SMutableCodeViewer.h
+++ b/Engine/Plugins/Mutable/Source/CustomizableObjectEditor/Private/MuCOE/SMutableCodeViewer.h
@@ -3,7 +3,6 @@
 #pragma once
 
 #include "MuR/ModelPrivate.h"
-#include "MuT/ErrorLog.h"
 #include "MuT/TypeInfo.h"
 #include "UObject/GCObject.h"
 #include "Widgets/SCompoundWidget.h"
@@ -250,8 +249,8 @@ private:
 	 */
 	const TArray<FSlateColor> ColorPerComputationalCost
 	{
-		FSlateColor(FLinearColor(1,1,1,1)),		// Standard cost color
-		FSlateColor(FLinearColor(1,0.4,0.2,1)),	// Expensive cost color
+		FSlateColor(FLinearColor(1,1,1,1)),			// Standard cost color
+		FSlateColor(FLinearColor(1,0.4,0.2,1)),		// Expensive cost color
 		FSlateColor(FLinearColor(1,0.1,0.1,1))		// Very Expensive cost color
 	};
 	
@@ -487,7 +486,8 @@ private:
 				uint32 ChildIndex = 0;
 				mu::ForEachReference(InProgram, InParentAddress, [this, &InParentAddress, &ChildIndex,&OutFoundChildrenData]( mu::OP::ADDRESS ChildAddress)
 				{
-					// If the parent does have a child then process it 
+					// If the parent does have a child then process it
+					// Ignoring child address 0 since those can not have children
 					if (ChildAddress) 
 					{
 						FItemCacheKey Key;
diff --git a/Engine/Plugins/Mutable/Source/MutableRuntime/Private/MuR/Operations.cpp b/Engine/Plugins/Mutable/Source/MutableRuntime/Private/MuR/Operations.cpp
--- a/Engine/Plugins/Mutable/Source/MutableRuntime/Private/MuR/Operations.cpp
+++ b/Engine/Plugins/Mutable/Source/MutableRuntime/Private/MuR/Operations.cpp
@@ -277,6 +277,7 @@ namespace mu
         {
 			OP::ScalarCurveArgs args = program.GetOpArgs<OP::ScalarCurveArgs>(at);
             f(args.time );
+        	f(args.curve);
             break;
         }
 
@@ -332,10 +333,7 @@ namespace mu
 				FMemory::Memcpy( &At, data, sizeof(OP::ADDRESS) );
 				data += sizeof(OP::ADDRESS);
 
-				if (At)
-				{
-					f(At);
-				}
+				f(At);
 			}
 
             break;
@@ -438,10 +436,7 @@ namespace mu
         {
 			OP::ImageLayerArgs args = program.GetOpArgs<OP::ImageLayerArgs>(at);
             f(args.base );
-            if ( args.mask )
-            {
-                f(args.mask );
-            }
+        	f(args.mask );
             f(args.blended );
             break;
         }
@@ -450,10 +445,7 @@ namespace mu
         {
 			OP::ImageLayerColourArgs args = program.GetOpArgs<OP::ImageLayerColourArgs>(at);
             f(args.base );
-            if ( args.mask )
-            {
-                f(args.mask );
-            }
+        	f(args.mask );
             f(args.colour );
             break;
         }
@@ -463,10 +455,7 @@ namespace mu
 			OP::ImageMultiLayerArgs args = program.GetOpArgs<OP::ImageMultiLayerArgs>(at);
             f(args.rangeSize );
             f(args.base );
-            if ( args.mask )
-            {
-                f(args.mask );
-            }
+            f(args.mask );
             f(args.blended );
             break;
         }
