[mutable] Remove ASTOpFixed legacy class. Convert remaining operations to their own ASTOp subclass.
#rnx

diff --git a/Engine/Plugins/Mutable/Source/CustomizableObject/Internal/MuCO/CustomizableObjectPrivate.h b/Engine/Plugins/Mutable/Source/CustomizableObject/Internal/MuCO/CustomizableObjectPrivate.h
--- a/Engine/Plugins/Mutable/Source/CustomizableObject/Internal/MuCO/CustomizableObjectPrivate.h
+++ b/Engine/Plugins/Mutable/Source/CustomizableObject/Internal/MuCO/CustomizableObjectPrivate.h
@@ -1826,6 +1826,8 @@ public:
 
 		MeshReferencesExtendedForCompilation,
 
+		RemoveObsoleteBoolOps,
+
 		// -----<new versions can be added above this line>--------
 		LastCustomizableObjectVersion
 	};
diff --git a/Engine/Plugins/Mutable/Source/MutableRuntime/Internal/MuR/Operations.h b/Engine/Plugins/Mutable/Source/MutableRuntime/Internal/MuR/Operations.h
--- a/Engine/Plugins/Mutable/Source/MutableRuntime/Internal/MuR/Operations.h
+++ b/Engine/Plugins/Mutable/Source/MutableRuntime/Internal/MuR/Operations.h
@@ -20,9 +20,7 @@ namespace mu
     //!
     enum class EOpType : uint16
     {
-        //-----------------------------------------------------------------------------------------
-        // No operation
-        //-----------------------------------------------------------------------------------------
+        /** No operation. */
         NONE,
 
         //-----------------------------------------------------------------------------------------
@@ -80,10 +78,6 @@ namespace mu
         // Boolean operations
         //-----------------------------------------------------------------------------------------
 
-        //! Compare two scalars, return true if the first is less than the second.
-		//! \TODO: Deprecated?
-		BO_LESS,
-
         //! Compare an integerexpression with an integer constant
         BO_EQUAL_INT_CONST,
 
@@ -96,18 +90,10 @@ namespace mu
         //! Left as an exercise to the reader to find out what this op does.
         BO_NOT,
 
-        //-----------------------------------------------------------------------------------------
-        // Integer operations
-        //-----------------------------------------------------------------------------------------
-
         //-----------------------------------------------------------------------------------------
         // Scalar operations
         //-----------------------------------------------------------------------------------------
 
-        //! Multiply a scalar value by another onw and add a third one to the result
-		//! \TODO: Deprecated?
-		SC_MULTIPLYADD,
-
         //! Apply an arithmetic operation to two scalars
         SC_ARITHMETIC,
 
@@ -375,35 +361,31 @@ namespace mu
 	static_assert(sizeof(FMeshBindColorChannelUsages) == sizeof(uint32));
 
 
-    //---------------------------------------------------------------------------------------------
-    //!
-    //---------------------------------------------------------------------------------------------
+    /** */
     template<typename ADDRESS_TYPE>
-    struct t_OP
+    struct TOperation
     {
         typedef ADDRESS_TYPE ADDRESS;
 
-        //-----------------------------------------------------------------------------------------
-        // Arguments for every operation type
-        //-----------------------------------------------------------------------------------------
+        /** Arguments for every operation type. */
         struct BoolConstantArgs
         {
-            bool value;
+            bool bValue;
         };
 
         struct IntConstantArgs
         {
-            int32 value;
+            int32 Value;
         };
 
         struct ScalarConstantArgs
         {
-            float value;
+            float Value;
         };
 
-        struct ColourConstantArgs
+        struct ColorConstantArgs
         {
-            float value[4];
+            FVector4f Value;
         };
 
     	struct MatrixConstantArgs
@@ -448,31 +430,20 @@ namespace mu
 		};
 
         //-------------------------------------------------------------------------------------
-        struct BoolLessArgs
-        {
-            ADDRESS a,b;
-        };
-
         struct BoolEqualScalarConstArgs
         {
-            ADDRESS value;
-            int16 constant;
+            ADDRESS Value;
+            int32 Constant;
         };
 
         struct BoolBinaryArgs
         {
-            ADDRESS a,b;
+            ADDRESS A,B;
         };
 
         struct BoolNotArgs
         {
-            ADDRESS source;
-        };
-
-        //-------------------------------------------------------------------------------------
-        struct ScalarMultiplyAddArgs
-        {
-            ADDRESS factor0,factor1,add;
+            ADDRESS A;
         };
 
         struct ScalarCurveArgs
@@ -481,12 +452,11 @@ namespace mu
             ADDRESS curve;  // Constant curve
         };
 
-        //-------------------------------------------------------------------------------------
         struct ColourSampleImageArgs
         {
-            ADDRESS image;
-            ADDRESS x,y;
-            uint8 filter;
+            ADDRESS Image;
+            ADDRESS X,Y;
+            uint8 Filter;
         };
 
         struct ColourSwizzleArgs
@@ -495,14 +465,9 @@ namespace mu
             ADDRESS sources[ MUTABLE_OP_MAX_SWIZZLE_CHANNELS ];
         };
 
-		struct ColourMultiplyArgs
-        {
-            ADDRESS a,b;
-        };
-
 		struct ColourFromScalarsArgs
 		{
-			ADDRESS v[MUTABLE_OP_MAX_SWIZZLE_CHANNELS];
+			ADDRESS V[MUTABLE_OP_MAX_SWIZZLE_CHANNELS];
 		};
 
         struct ArithmeticArgs
@@ -515,9 +480,9 @@ namespace mu
 				MULTIPLY,
 				DIVIDE
 			} OPERATION;
-            uint8 operation;
+            uint8 Operation;
 
-			ADDRESS a, b;
+			ADDRESS A, B;
 		};
 
 
@@ -1026,86 +991,10 @@ namespace mu
 			ADDRESS Mesh;
 			uint8 LayoutIndex;
 		};
-
-
-        //-------------------------------------------------------------------------------------
-        //-------------------------------------------------------------------------------------
-        //-------------------------------------------------------------------------------------
-        t_OP()
-        {
-            type = EOpType::NONE;
-            unused = 0;
-			FMemory::Memzero( &args, sizeof(args) );
-        }
-
-        bool operator==( const t_OP<ADDRESS>& o ) const
-        {
-            if (o.type!=type) return false;
-            //if (o.flags!=flags) return false;
-            if (FMemory::Memcmp(&o.args, &args, sizeof(FArgs))) return false;
-            return true;
-        }
-
-        bool operator<( const t_OP<ADDRESS>& o ) const
-        {
-            if (o.type<type) return true;
-            if (o.type>type) return false;
-//            if (o.flags<flags) return true;
-//            if (o.flags>flags) return false;
-            if (FMemory::Memcmp(&o.args, &args, sizeof(FArgs))<0) return true;
-            return false;
-        }
-
-        EOpType type;
-
-        typedef enum
-        {
-            //! In a particular state, this op is used only once. Knowing this help optimising the
-            //! build process since it allows reusing the resources in the operation result.
-            //F_SINGLE_USE = 1 << 0
-
-        } FLAGS;
-
-        //! Bitmask of FLAGS values
-        uint16 unused;
-
-        //! Constant arguments per operation type
-		//! Everything in this union needs to be converted to their own AST class eventually.
-        typedef union
-        {
-            IntConstantArgs IntConstant;
-            ScalarConstantArgs ScalarConstant;
-            ColourConstantArgs ColourConstant;
-        	MatrixConstantArgs TransformConstant;
-
-            //-------------------------------------------------------------------------------------
-            BoolLessArgs BoolLess;
-            BoolEqualScalarConstArgs BoolEqualScalarConst;
-            BoolBinaryArgs BoolBinary;
-            BoolNotArgs BoolNot;
-
-            //-------------------------------------------------------------------------------------
-            ScalarMultiplyAddArgs ScalarMultiplyAdd;
-            ArithmeticArgs ScalarArithmetic;
-
-            //-------------------------------------------------------------------------------------
-            ColourSampleImageArgs ColourSampleImage;
-            ColourSwizzleArgs ColourSwizzle;
-            ColourMultiplyArgs ColourMultiply;
-            ColourFromScalarsArgs ColourFromScalars;
-            ArithmeticArgs ColourArithmetic;
-
-        } FArgs;
-
-        FArgs args;
     };
 
-    typedef t_OP<uint32> OP;
+    typedef TOperation<uint32> OP;
 
-    // OP is not serialisable anymore, we link code now.
-    MUTABLE_DEFINE_POD_SERIALISABLE(OP);
-    MUTABLE_DEFINE_POD_VECTOR_SERIALISABLE(OP);
-    
 
 	//! Types of data handled by the Mutable runtime.
 	enum class EDataType : uint8
diff --git a/Engine/Plugins/Mutable/Source/MutableRuntime/Private/MuR/CodeRunner.cpp b/Engine/Plugins/Mutable/Source/MutableRuntime/Private/MuR/CodeRunner.cpp
--- a/Engine/Plugins/Mutable/Source/MutableRuntime/Private/MuR/CodeRunner.cpp
+++ b/Engine/Plugins/Mutable/Source/MutableRuntime/Private/MuR/CodeRunner.cpp
@@ -4956,7 +4956,7 @@ namespace mu
         case EOpType::BO_CONSTANT:
         {
 			OP::BoolConstantArgs Args = Program.GetOpArgs<OP::BoolConstantArgs>(Item.At);
-            bool result = Args.value;
+            bool result = Args.bValue;
             StoreBool( Item, result );
             break;
         }
@@ -4971,33 +4971,6 @@ namespace mu
             break;
         }
 
-        case EOpType::BO_LESS:
-        {
-			OP::BoolLessArgs Args = Program.GetOpArgs<OP::BoolLessArgs>(Item.At);
-            switch (Item.Stage)
-            {
-            case 0:
-                AddOp( FScheduledOp( Item.At, Item, 1),
-                        FScheduledOp( Args.a, Item),
-                        FScheduledOp( Args.b, Item) );
-                break;
-
-            case 1:
-            {
-                float a = LoadScalar( FCacheAddress(Args.a,Item) );
-                float b = LoadScalar( FCacheAddress(Args.b,Item) );
-                bool result = a<b;
-                StoreBool( Item, result );
-                break;
-            }
-
-            default:
-                check(false);
-            }
-
-            break;
-        }
-
         case EOpType::BO_AND:
         {
 			OP::BoolBinaryArgs Args = Program.GetOpArgs<OP::BoolBinaryArgs>(Item.At);
@@ -5007,9 +4980,9 @@ namespace mu
                 {
                     // Try to avoid the op entirely if we have some children cached
                     bool skip = false;
-                    if ( Args.a && GetMemory().IsValid( FCacheAddress(Args.a,Item) ) )
+                    if ( Args.A && GetMemory().IsValid( FCacheAddress(Args.A,Item) ) )
                     {
-                         bool a = LoadBool( FCacheAddress(Args.a,Item) );
+                         bool a = LoadBool( FCacheAddress(Args.A,Item) );
                          if (!a)
                          {
                             StoreBool( Item, false );
@@ -5017,9 +4990,9 @@ namespace mu
                          }
                     }
 
-                    if ( !skip && Args.b && GetMemory().IsValid( FCacheAddress(Args.b,Item) ) )
+                    if ( !skip && Args.B && GetMemory().IsValid( FCacheAddress(Args.B,Item) ) )
                     {
-                         bool b = LoadBool( FCacheAddress(Args.b,Item) );
+                         bool b = LoadBool( FCacheAddress(Args.B,Item) );
                          if (!b)
                          {
                             StoreBool( Item, false );
@@ -5030,14 +5003,14 @@ namespace mu
                     if (!skip)
                     {
                         AddOp( FScheduledOp( Item.At, Item, 1),
-                               FScheduledOp( Args.a, Item));
+                               FScheduledOp( Args.A, Item));
                     }
 				break;
                 }
 
             case 1:
             {
-                bool a = Args.a ? LoadBool( FCacheAddress(Args.a,Item) ) : true;
+                bool a = Args.A ? LoadBool( FCacheAddress(Args.A,Item) ) : true;
                 if (!a)
                 {
                     StoreBool( Item, false );
@@ -5045,7 +5018,7 @@ namespace mu
                 else
                 {
                     AddOp( FScheduledOp( Item.At, Item, 2),
-                           FScheduledOp( Args.b, Item));
+                           FScheduledOp( Args.B, Item));
                 }
                 break;
             }
@@ -5053,7 +5026,7 @@ namespace mu
             case 2:
             {
                 // We arrived here because a is true
-                bool b = Args.b ? LoadBool( FCacheAddress(Args.b,Item) ) : true;
+                bool b = Args.B ? LoadBool( FCacheAddress(Args.B,Item) ) : true;
                 StoreBool( Item, b );
                 break;
             }
@@ -5073,9 +5046,9 @@ namespace mu
                 {
                     // Try to avoid the op entirely if we have some children cached
                     bool skip = false;
-                    if ( Args.a && GetMemory().IsValid( FCacheAddress(Args.a,Item) ) )
+                    if ( Args.A && GetMemory().IsValid( FCacheAddress(Args.A,Item) ) )
                     {
-                         bool a = LoadBool( FCacheAddress(Args.a,Item) );
+                         bool a = LoadBool( FCacheAddress(Args.A,Item) );
                          if (a)
                          {
                             StoreBool( Item, true );
@@ -5083,9 +5056,9 @@ namespace mu
                          }
                     }
 
-                    if ( !skip && Args.b && GetMemory().IsValid( FCacheAddress(Args.b,Item) ) )
+                    if ( !skip && Args.B && GetMemory().IsValid( FCacheAddress(Args.B,Item) ) )
                     {
-                         bool b = LoadBool( FCacheAddress(Args.b,Item) );
+                         bool b = LoadBool( FCacheAddress(Args.B,Item) );
                          if (b)
                          {
                             StoreBool( Item, true );
@@ -5096,14 +5069,14 @@ namespace mu
                     if (!skip)
                     {
                         AddOp( FScheduledOp( Item.At, Item, 1),
-                               FScheduledOp( Args.a, Item));
+                               FScheduledOp( Args.A, Item));
                     }
 				break;
                 }
 
             case 1:
             {
-                bool a = Args.a ? LoadBool( FCacheAddress(Args.a,Item) ) : false;
+                bool a = Args.A ? LoadBool( FCacheAddress(Args.A,Item) ) : false;
                 if (a)
                 {
                     StoreBool( Item, true );
@@ -5111,7 +5084,7 @@ namespace mu
                 else
                 {
                     AddOp( FScheduledOp( Item.At, Item, 2),
-                           FScheduledOp( Args.b, Item));
+                           FScheduledOp( Args.B, Item));
                 }
                 break;
             }
@@ -5119,7 +5092,7 @@ namespace mu
             case 2:
             {
                 // We arrived here because a is false
-                bool b = Args.b ? LoadBool( FCacheAddress(Args.b,Item) ) : false;
+                bool b = Args.B ? LoadBool( FCacheAddress(Args.B,Item) ) : false;
                 StoreBool( Item, b );
                 break;
             }
@@ -5137,12 +5110,12 @@ namespace mu
             {
             case 0:
                     AddOp( FScheduledOp( Item.At, Item, 1),
-                           FScheduledOp( Args.source, Item) );
+                           FScheduledOp( Args.A, Item) );
                 break;
 
             case 1:
             {
-                bool result = !LoadBool( FCacheAddress(Args.source,Item) );
+                bool result = !LoadBool( FCacheAddress(Args.A,Item) );
                 StoreBool( Item, result );
                 break;
             }
@@ -5160,13 +5133,13 @@ namespace mu
             {
             case 0:
                     AddOp( FScheduledOp( Item.At, Item, 1),
-                           FScheduledOp( Args.value, Item) );
+                           FScheduledOp( Args.Value, Item) );
                 break;
 
             case 1:
             {
-                int32 a = LoadInt( FCacheAddress(Args.value,Item) );
-                bool result = a == Args.constant;
+                int32 a = LoadInt( FCacheAddress(Args.Value,Item) );
+                bool result = a == Args.Constant;
                 StoreBool( Item, result );
                 break;
             }
@@ -5197,7 +5170,7 @@ namespace mu
         case EOpType::NU_CONSTANT:
         {
 			OP::IntConstantArgs Args = Program.GetOpArgs<OP::IntConstantArgs>(Item.At);
-            int32 result = Args.value;
+            int32 result = Args.Value;
             StoreInt( Item, result );
             break;
         }
@@ -5252,7 +5225,7 @@ namespace mu
         case EOpType::SC_CONSTANT:
         {
 			OP::ScalarConstantArgs Args = Program.GetOpArgs<OP::ScalarConstantArgs>(Item.At);
-            float result = Args.value;
+            float result = Args.Value;
             StoreScalar( Item, result );
             break;
         }
@@ -5293,11 +5266,6 @@ namespace mu
             break;
         }
 
-        case EOpType::SC_MULTIPLYADD:
-            // \TODO
-            check( false );
-            break;
-
         case EOpType::SC_ARITHMETIC:
         {
 			OP::ArithmeticArgs Args = Program.GetOpArgs<OP::ArithmeticArgs>(Item.At);
@@ -5305,17 +5273,17 @@ namespace mu
             {
             case 0:
                     AddOp( FScheduledOp( Item.At, Item, 1),
-                           FScheduledOp( Args.a, Item),
-                           FScheduledOp( Args.b, Item) );
+                           FScheduledOp( Args.A, Item),
+                           FScheduledOp( Args.B, Item) );
                 break;
 
             case 1:
             {
-                float a = LoadScalar( FCacheAddress(Args.a,Item) );
-                float b = LoadScalar( FCacheAddress(Args.b,Item) );
+                float a = LoadScalar( FCacheAddress(Args.A,Item) );
+                float b = LoadScalar( FCacheAddress(Args.B,Item) );
 
                 float result = 1.0f;
-                switch (Args.operation)
+                switch (Args.Operation)
                 {
                 case OP::ArithmeticArgs::ADD:
                     result = a + b;
@@ -5407,13 +5375,8 @@ namespace mu
 
         case EOpType::CO_CONSTANT:
         {
-			OP::ColourConstantArgs Args = Program.GetOpArgs<OP::ColourConstantArgs>(Item.At);
-			FVector4f result;
-            result[0] = Args.value[0];
-            result[1] = Args.value[1];
-            result[2] = Args.value[2];
-            result[3] = Args.value[3];
-            StoreColor( Item, result );
+			OP::ColorConstantArgs Args = Program.GetOpArgs<OP::ColorConstantArgs>(Item.At);
+            StoreColor( Item, Args.Value );
             break;
         }
 
@@ -5434,23 +5397,23 @@ namespace mu
             {
             case 0:
                     AddOp( FScheduledOp( Item.At, Item, 1),
-                           FScheduledOp( Args.x, Item),
-                           FScheduledOp( Args.y, Item),
+                           FScheduledOp( Args.X, Item),
+                           FScheduledOp( Args.Y, Item),
 						   // Don't skip mips for the texture to sample
-                           FScheduledOp::FromOpAndOptions( Args.image, Item, 0) );
+                           FScheduledOp::FromOpAndOptions( Args.Image, Item, 0) );
                 break;
 
             case 1:
             {
-                float x = Args.x ? LoadScalar( FCacheAddress(Args.x,Item) ) : 0.5f;
-                float y = Args.y ? LoadScalar( FCacheAddress(Args.y,Item) ) : 0.5f;
+                float x = Args.X ? LoadScalar( FCacheAddress(Args.X,Item) ) : 0.5f;
+                float y = Args.Y ? LoadScalar( FCacheAddress(Args.Y,Item) ) : 0.5f;
 
-                TSharedPtr<const FImage> pImage = LoadImage(FScheduledOp::FromOpAndOptions(Args.image, Item, 0));
+                TSharedPtr<const FImage> pImage = LoadImage(FScheduledOp::FromOpAndOptions(Args.Image, Item, 0));
 
 				FVector4f result;
                 if (pImage)
                 {
-                    if (Args.filter)
+                    if (Args.Filter)
                     {
                         // TODO
                         result = pImage->Sample(FVector2f(x, y));
@@ -5521,10 +5484,10 @@ namespace mu
             {
             case 0:
                     AddOp( FScheduledOp( Item.At, Item, 1),
-                           FScheduledOp( Args.v[0], Item),
-                           FScheduledOp( Args.v[1], Item),
-                           FScheduledOp( Args.v[2], Item),
-                           FScheduledOp( Args.v[3], Item));
+                           FScheduledOp( Args.V[0], Item),
+                           FScheduledOp( Args.V[1], Item),
+                           FScheduledOp( Args.V[2], Item),
+                           FScheduledOp( Args.V[3], Item));
                 break;
 
             case 1:
@@ -5533,9 +5496,9 @@ namespace mu
 
 				for (int32 t = 0; t < MUTABLE_OP_MAX_SWIZZLE_CHANNELS; ++t)
 				{
-					if (Args.v[t])
+					if (Args.V[t])
 					{
-						Result[t] = LoadScalar(FCacheAddress(Args.v[t], Item));
+						Result[t] = LoadScalar(FCacheAddress(Args.V[t], Item));
 					}
 				}
 
@@ -5557,25 +5520,25 @@ namespace mu
             {
             case 0:
                     AddOp( FScheduledOp( Item.At, Item, 1),
-                           FScheduledOp( Args.a, Item),
-                           FScheduledOp( Args.b, Item));
+                           FScheduledOp( Args.A, Item),
+                           FScheduledOp( Args.B, Item));
                 break;
 
             case 1:
             {
-				EOpType otype = Program.GetOpType( Args.a );
+				EOpType otype = Program.GetOpType( Args.A );
                 EDataType dtype = GetOpDataType( otype );
                 check( dtype == EDataType::Color );
-                otype = Program.GetOpType( Args.b );
+                otype = Program.GetOpType( Args.B );
                 dtype = GetOpDataType( otype );
                 check( dtype == EDataType::Color );
-				FVector4f a = Args.a ? LoadColor( FCacheAddress( Args.a, Item ) )
+				FVector4f a = Args.A ? LoadColor( FCacheAddress( Args.A, Item ) )
                                  : FVector4f( 0, 0, 0, 0 );
-				FVector4f b = Args.b ? LoadColor( FCacheAddress( Args.b, Item ) )
+				FVector4f b = Args.B ? LoadColor( FCacheAddress( Args.B, Item ) )
                                  : FVector4f( 0, 0, 0, 0 );
 
 				FVector4f result = FVector4f(0,0,0,0);
-                switch (Args.operation)
+                switch (Args.Operation)
                 {
                 case OP::ArithmeticArgs::ADD:
                     result = a + b;
diff --git a/Engine/Plugins/Mutable/Source/MutableRuntime/Private/MuR/Operations.cpp b/Engine/Plugins/Mutable/Source/MutableRuntime/Private/MuR/Operations.cpp
--- a/Engine/Plugins/Mutable/Source/MutableRuntime/Private/MuR/Operations.cpp
+++ b/Engine/Plugins/Mutable/Source/MutableRuntime/Private/MuR/Operations.cpp
@@ -59,13 +59,11 @@ namespace mu
 		{ EDataType::Instance,		},	// IN_SWITCH
 		{ EDataType::ExtensionData,	},	// ED_SWITCH
 
-		{ EDataType::Bool,			},	// BO_LESS
 		{ EDataType::Bool,			},	// BO_EQUAL_SC_CONST
 		{ EDataType::Bool,			},	// BO_AND
 		{ EDataType::Bool,			},	// BO_OR
 		{ EDataType::Bool,			},	// BO_NOT
 
-		{ EDataType::Scalar,		},	// SC_MULTIPLYADD
 		{ EDataType::Scalar,		},	// SC_ARITHMETIC
 		{ EDataType::Scalar,		},	// SC_CURVE
 
@@ -150,100 +148,6 @@ namespace mu
 	}
 
 
-    //---------------------------------------------------------------------------------------------
-    void ForEachReference( OP& op, const TFunctionRef<void(OP::ADDRESS*)> f )
-    {
-		// Only operations that still use ASTOpFixed should be handled here.
-        switch ( op.type )
-        {
-        case EOpType::NONE:
-        case EOpType::BO_CONSTANT:
-        case EOpType::NU_CONSTANT:
-        case EOpType::SC_CONSTANT:
-        case EOpType::CO_CONSTANT:
-        case EOpType::IM_CONSTANT:
-        case EOpType::ME_CONSTANT:
-        case EOpType::LA_CONSTANT:
-        case EOpType::PR_CONSTANT:
-        case EOpType::BO_PARAMETER:
-        case EOpType::NU_PARAMETER:
-        case EOpType::SC_PARAMETER:
-        case EOpType::CO_PARAMETER:
-        case EOpType::PR_PARAMETER:
-        case EOpType::IM_PARAMETER:
-            break;
-
-        //-------------------------------------------------------------------------------------
-        case EOpType::BO_LESS:
-            f(&op.args.BoolLess.a );
-            f(&op.args.BoolLess.b );
-            break;
-
-        //-------------------------------------------------------------------------------------
-        case EOpType::BO_EQUAL_INT_CONST:
-            f(&op.args.BoolEqualScalarConst.value );
-            break;
-
-        //-------------------------------------------------------------------------------------
-        case EOpType::BO_AND:
-        case EOpType::BO_OR:
-            f(&op.args.BoolBinary.a );
-            f(&op.args.BoolBinary.b );
-            break;
-
-        //-------------------------------------------------------------------------------------
-        case EOpType::BO_NOT:
-            f(&op.args.BoolNot.source );
-            break;
-
-        //-------------------------------------------------------------------------------------
-        case EOpType::SC_MULTIPLYADD:
-            f(&op.args.ScalarMultiplyAdd.factor0 );
-            f(&op.args.ScalarMultiplyAdd.factor1 );
-            f(&op.args.ScalarMultiplyAdd.add );
-            break;
-
-        //-------------------------------------------------------------------------------------
-        case EOpType::SC_ARITHMETIC:
-            f(&op.args.ScalarArithmetic.a );
-            f(&op.args.ScalarArithmetic.b );
-            break;
-
-        //-------------------------------------------------------------------------------------
-        case EOpType::CO_SAMPLEIMAGE:
-            f(&op.args.ColourSampleImage.image );
-            f(&op.args.ColourSampleImage.x );
-            f(&op.args.ColourSampleImage.y );
-            break;
-
-        case EOpType::CO_SWIZZLE:
-            for (int32 t=0;t<MUTABLE_OP_MAX_SWIZZLE_CHANNELS;++t)
-            {
-                f(&op.args.ColourSwizzle.sources[t] );
-            }
-            break;
-
-        case EOpType::CO_FROMSCALARS:
-			for (int32 t = 0; t < MUTABLE_OP_MAX_SWIZZLE_CHANNELS; ++t)
-			{
-				f(&op.args.ColourFromScalars.v[t]);
-			}
-            break;
-
-        case EOpType::CO_ARITHMETIC:
-            f(&op.args.ColourArithmetic.a);
-            f(&op.args.ColourArithmetic.b);
-            break;
-
-        default:
-			check( false );
-            break;
-        }
-
-    }
-
-
-
     //---------------------------------------------------------------------------------------------
     void ForEachReference( const FProgram& program, OP::ADDRESS at, const TFunctionRef<void(OP::ADDRESS)> f )
     {
@@ -339,67 +243,44 @@ namespace mu
             break;
         }
 
-        //-------------------------------------------------------------------------------------
-        case EOpType::BO_LESS:
-        {
-			OP::BoolLessArgs args = program.GetOpArgs<OP::BoolLessArgs>(at);
-            f(args.a );
-            f(args.b );
-            break;
-        }
-
         //-------------------------------------------------------------------------------------
         case EOpType::BO_EQUAL_INT_CONST:
         {
 			OP::BoolEqualScalarConstArgs args = program.GetOpArgs<OP::BoolEqualScalarConstArgs>(at);
-            f(args.value );
+            f(args.Value );
             break;
         }
 
-        //-------------------------------------------------------------------------------------
         case EOpType::BO_AND:
         case EOpType::BO_OR:
         {
 			OP::BoolBinaryArgs args = program.GetOpArgs<OP::BoolBinaryArgs>(at);
-            f(args.a );
-            f(args.b );
+            f(args.A );
+            f(args.B );
             break;
         }
 
-        //-------------------------------------------------------------------------------------
         case EOpType::BO_NOT:
         {
 			OP::BoolNotArgs args = program.GetOpArgs<OP::BoolNotArgs>(at);
-            f(args.source );
-            break;
-        }
-
-        //-------------------------------------------------------------------------------------
-        case EOpType::SC_MULTIPLYADD:
-        {
-			OP::ScalarMultiplyAddArgs args = program.GetOpArgs<OP::ScalarMultiplyAddArgs>(at);
-            f(args.factor0 );
-            f(args.factor1 );
-            f(args.add );
+            f(args.A );
             break;
         }
 
-        //-------------------------------------------------------------------------------------
         case EOpType::SC_ARITHMETIC:
         {
 			OP::ArithmeticArgs args = program.GetOpArgs<OP::ArithmeticArgs>(at);
-            f(args.a );
-            f(args.b );
+            f(args.A );
+            f(args.B );
             break;
         }
 
-        //-------------------------------------------------------------------------------------
         case EOpType::CO_SAMPLEIMAGE:
         {
 			OP::ColourSampleImageArgs args = program.GetOpArgs<OP::ColourSampleImageArgs>(at);
-            f(args.image );
-            f(args.x );
-            f(args.y );
+            f(args.Image );
+            f(args.X );
+            f(args.Y );
             break;
         }
 
@@ -418,7 +299,7 @@ namespace mu
 			OP::ColourFromScalarsArgs args = program.GetOpArgs<OP::ColourFromScalarsArgs>(at);
 			for (int32 t = 0; t < MUTABLE_OP_MAX_SWIZZLE_CHANNELS; ++t)
 			{
-				f(args.v[t]);
+				f(args.V[t]);
 			}
             break;
         }
@@ -426,8 +307,8 @@ namespace mu
         case EOpType::CO_ARITHMETIC:
         {
 			OP::ArithmeticArgs args = program.GetOpArgs<OP::ArithmeticArgs>(at);
-            f(args.a);
-            f(args.b);
+            f(args.A);
+            f(args.B);
             break;
         }
 
diff --git a/Engine/Plugins/Mutable/Source/MutableTools/Internal/MuT/AST.h b/Engine/Plugins/Mutable/Source/MutableTools/Internal/MuT/AST.h
--- a/Engine/Plugins/Mutable/Source/MutableTools/Internal/MuT/AST.h
+++ b/Engine/Plugins/Mutable/Source/MutableTools/Internal/MuT/AST.h
@@ -1146,70 +1146,6 @@ namespace mu
     };
 
 
-    /** This class of operation is used to ease the transition from legacy code generation. */
-    class ASTOpFixed : public ASTOp
-    {
-    public:
-
-        //! Fixed size op structure. Any code ADDRESS in it, refers to the children TArray below.
-        OP op;
-
-        //! Operations used by this one. To be indexed with the ADDRESSes in the m_op.
-        //TArray<ASTChild> children;
-        std::array<ASTChild,8> children;
-        uint8 childCount = 1;
-
-    public:
-        ASTOpFixed();
-
-        ~ASTOpFixed() override;
-
-		virtual EOpType GetOpType() const override;
-		virtual Ptr<ASTOp> Clone( MapChildFuncRef mapChild ) const override;
-		virtual void ForEachChild( const TFunctionRef<void(ASTChild&)> f ) override;
-		virtual void Link( FProgram& program, FLinkerOptions* Options) override;
-		virtual bool IsEqual(const ASTOp& otherUntyped) const override;
-		virtual uint64 Hash() const override;
-		virtual Ptr<ASTOp> OptimiseSemantic(const FModelOptimizationOptions&, int32 Pass) const override;
-		virtual FBoolEvalResult EvaluateBool( ASTOpList& facts, FEvaluateBoolCache* cache ) const override;
-		virtual int32 EvaluateInt( ASTOpList& facts, bool &unknown ) const override;
-		virtual bool IsColourConstant(FVector4f& colour ) const override;
-
-        /** Add child for OP fixed nodes. */
-        inline void SetChild( OP::ADDRESS& at, const Ptr<ASTOp>& child )
-        {
-            // hack check
-            check( ((uint8*)&at)>=((uint8*)&op.args)
-                            &&
-                            ((uint8*)&at)<=((uint8*)&op.args)+sizeof(op.args));
-
-            if (!at)
-            {
-                if (child)
-                {
-                    check(childCount<8);
-                    at = (OP::ADDRESS)childCount++;
-                    children[at] = child;
-                }
-            }
-            else
-            {
-                children[at]=child;
-                if (!child)
-                {
-                    at=0;
-                }
-            }
-        }
-
-        inline void SetChild( OP::ADDRESS& at, const ASTChild& child )
-        {
-            SetChild(at,child.child());
-        }
-
-    };
-
-
     inline void ASTChild::AddParent()
     {
 		UE::TScopeLock Lock(Child->ParentsMutex);
@@ -1229,27 +1165,10 @@ namespace mu
 	}
 
 
-    //---------------------------------------------------------------------------------------------
-    //---------------------------------------------------------------------------------------------
-    //---------------------------------------------------------------------------------------------
-    class UniqueOpPool
+	/** */
+    class FUniqueOpPool
     {
     private:
-        struct op_pointer_hash
-        {
-			uint64 operator() (const Ptr<ASTOp>& n) const
-            {
-                return n->Hash();
-            }
-        };
-
-        struct op_pointer_equal
-        {
-            bool operator() (const Ptr<ASTOp>& a, const Ptr<ASTOp>& b) const
-            {
-                return *a==*b;
-            }
-        };
 
 		struct FKeyFuncs : BaseKeyFuncs<Ptr<ASTOp>, Ptr<ASTOp>, false>
 		{
diff --git a/Engine/Plugins/Mutable/Source/MutableTools/Private/MuT/AST.cpp b/Engine/Plugins/Mutable/Source/MutableTools/Private/MuT/AST.cpp
--- a/Engine/Plugins/Mutable/Source/MutableTools/Private/MuT/AST.cpp
+++ b/Engine/Plugins/Mutable/Source/MutableTools/Private/MuT/AST.cpp
@@ -963,327 +963,4 @@ void ASTOp::LinkRange(FProgram& program,
 	}
 }
 
-
-//-------------------------------------------------------------------------------------------------
-//-------------------------------------------------------------------------------------------------
-//-------------------------------------------------------------------------------------------------
-ASTOpFixed::ASTOpFixed()
-    : children({{
-			ASTChild(this,nullptr),ASTChild(this,nullptr),ASTChild(this,nullptr),ASTChild(this,nullptr),
-			ASTChild(this,nullptr),ASTChild(this,nullptr),ASTChild(this,nullptr),ASTChild(this,nullptr)
-		}})
-{
-}
-
-
-ASTOpFixed::~ASTOpFixed()
-{
-    // Explicit call needed to avoid recursive destruction
-    ASTOp::RemoveChildren();
-}
-
-
-void ASTOpFixed::ForEachChild( const TFunctionRef<void(ASTChild&)> f )
-{
-    // ugly but temporary while ASTOpFixed exists
-    ForEachReference( op, [&](OP::ADDRESS* pAt)
-    {
-        if (pAt && *pAt)
-        {
-            check(children.size()>*pAt);
-            f(children[*pAt]);
-        }
-    });
-}
-
-
-void ASTOpFixed::Link( FProgram& program, FLinkerOptions* )
-{
-	MUTABLE_CPUPROFILER_SCOPE(ASTOpFixed_Link);
-
-    if (!linkedAddress)
-    {
-        OP lop = op;
-
-        // Fix the linked op and generate children
-        ForEachReference(lop, [&](OP::ADDRESS* pChildOp)
-        {
-            Ptr<ASTOp> pASTChild = children[*pChildOp].Child;
-            if (pASTChild)
-            {
-                *pChildOp = pASTChild->linkedAddress;
-            }
-            else
-            {
-                *pChildOp = 0;
-            }
-        });
-
-        linkedAddress = (OP::ADDRESS)program.OpAddress.Num();        
-        program.OpAddress.Add((uint32)program.ByteCode.Num());
-
-        AppendCode(program.ByteCode,lop.type);
-        AppendCode(program.ByteCode,lop.args);
-    }
-}
-
-
-bool ASTOpFixed::IsEqual(const ASTOp& otherUntyped) const
-{
-    if (otherUntyped.GetOpType()==GetOpType())
-    {
-		const ASTOpFixed* other = static_cast<const ASTOpFixed*>(&otherUntyped);
-        return op==other->op && children==other->children;
-    }
-    return false;
-}
-
-
-EOpType ASTOpFixed::GetOpType() const 
-{ 
-	return op.type; 
-}
-
-
-mu::Ptr<ASTOp> ASTOpFixed::Clone( MapChildFuncRef mapChild ) const
-{
-    Ptr<ASTOpFixed> n = new ASTOpFixed();
-    n->op = op;
-    n->childCount = childCount;
-    // Skip the children 0, used to represent null op.
-    for(uint8_t i=1; i<childCount; ++i)
-    {
-        n->children[i] = mapChild(children[i].child());
-    }
-    return n;
-}
-
-
-uint64 ASTOpFixed::Hash() const
-{
-	uint64 res = std::hash<uint64>()(uint64(op.type));
-    for (const ASTChild& c: children)
-    {
-        hash_combine( res, c.child().get() );
-    }
-    return res;
-}
-
-
-ASTOp::FBoolEvalResult ASTOpFixed::EvaluateBool( ASTOpList& facts, FEvaluateBoolCache* cache ) const
-{
-    FEvaluateBoolCache localCache;
-    if (!cache)
-    {
-        cache = &localCache;
-    }
-    else
-    {
-        // Is this in the cache?
-		FEvaluateBoolCache::iterator it = cache->find(this);
-        if (it!=cache->end())
-        {
-            return it->second;
-        }
-    }
-
-    FBoolEvalResult result = BET_UNKNOWN;
-    switch(op.type)
-    {
-    case EOpType::BO_NOT:
-    {
-        if (children[op.args.BoolNot.source])
-        {
-            FBoolEvalResult child = children[op.args.BoolNot.source]->EvaluateBool(facts,cache);
-            if (child==BET_TRUE)
-            {
-                result = BET_FALSE;
-            }
-            else if (child==BET_FALSE)
-            {
-                result = BET_TRUE;
-            }
-        }
-        break;
-    }
-
-    case EOpType::BO_EQUAL_INT_CONST:
-    {
-        int32 intValue = op.args.BoolEqualScalarConst.constant;
-        const Ptr<ASTOp>& intExp = children[op.args.BoolEqualScalarConst.value].child();
-        bool intUnknown = true;
-        int intResult = intExp->EvaluateInt( facts, intUnknown );
-        if (intUnknown)
-        {
-            result = BET_UNKNOWN;
-        }
-        else if (intResult==intValue)
-        {
-            result = BET_TRUE;
-        }
-        else
-        {
-            result = BET_FALSE;
-        }
-        break;
-    }
-
-    case EOpType::BO_AND:
-    {
-        const Ptr<ASTOp>& a = children[op.args.BoolBinary.a].child();
-        const Ptr<ASTOp>& b = children[op.args.BoolBinary.b].child();
-        FBoolEvalResult resultA = BET_UNKNOWN;
-        FBoolEvalResult resultB = BET_UNKNOWN;
-        for ( size_t f=0; f<facts.Num(); ++f )
-        {
-            if ( a && resultA == BET_UNKNOWN )
-            {
-                resultA = a->EvaluateBool( facts, cache );
-
-                if ( resultA==BET_TRUE && resultB==BET_TRUE )
-                {
-                    result = BET_TRUE;
-                    break;
-                }
-                if ( resultA==BET_FALSE || resultB==BET_FALSE )
-                {
-                    result = BET_FALSE;
-                    break;
-                }
-            }
-            if ( b && resultB == BET_UNKNOWN )
-            {
-                resultB = b->EvaluateBool( facts, cache );
-
-                if ( resultA==BET_TRUE && resultB==BET_TRUE )
-                {
-                    result = BET_TRUE;
-                    break;
-                }
-                if ( resultA==BET_FALSE || resultB==BET_FALSE )
-                {
-                    result = BET_FALSE;
-                    break;
-                }
-            }
-        }
-        break;
-    }
-
-    case EOpType::BO_OR:
-    {
-        const Ptr<ASTOp>& a = children[op.args.BoolBinary.a].child();
-        const Ptr<ASTOp>& b = children[op.args.BoolBinary.b].child();
-        FBoolEvalResult resultA = BET_UNKNOWN;
-        FBoolEvalResult resultB = BET_UNKNOWN;
-        for ( size_t f=0; f<facts.Num(); ++f )
-        {
-            if ( a && resultA == BET_UNKNOWN )
-            {
-                resultA = a->EvaluateBool( facts, cache );
-
-                if ( resultA==BET_TRUE || resultB==BET_TRUE )
-                {
-                    result = BET_TRUE;
-                    break;
-                }
-                if ( resultA==BET_FALSE && resultB==BET_FALSE )
-                {
-                    result = BET_FALSE;
-                    break;
-                }
-            }
-            if ( b && resultB == BET_UNKNOWN )
-            {
-                resultB = b->EvaluateBool( facts, cache );
-
-                if ( resultA==BET_TRUE || resultB==BET_TRUE )
-                {
-                    result = BET_TRUE;
-                    break;
-                }
-                if ( resultA==BET_FALSE && resultB==BET_FALSE )
-                {
-                    result = BET_FALSE;
-                    break;
-                }
-            }
-        }
-        break;
-    }
-
-    default:
-        checkf( false, TEXT("Instruction not supported") );
-    }
-
-    (*cache)[this] = result;
-
-    return result;
-}
-
-
-int ASTOpFixed::EvaluateInt( ASTOpList& /*facts*/, bool &unknown ) const
-{
-    unknown = true;
-
-    switch (GetOpType())
-    {
-    case EOpType::NU_CONSTANT:
-        unknown = false;
-        return op.args.IntConstant.value;
-        break;
-
-    case EOpType::SC_CONSTANT:
-        unknown = false;
-        return (int)op.args.ScalarConstant.value;
-        break;
-
-        // TODO
-//    case IntExpression::IET_PARAMETER:
-//        for ( size_t f=0; (result==BET_UNKNOWN) && f<facts.Num(); ++f )
-//        {
-//            result = EvaluateIntEquality(facts[f].get(),intExp,value);
-//        }
-//        break;
-
-    default:
-        break;
-
-    }
-
-
-    return 0;
-}
-
-
-//-------------------------------------------------------------------------------------------------
-bool ASTOpFixed::IsColourConstant(FVector4f& colour) const
-{
-    bool res = false;
-    switch( op.type )
-    {
-    case EOpType::CO_CONSTANT:
-        res = true;
-        colour[0] = op.args.ColourConstant.value[0];
-        colour[1] = op.args.ColourConstant.value[1];
-        colour[2] = op.args.ColourConstant.value[2];
-        colour[3] = op.args.ColourConstant.value[3];
-        break;
-
-    case EOpType::CO_SAMPLEIMAGE:
-    case EOpType::CO_SWIZZLE:
-    case EOpType::CO_FROMSCALARS:
-    case EOpType::CO_ARITHMETIC:
-
-    default:
-        // TODO: Improve this test with more operations
-        //check( false );
-        res = false;
-        break;
-    }
-
-    return res;
-}
-
-
-}
+} // namespace mu
diff --git a/Engine/Plugins/Mutable/Source/MutableTools/Private/MuT/ASTOpAddExtensionData.cpp b/Engine/Plugins/Mutable/Source/MutableTools/Private/MuT/ASTOpAddExtensionData.cpp
--- a/Engine/Plugins/Mutable/Source/MutableTools/Private/MuT/ASTOpAddExtensionData.cpp
+++ b/Engine/Plugins/Mutable/Source/MutableTools/Private/MuT/ASTOpAddExtensionData.cpp
@@ -76,7 +76,7 @@ void ASTOpAddExtensionData::Link(FProgram& Program, FLinkerOptions*)
 	}
 
 	OP::InstanceAddExtensionDataArgs Args;
-	FMemory::Memset(Args, 0);
+	FMemory::Memzero(Args);
 
 	if (!Instance || !Instance->linkedAddress)
 	{
diff --git a/Engine/Plugins/Mutable/Source/MutableTools/Private/MuT/ASTOpBoolAnd.cpp b/Engine/Plugins/Mutable/Source/MutableTools/Private/MuT/ASTOpBoolAnd.cpp
new file mode 100644
--- /dev/null
+++ b/Engine/Plugins/Mutable/Source/MutableTools/Private/MuT/ASTOpBoolAnd.cpp
@@ -0,0 +1,227 @@
+// Copyright Epic Games, Inc. All Rights Reserved.
+
+#include "MuT/ASTOpBoolAnd.h"
+#include "MuT/ASTOpConstantBool.h"
+
+#include "HAL/PlatformMath.h"
+#include "MuR/ModelPrivate.h"
+#include "MuR/RefCounted.h"
+#include "MuR/Types.h"
+
+namespace mu
+{
+
+	ASTOpBoolAnd::ASTOpBoolAnd()
+		: A(this)
+		, B(this)
+	{
+	}
+
+
+	ASTOpBoolAnd::~ASTOpBoolAnd()
+	{
+		// Explicit call needed to avoid recursive destruction
+		ASTOp::RemoveChildren();
+	}
+
+
+	bool ASTOpBoolAnd::IsEqual(const ASTOp& OtherUntyped) const
+	{
+		if (OtherUntyped.GetOpType()==GetOpType())
+		{
+			const ASTOpBoolAnd* Other = static_cast<const ASTOpBoolAnd*>(&OtherUntyped);
+			return A == Other->A &&
+				B == Other->B;
+		}
+		return false;
+	}
+
+
+	uint64 ASTOpBoolAnd::Hash() const
+	{
+		uint64 Result = std::hash<void*>()(A.child().get());
+		hash_combine(Result, B.child().get());
+		return Result;
+	}
+
+
+	Ptr<ASTOp> ASTOpBoolAnd::Clone(MapChildFuncRef MapChild) const
+	{
+		Ptr<ASTOpBoolAnd> New = new ASTOpBoolAnd();
+		New->A = MapChild(A.child());
+		New->B = MapChild(B.child());
+		return New;
+	}
+
+
+	void ASTOpBoolAnd::ForEachChild(const TFunctionRef<void(ASTChild&)> Func)
+	{
+		Func(A);
+		Func(B);
+	}
+
+
+	void ASTOpBoolAnd::Link(FProgram& Program, FLinkerOptions*)
+	{
+		// Already linked?
+		if (!linkedAddress)
+		{
+			OP::BoolBinaryArgs Args;
+			FMemory::Memzero(Args);
+
+			if (A) Args.A = A->linkedAddress;
+			if (B) Args.B = B->linkedAddress;
+
+			linkedAddress = (OP::ADDRESS)Program.OpAddress.Num();
+			Program.OpAddress.Add(Program.ByteCode.Num());
+			AppendCode(Program.ByteCode, GetOpType());
+			AppendCode(Program.ByteCode, Args);
+		}
+
+	}
+
+
+	ASTOp::FBoolEvalResult ASTOpBoolAnd::EvaluateBool(ASTOpList& Facts, FEvaluateBoolCache* Cache) const
+	{
+		FEvaluateBoolCache LocalCache;
+		if (!Cache)
+		{
+			Cache = &LocalCache;
+		}
+		else
+		{
+			// Is this in the cache?
+			FEvaluateBoolCache::iterator it = Cache->find(this);
+			if (it != Cache->end())
+			{
+				return it->second;
+			}
+		}
+
+		FBoolEvalResult Result = BET_UNKNOWN;
+
+		FBoolEvalResult resultA = BET_UNKNOWN;
+		FBoolEvalResult resultB = BET_UNKNOWN;
+		for (int32 f = 0; f < Facts.Num(); ++f)
+		{
+			if (A && resultA == BET_UNKNOWN)
+			{
+				resultA = A->EvaluateBool(Facts, Cache);
+
+				if (resultA == BET_TRUE && resultB == BET_TRUE)
+				{
+					Result = BET_TRUE;
+					break;
+				}
+				if (resultA == BET_FALSE || resultB == BET_FALSE)
+				{
+					Result = BET_FALSE;
+					break;
+				}
+			}
+			if (B && resultB == BET_UNKNOWN)
+			{
+				resultB = B->EvaluateBool(Facts, Cache);
+
+				if (resultA == BET_TRUE && resultB == BET_TRUE)
+				{
+					Result = BET_TRUE;
+					break;
+				}
+				if (resultA == BET_FALSE || resultB == BET_FALSE)
+				{
+					Result = BET_FALSE;
+					break;
+				}
+			}
+		}
+
+		(*Cache)[this] = Result;
+
+		return Result;
+	}
+
+
+	Ptr<ASTOp> ASTOpBoolAnd::OptimiseSemantic(const FModelOptimizationOptions&, int32 Pass) const
+	{
+		Ptr<ASTOp> Result;
+
+		bool bChanged = false;
+		if (!A)
+		{
+			Result = B.child();
+			bChanged = true;
+		}
+		else if (!B)
+		{
+			Result = A.child();
+			bChanged = true;
+		}
+		else if (A->GetOpType() == EOpType::BO_CONSTANT)
+		{
+			if (static_cast<const ASTOpConstantBool*>(A.child().get())->bValue)
+			{
+				Result = B.child();
+				bChanged = true;
+			}
+			else
+			{
+				Result = A.child();
+				bChanged = true;
+			}
+		}
+		else if (B->GetOpType() == EOpType::BO_CONSTANT)
+		{
+			if (static_cast<const ASTOpConstantBool*>(B.child().get())->bValue)
+			{
+				Result = A.child();
+				bChanged = true;
+			}
+			else
+			{
+				Result = B.child();
+				bChanged = true;
+			}
+		}
+
+		// Common cases of repeated branch in children
+		else if (A->GetOpType() == EOpType::BO_AND)
+		{
+			const ASTOpBoolAnd* TypedA = static_cast<const ASTOpBoolAnd*>(A.child().get());
+			if (TypedA->A.child() == B.child()
+				||
+				TypedA->B.child() == B.child())
+			{
+				Result = A.child();
+				bChanged = true;
+			}
+		}
+		else if (B->GetOpType() == EOpType::BO_AND)
+		{
+			const ASTOpBoolAnd* TypedB = static_cast<const ASTOpBoolAnd*>(B.child().get());
+			if (TypedB->B.child() == A.child()
+				||
+				TypedB->B.child() == B.child())
+			{
+				Result = B.child();
+				bChanged = true;
+			}
+		}
+
+		else if (A.child() == B.child() || *A.child() == *B.child())
+		{
+			Result = A.child();
+			bChanged = true;
+		}
+
+		// if it became null, it means true (neutral AND argument)
+		if (bChanged && !Result)
+		{
+			Result = new ASTOpConstantBool(true);
+		}
+
+		return Result;
+	}
+
+
+}
diff --git a/Engine/Plugins/Mutable/Source/MutableTools/Private/MuT/ASTOpBoolAnd.h b/Engine/Plugins/Mutable/Source/MutableTools/Private/MuT/ASTOpBoolAnd.h
new file mode 100644
--- /dev/null
+++ b/Engine/Plugins/Mutable/Source/MutableTools/Private/MuT/ASTOpBoolAnd.h
@@ -0,0 +1,43 @@
+// Copyright Epic Games, Inc. All Rights Reserved.
+
+#pragma once
+
+#include "HAL/Platform.h"
+#include "MuR/Operations.h"
+#include "MuR/Ptr.h"
+#include "MuT/AST.h"
+
+
+namespace mu
+{
+	struct FProgram;
+
+	/** */
+	class ASTOpBoolAnd final : public ASTOp
+	{
+	public:
+
+		ASTChild A;
+		ASTChild B;
+
+	public:
+
+		ASTOpBoolAnd();
+		ASTOpBoolAnd(const ASTOpBoolAnd&) = delete;
+		~ASTOpBoolAnd();
+
+		// ASTOp interface
+		virtual EOpType GetOpType() const override { return EOpType::BO_AND; }
+		virtual uint64 Hash() const override;
+		virtual bool IsEqual(const ASTOp&) const override;
+		virtual Ptr<ASTOp> Clone(MapChildFuncRef) const override;
+		virtual void ForEachChild(const TFunctionRef<void(ASTChild&)>) override;
+		virtual void Link(FProgram&, FLinkerOptions*) override;
+		virtual FBoolEvalResult EvaluateBool(ASTOpList& Facts, FEvaluateBoolCache*) const override;
+		virtual Ptr<ASTOp> OptimiseSemantic(const FModelOptimizationOptions&, int32 Pass) const override;
+
+	};
+
+
+}
+
diff --git a/Engine/Plugins/Mutable/Source/MutableTools/Private/MuT/ASTOpBoolEqualIntConst.cpp b/Engine/Plugins/Mutable/Source/MutableTools/Private/MuT/ASTOpBoolEqualIntConst.cpp
new file mode 100644
--- /dev/null
+++ b/Engine/Plugins/Mutable/Source/MutableTools/Private/MuT/ASTOpBoolEqualIntConst.cpp
@@ -0,0 +1,130 @@
+// Copyright Epic Games, Inc. All Rights Reserved.
+
+#include "MuT/ASTOpBoolEqualIntConst.h"
+
+#include "MuT/ASTOpMeshMorph.h"
+#include "HAL/PlatformMath.h"
+#include "MuR/ModelPrivate.h"
+#include "MuR/RefCounted.h"
+#include "MuR/Types.h"
+
+
+namespace mu
+{
+
+	ASTOpBoolEqualIntConst::ASTOpBoolEqualIntConst()
+		: Value(this)
+	{
+	}
+
+
+	ASTOpBoolEqualIntConst::~ASTOpBoolEqualIntConst()
+	{
+		// Explicit call needed to avoid recursive destruction
+		ASTOp::RemoveChildren();
+	}
+
+
+	bool ASTOpBoolEqualIntConst::IsEqual(const ASTOp& otherUntyped) const
+	{
+		if (otherUntyped.GetOpType()==GetOpType())
+		{
+			const ASTOpBoolEqualIntConst* other = static_cast<const ASTOpBoolEqualIntConst*>(&otherUntyped);
+			return Value == other->Value && Constant == other->Constant;
+		}
+		return false;
+	}
+
+
+	mu::Ptr<ASTOp> ASTOpBoolEqualIntConst::Clone(MapChildFuncRef mapChild) const
+	{
+		Ptr<ASTOpBoolEqualIntConst> n = new ASTOpBoolEqualIntConst();
+		n->Value = mapChild(Value.child());
+		n->Constant = Constant;
+		return n;
+	}
+
+
+	void ASTOpBoolEqualIntConst::ForEachChild(const TFunctionRef<void(ASTChild&)> f)
+	{
+		f(Value);
+	}
+
+
+	uint64 ASTOpBoolEqualIntConst::Hash() const
+	{
+		uint64 Res = std::hash<ASTOp*>()(Value.child().get());
+		hash_combine(Res, Constant);
+		return Res;
+	}
+
+
+	void ASTOpBoolEqualIntConst::Link(FProgram& Program, FLinkerOptions*)
+	{
+		// Already linked?
+		if (!linkedAddress)
+		{
+			OP::BoolEqualScalarConstArgs Args;
+			FMemory::Memzero(Args);
+
+			Args.Constant = Constant;
+
+			if (Value)
+			{
+				Args.Value = Value->linkedAddress;
+			}
+
+			linkedAddress = (OP::ADDRESS)Program.OpAddress.Num();
+			Program.OpAddress.Add((uint32)Program.ByteCode.Num());
+			AppendCode(Program.ByteCode, GetOpType());
+			AppendCode(Program.ByteCode, Args);
+		}
+	}
+
+
+	ASTOp::FBoolEvalResult ASTOpBoolEqualIntConst::EvaluateBool(ASTOpList& Facts, FEvaluateBoolCache* Cache) const
+	{
+		FEvaluateBoolCache LocalCache;
+		if (!Cache)
+		{
+			Cache = &LocalCache;
+		}
+		else
+		{
+			// Is this in the cache?
+			FEvaluateBoolCache::iterator it = Cache->find(this);
+			if (it != Cache->end())
+			{
+				return it->second;
+			}
+		}
+
+		FBoolEvalResult Result = BET_UNKNOWN;
+
+		if (!Value)
+		{
+			return Result;
+		}
+
+		bool bIntUnknown = true;
+		int32 IntResult = Value->EvaluateInt(Facts, bIntUnknown);
+		if (bIntUnknown)
+		{
+			Result = BET_UNKNOWN;
+		}
+		else if (IntResult == Constant)
+		{
+			Result = BET_TRUE;
+		}
+		else
+		{
+			Result = BET_FALSE;
+		}
+
+		(*Cache)[this] = Result;
+
+		return Result;
+	}
+
+
+}
diff --git a/Engine/Plugins/Mutable/Source/MutableTools/Private/MuT/ASTOpBoolEqualIntConst.h b/Engine/Plugins/Mutable/Source/MutableTools/Private/MuT/ASTOpBoolEqualIntConst.h
new file mode 100644
--- /dev/null
+++ b/Engine/Plugins/Mutable/Source/MutableTools/Private/MuT/ASTOpBoolEqualIntConst.h
@@ -0,0 +1,40 @@
+// Copyright Epic Games, Inc. All Rights Reserved.
+
+#pragma once
+
+#include "HAL/Platform.h"
+#include "MuR/Operations.h"
+#include "MuR/Ptr.h"
+#include "MuT/AST.h"
+
+
+namespace mu
+{
+	struct FProgram;
+
+	/** */
+	class ASTOpBoolEqualIntConst final : public ASTOp
+	{
+	public:
+
+		ASTChild Value;
+		int32 Constant=0;
+
+	public:
+
+		ASTOpBoolEqualIntConst();
+		ASTOpBoolEqualIntConst(const ASTOpBoolEqualIntConst&) = delete;
+		virtual ~ASTOpBoolEqualIntConst() override;
+
+		// ASTOp interface
+		virtual EOpType GetOpType() const override { return EOpType::BO_EQUAL_INT_CONST; }
+		virtual uint64 Hash() const override;
+		virtual void ForEachChild(const TFunctionRef<void(ASTChild&)>) override;
+		virtual bool IsEqual(const ASTOp& otherUntyped) const override;
+		virtual Ptr<ASTOp> Clone(MapChildFuncRef mapChild) const override;
+		virtual void Link(FProgram& program, FLinkerOptions*) override;
+		virtual FBoolEvalResult EvaluateBool(ASTOpList& Facts, FEvaluateBoolCache*) const override;
+	};
+
+}
+
diff --git a/Engine/Plugins/Mutable/Source/MutableTools/Private/MuT/ASTOpBoolNot.cpp b/Engine/Plugins/Mutable/Source/MutableTools/Private/MuT/ASTOpBoolNot.cpp
new file mode 100644
--- /dev/null
+++ b/Engine/Plugins/Mutable/Source/MutableTools/Private/MuT/ASTOpBoolNot.cpp
@@ -0,0 +1,131 @@
+// Copyright Epic Games, Inc. All Rights Reserved.
+
+#include "MuT/ASTOpBoolNot.h"
+#include "MuT/ASTOpConstantBool.h"
+
+#include "HAL/PlatformMath.h"
+#include "MuR/ModelPrivate.h"
+#include "MuR/RefCounted.h"
+#include "MuR/Types.h"
+
+namespace mu
+{
+
+	ASTOpBoolNot::ASTOpBoolNot()
+		: A(this)
+	{
+	}
+
+
+	ASTOpBoolNot::~ASTOpBoolNot()
+	{
+		// Explicit call needed to avoid recursive destruction
+		ASTOp::RemoveChildren();
+	}
+
+
+	bool ASTOpBoolNot::IsEqual(const ASTOp& OtherUntyped) const
+	{
+		if (OtherUntyped.GetOpType()==GetOpType())
+		{
+			const ASTOpBoolNot* Other = static_cast<const ASTOpBoolNot*>(&OtherUntyped);
+			return A == Other->A;
+		}
+		return false;
+	}
+
+
+	uint64 ASTOpBoolNot::Hash() const
+	{
+		uint64 Result = std::hash<void*>()(A.child().get());
+		return Result;
+	}
+
+
+	Ptr<ASTOp> ASTOpBoolNot::Clone(MapChildFuncRef MapChild) const
+	{
+		Ptr<ASTOpBoolNot> New = new ASTOpBoolNot();
+		New->A = MapChild(A.child());
+		return New;
+	}
+
+
+	void ASTOpBoolNot::ForEachChild(const TFunctionRef<void(ASTChild&)> Func)
+	{
+		Func(A);
+	}
+
+
+	void ASTOpBoolNot::Link(FProgram& Program, FLinkerOptions*)
+	{
+		// Already linked?
+		if (!linkedAddress)
+		{
+			OP::BoolNotArgs Args;
+			FMemory::Memzero(Args);
+
+			if (A) Args.A = A->linkedAddress;
+
+			linkedAddress = (OP::ADDRESS)Program.OpAddress.Num();
+			Program.OpAddress.Add(Program.ByteCode.Num());
+			AppendCode(Program.ByteCode, GetOpType());
+			AppendCode(Program.ByteCode, Args);
+		}
+
+	}
+
+
+	ASTOp::FBoolEvalResult ASTOpBoolNot::EvaluateBool(ASTOpList& Facts, FEvaluateBoolCache* Cache) const
+	{
+		FEvaluateBoolCache LocalCache;
+		if (!Cache)
+		{
+			Cache = &LocalCache;
+		}
+		else
+		{
+			// Is this in the cache?
+			FEvaluateBoolCache::iterator it = Cache->find(this);
+			if (it != Cache->end())
+			{
+				return it->second;
+			}
+		}
+
+		FBoolEvalResult Result = BET_UNKNOWN;
+
+		if (A)
+		{
+			FBoolEvalResult child = A->EvaluateBool(Facts, Cache);
+			if (child == BET_TRUE)
+			{
+				Result = BET_FALSE;
+			}
+			else if (child == BET_FALSE)
+			{
+				Result = BET_TRUE;
+			}
+		}
+
+		(*Cache)[this] = Result;
+
+		return Result;
+	}
+
+
+	Ptr<ASTOp> ASTOpBoolNot::OptimiseSemantic(const FModelOptimizationOptions&, int32 Pass) const
+	{
+		Ptr<ASTOp> Result;
+
+		Ptr<ASTOp> Source = A.child();
+		if (Source && Source->GetOpType() == EOpType::BO_CONSTANT)
+		{
+			const ASTOpConstantBool* TypedSource = static_cast<const ASTOpConstantBool*>(Source.get());
+			Result = new ASTOpConstantBool(!TypedSource->bValue);
+		}
+
+		return Result;
+	}
+
+
+}
diff --git a/Engine/Plugins/Mutable/Source/MutableTools/Private/MuT/ASTOpBoolNot.h b/Engine/Plugins/Mutable/Source/MutableTools/Private/MuT/ASTOpBoolNot.h
new file mode 100644
--- /dev/null
+++ b/Engine/Plugins/Mutable/Source/MutableTools/Private/MuT/ASTOpBoolNot.h
@@ -0,0 +1,42 @@
+// Copyright Epic Games, Inc. All Rights Reserved.
+
+#pragma once
+
+#include "HAL/Platform.h"
+#include "MuR/Operations.h"
+#include "MuR/Ptr.h"
+#include "MuT/AST.h"
+
+
+namespace mu
+{
+	struct FProgram;
+
+	/** */
+	class ASTOpBoolNot final : public ASTOp
+	{
+	public:
+
+		ASTChild A;
+
+	public:
+
+		ASTOpBoolNot();
+		ASTOpBoolNot(const ASTOpBoolNot&) = delete;
+		~ASTOpBoolNot();
+
+		// ASTOp interface
+		virtual EOpType GetOpType() const override { return EOpType::BO_NOT; }
+		virtual uint64 Hash() const override;
+		virtual bool IsEqual(const ASTOp&) const override;
+		virtual Ptr<ASTOp> Clone(MapChildFuncRef) const override;
+		virtual void ForEachChild(const TFunctionRef<void(ASTChild&)>) override;
+		virtual void Link(FProgram&, FLinkerOptions*) override;
+		virtual FBoolEvalResult EvaluateBool(ASTOpList& Facts, FEvaluateBoolCache*) const override;
+		virtual Ptr<ASTOp> OptimiseSemantic(const FModelOptimizationOptions&, int32 Pass) const override;
+
+	};
+
+
+}
+
diff --git a/Engine/Plugins/Mutable/Source/MutableTools/Private/MuT/ASTOpBoolOr.cpp b/Engine/Plugins/Mutable/Source/MutableTools/Private/MuT/ASTOpBoolOr.cpp
new file mode 100644
--- /dev/null
+++ b/Engine/Plugins/Mutable/Source/MutableTools/Private/MuT/ASTOpBoolOr.cpp
@@ -0,0 +1,229 @@
+// Copyright Epic Games, Inc. All Rights Reserved.
+
+#include "MuT/ASTOpBoolOr.h"
+#include "MuT/ASTOpConstantBool.h"
+
+#include "HAL/PlatformMath.h"
+#include "MuR/ModelPrivate.h"
+#include "MuR/RefCounted.h"
+#include "MuR/Types.h"
+
+namespace mu
+{
+
+	ASTOpBoolOr::ASTOpBoolOr()
+		: A(this)
+		, B(this)
+	{
+	}
+
+
+	ASTOpBoolOr::~ASTOpBoolOr()
+	{
+		// Explicit call needed to avoid recursive destruction
+		ASTOp::RemoveChildren();
+	}
+
+
+	bool ASTOpBoolOr::IsEqual(const ASTOp& OtherUntyped) const
+	{
+		if (OtherUntyped.GetOpType()==GetOpType())
+		{
+			const ASTOpBoolOr* Other = static_cast<const ASTOpBoolOr*>(&OtherUntyped);
+			return A == Other->A &&
+				B == Other->B;
+		}
+		return false;
+	}
+
+
+	uint64 ASTOpBoolOr::Hash() const
+	{
+		uint64 Result = std::hash<void*>()(A.child().get());
+		hash_combine(Result, B.child().get());
+		return Result;
+	}
+
+
+	Ptr<ASTOp> ASTOpBoolOr::Clone(MapChildFuncRef MapChild) const
+	{
+		Ptr<ASTOpBoolOr> New = new ASTOpBoolOr();
+		New->A = MapChild(A.child());
+		New->B = MapChild(B.child());
+		return New;
+	}
+
+
+	void ASTOpBoolOr::ForEachChild(const TFunctionRef<void(ASTChild&)> Func)
+	{
+		Func(A);
+		Func(B);
+	}
+
+
+	void ASTOpBoolOr::Link(FProgram& Program, FLinkerOptions*)
+	{
+		// Already linked?
+		if (!linkedAddress)
+		{
+			OP::BoolBinaryArgs Args;
+			FMemory::Memzero(Args);
+
+			if (A) Args.A = A->linkedAddress;
+			if (B) Args.B = B->linkedAddress;
+
+			linkedAddress = (OP::ADDRESS)Program.OpAddress.Num();
+			Program.OpAddress.Add(Program.ByteCode.Num());
+			AppendCode(Program.ByteCode, GetOpType());
+			AppendCode(Program.ByteCode, Args);
+		}
+
+	}
+
+
+	ASTOp::FBoolEvalResult ASTOpBoolOr::EvaluateBool(ASTOpList& Facts, FEvaluateBoolCache* Cache) const
+	{
+		FEvaluateBoolCache LocalCache;
+		if (!Cache)
+		{
+			Cache = &LocalCache;
+		}
+		else
+		{
+			// Is this in the cache?
+			FEvaluateBoolCache::iterator it = Cache->find(this);
+			if (it != Cache->end())
+			{
+				return it->second;
+			}
+		}
+
+		FBoolEvalResult Result = BET_UNKNOWN;
+
+		FBoolEvalResult resultA = BET_UNKNOWN;
+		FBoolEvalResult resultB = BET_UNKNOWN;
+		for (int32 f = 0; f < Facts.Num(); ++f)
+		{
+			if (A && resultA == BET_UNKNOWN)
+			{
+				resultA = A->EvaluateBool(Facts, Cache);
+
+				if (resultA == BET_TRUE || resultB == BET_TRUE)
+				{
+					Result = BET_TRUE;
+					break;
+				}
+				if (resultA == BET_FALSE && resultB == BET_FALSE)
+				{
+					Result = BET_FALSE;
+					break;
+				}
+			}
+			if (B && resultB == BET_UNKNOWN)
+			{
+				resultB = B->EvaluateBool(Facts, Cache);
+
+				if (resultA == BET_TRUE || resultB == BET_TRUE)
+				{
+					Result = BET_TRUE;
+					break;
+				}
+				if (resultA == BET_FALSE && resultB == BET_FALSE)
+				{
+					Result = BET_FALSE;
+					break;
+				}
+			}
+		}
+
+		(*Cache)[this] = Result;
+
+		return Result;
+	}
+
+
+	Ptr<ASTOp> ASTOpBoolOr::OptimiseSemantic(const FModelOptimizationOptions&, int32 Pass) const
+	{
+		Ptr<ASTOp> Result;
+
+		bool bChanged = false;
+		Ptr<ASTOp> aAt = A.child();
+		Ptr<ASTOp> bAt = B.child();
+		if (!aAt)
+		{
+			Result = bAt;
+			bChanged = true;
+		}
+		else if (!bAt)
+		{
+			Result = aAt;
+			bChanged = true;
+		}
+		else if (aAt->GetOpType() == EOpType::BO_CONSTANT)
+		{
+			if (static_cast<const ASTOpConstantBool*>(aAt.get())->bValue)
+			{
+				Result = aAt;
+				bChanged = true;
+			}
+			else
+			{
+				Result = bAt;
+				bChanged = true;
+			}
+		}
+		else if (bAt->GetOpType() == EOpType::BO_CONSTANT)
+		{
+			if (static_cast<const ASTOpConstantBool*>(bAt.get())->bValue)
+			{
+				Result = bAt;
+				bChanged = true;
+			}
+			else
+			{
+				Result = aAt;
+				bChanged = true;
+			}
+		}
+
+		// Common cases of repeated branch in children
+		else if (aAt->GetOpType() == EOpType::BO_OR)
+		{
+			const ASTOpBoolOr* typedA = static_cast<const ASTOpBoolOr*>(aAt.get());
+			if (typedA->A.child() == bAt
+				||
+				typedA->B.child() == bAt)
+			{
+				Result = aAt;
+				bChanged = true;
+			}
+		}
+		else if (bAt->GetOpType() == EOpType::BO_OR)
+		{
+			const ASTOpBoolOr* typedB = static_cast<const ASTOpBoolOr*>(bAt.get());
+			if (typedB->B.child() == aAt
+				||
+				typedB->B.child() == bAt)
+			{
+				Result = bAt;
+				bChanged = true;
+			}
+		}
+
+		else if (aAt == bAt || *aAt == *bAt)
+		{
+			Result = aAt;
+			bChanged = true;
+		}
+
+		// if it became null, it means false (neutral OR argument)
+		if (bChanged && !Result)
+		{
+			Result = new ASTOpConstantBool(false);
+		}
+
+		return Result;
+	}
+
+
+}
diff --git a/Engine/Plugins/Mutable/Source/MutableTools/Private/MuT/ASTOpBoolOr.h b/Engine/Plugins/Mutable/Source/MutableTools/Private/MuT/ASTOpBoolOr.h
new file mode 100644
--- /dev/null
+++ b/Engine/Plugins/Mutable/Source/MutableTools/Private/MuT/ASTOpBoolOr.h
@@ -0,0 +1,43 @@
+// Copyright Epic Games, Inc. All Rights Reserved.
+
+#pragma once
+
+#include "HAL/Platform.h"
+#include "MuR/Operations.h"
+#include "MuR/Ptr.h"
+#include "MuT/AST.h"
+
+
+namespace mu
+{
+	struct FProgram;
+
+	/** */
+	class ASTOpBoolOr final : public ASTOp
+	{
+	public:
+
+		ASTChild A;
+		ASTChild B;
+
+	public:
+
+		ASTOpBoolOr();
+		ASTOpBoolOr(const ASTOpBoolOr&) = delete;
+		~ASTOpBoolOr();
+
+		// ASTOp interface
+		virtual EOpType GetOpType() const override { return EOpType::BO_OR; }
+		virtual uint64 Hash() const override;
+		virtual bool IsEqual(const ASTOp&) const override;
+		virtual Ptr<ASTOp> Clone(MapChildFuncRef) const override;
+		virtual void ForEachChild(const TFunctionRef<void(ASTChild&)>) override;
+		virtual void Link(FProgram&, FLinkerOptions*) override;
+		virtual FBoolEvalResult EvaluateBool(ASTOpList& Facts, FEvaluateBoolCache*) const override;
+		virtual Ptr<ASTOp> OptimiseSemantic(const FModelOptimizationOptions&, int32 Pass) const override;
+
+	};
+
+
+}
+
diff --git a/Engine/Plugins/Mutable/Source/MutableTools/Private/MuT/ASTOpColorArithmetic.cpp b/Engine/Plugins/Mutable/Source/MutableTools/Private/MuT/ASTOpColorArithmetic.cpp
new file mode 100644
--- /dev/null
+++ b/Engine/Plugins/Mutable/Source/MutableTools/Private/MuT/ASTOpColorArithmetic.cpp
@@ -0,0 +1,86 @@
+// Copyright Epic Games, Inc. All Rights Reserved.
+
+#include "MuT/ASTOpColorArithmetic.h"
+#include "MuT/ASTOpMeshAddTags.h"
+
+#include "HAL/PlatformMath.h"
+#include "MuR/ModelPrivate.h"
+#include "MuR/RefCounted.h"
+#include "MuR/Types.h"
+
+namespace mu
+{
+
+	ASTOpColorArithmetic::ASTOpColorArithmetic()
+		: A(this)
+		, B(this)
+	{
+	}
+
+
+	ASTOpColorArithmetic::~ASTOpColorArithmetic()
+	{
+		// Explicit call needed to avoid recursive destruction
+		ASTOp::RemoveChildren();
+	}
+
+
+	bool ASTOpColorArithmetic::IsEqual(const ASTOp& OtherUntyped) const
+	{
+		if (OtherUntyped.GetOpType()==GetOpType())
+		{
+			const ASTOpColorArithmetic* Other = static_cast<const ASTOpColorArithmetic*>(&OtherUntyped);
+			return A == Other->A &&
+				B == Other->B &&
+				Operation == Other->Operation;
+		}
+		return false;
+	}
+
+
+	uint64 ASTOpColorArithmetic::Hash() const
+	{
+		uint64 Result = std::hash<void*>()(A.child().get());
+		hash_combine(Result, B.child().get());
+		hash_combine(Result, Operation);
+		return Result;
+	}
+
+
+	Ptr<ASTOp> ASTOpColorArithmetic::Clone(MapChildFuncRef MapChild) const
+	{
+		Ptr<ASTOpColorArithmetic> New = new ASTOpColorArithmetic();
+		New->A = MapChild(A.child());
+		New->B = MapChild(B.child());
+		New->Operation = Operation;
+		return New;
+	}
+
+
+	void ASTOpColorArithmetic::ForEachChild(const TFunctionRef<void(ASTChild&)> Func)
+	{
+		Func(A);
+		Func(B);
+	}
+
+
+	void ASTOpColorArithmetic::Link(FProgram& Program, FLinkerOptions*)
+	{
+		// Already linked?
+		if (!linkedAddress)
+		{
+			OP::ArithmeticArgs Args;
+			FMemory::Memzero(Args);
+
+			if (A) Args.A = A->linkedAddress;
+			if (B) Args.B = B->linkedAddress;
+			Args.Operation = Operation;
+
+			linkedAddress = (OP::ADDRESS)Program.OpAddress.Num();
+			Program.OpAddress.Add(Program.ByteCode.Num());
+			AppendCode(Program.ByteCode, GetOpType());
+			AppendCode(Program.ByteCode, Args);
+		}
+
+	}
+}
diff --git a/Engine/Plugins/Mutable/Source/MutableTools/Private/MuT/ASTOpColorArithmetic.h b/Engine/Plugins/Mutable/Source/MutableTools/Private/MuT/ASTOpColorArithmetic.h
new file mode 100644
--- /dev/null
+++ b/Engine/Plugins/Mutable/Source/MutableTools/Private/MuT/ASTOpColorArithmetic.h
@@ -0,0 +1,39 @@
+// Copyright Epic Games, Inc. All Rights Reserved.
+
+#pragma once
+
+#include "HAL/Platform.h"
+#include "MuR/Operations.h"
+#include "MuR/Ptr.h"
+#include "MuT/AST.h"
+
+
+namespace mu
+{
+	struct FProgram;
+
+	/** */
+	class ASTOpColorArithmetic final : public ASTOp
+	{
+	public:
+
+		ASTChild A;
+		ASTChild B;
+		uint16 Operation=0;
+
+	public:
+
+		ASTOpColorArithmetic();
+		ASTOpColorArithmetic(const ASTOpColorArithmetic&) = delete;
+		~ASTOpColorArithmetic();
+
+		// ASTOp interface
+		virtual EOpType GetOpType() const override { return EOpType::CO_ARITHMETIC; }
+		virtual uint64 Hash() const override;
+		virtual bool IsEqual(const ASTOp&) const override;
+		virtual Ptr<ASTOp> Clone(MapChildFuncRef) const override;
+		virtual void ForEachChild(const TFunctionRef<void(ASTChild&)>) override;
+		virtual void Link(FProgram&, FLinkerOptions*) override;
+	};
+
+}
diff --git a/Engine/Plugins/Mutable/Source/MutableTools/Private/MuT/ASTOpColorFromScalars.cpp b/Engine/Plugins/Mutable/Source/MutableTools/Private/MuT/ASTOpColorFromScalars.cpp
new file mode 100644
--- /dev/null
+++ b/Engine/Plugins/Mutable/Source/MutableTools/Private/MuT/ASTOpColorFromScalars.cpp
@@ -0,0 +1,110 @@
+// Copyright Epic Games, Inc. All Rights Reserved.
+
+#include "MuT/ASTOpColorFromScalars.h"
+
+#include "Containers/Map.h"
+#include "HAL/UnrealMemory.h"
+#include "Math/UnrealMathSSE.h"
+#include "Misc/AssertionMacros.h"
+#include "MuR/ImagePrivate.h"
+#include "MuR/ModelPrivate.h"
+#include "MuR/MutableMath.h"
+#include "MuR/RefCounted.h"
+#include "MuR/Types.h"
+#include "MuT/ASTOpConditional.h"
+#include "MuT/ASTOpImageCompose.h"
+#include "MuT/ASTOpImagePatch.h"
+#include "MuT/ASTOpImagePixelFormat.h"
+#include "MuT/ASTOpImageMultiLayer.h"
+#include "MuT/ASTOpImageLayer.h"
+#include "MuT/ASTOpImageLayerColor.h"
+#include "MuT/ASTOpImageRasterMesh.h"
+#include "MuT/ASTOpImageTransform.h"
+#include "MuT/ASTOpImageMipmap.h"
+#include "MuT/ASTOpSwitch.h"
+
+namespace mu
+{
+
+	ASTOpColorFromScalars::ASTOpColorFromScalars()
+		: V{ ASTChild(this),ASTChild(this),ASTChild(this),ASTChild(this) }
+	{
+	}
+
+
+	ASTOpColorFromScalars::~ASTOpColorFromScalars()
+	{
+		// Explicit call needed to avoid recursive destruction
+		ASTOp::RemoveChildren();
+	}
+
+
+	bool ASTOpColorFromScalars::IsEqual(const ASTOp& otherUntyped) const
+	{
+		if (otherUntyped.GetOpType()==GetOpType())
+		{
+			const ASTOpColorFromScalars* Other = static_cast<const ASTOpColorFromScalars*>(&otherUntyped);
+			for (int32 i = 0; i<MUTABLE_OP_MAX_SWIZZLE_CHANNELS; ++i)
+			{
+				if (!(V[i] == Other->V[i]))
+				{
+					return false;
+				}
+			}
+			return true;
+		}
+		return false;
+	}
+
+
+	uint64 ASTOpColorFromScalars::Hash() const
+	{
+		uint64 res = std::hash<void*>()(V[0].child().get());
+		hash_combine(res, std::hash<void*>()(V[1].child().get()));
+		hash_combine(res, std::hash<void*>()(V[2].child().get()));
+		hash_combine(res, std::hash<void*>()(V[3].child().get()));
+		return res;
+	}
+
+
+	mu::Ptr<ASTOp> ASTOpColorFromScalars::Clone(MapChildFuncRef mapChild) const
+	{
+		mu::Ptr<ASTOpColorFromScalars> n = new ASTOpColorFromScalars();
+		for (int32 i = 0; i < MUTABLE_OP_MAX_SWIZZLE_CHANNELS; ++i)
+		{
+			n->V[i] = mapChild(V[i].child());
+		}
+		return n;
+	}
+
+
+	void ASTOpColorFromScalars::ForEachChild(const TFunctionRef<void(ASTChild&)> f)
+	{
+		for (int32 i = 0; i < MUTABLE_OP_MAX_SWIZZLE_CHANNELS; ++i)
+		{
+			f(V[i]);
+		}
+	}
+
+
+	void ASTOpColorFromScalars::Link(FProgram& program, FLinkerOptions*)
+	{
+		// Already linked?
+		if (!linkedAddress)
+		{
+			OP::ColourFromScalarsArgs Args;
+			FMemory::Memzero(Args);
+
+			for (int32 i = 0; i < MUTABLE_OP_MAX_SWIZZLE_CHANNELS; ++i)
+			{
+				if (V[i]) Args.V[i] = V[i]->linkedAddress;
+			}			
+
+			linkedAddress = (OP::ADDRESS)program.OpAddress.Num();
+			program.OpAddress.Add((uint32)program.ByteCode.Num());
+			AppendCode(program.ByteCode, GetOpType());
+			AppendCode(program.ByteCode, Args);
+		}
+	}
+
+}
diff --git a/Engine/Plugins/Mutable/Source/MutableTools/Private/MuT/ASTOpColorFromScalars.h b/Engine/Plugins/Mutable/Source/MutableTools/Private/MuT/ASTOpColorFromScalars.h
new file mode 100644
--- /dev/null
+++ b/Engine/Plugins/Mutable/Source/MutableTools/Private/MuT/ASTOpColorFromScalars.h
@@ -0,0 +1,37 @@
+// Copyright Epic Games, Inc. All Rights Reserved.
+
+#pragma once
+
+#include "HAL/PlatformMath.h"
+#include "MuR/Image.h"
+#include "MuR/Operations.h"
+#include "MuR/Ptr.h"
+#include "MuT/AST.h"
+
+
+namespace mu
+{
+	struct FProgram;
+
+	class ASTOpColorFromScalars final : public ASTOp
+	{
+	public:
+
+		ASTChild V[MUTABLE_OP_MAX_SWIZZLE_CHANNELS];
+
+	public:
+
+		ASTOpColorFromScalars();
+		ASTOpColorFromScalars(const ASTOpColorFromScalars&) = delete;
+		~ASTOpColorFromScalars();
+
+		virtual EOpType GetOpType() const override { return EOpType::CO_FROMSCALARS; }
+		virtual uint64 Hash() const override;
+		virtual bool IsEqual(const ASTOp& otherUntyped) const override;
+		virtual Ptr<ASTOp> Clone(MapChildFuncRef mapChild) const override;
+		virtual void ForEachChild(const TFunctionRef<void(ASTChild&)>) override;
+		virtual void Link(FProgram& program, FLinkerOptions* Options) override;
+	};
+
+}
+
diff --git a/Engine/Plugins/Mutable/Source/MutableTools/Private/MuT/ASTOpColorSampleImage.cpp b/Engine/Plugins/Mutable/Source/MutableTools/Private/MuT/ASTOpColorSampleImage.cpp
new file mode 100644
--- /dev/null
+++ b/Engine/Plugins/Mutable/Source/MutableTools/Private/MuT/ASTOpColorSampleImage.cpp
@@ -0,0 +1,90 @@
+// Copyright Epic Games, Inc. All Rights Reserved.
+
+#include "MuT/ASTOpColorSampleImage.h"
+
+#include "HAL/PlatformMath.h"
+#include "MuR/ModelPrivate.h"
+#include "MuR/RefCounted.h"
+#include "MuR/Types.h"
+
+namespace mu
+{
+
+	ASTOpColorSampleImage::ASTOpColorSampleImage()
+		: Image(this)
+		, X(this)
+		, Y(this)
+	{
+	}
+
+
+	ASTOpColorSampleImage::~ASTOpColorSampleImage()
+	{
+		// Explicit call needed to avoid recursive destruction
+		ASTOp::RemoveChildren();
+	}
+
+
+	bool ASTOpColorSampleImage::IsEqual(const ASTOp& OtherUntyped) const
+	{
+		if (OtherUntyped.GetOpType()==GetOpType())
+		{
+			const ASTOpColorSampleImage* Other = static_cast<const ASTOpColorSampleImage*>(&OtherUntyped);
+			return Image == Other->Image &&
+				X == Other->X &&
+				Y == Other->Y &&
+				Filter == Other->Filter;
+		}
+		return false;
+	}
+
+
+	uint64 ASTOpColorSampleImage::Hash() const
+	{
+		uint64 Result = std::hash<void*>()(Image.child().get());
+		hash_combine(Result, X.child().get());
+		hash_combine(Result, Y.child().get());
+		hash_combine(Result, Filter);
+		return Result;
+	}
+
+
+	Ptr<ASTOp> ASTOpColorSampleImage::Clone(MapChildFuncRef MapChild) const
+	{
+		Ptr<ASTOpColorSampleImage> New = new ASTOpColorSampleImage();
+		New->Image = MapChild(Image.child());
+		New->X = MapChild(X.child());
+		New->Y = MapChild(Y.child());
+		New->Filter = Filter;
+		return New;
+	}
+
+
+	void ASTOpColorSampleImage::ForEachChild(const TFunctionRef<void(ASTChild&)> Func)
+	{
+		Func(Image);
+		Func(X);
+		Func(Y);
+	}
+
+
+	void ASTOpColorSampleImage::Link(FProgram& Program, FLinkerOptions*)
+	{
+		// Already linked?
+		if (!linkedAddress)
+		{
+			OP::ColourSampleImageArgs Args;
+			FMemory::Memzero(Args);
+
+			if (Image) Args.Image = Image->linkedAddress;
+			if (X) Args.X = X->linkedAddress;
+			if (Y) Args.Y = Y->linkedAddress;
+			Args.Filter = Filter;
+
+			linkedAddress = (OP::ADDRESS)Program.OpAddress.Num();
+			Program.OpAddress.Add(Program.ByteCode.Num());
+			AppendCode(Program.ByteCode, GetOpType());
+			AppendCode(Program.ByteCode, Args);
+		}
+	}
+}
diff --git a/Engine/Plugins/Mutable/Source/MutableTools/Private/MuT/ASTOpColorSampleImage.h b/Engine/Plugins/Mutable/Source/MutableTools/Private/MuT/ASTOpColorSampleImage.h
new file mode 100644
--- /dev/null
+++ b/Engine/Plugins/Mutable/Source/MutableTools/Private/MuT/ASTOpColorSampleImage.h
@@ -0,0 +1,40 @@
+// Copyright Epic Games, Inc. All Rights Reserved.
+
+#pragma once
+
+#include "HAL/Platform.h"
+#include "MuR/Operations.h"
+#include "MuR/Ptr.h"
+#include "MuT/AST.h"
+
+
+namespace mu
+{
+	struct FProgram;
+
+	/** */
+	class ASTOpColorSampleImage final : public ASTOp
+	{
+	public:
+
+		ASTChild Image;
+		ASTChild X;
+		ASTChild Y;
+		uint8 Filter=0;
+
+	public:
+
+		ASTOpColorSampleImage();
+		ASTOpColorSampleImage(const ASTOpColorSampleImage&) = delete;
+		~ASTOpColorSampleImage();
+
+		// ASTOp interface
+		virtual EOpType GetOpType() const override { return EOpType::CO_SAMPLEIMAGE; }
+		virtual uint64 Hash() const override;
+		virtual bool IsEqual(const ASTOp&) const override;
+		virtual Ptr<ASTOp> Clone(MapChildFuncRef) const override;
+		virtual void ForEachChild(const TFunctionRef<void(ASTChild&)>) override;
+		virtual void Link(FProgram&, FLinkerOptions*) override;
+	};
+
+}
diff --git a/Engine/Plugins/Mutable/Source/MutableTools/Private/MuT/ASTOpColorSwizzle.cpp b/Engine/Plugins/Mutable/Source/MutableTools/Private/MuT/ASTOpColorSwizzle.cpp
new file mode 100644
--- /dev/null
+++ b/Engine/Plugins/Mutable/Source/MutableTools/Private/MuT/ASTOpColorSwizzle.cpp
@@ -0,0 +1,220 @@
+// Copyright Epic Games, Inc. All Rights Reserved.
+
+#include "MuT/ASTOpColorSwizzle.h"
+
+#include "Containers/Map.h"
+#include "HAL/UnrealMemory.h"
+#include "Math/UnrealMathSSE.h"
+#include "Misc/AssertionMacros.h"
+#include "MuR/ImagePrivate.h"
+#include "MuR/ModelPrivate.h"
+#include "MuR/MutableMath.h"
+#include "MuR/RefCounted.h"
+#include "MuR/Types.h"
+#include "MuT/ASTOpConditional.h"
+#include "MuT/ASTOpImageCompose.h"
+#include "MuT/ASTOpImagePatch.h"
+#include "MuT/ASTOpImagePixelFormat.h"
+#include "MuT/ASTOpImageMultiLayer.h"
+#include "MuT/ASTOpImageLayer.h"
+#include "MuT/ASTOpImageLayerColor.h"
+#include "MuT/ASTOpImageRasterMesh.h"
+#include "MuT/ASTOpImageTransform.h"
+#include "MuT/ASTOpImageMipmap.h"
+#include "MuT/ASTOpColorFromScalars.h"
+#include "MuT/ASTOpSwitch.h"
+
+namespace mu
+{
+
+	ASTOpColorSwizzle::ASTOpColorSwizzle()
+		: Sources{ ASTChild(this),ASTChild(this),ASTChild(this),ASTChild(this) }
+	{
+	}
+
+
+	ASTOpColorSwizzle::~ASTOpColorSwizzle()
+	{
+		// Explicit call needed to avoid recursive destruction
+		ASTOp::RemoveChildren();
+	}
+
+
+	bool ASTOpColorSwizzle::IsEqual(const ASTOp& otherUntyped) const
+	{
+		if (otherUntyped.GetOpType()==GetOpType())
+		{
+			const ASTOpColorSwizzle* Other = static_cast<const ASTOpColorSwizzle*>(&otherUntyped);
+			for (int32 i = 0; i<MUTABLE_OP_MAX_SWIZZLE_CHANNELS; ++i)
+			{
+				if (!(Sources[i] == Other->Sources[i] && SourceChannels[i] == Other->SourceChannels[i]))
+				{
+					return false;
+				}
+			}
+			return true;
+		}
+		return false;
+	}
+
+
+	uint64 ASTOpColorSwizzle::Hash() const
+	{
+		uint64 res = std::hash<void*>()(Sources[0].child().get());
+		hash_combine(res, std::hash<void*>()(Sources[1].child().get()));
+		hash_combine(res, std::hash<void*>()(Sources[2].child().get()));
+		hash_combine(res, std::hash<void*>()(Sources[3].child().get()));
+		hash_combine(res, std::hash<uint8>()(SourceChannels[0]));
+		hash_combine(res, std::hash<uint8>()(SourceChannels[1]));
+		hash_combine(res, std::hash<uint8>()(SourceChannels[2]));
+		hash_combine(res, std::hash<uint8>()(SourceChannels[3]));
+		return res;
+	}
+
+
+	mu::Ptr<ASTOp> ASTOpColorSwizzle::Clone(MapChildFuncRef mapChild) const
+	{
+		mu::Ptr<ASTOpColorSwizzle> n = new ASTOpColorSwizzle();
+		for (int32 i = 0; i < MUTABLE_OP_MAX_SWIZZLE_CHANNELS; ++i)
+		{
+			n->Sources[i] = mapChild(Sources[i].child());
+			n->SourceChannels[i] = SourceChannels[i];
+		}
+		return n;
+	}
+
+
+	void ASTOpColorSwizzle::ForEachChild(const TFunctionRef<void(ASTChild&)> f)
+	{
+		for (int32 i = 0; i < MUTABLE_OP_MAX_SWIZZLE_CHANNELS; ++i)
+		{
+			f(Sources[i]);
+		}
+	}
+
+
+	void ASTOpColorSwizzle::Link(FProgram& program, FLinkerOptions*)
+	{
+		// Already linked?
+		if (!linkedAddress)
+		{
+			OP::ColourSwizzleArgs Args;
+			FMemory::Memzero(Args);
+
+			for (int32 i = 0; i < MUTABLE_OP_MAX_SWIZZLE_CHANNELS; ++i)
+			{
+				if (Sources[i]) Args.sources[i] = Sources[i]->linkedAddress;
+				Args.sourceChannels[i] = SourceChannels[i];
+			}			
+
+			linkedAddress = (OP::ADDRESS)program.OpAddress.Num();
+			program.OpAddress.Add((uint32)program.ByteCode.Num());
+			AppendCode(program.ByteCode, GetOpType());
+			AppendCode(program.ByteCode, Args);
+		}
+	}
+
+
+	mu::Ptr<ASTOp> ASTOpColorSwizzle::OptimiseSemantic(const FModelOptimizationOptions&, int32 Pass) const
+	{
+		Ptr<ASTOp> at;
+
+		// Optimizations that can be applied per-channel
+		{
+			Ptr<ASTOpColorSwizzle> NewSwizzle;
+
+			for (int32 ChannelIndex = 0; ChannelIndex < MUTABLE_OP_MAX_SWIZZLE_CHANNELS; ++ChannelIndex)
+			{
+				Ptr<ASTOp> Candidate = Sources[ChannelIndex].child();
+				if (!Candidate)
+				{
+					continue;
+				}
+
+				switch (Candidate->GetOpType())
+				{
+
+				// Swizzle + swizzle = swizzle
+				case EOpType::CO_SWIZZLE:
+				{
+					if (!NewSwizzle)
+					{
+						NewSwizzle = mu::Clone<ASTOpColorSwizzle>(this);
+					}
+					const ASTOpColorSwizzle* TypedCandidate = static_cast<const ASTOpColorSwizzle*>(Candidate.get());
+					int32 CandidateChannel = SourceChannels[ChannelIndex];
+
+					NewSwizzle->Sources[ChannelIndex] = TypedCandidate->Sources[CandidateChannel].child();
+					NewSwizzle->SourceChannels[ChannelIndex] = TypedCandidate->SourceChannels[CandidateChannel];
+					break;
+				}
+
+				default:
+					break;
+
+				}
+			}
+
+			at = NewSwizzle;
+		}
+
+		// Not optimized yet?
+		if (!at)
+		{
+			// Optimizations that depend on all channels.
+			bool bAllChannelsSameType = true;
+			EOpType AllChannelsType = EOpType::NONE;
+			for (int32 ChannelIndex = 0; ChannelIndex < MUTABLE_OP_MAX_SWIZZLE_CHANNELS; ++ChannelIndex)
+			{
+				Ptr<ASTOp> Candidate = Sources[ChannelIndex].child();
+				if (!Candidate)
+				{
+					continue;
+				}
+
+				if (ChannelIndex == 0)
+				{
+					AllChannelsType = Candidate->GetOpType();
+				}
+				else if (Candidate->GetOpType() != AllChannelsType)
+				{
+					bAllChannelsSameType = false;
+					break;
+				}
+			}
+
+			if (bAllChannelsSameType)
+			{
+				switch (AllChannelsType)
+				{
+				case EOpType::CO_FROMSCALARS:
+				{
+					// We can remove the swizzle and replace it with a new FromScalars actually swizzling the inputs.
+
+					Ptr<ASTOpColorFromScalars> NewAt = new ASTOpColorFromScalars();
+
+					for (int32 ChannelIndex = 0; ChannelIndex < MUTABLE_OP_MAX_SWIZZLE_CHANNELS; ++ChannelIndex)
+					{
+						const ASTOp* SelectedSourceGeneric = Sources[ChannelIndex].child().get();
+						if (SelectedSourceGeneric)
+						{
+							const ASTOpColorFromScalars* SelectedSource = static_cast<const ASTOpColorFromScalars*>(SelectedSourceGeneric);
+							int32 SelectedChannel = SourceChannels[ChannelIndex];
+							Ptr<ASTOp> SelectedFloatInput = SelectedSource->V[SelectedChannel].child();
+							NewAt->V[ChannelIndex] = SelectedFloatInput;
+						}
+					}
+
+					at = NewAt;
+					break;
+				}
+
+				default: break;
+				}
+			}
+		}
+
+		return at;
+	}
+
+}
diff --git a/Engine/Plugins/Mutable/Source/MutableTools/Private/MuT/ASTOpColorSwizzle.h b/Engine/Plugins/Mutable/Source/MutableTools/Private/MuT/ASTOpColorSwizzle.h
new file mode 100644
--- /dev/null
+++ b/Engine/Plugins/Mutable/Source/MutableTools/Private/MuT/ASTOpColorSwizzle.h
@@ -0,0 +1,40 @@
+// Copyright Epic Games, Inc. All Rights Reserved.
+
+#pragma once
+
+#include "HAL/PlatformMath.h"
+#include "MuR/Image.h"
+#include "MuR/Operations.h"
+#include "MuR/Ptr.h"
+#include "MuT/AST.h"
+
+
+namespace mu
+{
+	struct FProgram;
+
+	class ASTOpColorSwizzle final : public ASTOp
+	{
+	public:
+
+		ASTChild Sources[MUTABLE_OP_MAX_SWIZZLE_CHANNELS];
+
+		uint8 SourceChannels[MUTABLE_OP_MAX_SWIZZLE_CHANNELS] = { 0,0,0,0 };
+
+	public:
+
+		ASTOpColorSwizzle();
+		ASTOpColorSwizzle(const ASTOpColorSwizzle&) = delete;
+		~ASTOpColorSwizzle();
+
+		virtual EOpType GetOpType() const override { return EOpType::CO_SWIZZLE; }
+		virtual uint64 Hash() const override;
+		virtual bool IsEqual(const ASTOp& otherUntyped) const override;
+		virtual Ptr<ASTOp> Clone(MapChildFuncRef mapChild) const override;
+		virtual void ForEachChild(const TFunctionRef<void(ASTChild&)>) override;
+		virtual void Link(FProgram& program, FLinkerOptions* Options) override;
+		virtual Ptr<ASTOp> OptimiseSemantic(const FModelOptimizationOptions&, int32 Pass) const override;
+	};
+
+}
+
diff --git a/Engine/Plugins/Mutable/Source/MutableTools/Private/MuT/ASTOpConditional.cpp b/Engine/Plugins/Mutable/Source/MutableTools/Private/MuT/ASTOpConditional.cpp
--- a/Engine/Plugins/Mutable/Source/MutableTools/Private/MuT/ASTOpConditional.cpp
+++ b/Engine/Plugins/Mutable/Source/MutableTools/Private/MuT/ASTOpConditional.cpp
@@ -109,18 +109,18 @@ namespace mu
 		// Already linked?
 		if (!linkedAddress)
 		{
-			OP::ConditionalArgs args;
-			memset(&args, 0, sizeof(args));
+			OP::ConditionalArgs Args;
+			FMemory::Memzero(Args);
 
-			if (condition) args.condition = condition->linkedAddress;
-			if (yes) args.yes = yes->linkedAddress;
-			if (no) args.no = no->linkedAddress;
+			if (condition) Args.condition = condition->linkedAddress;
+			if (yes) Args.yes = yes->linkedAddress;
+			if (no) Args.no = no->linkedAddress;
 
 			linkedAddress = (OP::ADDRESS)program.OpAddress.Num();
 			//program.m_code.push_back(op);
 			program.OpAddress.Add((uint32_t)program.ByteCode.Num());
 			AppendCode(program.ByteCode, type);
-			AppendCode(program.ByteCode, args);
+			AppendCode(program.ByteCode, Args);
 
 		}
 	}
@@ -298,8 +298,8 @@ namespace mu
 		// Constant condition?
 		if (condition->GetOpType() == EOpType::BO_CONSTANT)
 		{
-			const ASTOpConstantBool* typedCondition = static_cast<const ASTOpConstantBool*>(condition.child().get());
-			if (typedCondition->value)
+			const ASTOpConstantBool* TypedCondition = static_cast<const ASTOpConstantBool*>(condition.child().get());
+			if (TypedCondition->bValue)
 			{
 				return yes.child();
 			}
diff --git a/Engine/Plugins/Mutable/Source/MutableTools/Private/MuT/ASTOpConstantBool.cpp b/Engine/Plugins/Mutable/Source/MutableTools/Private/MuT/ASTOpConstantBool.cpp
--- a/Engine/Plugins/Mutable/Source/MutableTools/Private/MuT/ASTOpConstantBool.cpp
+++ b/Engine/Plugins/Mutable/Source/MutableTools/Private/MuT/ASTOpConstantBool.cpp
@@ -11,13 +11,9 @@
 namespace mu
 {
 
-
-	//-------------------------------------------------------------------------------------------------
-	//-------------------------------------------------------------------------------------------------
-	//-------------------------------------------------------------------------------------------------
-	ASTOpConstantBool::ASTOpConstantBool(bool v)
+	ASTOpConstantBool::ASTOpConstantBool(bool InValue)
 	{
-		value = v;
+		bValue = InValue;
 	}
 
 
@@ -26,12 +22,12 @@ namespace mu
 	}
 
 
-	bool ASTOpConstantBool::IsEqual(const ASTOp& otherUntyped) const
+	bool ASTOpConstantBool::IsEqual(const ASTOp& OtherUntyped) const
 	{
-		if (otherUntyped.GetOpType() == GetOpType())
+		if (OtherUntyped.GetOpType() == GetOpType())
 		{
-			const ASTOpConstantBool* other = static_cast<const ASTOpConstantBool*>(&otherUntyped);
-			return value == other->value;
+			const ASTOpConstantBool* other = static_cast<const ASTOpConstantBool*>(&OtherUntyped);
+			return bValue == other->bValue;
 		}
 		return false;
 	}
@@ -39,40 +35,39 @@ namespace mu
 
 	uint64 ASTOpConstantBool::Hash() const
 	{
-		uint64 res = std::hash<uint64>()(uint64(EOpType::BO_CONSTANT));
-		hash_combine(res, value);
-		return res;
+		uint64 Result = std::hash<uint64>()(uint64(EOpType::BO_CONSTANT));
+		hash_combine(Result, bValue);
+		return Result;
 	}
 
 
 	mu::Ptr<ASTOp> ASTOpConstantBool::Clone(MapChildFuncRef) const
 	{
-		Ptr<ASTOpConstantBool> n = new ASTOpConstantBool();
-		n->value = value;
-		return n;
+		Ptr<ASTOpConstantBool> New = new ASTOpConstantBool();
+		New->bValue = bValue;
+		return New;
 	}
 
 
-	void ASTOpConstantBool::Link(FProgram& program, FLinkerOptions*)
+	void ASTOpConstantBool::Link(FProgram& Program, FLinkerOptions*)
 	{
 		if (!linkedAddress)
 		{
-			OP::BoolConstantArgs args;
-			memset(&args, 0, sizeof(args));
-			args.value = value;
-
-			linkedAddress = (OP::ADDRESS)program.OpAddress.Num();
-			//        program.m_code.push_back(op);
-			program.OpAddress.Add((uint32_t)program.ByteCode.Num());
-			AppendCode(program.ByteCode, EOpType::BO_CONSTANT);
-			AppendCode(program.ByteCode, args);
+			OP::BoolConstantArgs Args;
+			FMemory::Memzero(Args);
+			Args.bValue = bValue;
+
+			linkedAddress = (OP::ADDRESS)Program.OpAddress.Num();
+			Program.OpAddress.Add(Program.ByteCode.Num());
+			AppendCode(Program.ByteCode, GetOpType());
+			AppendCode(Program.ByteCode, Args);
 		}
 	}
 
 
 	ASTOp::FBoolEvalResult ASTOpConstantBool::EvaluateBool(ASTOpList&, FEvaluateBoolCache*) const
 	{
-		return value ? BET_TRUE : BET_FALSE;
+		return bValue ? BET_TRUE : BET_FALSE;
 	}
 
 }
diff --git a/Engine/Plugins/Mutable/Source/MutableTools/Private/MuT/ASTOpConstantBool.h b/Engine/Plugins/Mutable/Source/MutableTools/Private/MuT/ASTOpConstantBool.h
--- a/Engine/Plugins/Mutable/Source/MutableTools/Private/MuT/ASTOpConstantBool.h
+++ b/Engine/Plugins/Mutable/Source/MutableTools/Private/MuT/ASTOpConstantBool.h
@@ -10,30 +10,25 @@
 
 namespace mu
 {
-struct FProgram;
+	struct FProgram;
 
-
-	//---------------------------------------------------------------------------------------------
-	//!
-	//---------------------------------------------------------------------------------------------
 	class ASTOpConstantBool final : public ASTOp
 	{
 	public:
 
-		//!
-		bool value;
+		bool bValue = true;
 
 	public:
 
-		ASTOpConstantBool(bool value = true);
+		ASTOpConstantBool(bool bValue = true);
 
-		EOpType GetOpType() const override { return EOpType::BO_CONSTANT; }
-		uint64 Hash() const override;
-		void ForEachChild(const TFunctionRef<void(ASTChild&)>) override;
-		bool IsEqual(const ASTOp& otherUntyped) const override;
-		Ptr<ASTOp> Clone(MapChildFuncRef mapChild) const override;
-		void Link(FProgram& program, FLinkerOptions* Options) override;
-		FBoolEvalResult EvaluateBool(ASTOpList& facts, FEvaluateBoolCache* cache) const override;
+		virtual EOpType GetOpType() const override { return EOpType::BO_CONSTANT; }
+		virtual uint64 Hash() const override;
+		virtual void ForEachChild(const TFunctionRef<void(ASTChild&)>) override;
+		virtual bool IsEqual(const ASTOp&) const override;
+		virtual Ptr<ASTOp> Clone(MapChildFuncRef) const override;
+		virtual void Link(FProgram&, FLinkerOptions*) override;
+		virtual FBoolEvalResult EvaluateBool(ASTOpList& Facts, FEvaluateBoolCache*) const override;
 	};
 
 
diff --git a/Engine/Plugins/Mutable/Source/MutableTools/Private/MuT/ASTOpConstantColor.cpp b/Engine/Plugins/Mutable/Source/MutableTools/Private/MuT/ASTOpConstantColor.cpp
new file mode 100644
--- /dev/null
+++ b/Engine/Plugins/Mutable/Source/MutableTools/Private/MuT/ASTOpConstantColor.cpp
@@ -0,0 +1,71 @@
+// Copyright Epic Games, Inc. All Rights Reserved.
+
+#include "MuT/ASTOpConstantColor.h"
+
+#include "HAL/PlatformMath.h"
+#include "MuR/ModelPrivate.h"
+#include "MuR/RefCounted.h"
+#include "MuR/Types.h"
+
+
+namespace mu
+{
+
+	void ASTOpConstantColor::ForEachChild(const TFunctionRef<void(ASTChild&)>)
+	{
+	}
+
+
+	bool ASTOpConstantColor::IsEqual(const ASTOp& OtherUntyped) const
+	{
+		if (OtherUntyped.GetOpType() == GetOpType())
+		{
+			const ASTOpConstantColor* other = static_cast<const ASTOpConstantColor*>(&OtherUntyped);
+			return Value == other->Value;
+		}
+		return false;
+	}
+
+
+	uint64 ASTOpConstantColor::Hash() const
+	{
+		uint64 Result = std::hash<uint64>()(uint64(GetOpType()));
+		hash_combine(Result, Value[0]);
+		hash_combine(Result, Value[1]);
+		hash_combine(Result, Value[2]);
+		hash_combine(Result, Value[3]);
+		return Result;
+	}
+
+
+	mu::Ptr<ASTOp> ASTOpConstantColor::Clone(MapChildFuncRef) const
+	{
+		Ptr<ASTOpConstantColor> New = new ASTOpConstantColor();
+		New->Value = Value;
+		return New;
+	}
+
+
+	void ASTOpConstantColor::Link(FProgram& Program, FLinkerOptions*)
+	{
+		if (!linkedAddress)
+		{
+			OP::ColorConstantArgs Args;
+			FMemory::Memzero(Args);
+			Args.Value = Value;
+
+			linkedAddress = (OP::ADDRESS)Program.OpAddress.Num();
+			Program.OpAddress.Add(Program.ByteCode.Num());
+			AppendCode(Program.ByteCode, GetOpType());
+			AppendCode(Program.ByteCode, Args);
+		}
+	}
+
+
+	bool ASTOpConstantColor::IsColourConstant(FVector4f& OutColour) const
+	{
+		OutColour = Value;
+		return true;
+	}
+
+}
diff --git a/Engine/Plugins/Mutable/Source/MutableTools/Private/MuT/ASTOpConstantColor.h b/Engine/Plugins/Mutable/Source/MutableTools/Private/MuT/ASTOpConstantColor.h
new file mode 100644
--- /dev/null
+++ b/Engine/Plugins/Mutable/Source/MutableTools/Private/MuT/ASTOpConstantColor.h
@@ -0,0 +1,33 @@
+// Copyright Epic Games, Inc. All Rights Reserved.
+
+#pragma once
+
+#include "HAL/Platform.h"
+#include "MuR/Operations.h"
+#include "MuR/Ptr.h"
+#include "MuT/AST.h"
+
+
+namespace mu
+{
+	struct FProgram;
+
+	class ASTOpConstantColor final : public ASTOp
+	{
+	public:
+
+		FVector4f Value = FVector4f(0.0f,0.0f,0.0f,0.0f);
+
+	public:
+
+		virtual EOpType GetOpType() const override { return EOpType::CO_CONSTANT; }
+		virtual uint64 Hash() const override;
+		virtual void ForEachChild(const TFunctionRef<void(ASTChild&)>) override;
+		virtual bool IsEqual(const ASTOp&) const override;
+		virtual Ptr<ASTOp> Clone(MapChildFuncRef) const override;
+		virtual void Link(FProgram&, FLinkerOptions*) override;
+		virtual bool IsColourConstant(FVector4f& OutColour) const override;
+	};
+
+}
+
diff --git a/Engine/Plugins/Mutable/Source/MutableTools/Private/MuT/ASTOpConstantExtensionData.cpp b/Engine/Plugins/Mutable/Source/MutableTools/Private/MuT/ASTOpConstantExtensionData.cpp
--- a/Engine/Plugins/Mutable/Source/MutableTools/Private/MuT/ASTOpConstantExtensionData.cpp
+++ b/Engine/Plugins/Mutable/Source/MutableTools/Private/MuT/ASTOpConstantExtensionData.cpp
@@ -47,7 +47,7 @@ void ASTOpConstantExtensionData::Link(FProgram& Program, FLinkerOptions* Options
 	if (!linkedAddress)
 	{
 		OP::ResourceConstantArgs Args;
-		FMemory::Memset(Args, 0);
+		FMemory::Memzero(Args);
 
 		Args.value = Program.AddConstant(Value);
 
diff --git a/Engine/Plugins/Mutable/Source/MutableTools/Private/MuT/ASTOpConstantInt.cpp b/Engine/Plugins/Mutable/Source/MutableTools/Private/MuT/ASTOpConstantInt.cpp
new file mode 100644
--- /dev/null
+++ b/Engine/Plugins/Mutable/Source/MutableTools/Private/MuT/ASTOpConstantInt.cpp
@@ -0,0 +1,74 @@
+// Copyright Epic Games, Inc. All Rights Reserved.
+
+#include "MuT/ASTOpConstantInt.h"
+
+#include "HAL/PlatformMath.h"
+#include "MuR/ModelPrivate.h"
+#include "MuR/RefCounted.h"
+#include "MuR/Types.h"
+
+
+namespace mu
+{
+
+	ASTOpConstantInt::ASTOpConstantInt(int32 InValue)
+	{
+		Value = InValue;
+	}
+
+
+	void ASTOpConstantInt::ForEachChild(const TFunctionRef<void(ASTChild&)>)
+	{
+	}
+
+
+	bool ASTOpConstantInt::IsEqual(const ASTOp& OtherUntyped) const
+	{
+		if (OtherUntyped.GetOpType() == GetOpType())
+		{
+			const ASTOpConstantInt* Other = static_cast<const ASTOpConstantInt*>(&OtherUntyped);
+			return Value == Other->Value;
+		}
+		return false;
+	}
+
+
+	uint64 ASTOpConstantInt::Hash() const
+	{
+		uint64 Result = std::hash<uint64>()(uint64(GetOpType()));
+		hash_combine(Result, Value);
+		return Result;
+	}
+
+
+	mu::Ptr<ASTOp> ASTOpConstantInt::Clone(MapChildFuncRef) const
+	{
+		Ptr<ASTOpConstantInt> New = new ASTOpConstantInt();
+		New->Value = Value;
+		return New;
+	}
+
+
+	void ASTOpConstantInt::Link(FProgram& Program, FLinkerOptions*)
+	{
+		if (!linkedAddress)
+		{
+			OP::IntConstantArgs Args;
+			FMemory::Memzero(Args);
+			Args.Value = Value;
+
+			linkedAddress = (OP::ADDRESS)Program.OpAddress.Num();
+			Program.OpAddress.Add(Program.ByteCode.Num());
+			AppendCode(Program.ByteCode, GetOpType());
+			AppendCode(Program.ByteCode, Args);
+		}
+	}
+
+
+	int32 ASTOpConstantInt::EvaluateInt(ASTOpList&, bool& bOutUnknown) const
+	{
+		bOutUnknown =  false;
+		return Value;
+	}
+
+}
diff --git a/Engine/Plugins/Mutable/Source/MutableTools/Private/MuT/ASTOpConstantInt.h b/Engine/Plugins/Mutable/Source/MutableTools/Private/MuT/ASTOpConstantInt.h
new file mode 100644
--- /dev/null
+++ b/Engine/Plugins/Mutable/Source/MutableTools/Private/MuT/ASTOpConstantInt.h
@@ -0,0 +1,35 @@
+// Copyright Epic Games, Inc. All Rights Reserved.
+
+#pragma once
+
+#include "HAL/Platform.h"
+#include "MuR/Operations.h"
+#include "MuR/Ptr.h"
+#include "MuT/AST.h"
+
+
+namespace mu
+{
+	struct FProgram;
+
+	class ASTOpConstantInt final : public ASTOp
+	{
+	public:
+
+		int32 Value = 0;
+
+	public:
+
+		ASTOpConstantInt(int Value = 0);
+
+		virtual EOpType GetOpType() const override { return EOpType::NU_CONSTANT; }
+		virtual uint64 Hash() const override;
+		virtual void ForEachChild(const TFunctionRef<void(ASTChild&)>) override;
+		virtual bool IsEqual(const ASTOp&) const override;
+		virtual Ptr<ASTOp> Clone(MapChildFuncRef) const override;
+		virtual void Link(FProgram&, FLinkerOptions*) override;
+		virtual int32 EvaluateInt(ASTOpList& Facts, bool& bOutUnknown) const override;
+	};
+
+}
+
diff --git a/Engine/Plugins/Mutable/Source/MutableTools/Private/MuT/ASTOpConstantMatrix.cpp b/Engine/Plugins/Mutable/Source/MutableTools/Private/MuT/ASTOpConstantMatrix.cpp
--- a/Engine/Plugins/Mutable/Source/MutableTools/Private/MuT/ASTOpConstantMatrix.cpp
+++ b/Engine/Plugins/Mutable/Source/MutableTools/Private/MuT/ASTOpConstantMatrix.cpp
@@ -35,12 +35,13 @@ void mu::ASTOpConstantMatrix::Link(FProgram& program, FLinkerOptions* Options)
 {
 	if (!linkedAddress)
 	{
-		OP::MatrixConstantArgs args;
-		args.value = program.AddConstant(value);
+		OP::MatrixConstantArgs Args;
+		FMemory::Memzero(Args);
+		Args.value = program.AddConstant(value);
 
 		linkedAddress = static_cast<OP::ADDRESS>(program.OpAddress.Num());
 		program.OpAddress.Add(static_cast<uint32_t>(program.ByteCode.Num()));
 		AppendCode(program.ByteCode, EOpType::MA_CONSTANT);
-		AppendCode(program.ByteCode, args);
+		AppendCode(program.ByteCode, Args);
 	}
 }
diff --git a/Engine/Plugins/Mutable/Source/MutableTools/Private/MuT/ASTOpConstantProjector.cpp b/Engine/Plugins/Mutable/Source/MutableTools/Private/MuT/ASTOpConstantProjector.cpp
--- a/Engine/Plugins/Mutable/Source/MutableTools/Private/MuT/ASTOpConstantProjector.cpp
+++ b/Engine/Plugins/Mutable/Source/MutableTools/Private/MuT/ASTOpConstantProjector.cpp
@@ -50,15 +50,15 @@ namespace mu
 	{
 		if (!linkedAddress)
 		{
-			OP::ResourceConstantArgs args;
-			memset(&args, 0, sizeof(args));
-			args.value = program.AddConstant(value);
+			OP::ResourceConstantArgs Args;
+			FMemory::Memzero(Args);
+			Args.value = program.AddConstant(value);
 
 			linkedAddress = (OP::ADDRESS)program.OpAddress.Num();
 			//program.m_code.push_back(op);
 			program.OpAddress.Add((uint32_t)program.ByteCode.Num());
 			AppendCode(program.ByteCode, EOpType::PR_CONSTANT);
-			AppendCode(program.ByteCode, args);
+			AppendCode(program.ByteCode, Args);
 		}
 	}
 
diff --git a/Engine/Plugins/Mutable/Source/MutableTools/Private/MuT/ASTOpConstantResource.cpp b/Engine/Plugins/Mutable/Source/MutableTools/Private/MuT/ASTOpConstantResource.cpp
--- a/Engine/Plugins/Mutable/Source/MutableTools/Private/MuT/ASTOpConstantResource.cpp
+++ b/Engine/Plugins/Mutable/Source/MutableTools/Private/MuT/ASTOpConstantResource.cpp
@@ -409,23 +409,23 @@ namespace mu
 		{
 			if (Type == EOpType::ME_CONSTANT)
 			{
-				OP::MeshConstantArgs args;
-				FMemory::Memset(&args, 0, sizeof(args));
+				OP::MeshConstantArgs Args;
+				FMemory::Memzero(Args);
 
 				TSharedPtr<FMesh> MeshData = static_cast<const FMesh*>(GetValue().Get())->Clone();
 				check(MeshData);
 
-				args.skeleton = -1;
+				Args.skeleton = -1;
 				if (TSharedPtr<const FSkeleton> pSkeleton = MeshData->GetSkeleton())
 				{
-					args.skeleton = Program.AddConstant(pSkeleton);
+					Args.skeleton = Program.AddConstant(pSkeleton);
 					MeshData->SetSkeleton(nullptr);
 				}
 
-				args.physicsBody = -1;
+				Args.physicsBody = -1;
 				if (TSharedPtr<const FPhysicsBody> pPhysicsBody = MeshData->GetPhysicsBody())
 				{
-					args.physicsBody = Program.AddConstant(pPhysicsBody);
+					Args.physicsBody = Program.AddConstant(pPhysicsBody);
 					MeshData->SetPhysicsBody(nullptr);
 				}
 
@@ -434,27 +434,27 @@ namespace mu
 				const int32* IndexPtr = Options->MeshConstantMap.Find(Key);
 				if (!IndexPtr)
 				{
-					args.value = Program.ConstantMeshesPermanent.Add(MeshData);
+					Args.value = Program.ConstantMeshesPermanent.Add(MeshData);
 
 					int32 DataDescIndex = Options->AdditionalData.SourceMeshPerConstant.Add(SourceDataDescriptor);
-					check(DataDescIndex == args.value);
+					check(DataDescIndex == Args.value);
 
-					Options->MeshConstantMap.Add(MeshData, int32(args.value));
+					Options->MeshConstantMap.Add(MeshData, int32(Args.value));
 				}
 				else
 				{
-					args.value = *IndexPtr;
+					Args.value = *IndexPtr;
 				}
 
 				linkedAddress = (OP::ADDRESS)Program.OpAddress.Num();
 				Program.OpAddress.Add((uint32)Program.ByteCode.Num());
 				AppendCode(Program.ByteCode, Type);
-				AppendCode(Program.ByteCode, args);
+				AppendCode(Program.ByteCode, Args);
 			}
 			else
 			{
-				OP::ResourceConstantArgs args;
-				FMemory::Memset(&args, 0, sizeof(args));
+				OP::ResourceConstantArgs Args;
+				FMemory::Memzero(Args);
 
 				bool bValidData = true;
 
@@ -472,10 +472,10 @@ namespace mu
 					}
 					else
 					{
-						args.value = AddConstantImage( Program, pTyped, *Options);
+						Args.value = AddConstantImage( Program, pTyped, *Options);
 
 						int32 DataDescIndex = Options->AdditionalData.SourceImagePerConstant.Add(SourceDataDescriptor);
-						check(DataDescIndex == args.value);
+						check(DataDescIndex == Args.value);
 					}
 
 					break;
@@ -484,7 +484,7 @@ namespace mu
 				{
 					TSharedPtr<const FLayout> pTyped = StaticCastSharedPtr<const FLayout>(GetValue());
 					check(pTyped);
-					args.value = Program.AddConstant(pTyped);
+					Args.value = Program.AddConstant(pTyped);
 					break;
 				}
 				default:
@@ -496,7 +496,7 @@ namespace mu
 					linkedAddress = (OP::ADDRESS)Program.OpAddress.Num();
 					Program.OpAddress.Add((uint32)Program.ByteCode.Num());
 					AppendCode(Program.ByteCode, Type);
-					AppendCode(Program.ByteCode, args);
+					AppendCode(Program.ByteCode, Args);
 				}
 				else
 				{
diff --git a/Engine/Plugins/Mutable/Source/MutableTools/Private/MuT/ASTOpConstantScalar.cpp b/Engine/Plugins/Mutable/Source/MutableTools/Private/MuT/ASTOpConstantScalar.cpp
new file mode 100644
--- /dev/null
+++ b/Engine/Plugins/Mutable/Source/MutableTools/Private/MuT/ASTOpConstantScalar.cpp
@@ -0,0 +1,74 @@
+// Copyright Epic Games, Inc. All Rights Reserved.
+
+#include "MuT/ASTOpConstantScalar.h"
+
+#include "HAL/PlatformMath.h"
+#include "MuR/ModelPrivate.h"
+#include "MuR/RefCounted.h"
+#include "MuR/Types.h"
+
+
+namespace mu
+{
+
+	ASTOpConstantScalar::ASTOpConstantScalar(float InValue)
+	{
+		Value = InValue;
+	}
+
+
+	void ASTOpConstantScalar::ForEachChild(const TFunctionRef<void(ASTChild&)>)
+	{
+	}
+
+
+	bool ASTOpConstantScalar::IsEqual(const ASTOp& OtherUntyped) const
+	{
+		if (OtherUntyped.GetOpType() == GetOpType())
+		{
+			const ASTOpConstantScalar* other = static_cast<const ASTOpConstantScalar*>(&OtherUntyped);
+			return Value == other->Value;
+		}
+		return false;
+	}
+
+
+	uint64 ASTOpConstantScalar::Hash() const
+	{
+		uint64 Result = std::hash<uint64>()(uint64(GetOpType()));
+		hash_combine(Result, Value);
+		return Result;
+	}
+
+
+	mu::Ptr<ASTOp> ASTOpConstantScalar::Clone(MapChildFuncRef) const
+	{
+		Ptr<ASTOpConstantScalar> New = new ASTOpConstantScalar();
+		New->Value = Value;
+		return New;
+	}
+
+
+	void ASTOpConstantScalar::Link(FProgram& Program, FLinkerOptions*)
+	{
+		if (!linkedAddress)
+		{
+			OP::ScalarConstantArgs Args;
+			FMemory::Memzero(Args);
+			Args.Value = Value;
+
+			linkedAddress = (OP::ADDRESS)Program.OpAddress.Num();
+			Program.OpAddress.Add(Program.ByteCode.Num());
+			AppendCode(Program.ByteCode, GetOpType());
+			AppendCode(Program.ByteCode, Args);
+		}
+	}
+
+
+	int32 ASTOpConstantScalar::EvaluateInt(ASTOpList&, bool& bOutUnknown) const
+	{
+		bOutUnknown =  false;
+		return Value;
+	}
+
+}
diff --git a/Engine/Plugins/Mutable/Source/MutableTools/Private/MuT/ASTOpConstantScalar.h b/Engine/Plugins/Mutable/Source/MutableTools/Private/MuT/ASTOpConstantScalar.h
new file mode 100644
--- /dev/null
+++ b/Engine/Plugins/Mutable/Source/MutableTools/Private/MuT/ASTOpConstantScalar.h
@@ -0,0 +1,35 @@
+// Copyright Epic Games, Inc. All Rights Reserved.
+
+#pragma once
+
+#include "HAL/Platform.h"
+#include "MuR/Operations.h"
+#include "MuR/Ptr.h"
+#include "MuT/AST.h"
+
+
+namespace mu
+{
+	struct FProgram;
+
+	class ASTOpConstantScalar final : public ASTOp
+	{
+	public:
+
+		float Value = 0.0f;
+
+	public:
+
+		ASTOpConstantScalar(float Value = 0.0f);
+
+		virtual EOpType GetOpType() const override { return EOpType::SC_CONSTANT; }
+		virtual uint64 Hash() const override;
+		virtual void ForEachChild(const TFunctionRef<void(ASTChild&)>) override;
+		virtual bool IsEqual(const ASTOp&) const override;
+		virtual Ptr<ASTOp> Clone(MapChildFuncRef) const override;
+		virtual void Link(FProgram&, FLinkerOptions*) override;
+		virtual int32 EvaluateInt(ASTOpList& Facts, bool& bOutUnknown) const override;
+	};
+
+}
+
diff --git a/Engine/Plugins/Mutable/Source/MutableTools/Private/MuT/ASTOpConstantString.cpp b/Engine/Plugins/Mutable/Source/MutableTools/Private/MuT/ASTOpConstantString.cpp
--- a/Engine/Plugins/Mutable/Source/MutableTools/Private/MuT/ASTOpConstantString.cpp
+++ b/Engine/Plugins/Mutable/Source/MutableTools/Private/MuT/ASTOpConstantString.cpp
@@ -50,14 +50,14 @@ namespace mu
 	{
 		if (!linkedAddress)
 		{
-			OP::ResourceConstantArgs args;
-			FMemory::Memset(&args, 0, sizeof(args));
-			args.value = program.AddConstant(value);
+			OP::ResourceConstantArgs Args;
+			FMemory::Memzero(Args);
+			Args.value = program.AddConstant(value);
 
 			linkedAddress = (OP::ADDRESS)program.OpAddress.Num();
 			program.OpAddress.Add((uint32_t)program.ByteCode.Num());
 			AppendCode(program.ByteCode, EOpType::ST_CONSTANT);
-			AppendCode(program.ByteCode, args);
+			AppendCode(program.ByteCode, Args);
 		}
 	}
 
diff --git a/Engine/Plugins/Mutable/Source/MutableTools/Private/MuT/ASTOpImageBinarize.cpp b/Engine/Plugins/Mutable/Source/MutableTools/Private/MuT/ASTOpImageBinarize.cpp
--- a/Engine/Plugins/Mutable/Source/MutableTools/Private/MuT/ASTOpImageBinarize.cpp
+++ b/Engine/Plugins/Mutable/Source/MutableTools/Private/MuT/ASTOpImageBinarize.cpp
@@ -80,7 +80,8 @@ namespace mu
 		// Already linked?
 		if (!linkedAddress)
 		{
-			OP::ImageBinariseArgs Args = {};
+			OP::ImageBinariseArgs Args;
+			FMemory::Memzero(Args);
 
 			if (Base)
 			{
diff --git a/Engine/Plugins/Mutable/Source/MutableTools/Private/MuT/ASTOpImageBlankLayout.cpp b/Engine/Plugins/Mutable/Source/MutableTools/Private/MuT/ASTOpImageBlankLayout.cpp
--- a/Engine/Plugins/Mutable/Source/MutableTools/Private/MuT/ASTOpImageBlankLayout.cpp
+++ b/Engine/Plugins/Mutable/Source/MutableTools/Private/MuT/ASTOpImageBlankLayout.cpp
@@ -82,7 +82,8 @@ namespace mu
 		// Already linked?
 		if (!linkedAddress)
 		{
-			OP::ImageBlankLayoutArgs Args={};
+			OP::ImageBlankLayoutArgs Args;
+			FMemory::Memzero(Args);
 
 			if (Layout)
 			{
diff --git a/Engine/Plugins/Mutable/Source/MutableTools/Private/MuT/ASTOpImageColorMap.cpp b/Engine/Plugins/Mutable/Source/MutableTools/Private/MuT/ASTOpImageColorMap.cpp
--- a/Engine/Plugins/Mutable/Source/MutableTools/Private/MuT/ASTOpImageColorMap.cpp
+++ b/Engine/Plugins/Mutable/Source/MutableTools/Private/MuT/ASTOpImageColorMap.cpp
@@ -85,7 +85,8 @@ namespace mu
 		// Already linked?
 		if (!linkedAddress)
 		{
-			OP::ImageColourMapArgs Args = {};
+			OP::ImageColourMapArgs Args;
+			FMemory::Memzero(Args);
 
 			if (Base)
 			{
diff --git a/Engine/Plugins/Mutable/Source/MutableTools/Private/MuT/ASTOpImageCompose.cpp b/Engine/Plugins/Mutable/Source/MutableTools/Private/MuT/ASTOpImageCompose.cpp
--- a/Engine/Plugins/Mutable/Source/MutableTools/Private/MuT/ASTOpImageCompose.cpp
+++ b/Engine/Plugins/Mutable/Source/MutableTools/Private/MuT/ASTOpImageCompose.cpp
@@ -99,19 +99,19 @@ void ASTOpImageCompose::Link( FProgram& program, FLinkerOptions*)
     // Already linked?
     if (!linkedAddress)
     {
-        OP::ImageComposeArgs args;
-        memset( &args,0, sizeof(args) );
+        OP::ImageComposeArgs Args;
+		FMemory::Memzero(Args);
 
-		if (Layout) args.layout = Layout->linkedAddress;
-		if (Base) args.base = Base->linkedAddress;
-		if (BlockImage) args.blockImage = BlockImage->linkedAddress;
-		if (Mask) args.mask = Mask->linkedAddress;
-        args.BlockId = BlockId;
+		if (Layout) Args.layout = Layout->linkedAddress;
+		if (Base) Args.base = Base->linkedAddress;
+		if (BlockImage) Args.blockImage = BlockImage->linkedAddress;
+		if (Mask) Args.mask = Mask->linkedAddress;
+        Args.BlockId = BlockId;
 
         linkedAddress = (OP::ADDRESS)program.OpAddress.Num();
         program.OpAddress.Add((uint32_t)program.ByteCode.Num());
         AppendCode(program.ByteCode,EOpType::IM_COMPOSE);
-        AppendCode(program.ByteCode,args);
+        AppendCode(program.ByteCode,Args);
     }
 
 }
diff --git a/Engine/Plugins/Mutable/Source/MutableTools/Private/MuT/ASTOpImageCrop.cpp b/Engine/Plugins/Mutable/Source/MutableTools/Private/MuT/ASTOpImageCrop.cpp
--- a/Engine/Plugins/Mutable/Source/MutableTools/Private/MuT/ASTOpImageCrop.cpp
+++ b/Engine/Plugins/Mutable/Source/MutableTools/Private/MuT/ASTOpImageCrop.cpp
@@ -80,13 +80,15 @@ namespace mu
 		// Already linked?
 		if (!linkedAddress)
 		{
-			OP::ImageCropArgs Args={};
+			OP::ImageCropArgs Args;
+			FMemory::Memzero(Args);
 
 			if (Source)
 			{
 				Args.source = Source->linkedAddress;
 			}
 
+			if (Source) Args.source = Source->linkedAddress;
 			Args.minX = Min[0];
 			Args.minY = Min[1];
 			Args.sizeX = Size[0];
diff --git a/Engine/Plugins/Mutable/Source/MutableTools/Private/MuT/ASTOpImageDisplace.cpp b/Engine/Plugins/Mutable/Source/MutableTools/Private/MuT/ASTOpImageDisplace.cpp
--- a/Engine/Plugins/Mutable/Source/MutableTools/Private/MuT/ASTOpImageDisplace.cpp
+++ b/Engine/Plugins/Mutable/Source/MutableTools/Private/MuT/ASTOpImageDisplace.cpp
@@ -80,7 +80,8 @@ namespace mu
 		// Already linked?
 		if (!linkedAddress)
 		{
-			OP::ImageDisplaceArgs Args = {};
+			OP::ImageDisplaceArgs Args;
+			FMemory::Memzero(Args);
 
 			if (Source)
 			{
diff --git a/Engine/Plugins/Mutable/Source/MutableTools/Private/MuT/ASTOpImageInterpolate.cpp b/Engine/Plugins/Mutable/Source/MutableTools/Private/MuT/ASTOpImageInterpolate.cpp
--- a/Engine/Plugins/Mutable/Source/MutableTools/Private/MuT/ASTOpImageInterpolate.cpp
+++ b/Engine/Plugins/Mutable/Source/MutableTools/Private/MuT/ASTOpImageInterpolate.cpp
@@ -98,7 +98,8 @@ namespace mu
 		// Already linked?
 		if (!linkedAddress)
 		{
-			OP::ImageInterpolateArgs Args = {};
+			OP::ImageInterpolateArgs Args;
+			FMemory::Memzero(Args);
 
 			if (Factor)
 			{
diff --git a/Engine/Plugins/Mutable/Source/MutableTools/Private/MuT/ASTOpImageInvert.cpp b/Engine/Plugins/Mutable/Source/MutableTools/Private/MuT/ASTOpImageInvert.cpp
--- a/Engine/Plugins/Mutable/Source/MutableTools/Private/MuT/ASTOpImageInvert.cpp
+++ b/Engine/Plugins/Mutable/Source/MutableTools/Private/MuT/ASTOpImageInvert.cpp
@@ -75,7 +75,8 @@ namespace mu
 		// Already linked?
 		if (!linkedAddress)
 		{
-			OP::ImageInvertArgs Args = {};
+			OP::ImageInvertArgs Args;
+			FMemory::Memzero(Args);
 
 			if (Base)
 			{
diff --git a/Engine/Plugins/Mutable/Source/MutableTools/Private/MuT/ASTOpImageLayer.cpp b/Engine/Plugins/Mutable/Source/MutableTools/Private/MuT/ASTOpImageLayer.cpp
--- a/Engine/Plugins/Mutable/Source/MutableTools/Private/MuT/ASTOpImageLayer.cpp
+++ b/Engine/Plugins/Mutable/Source/MutableTools/Private/MuT/ASTOpImageLayer.cpp
@@ -98,22 +98,22 @@ namespace mu
 		// Already linked?
 		if (!linkedAddress)
 		{
-			OP::ImageLayerArgs args;
-			FMemory::Memzero(&args, sizeof(args));
+			OP::ImageLayerArgs Args;
+			FMemory::Memzero(Args);
 
-			args.blendType = (uint8)blendType;
-			args.blendTypeAlpha = (uint8)blendTypeAlpha;
-			args.BlendAlphaSourceChannel = BlendAlphaSourceChannel;
-			args.flags = Flags;
+			Args.blendType = (uint8)blendType;
+			Args.blendTypeAlpha = (uint8)blendTypeAlpha;
+			Args.BlendAlphaSourceChannel = BlendAlphaSourceChannel;
+			Args.flags = Flags;
 
-			if (base) args.base = base->linkedAddress;
-			if (blend) args.blended = blend->linkedAddress;
-			if (mask) args.mask = mask->linkedAddress;
+			if (base) Args.base = base->linkedAddress;
+			if (blend) Args.blended = blend->linkedAddress;
+			if (mask) Args.mask = mask->linkedAddress;
 
 			linkedAddress = (OP::ADDRESS)program.OpAddress.Num();
 			program.OpAddress.Add((uint32)program.ByteCode.Num());
 			AppendCode(program.ByteCode, GetOpType());
-			AppendCode(program.ByteCode, args);
+			AppendCode(program.ByteCode, Args);
 		}
 	}
 
@@ -561,11 +561,11 @@ namespace mu
 //                OP op = program.m_code[baseAt];
 
 //                OP aOp = program.m_code[at];
-//                aOp.args.ImageLayer.base = program.m_code[baseAt].args.Conditional.yes;
+//                aOp.Args.ImageLayer.base = program.m_code[baseAt].Args.Conditional.yes;
 //                op.args.Conditional.yes = program.AddOp( aOp );
 
 //                OP bOp = program.m_code[at];
-//                bOp.args.ImageLayer.base = program.m_code[baseAt].args.Conditional.no;
+//                bOp.Args.ImageLayer.base = program.m_code[baseAt].Args.Conditional.no;
 //                op.args.Conditional.no = program.AddOp( bOp );
 
 //                at = program.AddOp( op );
@@ -642,12 +642,12 @@ namespace mu
 //                };
 //                auto newRoot = neutral.Apply
 //                         (
-//                            program.m_code[at].args.ImageLayer.base, program,
+//                            program.m_code[at].Args.ImageLayer.base, program,
 //                            Image, supportedOps,
-//                            &templateOp, &templateOp.args.ImageLayer.base);
+//                            &templateOp, &templateOp.Args.ImageLayer.base);
 
 //                // If there is a change
-//                if (newRoot!=program.m_code[at].args.ImageLayer.base)
+//                if (newRoot!=program.m_code[at].Args.ImageLayer.base)
 //                {
 //                    bModified = true;
 //                    at = newRoot;
diff --git a/Engine/Plugins/Mutable/Source/MutableTools/Private/MuT/ASTOpImageLayerColor.cpp b/Engine/Plugins/Mutable/Source/MutableTools/Private/MuT/ASTOpImageLayerColor.cpp
--- a/Engine/Plugins/Mutable/Source/MutableTools/Private/MuT/ASTOpImageLayerColor.cpp
+++ b/Engine/Plugins/Mutable/Source/MutableTools/Private/MuT/ASTOpImageLayerColor.cpp
@@ -6,6 +6,7 @@
 #include "MuT/ASTOpImageSwizzle.h"
 #include "MuT/ASTOpImageRasterMesh.h"
 #include "MuT/ASTOpImageDisplace.h"
+#include "MuT/ASTOpConstantColor.h"
 #include "Containers/Map.h"
 #include "HAL/PlatformMath.h"
 #include "MuR/ModelPrivate.h"
@@ -91,23 +92,23 @@ namespace mu
 		// Already linked?
 		if (!linkedAddress)
 		{
-			OP::ImageLayerColourArgs args;
-			FMemory::Memzero(&args, sizeof(args));
+			OP::ImageLayerColourArgs Args;
+			FMemory::Memzero(Args);
 
-			args.blendType = (uint8)blendType;
-			args.blendTypeAlpha = (uint8)blendTypeAlpha;
-			args.BlendAlphaSourceChannel = BlendAlphaSourceChannel;
-			args.flags = Flags;
+			Args.blendType = (uint8)blendType;
+			Args.blendTypeAlpha = (uint8)blendTypeAlpha;
+			Args.BlendAlphaSourceChannel = BlendAlphaSourceChannel;
+			Args.flags = Flags;
 
 			check(base);
-			if (base) args.base = base->linkedAddress;
-			if (color) args.colour = color->linkedAddress;
-			if (mask) args.mask = mask->linkedAddress;
+			if (base) Args.base = base->linkedAddress;
+			if (color) Args.colour = color->linkedAddress;
+			if (mask) Args.mask = mask->linkedAddress;
 
 			linkedAddress = (OP::ADDRESS)program.OpAddress.Num();
 			program.OpAddress.Add((uint32)program.ByteCode.Num());
 			AppendCode(program.ByteCode, GetOpType());
-			AppendCode(program.ByteCode, args);
+			AppendCode(program.ByteCode, Args);
 		}
 	}
 
@@ -210,9 +211,8 @@ namespace mu
 			{
 				if (color)
 				{
-					const ASTOpFixed* TypedColor = static_cast<const ASTOpFixed*>(color.child().get());
-					const float* Value = TypedColor->op.args.ColourConstant.value;
-					ColorConst.Set(Value[0], Value[1], Value[2], Value[3]);
+					const ASTOpConstantColor* TypedColor = static_cast<const ASTOpConstantColor*>(color.child().get());
+					ColorConst = TypedColor->Value;
 				}
 
 				if (!bAlphaUnchanged)
@@ -341,11 +341,11 @@ namespace mu
 //                OP op = program.m_code[baseAt];
 
 //                OP aOp = program.m_code[at];
-//                aOp.args.ImageLayerColour.base = program.m_code[baseAt].args.Conditional.yes;
+//                aOp.Args.ImageLayerColour.base = program.m_code[baseAt].Args.Conditional.yes;
 //                op.args.Conditional.yes = program.AddOp( aOp );
 
 //                OP bOp = program.m_code[at];
-//                bOp.args.ImageLayerColour.base = program.m_code[baseAt].args.Conditional.no;
+//                bOp.Args.ImageLayerColour.base = program.m_code[baseAt].Args.Conditional.no;
 //                op.args.Conditional.no = program.AddOp( bOp );
 
 //                at = program.AddOp( op );
diff --git a/Engine/Plugins/Mutable/Source/MutableTools/Private/MuT/ASTOpImageLuminance.cpp b/Engine/Plugins/Mutable/Source/MutableTools/Private/MuT/ASTOpImageLuminance.cpp
--- a/Engine/Plugins/Mutable/Source/MutableTools/Private/MuT/ASTOpImageLuminance.cpp
+++ b/Engine/Plugins/Mutable/Source/MutableTools/Private/MuT/ASTOpImageLuminance.cpp
@@ -75,7 +75,8 @@ namespace mu
 		// Already linked?
 		if (!linkedAddress)
 		{
-			OP::ImageLuminanceArgs Args = {};
+			OP::ImageLuminanceArgs Args;
+			FMemory::Memzero(Args);
 
 			if (Base)
 			{
diff --git a/Engine/Plugins/Mutable/Source/MutableTools/Private/MuT/ASTOpImageMakeGrowMap.cpp b/Engine/Plugins/Mutable/Source/MutableTools/Private/MuT/ASTOpImageMakeGrowMap.cpp
--- a/Engine/Plugins/Mutable/Source/MutableTools/Private/MuT/ASTOpImageMakeGrowMap.cpp
+++ b/Engine/Plugins/Mutable/Source/MutableTools/Private/MuT/ASTOpImageMakeGrowMap.cpp
@@ -74,17 +74,17 @@ namespace mu
 		// Already linked?
 		if (!linkedAddress)
 		{
-			OP::ImageMakeGrowMapArgs args;
-			FMemory::Memzero(&args, sizeof(args));
+			OP::ImageMakeGrowMapArgs Args;
+			FMemory::Memzero(Args);
 
-			args.border = Border;
-			if (Mask) args.mask = Mask->linkedAddress;
+			Args.border = Border;
+			if (Mask) Args.mask = Mask->linkedAddress;
 
 			linkedAddress = (OP::ADDRESS)program.OpAddress.Num();
 			//program.m_code.push_back(op);
 			program.OpAddress.Add((uint32_t)program.ByteCode.Num());
 			AppendCode(program.ByteCode, GetOpType());
-			AppendCode(program.ByteCode, args);
+			AppendCode(program.ByteCode, Args);
 		}
 	}
 
diff --git a/Engine/Plugins/Mutable/Source/MutableTools/Private/MuT/ASTOpImageMipmap.cpp b/Engine/Plugins/Mutable/Source/MutableTools/Private/MuT/ASTOpImageMipmap.cpp
--- a/Engine/Plugins/Mutable/Source/MutableTools/Private/MuT/ASTOpImageMipmap.cpp
+++ b/Engine/Plugins/Mutable/Source/MutableTools/Private/MuT/ASTOpImageMipmap.cpp
@@ -87,21 +87,21 @@ namespace mu
 		// Already linked?
 		if (!linkedAddress)
 		{
-			OP::ImageMipmapArgs args;
-			FMemory::Memzero(&args, sizeof(args));
+			OP::ImageMipmapArgs Args;
+			FMemory::Memzero(Args);
 
-			args.levels = Levels;
-			args.blockLevels = BlockLevels;
-			args.onlyTail = bOnlyTail;
-			args.AddressMode = AddressMode;
-			args.FilterType = FilterType;
-			if (Source) args.source = Source->linkedAddress;
+			Args.levels = Levels;
+			Args.blockLevels = BlockLevels;
+			Args.onlyTail = bOnlyTail;
+			Args.AddressMode = AddressMode;
+			Args.FilterType = FilterType;
+			if (Source) Args.source = Source->linkedAddress;
 
 			linkedAddress = (OP::ADDRESS)program.OpAddress.Num();
 			//program.m_code.push_back(op);
 			program.OpAddress.Add((uint32_t)program.ByteCode.Num());
 			AppendCode(program.ByteCode, GetOpType());
-			AppendCode(program.ByteCode, args);
+			AppendCode(program.ByteCode, Args);
 		}
 	}
 
diff --git a/Engine/Plugins/Mutable/Source/MutableTools/Private/MuT/ASTOpImageMultiLayer.cpp b/Engine/Plugins/Mutable/Source/MutableTools/Private/MuT/ASTOpImageMultiLayer.cpp
--- a/Engine/Plugins/Mutable/Source/MutableTools/Private/MuT/ASTOpImageMultiLayer.cpp
+++ b/Engine/Plugins/Mutable/Source/MutableTools/Private/MuT/ASTOpImageMultiLayer.cpp
@@ -94,26 +94,26 @@ namespace mu
 		// Already linked?
 		if (!linkedAddress)
 		{
-			OP::ImageMultiLayerArgs args;
-			FMemory::Memzero(&args, sizeof(args));
+			OP::ImageMultiLayerArgs Args;
+			FMemory::Memzero(Args);
 
-			args.blendType = (uint8)blendType;
-			args.blendTypeAlpha = (uint8)blendTypeAlpha;
-			args.BlendAlphaSourceChannel = BlendAlphaSourceChannel;
-			args.bUseMaskFromBlended = bUseMaskFromBlended;
+			Args.blendType = (uint8)blendType;
+			Args.blendTypeAlpha = (uint8)blendTypeAlpha;
+			Args.BlendAlphaSourceChannel = BlendAlphaSourceChannel;
+			Args.bUseMaskFromBlended = bUseMaskFromBlended;
 
-			if (base) args.base = base->linkedAddress;
-			if (blend) args.blended = blend->linkedAddress;
-			if (mask) args.mask = mask->linkedAddress;
+			if (base) Args.base = base->linkedAddress;
+			if (blend) Args.blended = blend->linkedAddress;
+			if (mask) Args.mask = mask->linkedAddress;
 			if (range.rangeSize)
 			{
-				LinkRange(program, range, args.rangeSize, args.rangeId);
+				LinkRange(program, range, Args.rangeSize, Args.rangeId);
 			}
 
 			linkedAddress = (OP::ADDRESS)program.OpAddress.Num();
 			program.OpAddress.Add((uint32)program.ByteCode.Num());
 			AppendCode(program.ByteCode, GetOpType());
-			AppendCode(program.ByteCode, args);
+			AppendCode(program.ByteCode, Args);
 		}
 	}
 
diff --git a/Engine/Plugins/Mutable/Source/MutableTools/Private/MuT/ASTOpImageNormalComposite.cpp b/Engine/Plugins/Mutable/Source/MutableTools/Private/MuT/ASTOpImageNormalComposite.cpp
--- a/Engine/Plugins/Mutable/Source/MutableTools/Private/MuT/ASTOpImageNormalComposite.cpp
+++ b/Engine/Plugins/Mutable/Source/MutableTools/Private/MuT/ASTOpImageNormalComposite.cpp
@@ -78,26 +78,26 @@ namespace mu
 		// Already linked?
 		if (!linkedAddress)
 		{
-			OP::ImageNormalCompositeArgs args;
-			FMemory::Memset(&args, 0, sizeof(args));
+			OP::ImageNormalCompositeArgs Args;
+			FMemory::Memzero(Args);
 
 			if (Base)
 			{
-				args.base = Base->linkedAddress;
+				Args.base = Base->linkedAddress;
 			}
 
 			if (Normal)
 			{
-				args.normal = Normal->linkedAddress;
+				Args.normal = Normal->linkedAddress;
 			}
 
-			args.power = Power;
-			args.mode = Mode;
+			Args.power = Power;
+			Args.mode = Mode;
 
 			linkedAddress = (OP::ADDRESS)program.OpAddress.Num();
 			program.OpAddress.Add((uint32_t)program.ByteCode.Num());
 			AppendCode(program.ByteCode, EOpType::IM_NORMALCOMPOSITE);
-			AppendCode(program.ByteCode, args);
+			AppendCode(program.ByteCode, Args);
 		}
 
 	}
diff --git a/Engine/Plugins/Mutable/Source/MutableTools/Private/MuT/ASTOpImagePatch.cpp b/Engine/Plugins/Mutable/Source/MutableTools/Private/MuT/ASTOpImagePatch.cpp
--- a/Engine/Plugins/Mutable/Source/MutableTools/Private/MuT/ASTOpImagePatch.cpp
+++ b/Engine/Plugins/Mutable/Source/MutableTools/Private/MuT/ASTOpImagePatch.cpp
@@ -75,18 +75,18 @@ namespace mu
 		// Already linked?
 		if (!linkedAddress)
 		{
-			OP::ImagePatchArgs args;
-			memset(&args, 0, sizeof(args));
+			OP::ImagePatchArgs Args;
+			FMemory::Memzero(Args);
 
-			if (base) args.base = base->linkedAddress;
-			if (patch) args.patch = patch->linkedAddress;
-			args.minX = location[0];
-			args.minY = location[1];
+			if (base) Args.base = base->linkedAddress;
+			if (patch) Args.patch = patch->linkedAddress;
+			Args.minX = location[0];
+			Args.minY = location[1];
 
 			linkedAddress = (OP::ADDRESS)program.OpAddress.Num();
 			program.OpAddress.Add((uint32_t)program.ByteCode.Num());
 			AppendCode(program.ByteCode, EOpType::IM_PATCH);
-			AppendCode(program.ByteCode, args);
+			AppendCode(program.ByteCode, Args);
 		}
 
 	}
diff --git a/Engine/Plugins/Mutable/Source/MutableTools/Private/MuT/ASTOpImagePixelFormat.cpp b/Engine/Plugins/Mutable/Source/MutableTools/Private/MuT/ASTOpImagePixelFormat.cpp
--- a/Engine/Plugins/Mutable/Source/MutableTools/Private/MuT/ASTOpImagePixelFormat.cpp
+++ b/Engine/Plugins/Mutable/Source/MutableTools/Private/MuT/ASTOpImagePixelFormat.cpp
@@ -86,18 +86,18 @@ namespace mu
 		// Already linked?
 		if (!linkedAddress)
 		{
-			OP::ImagePixelFormatArgs args;
-			FMemory::Memzero(&args, sizeof(args));
+			OP::ImagePixelFormatArgs Args;
+			FMemory::Memzero(Args);
 
-			args.format = Format;
-			args.formatIfAlpha = FormatIfAlpha;
-			if (Source) args.source = Source->linkedAddress;
+			Args.format = Format;
+			Args.formatIfAlpha = FormatIfAlpha;
+			if (Source) Args.source = Source->linkedAddress;
 
 			linkedAddress = (OP::ADDRESS)program.OpAddress.Num();
 			//program.m_code.push_back(op);
 			program.OpAddress.Add((uint32_t)program.ByteCode.Num());
 			AppendCode(program.ByteCode, EOpType::IM_PIXELFORMAT);
-			AppendCode(program.ByteCode, args);
+			AppendCode(program.ByteCode, Args);
 		}
 
 	}
diff --git a/Engine/Plugins/Mutable/Source/MutableTools/Private/MuT/ASTOpImagePlainColor.cpp b/Engine/Plugins/Mutable/Source/MutableTools/Private/MuT/ASTOpImagePlainColor.cpp
--- a/Engine/Plugins/Mutable/Source/MutableTools/Private/MuT/ASTOpImagePlainColor.cpp
+++ b/Engine/Plugins/Mutable/Source/MutableTools/Private/MuT/ASTOpImagePlainColor.cpp
@@ -80,7 +80,8 @@ namespace mu
 		// Already linked?
 		if (!linkedAddress)
 		{
-			OP::ImagePlainColorArgs Args={};
+			OP::ImagePlainColorArgs Args;
+			FMemory::Memzero(Args);
 
 			if (Color)
 			{
diff --git a/Engine/Plugins/Mutable/Source/MutableTools/Private/MuT/ASTOpImageRasterMesh.cpp b/Engine/Plugins/Mutable/Source/MutableTools/Private/MuT/ASTOpImageRasterMesh.cpp
--- a/Engine/Plugins/Mutable/Source/MutableTools/Private/MuT/ASTOpImageRasterMesh.cpp
+++ b/Engine/Plugins/Mutable/Source/MutableTools/Private/MuT/ASTOpImageRasterMesh.cpp
@@ -137,34 +137,34 @@ namespace mu
 		// Already linked?
 		if (!linkedAddress)
 		{
-			OP::ImageRasterMeshArgs args;
-			FMemory::Memzero(&args, sizeof(args));
-
-			args.BlockId = BlockId;
-			args.LayoutIndex = LayoutIndex;
-			args.sizeX = SizeX;
-			args.sizeY = SizeY;
-			args.SourceSizeX = SourceSizeX;
-			args.SourceSizeY = SourceSizeY;
-			args.CropMinX = CropMinX;
-			args.CropMinY = CropMinY;
-			args.UncroppedSizeX = UncroppedSizeX;
-			args.UncroppedSizeY = UncroppedSizeY;
-			args.bIsRGBFadingEnabled = bIsRGBFadingEnabled;
-			args.bIsAlphaFadingEnabled = bIsAlphaFadingEnabled;
-			args.SamplingMethod = static_cast<uint8>(SamplingMethod);
-			args.MinFilterMethod = static_cast<uint8>(MinFilterMethod);
-
-			if (mesh) args.mesh = mesh->linkedAddress;
-			if (image) args.image = image->linkedAddress;
-			if (angleFadeProperties) args.angleFadeProperties = angleFadeProperties->linkedAddress;
-			if (mask) args.mask = mask->linkedAddress;
-			if (projector) args.projector = projector->linkedAddress;
+			OP::ImageRasterMeshArgs Args;
+			FMemory::Memzero(Args);
+
+			Args.BlockId = BlockId;
+			Args.LayoutIndex = LayoutIndex;
+			Args.sizeX = SizeX;
+			Args.sizeY = SizeY;
+			Args.SourceSizeX = SourceSizeX;
+			Args.SourceSizeY = SourceSizeY;
+			Args.CropMinX = CropMinX;
+			Args.CropMinY = CropMinY;
+			Args.UncroppedSizeX = UncroppedSizeX;
+			Args.UncroppedSizeY = UncroppedSizeY;
+			Args.bIsRGBFadingEnabled = bIsRGBFadingEnabled;
+			Args.bIsAlphaFadingEnabled = bIsAlphaFadingEnabled;
+			Args.SamplingMethod = static_cast<uint8>(SamplingMethod);
+			Args.MinFilterMethod = static_cast<uint8>(MinFilterMethod);
+
+			if (mesh) Args.mesh = mesh->linkedAddress;
+			if (image) Args.image = image->linkedAddress;
+			if (angleFadeProperties) Args.angleFadeProperties = angleFadeProperties->linkedAddress;
+			if (mask) Args.mask = mask->linkedAddress;
+			if (projector) Args.projector = projector->linkedAddress;
 
 			linkedAddress = (OP::ADDRESS)program.OpAddress.Num();
 			program.OpAddress.Add((uint32)program.ByteCode.Num());
 			AppendCode(program.ByteCode, GetOpType());
-			AppendCode(program.ByteCode, args);
+			AppendCode(program.ByteCode, Args);
 		}
 	}
 
@@ -344,11 +344,11 @@ namespace mu
 				//	nop->type = EOpType::IM_CONDITIONAL;
 
 				//	Ptr<ASTOpFixed> aOp = mu::Clone<ASTOpFixed>(this);
-				//	aOp->SetChild(aOp->op.args.ImageRasterMesh.mesh, nop->yes);
+				//	aOp->SetChild(aOp->op.Args.ImageRasterMesh.mesh, nop->yes);
 				//	nop->yes = aOp;
 
 				//	Ptr<ASTOpFixed> bOp = mu::Clone<ASTOpFixed>(this);
-				//	bOp->SetChild(bOp->op.args.ImageRasterMesh.mesh, nop->no);
+				//	bOp->SetChild(bOp->op.Args.ImageRasterMesh.mesh, nop->no);
 				//	nop->no = bOp;
 
 				//	at = nop;
diff --git a/Engine/Plugins/Mutable/Source/MutableTools/Private/MuT/ASTOpImageResize.cpp b/Engine/Plugins/Mutable/Source/MutableTools/Private/MuT/ASTOpImageResize.cpp
--- a/Engine/Plugins/Mutable/Source/MutableTools/Private/MuT/ASTOpImageResize.cpp
+++ b/Engine/Plugins/Mutable/Source/MutableTools/Private/MuT/ASTOpImageResize.cpp
@@ -82,7 +82,8 @@ namespace mu
 		// Already linked?
 		if (!linkedAddress)
 		{
-			OP::ImageResizeArgs Args = {};
+			OP::ImageResizeArgs Args;
+			FMemory::Memzero(Args);
 
 			if (Source)
 			{
diff --git a/Engine/Plugins/Mutable/Source/MutableTools/Private/MuT/ASTOpImageResizeLike.cpp b/Engine/Plugins/Mutable/Source/MutableTools/Private/MuT/ASTOpImageResizeLike.cpp
--- a/Engine/Plugins/Mutable/Source/MutableTools/Private/MuT/ASTOpImageResizeLike.cpp
+++ b/Engine/Plugins/Mutable/Source/MutableTools/Private/MuT/ASTOpImageResizeLike.cpp
@@ -80,7 +80,8 @@ namespace mu
 		// Already linked?
 		if (!linkedAddress)
 		{
-			OP::ImageResizeLikeArgs Args = {};
+			OP::ImageResizeLikeArgs Args;
+			FMemory::Memzero(Args);
 
 			if (Source)
 			{
diff --git a/Engine/Plugins/Mutable/Source/MutableTools/Private/MuT/ASTOpImageResizeRel.cpp b/Engine/Plugins/Mutable/Source/MutableTools/Private/MuT/ASTOpImageResizeRel.cpp
--- a/Engine/Plugins/Mutable/Source/MutableTools/Private/MuT/ASTOpImageResizeRel.cpp
+++ b/Engine/Plugins/Mutable/Source/MutableTools/Private/MuT/ASTOpImageResizeRel.cpp
@@ -87,7 +87,8 @@ namespace mu
 		// Already linked?
 		if (!linkedAddress)
 		{
-			OP::ImageResizeRelArgs Args = {};
+			OP::ImageResizeRelArgs Args;
+			FMemory::Memzero(Args);
 
 			if (Source)
 			{
diff --git a/Engine/Plugins/Mutable/Source/MutableTools/Private/MuT/ASTOpImageSaturate.cpp b/Engine/Plugins/Mutable/Source/MutableTools/Private/MuT/ASTOpImageSaturate.cpp
--- a/Engine/Plugins/Mutable/Source/MutableTools/Private/MuT/ASTOpImageSaturate.cpp
+++ b/Engine/Plugins/Mutable/Source/MutableTools/Private/MuT/ASTOpImageSaturate.cpp
@@ -80,7 +80,8 @@ namespace mu
 		// Already linked?
 		if (!linkedAddress)
 		{
-			OP::ImageSaturateArgs Args = {};
+			OP::ImageSaturateArgs Args;
+			FMemory::Memzero(Args);
 
 			if (Base)
 			{
diff --git a/Engine/Plugins/Mutable/Source/MutableTools/Private/MuT/ASTOpImageSwizzle.cpp b/Engine/Plugins/Mutable/Source/MutableTools/Private/MuT/ASTOpImageSwizzle.cpp
--- a/Engine/Plugins/Mutable/Source/MutableTools/Private/MuT/ASTOpImageSwizzle.cpp
+++ b/Engine/Plugins/Mutable/Source/MutableTools/Private/MuT/ASTOpImageSwizzle.cpp
@@ -26,6 +26,7 @@
 #include "MuT/ASTOpImagePlainColor.h"
 #include "MuT/ASTOpImageDisplace.h"
 #include "MuT/ASTOpImageInvert.h"
+#include "MuT/ASTOpColorSwizzle.h"
 #include "MuT/ASTOpSwitch.h"
 
 namespace mu
@@ -104,20 +105,20 @@ namespace mu
 		// Already linked?
 		if (!linkedAddress)
 		{
-			OP::ImageSwizzleArgs args;
-			FMemory::Memzero(&args, sizeof(args));
+			OP::ImageSwizzleArgs Args;
+			FMemory::Memzero(Args);
 
-			args.format = Format;
+			Args.format = Format;
 			for (int32 i = 0; i < MUTABLE_OP_MAX_SWIZZLE_CHANNELS; ++i)
 			{
-				if (Sources[i]) args.sources[i] = Sources[i]->linkedAddress;
-				args.sourceChannels[i] = SourceChannels[i];
+				if (Sources[i]) Args.sources[i] = Sources[i]->linkedAddress;
+				Args.sourceChannels[i] = SourceChannels[i];
 			}			
 
 			linkedAddress = (OP::ADDRESS)program.OpAddress.Num();
 			program.OpAddress.Add((uint32)program.ByteCode.Num());
 			AppendCode(program.ByteCode, GetOpType());
-			AppendCode(program.ByteCode, args);
+			AppendCode(program.ByteCode, Args);
 		}
 	}
 
@@ -508,7 +509,7 @@ namespace mu
 						const ASTOpImageRasterMesh* Typed = static_cast<const ASTOpImageRasterMesh*>(Sources[c].child().get());
 						check(Typed);
 
-						// Compare all args but the source image
+						// Compare all Args but the source image
 						if (Typed->mesh.child() != FirstRasterMesh->mesh.child()
 							|| Typed->angleFadeProperties.child() != FirstRasterMesh->angleFadeProperties.child()
 							|| Typed->mask.child() != FirstRasterMesh->mask.child()
@@ -582,7 +583,7 @@ namespace mu
 						const ASTOpImageTransform* Typed = static_cast<const ASTOpImageTransform*>(Sources[c].child().get());
 						check(Typed);
 
-						// Compare all args but the source image
+						// Compare all Args but the source image
 						if (Typed->ScaleX.child() != FirstTransform->ScaleX.child()
 							|| Typed->ScaleY.child() != FirstTransform->ScaleY.child()
 							|| Typed->OffsetX.child() != FirstTransform->OffsetX.child()
@@ -636,7 +637,7 @@ namespace mu
 			//		{
 			//			const ASTOpFixed* Typed = static_cast<const ASTOpFixed*>(Sources[c].child().get());
 			//			check(Typed);
-			//			// Compare all args but the source image
+			//			// Compare all Args but the source image
 			//			OP::ImageResizeArgs ArgCopy = FirstResize->op.args.ImageResize;
 			//			ArgCopy.source = Typed->op.args.ImageResize.source;
 
@@ -705,16 +706,15 @@ namespace mu
 			if (!at && sourceType == EOpType::IM_PLAINCOLOUR)
 			{
 				Ptr<ASTOpImagePlainColor> NewPlain = mu::Clone<ASTOpImagePlainColor>(channelSourceAt);
-				Ptr<ASTOpFixed> NewSwizzle = new ASTOpFixed;
-				NewSwizzle->op.type = EOpType::CO_SWIZZLE;
+				Ptr<ASTOpColorSwizzle> NewSwizzle = new ASTOpColorSwizzle;
 				for (int c = 0; c < MUTABLE_OP_MAX_SWIZZLE_CHANNELS; ++c)
 				{
 					if (Sources[c])
 					{
 						const ASTOpImagePlainColor* TypedPlain = static_cast<const ASTOpImagePlainColor*>(Sources[c].child().get());
-						NewSwizzle->SetChild(NewSwizzle->op.args.ColourSwizzle.sources[c], TypedPlain->Color);
+						NewSwizzle->Sources[c] = TypedPlain->Color.child();
 					}
-					NewSwizzle->op.args.ColourSwizzle.sourceChannels[c] = SourceChannels[c];
+					NewSwizzle->SourceChannels[c] = SourceChannels[c];
 				}
 				NewPlain->Color = NewSwizzle;
 				NewPlain->Format = Format;
@@ -1270,12 +1270,11 @@ namespace mu
 			nop->base = Visit(nop->base.child(), CurrentSwizzleOp);
 
 			// We need to swizzle the colour too
-			Ptr<ASTOpFixed> NewColorOp = new ASTOpFixed;
-			NewColorOp->op.type = EOpType::CO_SWIZZLE;
+			Ptr<ASTOpColorSwizzle> NewColorOp = new ASTOpColorSwizzle;
 			for (int i = 0; i < MUTABLE_OP_MAX_SWIZZLE_CHANNELS; ++i)
 			{
-				NewColorOp->SetChild(NewColorOp->op.args.ColourSwizzle.sources[i], nop->color);
-				NewColorOp->op.args.ColourSwizzle.sourceChannels[i] = CurrentSwizzleOp->SourceChannels[i];
+				NewColorOp->Sources[i] = nop->color.child();
+				NewColorOp->SourceChannels[i] = CurrentSwizzleOp->SourceChannels[i];
 			}
 			nop->color = NewColorOp;
 
@@ -1362,13 +1361,10 @@ namespace mu
 
 		case EOpType::IM_BLANKLAYOUT:
 		{
-			Ptr<ASTOpFixed> NewOp = mu::Clone<ASTOpFixed>(at);
-
 			// We can remove the swizzle entirely.
 			// It is not 100% equivalent, because blank layouts are initialized with 0,0,0,1 so the result could be
 			// different, but those pixels shouldn't be used anyway.
-
-			newAt = NewOp;
+			newAt = mu::Clone<ASTOp>(at);
 			break;
 		}
 
diff --git a/Engine/Plugins/Mutable/Source/MutableTools/Private/MuT/ASTOpInstanceAdd.cpp b/Engine/Plugins/Mutable/Source/MutableTools/Private/MuT/ASTOpInstanceAdd.cpp
--- a/Engine/Plugins/Mutable/Source/MutableTools/Private/MuT/ASTOpInstanceAdd.cpp
+++ b/Engine/Plugins/Mutable/Source/MutableTools/Private/MuT/ASTOpInstanceAdd.cpp
@@ -103,15 +103,15 @@ namespace mu
 		// Already linked?
 		if (!linkedAddress)
 		{
-			OP::InstanceAddArgs args;
-			FMemory::Memzero(&args, sizeof(args));
-			args.id = id;
-			args.ExternalId = ExternalId;
-			args.SharedSurfaceId = SharedSurfaceId;
-			args.name = program.AddConstant(name);
+			OP::InstanceAddArgs Args;
+			FMemory::Memzero(Args);
+			Args.id = id;
+			Args.ExternalId = ExternalId;
+			Args.SharedSurfaceId = SharedSurfaceId;
+			Args.name = program.AddConstant(name);
 
-			if (instance) args.instance = instance->linkedAddress;
-			if (value) args.value = value->linkedAddress;
+			if (instance) Args.instance = instance->linkedAddress;
+			if (value) Args.value = value->linkedAddress;
 
 			if (type == EOpType::IN_ADDIMAGE ||
 				type == EOpType::IN_ADDMESH)
@@ -141,11 +141,11 @@ namespace mu
 
 				if (it != INDEX_NONE)
 				{
-					args.relevantParametersListIndex = it;
+					Args.relevantParametersListIndex = it;
 				}
 				else
 				{
-					args.relevantParametersListIndex = uint32_t(program.ParameterLists.Num());
+					Args.relevantParametersListIndex = uint32_t(program.ParameterLists.Num());
 					program.ParameterLists.Add(params);
 				}
 			}
@@ -153,7 +153,7 @@ namespace mu
 			linkedAddress = (OP::ADDRESS)program.OpAddress.Num();
 			program.OpAddress.Add((uint32_t)program.ByteCode.Num());
 			AppendCode(program.ByteCode, type);
-			AppendCode(program.ByteCode, args);
+			AppendCode(program.ByteCode, Args);
 		}
 
 	}
diff --git a/Engine/Plugins/Mutable/Source/MutableTools/Private/MuT/ASTOpLayoutFromMesh.cpp b/Engine/Plugins/Mutable/Source/MutableTools/Private/MuT/ASTOpLayoutFromMesh.cpp
--- a/Engine/Plugins/Mutable/Source/MutableTools/Private/MuT/ASTOpLayoutFromMesh.cpp
+++ b/Engine/Plugins/Mutable/Source/MutableTools/Private/MuT/ASTOpLayoutFromMesh.cpp
@@ -82,17 +82,17 @@ namespace mu
 		// Already linked?
 		if (!linkedAddress)
 		{
-			OP::LayoutFromMeshArgs args;
-			FMemory::Memzero(&args, sizeof(args));
+			OP::LayoutFromMeshArgs Args;
+			FMemory::Memzero(Args);
 
-			args.LayoutIndex = LayoutIndex;
-			if (Mesh) args.Mesh = Mesh->linkedAddress;
+			Args.LayoutIndex = LayoutIndex;
+			if (Mesh) Args.Mesh = Mesh->linkedAddress;
 
 			linkedAddress = (OP::ADDRESS)program.OpAddress.Num();
 			//program.m_code.push_back(op);
 			program.OpAddress.Add((uint32_t)program.ByteCode.Num());
 			AppendCode(program.ByteCode, GetOpType());
-			AppendCode(program.ByteCode, args);
+			AppendCode(program.ByteCode, Args);
 		}
 
 	}
diff --git a/Engine/Plugins/Mutable/Source/MutableTools/Private/MuT/ASTOpLayoutMerge.cpp b/Engine/Plugins/Mutable/Source/MutableTools/Private/MuT/ASTOpLayoutMerge.cpp
--- a/Engine/Plugins/Mutable/Source/MutableTools/Private/MuT/ASTOpLayoutMerge.cpp
+++ b/Engine/Plugins/Mutable/Source/MutableTools/Private/MuT/ASTOpLayoutMerge.cpp
@@ -78,16 +78,16 @@ namespace mu
 		// Already linked?
 		if (!linkedAddress)
 		{
-			OP::LayoutMergeArgs args;
-			FMemory::Memzero(&args, sizeof(args));
+			OP::LayoutMergeArgs Args;
+			FMemory::Memzero(Args);
 
-			if (Base) args.Base = Base->linkedAddress;
-			if (Added) args.Added = Added->linkedAddress;
+			if (Base) Args.Base = Base->linkedAddress;
+			if (Added) Args.Added = Added->linkedAddress;
 
 			linkedAddress = (OP::ADDRESS)program.OpAddress.Num();
 			program.OpAddress.Add((uint32_t)program.ByteCode.Num());
 			AppendCode(program.ByteCode, GetOpType());
-			AppendCode(program.ByteCode, args);
+			AppendCode(program.ByteCode, Args);
 		}
 
 	}
diff --git a/Engine/Plugins/Mutable/Source/MutableTools/Private/MuT/ASTOpLayoutPack.cpp b/Engine/Plugins/Mutable/Source/MutableTools/Private/MuT/ASTOpLayoutPack.cpp
--- a/Engine/Plugins/Mutable/Source/MutableTools/Private/MuT/ASTOpLayoutPack.cpp
+++ b/Engine/Plugins/Mutable/Source/MutableTools/Private/MuT/ASTOpLayoutPack.cpp
@@ -74,15 +74,15 @@ namespace mu
 		// Already linked?
 		if (!linkedAddress)
 		{
-			OP::LayoutRemoveBlocksArgs args;
-			FMemory::Memzero(&args, sizeof(args));
+			OP::LayoutRemoveBlocksArgs Args;
+			FMemory::Memzero(Args);
 
-			if (Source) args.Source = Source->linkedAddress;
+			if (Source) Args.Source = Source->linkedAddress;
 
 			linkedAddress = (OP::ADDRESS)program.OpAddress.Num();
 			program.OpAddress.Add((uint32_t)program.ByteCode.Num());
 			AppendCode(program.ByteCode, GetOpType());
-			AppendCode(program.ByteCode, args);
+			AppendCode(program.ByteCode, Args);
 		}
 
 	}
diff --git a/Engine/Plugins/Mutable/Source/MutableTools/Private/MuT/ASTOpLayoutRemoveBlocks.cpp b/Engine/Plugins/Mutable/Source/MutableTools/Private/MuT/ASTOpLayoutRemoveBlocks.cpp
--- a/Engine/Plugins/Mutable/Source/MutableTools/Private/MuT/ASTOpLayoutRemoveBlocks.cpp
+++ b/Engine/Plugins/Mutable/Source/MutableTools/Private/MuT/ASTOpLayoutRemoveBlocks.cpp
@@ -78,16 +78,16 @@ namespace mu
 		// Already linked?
 		if (!linkedAddress)
 		{
-			OP::LayoutRemoveBlocksArgs args;
-			FMemory::Memzero(&args, sizeof(args));
+			OP::LayoutRemoveBlocksArgs Args;
+			FMemory::Memzero(Args);
 
-			if (Source) args.Source = Source->linkedAddress;
-			if (ReferenceLayout) args.ReferenceLayout = ReferenceLayout->linkedAddress;
+			if (Source) Args.Source = Source->linkedAddress;
+			if (ReferenceLayout) Args.ReferenceLayout = ReferenceLayout->linkedAddress;
 
 			linkedAddress = (OP::ADDRESS)program.OpAddress.Num();
 			program.OpAddress.Add((uint32_t)program.ByteCode.Num());
 			AppendCode(program.ByteCode, EOpType::LA_REMOVEBLOCKS);
-			AppendCode(program.ByteCode, args);
+			AppendCode(program.ByteCode, Args);
 		}
 
 	}
diff --git a/Engine/Plugins/Mutable/Source/MutableTools/Private/MuT/ASTOpMeshApplyLayout.cpp b/Engine/Plugins/Mutable/Source/MutableTools/Private/MuT/ASTOpMeshApplyLayout.cpp
--- a/Engine/Plugins/Mutable/Source/MutableTools/Private/MuT/ASTOpMeshApplyLayout.cpp
+++ b/Engine/Plugins/Mutable/Source/MutableTools/Private/MuT/ASTOpMeshApplyLayout.cpp
@@ -69,7 +69,8 @@ namespace mu
 		// Already linked?
 		if (!linkedAddress)
 		{
-			OP::MeshApplyLayoutArgs Args = {};
+			OP::MeshApplyLayoutArgs Args;
+			FMemory::Memzero(Args);
 
 			if (Mesh) Args.Mesh = Mesh->linkedAddress;
 			if (Layout) Args.Layout = Layout->linkedAddress;
diff --git a/Engine/Plugins/Mutable/Source/MutableTools/Private/MuT/ASTOpMeshApplyPose.cpp b/Engine/Plugins/Mutable/Source/MutableTools/Private/MuT/ASTOpMeshApplyPose.cpp
--- a/Engine/Plugins/Mutable/Source/MutableTools/Private/MuT/ASTOpMeshApplyPose.cpp
+++ b/Engine/Plugins/Mutable/Source/MutableTools/Private/MuT/ASTOpMeshApplyPose.cpp
@@ -66,7 +66,8 @@ namespace mu
 		// Already linked?
 		if (!linkedAddress)
 		{
-			OP::MeshApplyPoseArgs Args = {};
+			OP::MeshApplyPoseArgs Args;
+			FMemory::Memzero(Args);
 
 			if (Base) Args.base = Base->linkedAddress;
 			if (Pose) Args.pose = Pose->linkedAddress;
diff --git a/Engine/Plugins/Mutable/Source/MutableTools/Private/MuT/ASTOpMeshApplyShape.cpp b/Engine/Plugins/Mutable/Source/MutableTools/Private/MuT/ASTOpMeshApplyShape.cpp
--- a/Engine/Plugins/Mutable/Source/MutableTools/Private/MuT/ASTOpMeshApplyShape.cpp
+++ b/Engine/Plugins/Mutable/Source/MutableTools/Private/MuT/ASTOpMeshApplyShape.cpp
@@ -91,7 +91,7 @@ namespace mu
 		if (!linkedAddress)
 		{
 			OP::MeshApplyShapeArgs Args;
-			FMemory::Memzero(&Args, sizeof(Args));
+			FMemory::Memzero(Args);
 
 			constexpr EMeshBindShapeFlags NoFlags = EMeshBindShapeFlags::None;
 			EMeshBindShapeFlags BindFlags = NoFlags;
diff --git a/Engine/Plugins/Mutable/Source/MutableTools/Private/MuT/ASTOpMeshBindShape.cpp b/Engine/Plugins/Mutable/Source/MutableTools/Private/MuT/ASTOpMeshBindShape.cpp
--- a/Engine/Plugins/Mutable/Source/MutableTools/Private/MuT/ASTOpMeshBindShape.cpp
+++ b/Engine/Plugins/Mutable/Source/MutableTools/Private/MuT/ASTOpMeshBindShape.cpp
@@ -131,7 +131,7 @@ namespace mu
 		if (!linkedAddress)
 		{
 			OP::MeshBindShapeArgs Args;
-			FMemory::Memzero(&Args, sizeof(Args));
+			FMemory::Memzero(Args);
 
 			constexpr EMeshBindShapeFlags NoFlags = EMeshBindShapeFlags::None;
 			EMeshBindShapeFlags BindFlags = NoFlags;
diff --git a/Engine/Plugins/Mutable/Source/MutableTools/Private/MuT/ASTOpMeshClipDeform.cpp b/Engine/Plugins/Mutable/Source/MutableTools/Private/MuT/ASTOpMeshClipDeform.cpp
--- a/Engine/Plugins/Mutable/Source/MutableTools/Private/MuT/ASTOpMeshClipDeform.cpp
+++ b/Engine/Plugins/Mutable/Source/MutableTools/Private/MuT/ASTOpMeshClipDeform.cpp
@@ -68,25 +68,25 @@ namespace mu
 		// Already linked?
 		if (!linkedAddress)
 		{
-			OP::MeshClipDeformArgs args;
-			FMemory::Memzero(&args, sizeof(args));
+			OP::MeshClipDeformArgs Args;
+			FMemory::Memzero(Args);
 
-			args.FaceCullStrategy = FaceCullStrategy;
+			Args.FaceCullStrategy = FaceCullStrategy;
 
 			if (Mesh)
 			{
-				args.mesh = Mesh->linkedAddress;
+				Args.mesh = Mesh->linkedAddress;
 			}
 
 			if (ClipShape)
 			{
-				args.clipShape = ClipShape->linkedAddress;
+				Args.clipShape = ClipShape->linkedAddress;
 			}
 
 			linkedAddress = (OP::ADDRESS)program.OpAddress.Num();
 			program.OpAddress.Add((uint32_t)program.ByteCode.Num());
 			AppendCode(program.ByteCode, EOpType::ME_CLIPDEFORM);
-			AppendCode(program.ByteCode, args);
+			AppendCode(program.ByteCode, Args);
 		}
 
 	}
diff --git a/Engine/Plugins/Mutable/Source/MutableTools/Private/MuT/ASTOpMeshClipMorphPlane.cpp b/Engine/Plugins/Mutable/Source/MutableTools/Private/MuT/ASTOpMeshClipMorphPlane.cpp
--- a/Engine/Plugins/Mutable/Source/MutableTools/Private/MuT/ASTOpMeshClipMorphPlane.cpp
+++ b/Engine/Plugins/Mutable/Source/MutableTools/Private/MuT/ASTOpMeshClipMorphPlane.cpp
@@ -83,36 +83,36 @@ namespace mu
 		// Already linked?
 		if (!linkedAddress)
 		{
-			OP::MeshClipMorphPlaneArgs args;
-			FMemory::Memzero(&args, sizeof(args));
+			OP::MeshClipMorphPlaneArgs Args;
+			FMemory::Memzero(Args);
 
-			if (Source) args.Source = Source->linkedAddress;
+			if (Source) Args.Source = Source->linkedAddress;
 
-			args.MorphShape = (OP::ADDRESS)program.ConstantShapes.Num();
+			Args.MorphShape = (OP::ADDRESS)program.ConstantShapes.Num();
 			program.ConstantShapes.Add(MorphShape);
 
-			args.FaceCullStrategy = FaceCullStrategy;
+			Args.FaceCullStrategy = FaceCullStrategy;
 
-			args.VertexSelectionType = VertexSelectionType;
+			Args.VertexSelectionType = VertexSelectionType;
 			if (VertexSelectionType == EClipVertexSelectionType::BoneHierarchy)
 			{
 				check(VertexSelectionBone.Id < MAX_uint32);
-				args.VertexSelectionShapeOrBone = VertexSelectionBone.Id;
+				Args.VertexSelectionShapeOrBone = VertexSelectionBone.Id;
 			}
 			else if (VertexSelectionType == EClipVertexSelectionType::Shape)
 			{
-				args.VertexSelectionShapeOrBone = program.AddConstant(SelectionShape);
+				Args.VertexSelectionShapeOrBone = program.AddConstant(SelectionShape);
 			}
 
-			args.Dist = Dist;
-			args.Factor = Factor;
-			args.MaxBoneRadius = VertexSelectionBoneMaxRadius;
+			Args.Dist = Dist;
+			Args.Factor = Factor;
+			Args.MaxBoneRadius = VertexSelectionBoneMaxRadius;
 
 			linkedAddress = (OP::ADDRESS)program.OpAddress.Num();
 			//program.m_code.push_back(op);
 			program.OpAddress.Add((uint32_t)program.ByteCode.Num());
 			AppendCode(program.ByteCode, EOpType::ME_CLIPMORPHPLANE);
-			AppendCode(program.ByteCode, args);
+			AppendCode(program.ByteCode, Args);
 		}
 
 	}
diff --git a/Engine/Plugins/Mutable/Source/MutableTools/Private/MuT/ASTOpMeshClipWithMesh.cpp b/Engine/Plugins/Mutable/Source/MutableTools/Private/MuT/ASTOpMeshClipWithMesh.cpp
--- a/Engine/Plugins/Mutable/Source/MutableTools/Private/MuT/ASTOpMeshClipWithMesh.cpp
+++ b/Engine/Plugins/Mutable/Source/MutableTools/Private/MuT/ASTOpMeshClipWithMesh.cpp
@@ -66,7 +66,8 @@ namespace mu
 		// Already linked?
 		if (!linkedAddress)
 		{
-			OP::MeshApplyPoseArgs Args = {};
+			OP::MeshApplyPoseArgs Args;
+			FMemory::Memzero(Args);
 
 			if (Source) Args.base = Source->linkedAddress;
 			if (ClipMesh) Args.pose = ClipMesh->linkedAddress;
diff --git a/Engine/Plugins/Mutable/Source/MutableTools/Private/MuT/ASTOpMeshFormat.cpp b/Engine/Plugins/Mutable/Source/MutableTools/Private/MuT/ASTOpMeshFormat.cpp
--- a/Engine/Plugins/Mutable/Source/MutableTools/Private/MuT/ASTOpMeshFormat.cpp
+++ b/Engine/Plugins/Mutable/Source/MutableTools/Private/MuT/ASTOpMeshFormat.cpp
@@ -87,22 +87,22 @@ void ASTOpMeshFormat::Link( FProgram& program, FLinkerOptions* )
     // Already linked?
     if (!linkedAddress)
     {
-        OP::MeshFormatArgs args;
-        FMemory::Memzero( &args, sizeof(args) );
+        OP::MeshFormatArgs Args;
+		FMemory::Memzero(Args);
 
-		args.Flags = Flags;
+		Args.Flags = Flags;
 		if (bOptimizeBuffers)
 		{
-			args.Flags = args.Flags | OP::MeshFormatArgs::OptimizeBuffers;
+			Args.Flags = Args.Flags | OP::MeshFormatArgs::OptimizeBuffers;
 		}
 
-		if (Source) args.source = Source->linkedAddress;
-		if (Format) args.format = Format->linkedAddress;
+		if (Source) Args.source = Source->linkedAddress;
+		if (Format) Args.format = Format->linkedAddress;
 
         linkedAddress = (OP::ADDRESS)program.OpAddress.Num();
         program.OpAddress.Add((uint32)program.ByteCode.Num());
         AppendCode(program.ByteCode,EOpType::ME_FORMAT);
-        AppendCode(program.ByteCode,args);
+        AppendCode(program.ByteCode,Args);
     }
 
 }
diff --git a/Engine/Plugins/Mutable/Source/MutableTools/Private/MuT/ASTOpMeshMaskClipMesh.cpp b/Engine/Plugins/Mutable/Source/MutableTools/Private/MuT/ASTOpMeshMaskClipMesh.cpp
--- a/Engine/Plugins/Mutable/Source/MutableTools/Private/MuT/ASTOpMeshMaskClipMesh.cpp
+++ b/Engine/Plugins/Mutable/Source/MutableTools/Private/MuT/ASTOpMeshMaskClipMesh.cpp
@@ -79,17 +79,17 @@ namespace mu
 		// Already linked?
 		if (!linkedAddress)
 		{
-			OP::MeshMaskClipMeshArgs args;
-			memset(&args, 0, sizeof(args));
+			OP::MeshMaskClipMeshArgs Args;
+			FMemory::Memzero(Args);
 
-			if (source) args.source = source->linkedAddress;
-			if (clip) args.clip = clip->linkedAddress;
+			if (source) Args.source = source->linkedAddress;
+			if (clip) Args.clip = clip->linkedAddress;
 
 			linkedAddress = (OP::ADDRESS)program.OpAddress.Num();
 			//program.m_code.push_back(op);
 			program.OpAddress.Add((uint32_t)program.ByteCode.Num());
 			AppendCode(program.ByteCode, EOpType::ME_MASKCLIPMESH);
-			AppendCode(program.ByteCode, args);
+			AppendCode(program.ByteCode, Args);
 		}
 	}
 
diff --git a/Engine/Plugins/Mutable/Source/MutableTools/Private/MuT/ASTOpMeshMaskClipUVMask.cpp b/Engine/Plugins/Mutable/Source/MutableTools/Private/MuT/ASTOpMeshMaskClipUVMask.cpp
--- a/Engine/Plugins/Mutable/Source/MutableTools/Private/MuT/ASTOpMeshMaskClipUVMask.cpp
+++ b/Engine/Plugins/Mutable/Source/MutableTools/Private/MuT/ASTOpMeshMaskClipUVMask.cpp
@@ -92,19 +92,19 @@ namespace mu
 		// Already linked?
 		if (!linkedAddress)
 		{
-			OP::MeshMaskClipUVMaskArgs args;
-			FMemory::Memzero(args);
+			OP::MeshMaskClipUVMaskArgs Args;
+			FMemory::Memzero(Args);
 
-			if (Source) args.Source = Source->linkedAddress;
-			if (UVSource) args.UVSource = UVSource->linkedAddress;
-			if (MaskImage) args.MaskImage = MaskImage->linkedAddress;
-			if (MaskLayout) args.MaskLayout = MaskLayout->linkedAddress;
-			args.LayoutIndex = LayoutIndex;
+			if (Source) Args.Source = Source->linkedAddress;
+			if (UVSource) Args.UVSource = UVSource->linkedAddress;
+			if (MaskImage) Args.MaskImage = MaskImage->linkedAddress;
+			if (MaskLayout) Args.MaskLayout = MaskLayout->linkedAddress;
+			Args.LayoutIndex = LayoutIndex;
 
 			linkedAddress = (OP::ADDRESS)program.OpAddress.Num();
 			program.OpAddress.Add((uint32)program.ByteCode.Num());
 			AppendCode(program.ByteCode, EOpType::ME_MASKCLIPUVMASK);
-			AppendCode(program.ByteCode, args);
+			AppendCode(program.ByteCode, Args);
 		}
 	}
 
diff --git a/Engine/Plugins/Mutable/Source/MutableTools/Private/MuT/ASTOpMeshMaskDiff.cpp b/Engine/Plugins/Mutable/Source/MutableTools/Private/MuT/ASTOpMeshMaskDiff.cpp
--- a/Engine/Plugins/Mutable/Source/MutableTools/Private/MuT/ASTOpMeshMaskDiff.cpp
+++ b/Engine/Plugins/Mutable/Source/MutableTools/Private/MuT/ASTOpMeshMaskDiff.cpp
@@ -67,7 +67,8 @@ namespace mu
 		// Already linked?
 		if (!linkedAddress)
 		{
-			OP::MeshMaskDiffArgs Args = {};
+			OP::MeshMaskDiffArgs Args;
+			FMemory::Memzero(Args);
 
 			if (Source) Args.Source = Source->linkedAddress;
 			if (Fragment) Args.Fragment = Fragment->linkedAddress;
diff --git a/Engine/Plugins/Mutable/Source/MutableTools/Private/MuT/ASTOpMeshMerge.cpp b/Engine/Plugins/Mutable/Source/MutableTools/Private/MuT/ASTOpMeshMerge.cpp
--- a/Engine/Plugins/Mutable/Source/MutableTools/Private/MuT/ASTOpMeshMerge.cpp
+++ b/Engine/Plugins/Mutable/Source/MutableTools/Private/MuT/ASTOpMeshMerge.cpp
@@ -69,7 +69,8 @@ namespace mu
 		// Already linked?
 		if (!linkedAddress)
 		{
-			OP::MeshMergeArgs Args = {};
+			OP::MeshMergeArgs Args;
+			FMemory::Memzero(Args);
 
 			if (Base) Args.Base = Base->linkedAddress;
 			if (Added) Args.Added = Added->linkedAddress;
diff --git a/Engine/Plugins/Mutable/Source/MutableTools/Private/MuT/ASTOpMeshMorphReshape.cpp b/Engine/Plugins/Mutable/Source/MutableTools/Private/MuT/ASTOpMeshMorphReshape.cpp
--- a/Engine/Plugins/Mutable/Source/MutableTools/Private/MuT/ASTOpMeshMorphReshape.cpp
+++ b/Engine/Plugins/Mutable/Source/MutableTools/Private/MuT/ASTOpMeshMorphReshape.cpp
@@ -69,7 +69,7 @@ namespace mu
 		if (!linkedAddress)
 		{
 			OP::MeshMorphReshapeArgs Args;
-			memset(&Args, 0, sizeof(Args));
+			FMemory::Memzero(Args);
 
 			if (Morph)
 			{
diff --git a/Engine/Plugins/Mutable/Source/MutableTools/Private/MuT/ASTOpMeshOptimizeSkinning.cpp b/Engine/Plugins/Mutable/Source/MutableTools/Private/MuT/ASTOpMeshOptimizeSkinning.cpp
--- a/Engine/Plugins/Mutable/Source/MutableTools/Private/MuT/ASTOpMeshOptimizeSkinning.cpp
+++ b/Engine/Plugins/Mutable/Source/MutableTools/Private/MuT/ASTOpMeshOptimizeSkinning.cpp
@@ -66,7 +66,8 @@ namespace mu
 		// Already linked?
 		if (!linkedAddress)
 		{
-			OP::MeshOptimizeSkinningArgs Args = {};
+			OP::MeshOptimizeSkinningArgs Args;
+			FMemory::Memzero(Args);
 
 			if (Source)
 			{
diff --git a/Engine/Plugins/Mutable/Source/MutableTools/Private/MuT/ASTOpMeshProject.cpp b/Engine/Plugins/Mutable/Source/MutableTools/Private/MuT/ASTOpMeshProject.cpp
--- a/Engine/Plugins/Mutable/Source/MutableTools/Private/MuT/ASTOpMeshProject.cpp
+++ b/Engine/Plugins/Mutable/Source/MutableTools/Private/MuT/ASTOpMeshProject.cpp
@@ -68,7 +68,8 @@ namespace mu
 		// Already linked?
 		if (!linkedAddress)
 		{
-			OP::MeshProjectArgs Args = {};
+			OP::MeshProjectArgs Args;
+			FMemory::Memzero(Args);
 
 			if (Mesh) Args.Mesh = Mesh->linkedAddress;
 			if (Projector) Args.Projector = Projector->linkedAddress;
diff --git a/Engine/Plugins/Mutable/Source/MutableTools/Private/MuT/ASTOpMeshSetSkeleton.cpp b/Engine/Plugins/Mutable/Source/MutableTools/Private/MuT/ASTOpMeshSetSkeleton.cpp
--- a/Engine/Plugins/Mutable/Source/MutableTools/Private/MuT/ASTOpMeshSetSkeleton.cpp
+++ b/Engine/Plugins/Mutable/Source/MutableTools/Private/MuT/ASTOpMeshSetSkeleton.cpp
@@ -66,7 +66,8 @@ namespace mu
 		// Already linked?
 		if (!linkedAddress)
 		{
-			OP::MeshApplyPoseArgs Args = {};
+			OP::MeshApplyPoseArgs Args;
+			FMemory::Memzero(Args);
 
 			if (Source) Args.base = Source->linkedAddress;
 			if (Skeleton) Args.pose = Skeleton->linkedAddress;
diff --git a/Engine/Plugins/Mutable/Source/MutableTools/Private/MuT/ASTOpMeshTransform.cpp b/Engine/Plugins/Mutable/Source/MutableTools/Private/MuT/ASTOpMeshTransform.cpp
--- a/Engine/Plugins/Mutable/Source/MutableTools/Private/MuT/ASTOpMeshTransform.cpp
+++ b/Engine/Plugins/Mutable/Source/MutableTools/Private/MuT/ASTOpMeshTransform.cpp
@@ -64,18 +64,18 @@ namespace mu
 		// Already linked?
 		if (!linkedAddress)
 		{
-			OP::MeshTransformArgs args;
-			memset(&args, 0, sizeof(args));
+			OP::MeshTransformArgs Args;
+			FMemory::Memzero(Args);
 
-			if (source) args.source = source->linkedAddress;
+			if (source) Args.source = source->linkedAddress;
 
-			args.matrix = program.AddConstant(matrix);
+			Args.matrix = program.AddConstant(matrix);
 
 			linkedAddress = (OP::ADDRESS)program.OpAddress.Num();
 			//program.m_code.push_back(op);
 			program.OpAddress.Add((uint32_t)program.ByteCode.Num());
 			AppendCode(program.ByteCode, EOpType::ME_TRANSFORM);
-			AppendCode(program.ByteCode, args);
+			AppendCode(program.ByteCode, Args);
 		}
 
 	}
diff --git a/Engine/Plugins/Mutable/Source/MutableTools/Private/MuT/ASTOpMeshTransformWithBoundingMesh.cpp b/Engine/Plugins/Mutable/Source/MutableTools/Private/MuT/ASTOpMeshTransformWithBoundingMesh.cpp
--- a/Engine/Plugins/Mutable/Source/MutableTools/Private/MuT/ASTOpMeshTransformWithBoundingMesh.cpp
+++ b/Engine/Plugins/Mutable/Source/MutableTools/Private/MuT/ASTOpMeshTransformWithBoundingMesh.cpp
@@ -58,23 +58,24 @@ void mu::ASTOpMeshTransformWithBoundingMesh::Link(FProgram& program, FLinkerOpti
 {
 	if (!linkedAddress)
 	{
-		OP::MeshTransformWithinMeshArgs args = {};
+		OP::MeshTransformWithinMeshArgs Args;
+		FMemory::Memzero(Args);
 		if (source)
 		{
-			args.sourceMesh = source->linkedAddress;
+			Args.sourceMesh = source->linkedAddress;
 		}
 		if (boundingMesh)
 		{
-			args.boundingMesh = boundingMesh->linkedAddress;
+			Args.boundingMesh = boundingMesh->linkedAddress;
 		}
 		if (matrix)
 		{
-			args.matrix = matrix->linkedAddress;
+			Args.matrix = matrix->linkedAddress;
 		}
 		linkedAddress = static_cast<OP::ADDRESS>(program.OpAddress.Num());
 		program.OpAddress.Add(program.ByteCode.Num());
 		AppendCode(program.ByteCode, EOpType::ME_TRANSFORMWITHMESH);
-		AppendCode(program.ByteCode, args);
+		AppendCode(program.ByteCode, Args);
 	}
 }
 
diff --git a/Engine/Plugins/Mutable/Source/MutableTools/Private/MuT/ASTOpParameter.cpp b/Engine/Plugins/Mutable/Source/MutableTools/Private/MuT/ASTOpParameter.cpp
--- a/Engine/Plugins/Mutable/Source/MutableTools/Private/MuT/ASTOpParameter.cpp
+++ b/Engine/Plugins/Mutable/Source/MutableTools/Private/MuT/ASTOpParameter.cpp
@@ -2,6 +2,7 @@
 
 #include "MuT/ASTOpParameter.h"
 
+#include "MuT/ASTOpBoolEqualIntConst.h"
 #include "HAL/PlatformMath.h"
 #include "Misc/AssertionMacros.h"
 #include "MuR/ModelPrivate.h"
@@ -82,8 +83,8 @@ namespace mu
 		// Already linked?
 		if (!linkedAddress)
 		{
-			OP::ParameterArgs args;
-			FMemory::Memzero(&args, sizeof(args));
+			OP::ParameterArgs Args;
+			FMemory::Memzero(Args);
 
 			LinkedParameterIndex = Program.Parameters.Find(Parameter);
 
@@ -91,7 +92,7 @@ namespace mu
 			// parameters map in the CodeGenerator_FirstPass.
 			check(LinkedParameterIndex != INDEX_NONE);
 
-			args.variable = (OP::ADDRESS)LinkedParameterIndex;
+			Args.variable = (OP::ADDRESS)LinkedParameterIndex;
 
 			for (const auto& d : Ranges)
 			{
@@ -105,26 +106,26 @@ namespace mu
 
 			Program.OpAddress.Add((uint32)Program.ByteCode.Num());
 			AppendCode(Program.ByteCode, Type);
-			AppendCode(Program.ByteCode, args);
+			AppendCode(Program.ByteCode, Args);
 		}
 	}
 
 
-	int ASTOpParameter::EvaluateInt(ASTOpList& facts, bool& unknown) const
+	int32 ASTOpParameter::EvaluateInt(ASTOpList& facts, bool& bOutUnknown) const
 	{
-		unknown = true;
+		bOutUnknown = true;
 
 		// Check the facts, in case we have the value for our parameter.
 		for (const auto& f : facts)
 		{
 			if (f->GetOpType() == EOpType::BO_EQUAL_INT_CONST)
 			{
-				const auto typedFact = static_cast<const ASTOpFixed*>(f.get());
-				auto value = typedFact->children[typedFact->op.args.BoolEqualScalarConst.value].child();
-				if (value.get() == this)
+				const ASTOpBoolEqualIntConst* typedFact = static_cast<const ASTOpBoolEqualIntConst*>(f.get());
+				Ptr<ASTOp> Value = typedFact->Value.child();
+				if (Value.get() == this)
 				{
-					unknown = false;
-					return typedFact->op.args.BoolEqualScalarConst.constant;
+					bOutUnknown = false;
+					return typedFact->Constant;
 				}
 				else
 				{
diff --git a/Engine/Plugins/Mutable/Source/MutableTools/Private/MuT/ASTOpScalarArithmetic.cpp b/Engine/Plugins/Mutable/Source/MutableTools/Private/MuT/ASTOpScalarArithmetic.cpp
new file mode 100644
--- /dev/null
+++ b/Engine/Plugins/Mutable/Source/MutableTools/Private/MuT/ASTOpScalarArithmetic.cpp
@@ -0,0 +1,86 @@
+// Copyright Epic Games, Inc. All Rights Reserved.
+
+#include "MuT/ASTOpScalarArithmetic.h"
+#include "MuT/ASTOpMeshAddTags.h"
+
+#include "HAL/PlatformMath.h"
+#include "MuR/ModelPrivate.h"
+#include "MuR/RefCounted.h"
+#include "MuR/Types.h"
+
+namespace mu
+{
+
+	ASTOpScalarArithmetic::ASTOpScalarArithmetic()
+		: A(this)
+		, B(this)
+	{
+	}
+
+
+	ASTOpScalarArithmetic::~ASTOpScalarArithmetic()
+	{
+		// Explicit call needed to avoid recursive destruction
+		ASTOp::RemoveChildren();
+	}
+
+
+	bool ASTOpScalarArithmetic::IsEqual(const ASTOp& OtherUntyped) const
+	{
+		if (OtherUntyped.GetOpType()==GetOpType())
+		{
+			const ASTOpScalarArithmetic* Other = static_cast<const ASTOpScalarArithmetic*>(&OtherUntyped);
+			return A == Other->A &&
+				B == Other->B &&
+				Operation == Other->Operation;
+		}
+		return false;
+	}
+
+
+	uint64 ASTOpScalarArithmetic::Hash() const
+	{
+		uint64 Result = std::hash<void*>()(A.child().get());
+		hash_combine(Result, B.child().get());
+		hash_combine(Result, Operation);
+		return Result;
+	}
+
+
+	Ptr<ASTOp> ASTOpScalarArithmetic::Clone(MapChildFuncRef MapChild) const
+	{
+		Ptr<ASTOpScalarArithmetic> New = new ASTOpScalarArithmetic();
+		New->A = MapChild(A.child());
+		New->B = MapChild(B.child());
+		New->Operation = Operation;
+		return New;
+	}
+
+
+	void ASTOpScalarArithmetic::ForEachChild(const TFunctionRef<void(ASTChild&)> Func)
+	{
+		Func(A);
+		Func(B);
+	}
+
+
+	void ASTOpScalarArithmetic::Link(FProgram& Program, FLinkerOptions*)
+	{
+		// Already linked?
+		if (!linkedAddress)
+		{
+			OP::ArithmeticArgs Args;
+			FMemory::Memzero(Args);
+
+			if (A) Args.A = A->linkedAddress;
+			if (B) Args.B = B->linkedAddress;
+			Args.Operation = Operation;
+
+			linkedAddress = (OP::ADDRESS)Program.OpAddress.Num();
+			Program.OpAddress.Add(Program.ByteCode.Num());
+			AppendCode(Program.ByteCode, GetOpType());
+			AppendCode(Program.ByteCode, Args);
+		}
+
+	}
+}
diff --git a/Engine/Plugins/Mutable/Source/MutableTools/Private/MuT/ASTOpScalarArithmetic.h b/Engine/Plugins/Mutable/Source/MutableTools/Private/MuT/ASTOpScalarArithmetic.h
new file mode 100644
--- /dev/null
+++ b/Engine/Plugins/Mutable/Source/MutableTools/Private/MuT/ASTOpScalarArithmetic.h
@@ -0,0 +1,41 @@
+// Copyright Epic Games, Inc. All Rights Reserved.
+
+#pragma once
+
+#include "HAL/Platform.h"
+#include "MuR/Operations.h"
+#include "MuR/Ptr.h"
+#include "MuT/AST.h"
+
+
+namespace mu
+{
+	struct FProgram;
+
+	/** */
+	class ASTOpScalarArithmetic final : public ASTOp
+	{
+	public:
+
+		ASTChild A;
+		ASTChild B;
+		uint16 Operation=0;
+
+	public:
+
+		ASTOpScalarArithmetic();
+		ASTOpScalarArithmetic(const ASTOpScalarArithmetic&) = delete;
+		~ASTOpScalarArithmetic();
+
+		// ASTOp interface
+		virtual EOpType GetOpType() const override { return EOpType::SC_ARITHMETIC; }
+		virtual uint64 Hash() const override;
+		virtual bool IsEqual(const ASTOp&) const override;
+		virtual Ptr<ASTOp> Clone(MapChildFuncRef) const override;
+		virtual void ForEachChild(const TFunctionRef<void(ASTChild&)>) override;
+		virtual void Link(FProgram&, FLinkerOptions*) override;
+	};
+
+
+}
+
diff --git a/Engine/Plugins/Mutable/Source/MutableTools/Private/MuT/ASTOpScalarCurve.cpp b/Engine/Plugins/Mutable/Source/MutableTools/Private/MuT/ASTOpScalarCurve.cpp
--- a/Engine/Plugins/Mutable/Source/MutableTools/Private/MuT/ASTOpScalarCurve.cpp
+++ b/Engine/Plugins/Mutable/Source/MutableTools/Private/MuT/ASTOpScalarCurve.cpp
@@ -63,16 +63,16 @@ namespace mu
 	{
 		if (!linkedAddress)
 		{
-			OP::ScalarCurveArgs args;
-			memset(&args, 0, sizeof(args));
-			args.time = time ? time->linkedAddress : 0;
-			args.curve = program.AddConstant(Curve);
+			OP::ScalarCurveArgs Args;
+			FMemory::Memzero(Args);
+			Args.time = time ? time->linkedAddress : 0;
+			Args.curve = program.AddConstant(Curve);
 
 			linkedAddress = (OP::ADDRESS)program.OpAddress.Num();
 			//program.m_code.push_back(op);
 			program.OpAddress.Add((uint32_t)program.ByteCode.Num());
 			AppendCode(program.ByteCode, EOpType::SC_CURVE);
-			AppendCode(program.ByteCode, args);
+			AppendCode(program.ByteCode, Args);
 		}
 
 	}
diff --git a/Engine/Plugins/Mutable/Source/MutableTools/Private/MuT/ASTOpSwitch.cpp b/Engine/Plugins/Mutable/Source/MutableTools/Private/MuT/ASTOpSwitch.cpp
--- a/Engine/Plugins/Mutable/Source/MutableTools/Private/MuT/ASTOpSwitch.cpp
+++ b/Engine/Plugins/Mutable/Source/MutableTools/Private/MuT/ASTOpSwitch.cpp
@@ -3,6 +3,7 @@
 #include "MuT/ASTOpSwitch.h"
 
 #include "MuT/ASTOpParameter.h"
+#include "MuT/ASTOpConstantInt.h"
 #include "Containers/Map.h"
 #include "Misc/AssertionMacros.h"
 #include "MuR/ImagePrivate.h"
@@ -420,13 +421,11 @@ namespace mu
 		{
 			Ptr<ASTOp> Branch = def.child();
 
-			const ASTOpFixed* typedCondition = static_cast<const ASTOpFixed*>(variable.child().get());
+			const ASTOpConstantInt* typedCondition = static_cast<const ASTOpConstantInt*>(variable.child().get());
 			for (int32 o = 0; o < cases.Num(); ++o)
 			{
 				if (cases[o].branch &&
-					typedCondition->op.args.IntConstant.value
-					==
-					cases[o].condition)
+					typedCondition->Value == (int32)cases[o].condition)
 				{
 					Branch = cases[o].branch.child();
 					break;
diff --git a/Engine/Plugins/Mutable/Source/MutableTools/Private/MuT/CodeGenerator.cpp b/Engine/Plugins/Mutable/Source/MutableTools/Private/MuT/CodeGenerator.cpp
--- a/Engine/Plugins/Mutable/Source/MutableTools/Private/MuT/CodeGenerator.cpp
+++ b/Engine/Plugins/Mutable/Source/MutableTools/Private/MuT/CodeGenerator.cpp
@@ -19,6 +19,7 @@
 #include "MuT/ASTOpConditional.h"
 #include "MuT/ASTOpConstantBool.h"
 #include "MuT/ASTOpConstantResource.h"
+#include "MuT/ASTOpBoolAnd.h"
 #include "MuT/ASTOpImageCompose.h"
 #include "MuT/ASTOpImageMipmap.h"
 #include "MuT/ASTOpImagePixelFormat.h"
@@ -616,10 +617,9 @@ namespace mu
 			if (Component.ComponentCondition || Component.ObjectCondition)
 			{
 				// TODO: This could be done earlier?
-				Ptr<ASTOpFixed> ConditionOp = new ASTOpFixed();
-				ConditionOp->op.type = EOpType::BO_AND;
-				ConditionOp->SetChild(ConditionOp->op.args.BoolBinary.a, Component.ObjectCondition);
-				ConditionOp->SetChild(ConditionOp->op.args.BoolBinary.b, Component.ComponentCondition);
+				Ptr<ASTOpBoolAnd> ConditionOp = new ASTOpBoolAnd();
+				ConditionOp->A = Component.ObjectCondition;
+				ConditionOp->B = Component.ComponentCondition;
 
 				Ptr<ASTOpConditional> IfOp = new ASTOpConditional();
 				IfOp->type = EOpType::IN_CONDITIONAL;
@@ -1731,7 +1731,7 @@ namespace mu
                     // In case there is no group node, we generate a constant true condition
                     // This condition will be overwritten by the group nodes.
                     Ptr<ASTOpConstantBool> op = new ASTOpConstantBool();
-                    op->value = true;
+                    op->bValue = true;
                     paramOp = op;
                 }
 
diff --git a/Engine/Plugins/Mutable/Source/MutableTools/Private/MuT/CodeGenerator_Bool.cpp b/Engine/Plugins/Mutable/Source/MutableTools/Private/MuT/CodeGenerator_Bool.cpp
--- a/Engine/Plugins/Mutable/Source/MutableTools/Private/MuT/CodeGenerator_Bool.cpp
+++ b/Engine/Plugins/Mutable/Source/MutableTools/Private/MuT/CodeGenerator_Bool.cpp
@@ -11,6 +11,8 @@
 #include "MuT/AST.h"
 #include "MuT/ASTOpConstantBool.h"
 #include "MuT/ASTOpParameter.h"
+#include "MuT/ASTOpBoolAnd.h"
+#include "MuT/ASTOpBoolNot.h"
 #include "MuT/CodeGenerator.h"
 #include "MuT/NodeBool.h"
 #include "MuT/NodeBoolPrivate.h"
@@ -74,10 +76,7 @@ namespace mu
 	{
 		const NodeBoolConstant::Private& node = *Typed->GetPrivate();
 
-		Ptr<ASTOpConstantBool> op = new ASTOpConstantBool();
-		op->value = node.Value;
-
-		result.op = op;
+		result.op = new ASTOpConstantBool(node.Value);
 	}
 
 
@@ -124,12 +123,11 @@ namespace mu
 	{
 		const NodeBoolNot::Private& node = *Typed->GetPrivate();
 
-		Ptr<ASTOpFixed> op = new ASTOpFixed();
-		op->op.type = EOpType::BO_NOT;
+		Ptr<ASTOpBoolNot> op = new ASTOpBoolNot();
 
 		FBoolGenerationResult ChildResult;
 		GenerateBool(ChildResult, Options, node.m_pSource);
-		op->SetChild(op->op.args.BoolNot.source, ChildResult.op);
+		op->A = ChildResult.op;
 
 		result.op = op;
 	}
@@ -139,15 +137,14 @@ namespace mu
 	{
 		const NodeBoolAnd::Private& node = *Typed->GetPrivate();
 
-		Ptr<ASTOpFixed> op = new ASTOpFixed();
-		op->op.type = EOpType::BO_AND;
+		Ptr<ASTOpBoolAnd> op = new ASTOpBoolAnd;
 
 		FBoolGenerationResult ChildResult;
 		GenerateBool(ChildResult, Options, node.m_pA);
-		op->SetChild(op->op.args.BoolBinary.a, ChildResult.op);
+		op->A = ChildResult.op;
 
 		GenerateBool(ChildResult, Options, node.m_pB);
-		op->SetChild(op->op.args.BoolBinary.b, ChildResult.op);
+		op->B = ChildResult.op;
 
 		result.op = op;
 	}
diff --git a/Engine/Plugins/Mutable/Source/MutableTools/Private/MuT/CodeGenerator_Color.cpp b/Engine/Plugins/Mutable/Source/MutableTools/Private/MuT/CodeGenerator_Color.cpp
--- a/Engine/Plugins/Mutable/Source/MutableTools/Private/MuT/CodeGenerator_Color.cpp
+++ b/Engine/Plugins/Mutable/Source/MutableTools/Private/MuT/CodeGenerator_Color.cpp
@@ -16,6 +16,10 @@
 #include "MuT/ASTOpConditional.h"
 #include "MuT/ASTOpParameter.h"
 #include "MuT/ASTOpSwitch.h"
+#include "MuT/ASTOpConstantColor.h"
+#include "MuT/ASTOpColorArithmetic.h"
+#include "MuT/ASTOpColorSampleImage.h"
+#include "MuT/ASTOpColorFromScalars.h"
 #include "MuT/CodeGenerator.h"
 #include "MuT/CodeGenerator_FirstPass.h"
 #include "MuT/ErrorLog.h"
@@ -119,14 +123,10 @@ namespace mu
 	{
 		const NodeColourConstant& node = *Typed;
 
-		Ptr<ASTOpFixed> op = new ASTOpFixed();
-		op->op.type = EOpType::CO_CONSTANT;
-		op->op.args.ColourConstant.value[0] = node.Value.X;
-		op->op.args.ColourConstant.value[1] = node.Value.Y;
-		op->op.args.ColourConstant.value[2] = node.Value.Z;
-		op->op.args.ColourConstant.value[3] = node.Value.W;
+		Ptr<ASTOpConstantColor> Op = new ASTOpConstantColor;
+		Op->Value = node.Value;
 
-		result.op = op;
+		result.op = Op;
 	}
 
 
@@ -295,8 +295,7 @@ namespace mu
 		const NodeColourSampleImage& node = *Typed;
 
 		// Generate the code
-		Ptr<ASTOpFixed> op = new ASTOpFixed();
-		op->op.type = EOpType::CO_SAMPLEIMAGE;
+		Ptr<ASTOpColorSampleImage> op = new ASTOpColorSampleImage();
 
 		// Source image
 		int32 ComponentId = -1; // TODO.
@@ -319,7 +318,7 @@ namespace mu
 			base = GenerateMissingImageCode(TEXT("Sample image"), EImageFormat::RGB_UByte, Typed->GetMessageContext(), ImageOptions);
 		}
 		base = GenerateImageFormat(base, EImageFormat::RGB_UByte);
-		op->SetChild(op->op.args.ColourSampleImage.image, base);
+		op->Image = base;
 
 		FScalarGenerationResult ChildResult;
 
@@ -327,7 +326,7 @@ namespace mu
 		if (NodeScalar* pX = node.X.get())
 		{
 			GenerateScalar(ChildResult, Options, pX);
-			op->SetChild(op->op.args.ColourSampleImage.x, ChildResult.op);
+			op->X = ChildResult.op;
 		}
 		else
 		{
@@ -335,7 +334,7 @@ namespace mu
 			NodeScalarConstantPtr pNode = new NodeScalarConstant();
 			pNode->SetValue(0.5f);
 			GenerateScalar(ChildResult, Options, pNode);
-			op->SetChild(op->op.args.ColourSampleImage.x, ChildResult.op);
+			op->X = ChildResult.op;
 		}
 
 
@@ -343,7 +342,7 @@ namespace mu
 		if (NodeScalar* pY = node.Y.get())
 		{
 			GenerateScalar(ChildResult, Options, pY);
-			op->SetChild(op->op.args.ColourSampleImage.y, ChildResult.op);
+			op->Y = ChildResult.op;
 		}
 		else
 		{
@@ -351,11 +350,11 @@ namespace mu
 			NodeScalarConstantPtr pNode = new NodeScalarConstant();
 			pNode->SetValue(0.5f);
 			GenerateScalar(ChildResult, Options, pNode);
-			op->SetChild(op->op.args.ColourSampleImage.y, ChildResult.op);
+			op->Y = ChildResult.op;
 		}
 
 		// TODO
-		op->op.args.ColourSampleImage.filter = 0;
+		op->Filter = 0;
 
 		result.op = op;
 	}
@@ -366,8 +365,7 @@ namespace mu
 	{
 		const NodeColourFromScalars& node = *Typed;
 
-		Ptr<ASTOpFixed> op = new ASTOpFixed();
-		op->op.type = EOpType::CO_FROMSCALARS;
+		Ptr<ASTOpColorFromScalars> op = new ASTOpColorFromScalars();
 
 		FScalarGenerationResult ChildResult;
 
@@ -375,56 +373,56 @@ namespace mu
 		if (NodeScalar* pX = node.X.get())
 		{
 			GenerateScalar(ChildResult, Options, pX);
-			op->SetChild(op->op.args.ColourFromScalars.v[0], ChildResult.op );
+			op->V[0] = ChildResult.op;
 		}
 		else
 		{
 			NodeScalarConstantPtr pNode = new NodeScalarConstant();
 			pNode->SetValue(1.0f);
 			GenerateScalar(ChildResult, Options, pNode);
-			op->SetChild(op->op.args.ColourFromScalars.v[0], ChildResult.op);
+			op->V[0] = ChildResult.op;
 		}
 
 		// Y
 		if (NodeScalar* pY = node.Y.get())
 		{
 			GenerateScalar(ChildResult, Options, pY);
-			op->SetChild(op->op.args.ColourFromScalars.v[1], ChildResult.op);
+			op->V[1] = ChildResult.op;
 		}
 		else
 		{
 			NodeScalarConstantPtr pNode = new NodeScalarConstant();
 			pNode->SetValue(1.0f);
 			GenerateScalar(ChildResult, Options, pNode);
-			op->SetChild(op->op.args.ColourFromScalars.v[1], ChildResult.op);
+			op->V[1] = ChildResult.op;
 		}
 
 		// Z
 		if (NodeScalar* pZ = node.Z.get())
 		{
 			GenerateScalar(ChildResult, Options, pZ);
-			op->SetChild(op->op.args.ColourFromScalars.v[2], ChildResult.op);
+			op->V[2] = ChildResult.op;
 		}
 		else
 		{
 			NodeScalarConstantPtr pNode = new NodeScalarConstant();
 			pNode->SetValue(1.0f);
 			GenerateScalar(ChildResult, Options, pNode);
-			op->SetChild(op->op.args.ColourFromScalars.v[2], ChildResult.op);
+			op->V[2] = ChildResult.op;
 		}
 
 		// W
 		if (NodeScalar* pW = node.W.get())
 		{
 			GenerateScalar(ChildResult, Options, pW);
-			op->SetChild(op->op.args.ColourFromScalars.v[3], ChildResult.op);
+			op->V[3] = ChildResult.op;
 		}
 		else
 		{
 			NodeScalarConstantPtr pNode = new NodeScalarConstant();
 			pNode->SetValue(1.0f);
 			GenerateScalar(ChildResult, Options, pNode);
-			op->SetChild(op->op.args.ColourFromScalars.v[3], ChildResult.op);
+			op->V[3] = ChildResult.op;
 		}
 
 		result.op = op;
@@ -436,18 +434,17 @@ namespace mu
 	{
 		const NodeColourArithmeticOperation& node = *Typed;
 
-		Ptr<ASTOpFixed> op = new ASTOpFixed();
-		op->op.type = EOpType::CO_ARITHMETIC;
+		Ptr<ASTOpColorArithmetic> op = new ASTOpColorArithmetic;
 
 		switch (node.Operation)
 		{
-		case NodeColourArithmeticOperation::EOperation::Add: op->op.args.ColourArithmetic.operation = OP::ArithmeticArgs::ADD; break;
-		case NodeColourArithmeticOperation::EOperation::Subtract: op->op.args.ColourArithmetic.operation = OP::ArithmeticArgs::SUBTRACT; break;
-		case NodeColourArithmeticOperation::EOperation::Multiply: op->op.args.ColourArithmetic.operation = OP::ArithmeticArgs::MULTIPLY; break;
-		case NodeColourArithmeticOperation::EOperation::Divide: op->op.args.ColourArithmetic.operation = OP::ArithmeticArgs::DIVIDE; break;
+		case NodeColourArithmeticOperation::EOperation::Add: op->Operation = OP::ArithmeticArgs::ADD; break;
+		case NodeColourArithmeticOperation::EOperation::Subtract: op->Operation = OP::ArithmeticArgs::SUBTRACT; break;
+		case NodeColourArithmeticOperation::EOperation::Multiply: op->Operation = OP::ArithmeticArgs::MULTIPLY; break;
+		case NodeColourArithmeticOperation::EOperation::Divide: op->Operation = OP::ArithmeticArgs::DIVIDE; break;
 		default:
 			checkf(false, TEXT("Unknown arithmetic operation."));
-			op->op.args.ColourArithmetic.operation = OP::ArithmeticArgs::NONE;
+			op->Operation = OP::ArithmeticArgs::NONE;
 			break;
 		}
 
@@ -457,24 +454,22 @@ namespace mu
 		if (NodeColour* pA = node.A.get())
 		{
 			GenerateColor(ChildResult, Options, pA );
-			op->SetChild(op->op.args.ColourArithmetic.a, ChildResult.op);
+			op->A = ChildResult.op;
 		}
 		else
 		{
-			op->SetChild(op->op.args.ColourArithmetic.a,
-				CodeGenerator::GenerateMissingColourCode(TEXT("ColourArithmetic A"), Typed->GetMessageContext()));
+			op->A = GenerateMissingColourCode(TEXT("ColourArithmetic A"), Typed->GetMessageContext());
 		}
 
 		// B
 		if (NodeColour* pB = node.B.get())
 		{
 			GenerateColor(ChildResult, Options, pB);
-			op->SetChild(op->op.args.ColourArithmetic.b, ChildResult.op);
+			op->B = ChildResult.op;
 		}
 		else
 		{
-			op->SetChild(op->op.args.ColourArithmetic.b,
-				CodeGenerator::GenerateMissingColourCode(TEXT("ColourArithmetic B"), Typed->GetMessageContext()));
+			op->B = GenerateMissingColourCode(TEXT("ColourArithmetic B"), Typed->GetMessageContext());
 		}
 
 		result.op = op;
diff --git a/Engine/Plugins/Mutable/Source/MutableTools/Private/MuT/CodeGenerator_FirstPass.cpp b/Engine/Plugins/Mutable/Source/MutableTools/Private/MuT/CodeGenerator_FirstPass.cpp
--- a/Engine/Plugins/Mutable/Source/MutableTools/Private/MuT/CodeGenerator_FirstPass.cpp
+++ b/Engine/Plugins/Mutable/Source/MutableTools/Private/MuT/CodeGenerator_FirstPass.cpp
@@ -12,6 +12,8 @@
 #include "MuR/Platform.h"
 #include "MuT/AST.h"
 #include "MuT/ASTOpConstantBool.h"
+#include "MuT/ASTOpBoolEqualIntConst.h"
+#include "MuT/ASTOpBoolAnd.h"
 #include "MuT/ASTOpParameter.h"
 #include "MuT/CompilerPrivate.h"
 #include "MuT/ErrorLog.h"
@@ -381,23 +383,23 @@ namespace mu
 		for (int32 t = 0; t < InNode->Options.Num(); ++t)
 		{
 			// Create a comparison operation as the boolean parameter for the child
-			Ptr<ASTOpFixed> ParamOp = new ASTOpFixed();
-			ParamOp->op.type = EOpType::BO_EQUAL_INT_CONST;
-			ParamOp->SetChild(ParamOp->op.args.BoolEqualScalarConst.value, ScalarResult.op);
-			ParamOp->op.args.BoolEqualScalarConst.constant = (int16)t;
+			Ptr<ASTOpBoolEqualIntConst> ParamOp = new ASTOpBoolEqualIntConst;
+			ParamOp->Value = ScalarResult.op;
+			ParamOp->Constant = t;
+
+			Ptr<ASTOp> CurrentOp = ParamOp;
 
 			// Combine the new condition with previous conditions coming from parent objects
 			if (CurrentCondition.Last().ObjectCondition)
 			{
-				Ptr<ASTOpFixed> op = new ASTOpFixed();
-				op->op.type = EOpType::BO_AND;
-				op->SetChild(op->op.args.BoolBinary.a, CurrentCondition.Last().ObjectCondition);
-				op->SetChild(op->op.args.BoolBinary.b, ParamOp);
-				ParamOp = op;
+				Ptr<ASTOpBoolAnd> op = new ASTOpBoolAnd();
+				op->A = CurrentCondition.Last().ObjectCondition;
+				op->B = CurrentOp;
+				CurrentOp = op;
 			}
 
 			FConditionContext data;
-			data.ObjectCondition = ParamOp;
+			data.ObjectCondition = CurrentOp;
 			CurrentCondition.Push(data);
 
 			if (InNode->Options[t])
@@ -513,23 +515,23 @@ namespace mu
 		for (int32 t = 0; t < InNode->Options.Num(); ++t)
 		{
 			// Create a comparison operation as the boolean parameter for the child
-			Ptr<ASTOpFixed> ParamOp = new ASTOpFixed();
-			ParamOp->op.type = EOpType::BO_EQUAL_INT_CONST;
-			ParamOp->SetChild(ParamOp->op.args.BoolEqualScalarConst.value, ScalarResult.op);
-			ParamOp->op.args.BoolEqualScalarConst.constant = (int16)t;
+			Ptr<ASTOpBoolEqualIntConst> ParamOp = new ASTOpBoolEqualIntConst();
+			ParamOp->Value = ScalarResult.op;
+			ParamOp->Constant = t;
+
+			Ptr<ASTOp> CurrentOp = ParamOp;
 
 			// Combine the new condition with previous conditions coming from parent objects
 			if (CurrentCondition.Last().ObjectCondition)
 			{
-				Ptr<ASTOpFixed> op = new ASTOpFixed();
-				op->op.type = EOpType::BO_AND;
-				op->SetChild(op->op.args.BoolBinary.a, CurrentCondition.Last().ObjectCondition);
-				op->SetChild(op->op.args.BoolBinary.b, ParamOp);
-				ParamOp = op;
+				Ptr<ASTOpBoolAnd> op = new ASTOpBoolAnd();
+				op->A = CurrentCondition.Last().ObjectCondition;
+				op->B = CurrentOp;
+				CurrentOp = op;
 			}
 
 			FConditionContext data;
-			data.ObjectCondition = ParamOp;
+			data.ObjectCondition = CurrentOp;
 			CurrentCondition.Push(data);
 
 			if (InNode->Options[t])
@@ -680,10 +682,7 @@ namespace mu
                     case NodeObjectGroup::CS_ALWAYS_ALL:
                     {
                         // Create a constant true boolean that the optimiser will remove later.
-                        Ptr<ASTOpConstantBool> op = new ASTOpConstantBool();
-                        op->value = true;
-
-                        paramOp = op;
+                        paramOp = new ASTOpConstantBool(true);
                         break;
                     }
 
@@ -699,10 +698,9 @@ namespace mu
                         check(enumOp);
 
                         // Create a comparison operation as the boolean parameter for the child
-                        Ptr<ASTOpFixed> op = new ASTOpFixed();
-                        op->op.type = EOpType::BO_EQUAL_INT_CONST;
-                        op->SetChild( op->op.args.BoolEqualScalarConst.value, enumOp);
-                        op->op.args.BoolEqualScalarConst.constant = (int16_t)t;
+                        Ptr<ASTOpBoolEqualIntConst> op = new ASTOpBoolEqualIntConst();
+                        op->Value = enumOp;
+                        op->Constant = t;
 
                         paramOp = op;
                         break;
@@ -715,10 +713,9 @@ namespace mu
                 // Combine the new condition with previous conditions coming from parent objects
                 if (CurrentCondition.Last().ObjectCondition)
                 {
-                    Ptr<ASTOpFixed> op = new ASTOpFixed();
-                    op->op.type = EOpType::BO_AND;
-                    op->SetChild( op->op.args.BoolBinary.a,CurrentCondition.Last().ObjectCondition);
-                    op->SetChild( op->op.args.BoolBinary.b,paramOp);
+                    Ptr<ASTOpBoolAnd> op = new ASTOpBoolAnd();
+                    op->A = CurrentCondition.Last().ObjectCondition;
+                    op->B = paramOp;
                     paramOp = op;
                 }
 
diff --git a/Engine/Plugins/Mutable/Source/MutableTools/Private/MuT/CodeGenerator_Image.cpp b/Engine/Plugins/Mutable/Source/MutableTools/Private/MuT/CodeGenerator_Image.cpp
--- a/Engine/Plugins/Mutable/Source/MutableTools/Private/MuT/CodeGenerator_Image.cpp
+++ b/Engine/Plugins/Mutable/Source/MutableTools/Private/MuT/CodeGenerator_Image.cpp
@@ -1263,35 +1263,6 @@ namespace mu
 
         Ptr<ASTOp> switchAt = op;
 
-        // Make sure all options are the same format and size
-		// Disabled: This is not always desirable. For example if the image is going to be used in a 
-		// projector, the size doesn't need to be constrained.
-        //auto desc = switchAt->GetImageDesc( true );
-        //if ( desc.m_format == EImageFormat::None )
-        //{
-        //    // TODO: Look for the most generic of the options?
-        //    // For now force a decently generic one
-        //    desc.m_format = EImageFormat::RGBA_UByte;
-        //}
-
-        //if (desc.m_size[0]!=0 && desc.m_size[1]!=0)
-        //{
-        //    Ptr<ASTOpFixed> sop = new ASTOpFixed();
-        //    sop->op.type = EOpType::IM_RESIZE;
-        //    sop->op.args.ImageResize.size[0] = desc.m_size[0];
-        //    sop->op.args.ImageResize.size[1] = desc.m_size[1];
-        //    sop->SetChild( sop->op.args.ImageResize.source, switchAt );
-        //    switchAt = sop;
-        //}
-
-        //{
-        //    Ptr<ASTOpFixed> fop = new ASTOpFixed();
-        //    fop->op.type = EOpType::IM_PIXELFORMAT;
-        //    fop->op.args.ImagePixelFormat.format = desc.m_format;
-        //    fop->SetChild( fop->op.args.ImagePixelFormat.source, switchAt );
-        //    switchAt = fop;
-        //}
-
         Result.op = switchAt;
     }
 
diff --git a/Engine/Plugins/Mutable/Source/MutableTools/Private/MuT/CodeGenerator_Scalar.cpp b/Engine/Plugins/Mutable/Source/MutableTools/Private/MuT/CodeGenerator_Scalar.cpp
--- a/Engine/Plugins/Mutable/Source/MutableTools/Private/MuT/CodeGenerator_Scalar.cpp
+++ b/Engine/Plugins/Mutable/Source/MutableTools/Private/MuT/CodeGenerator_Scalar.cpp
@@ -15,6 +15,8 @@
 #include "MuT/ASTOpConditional.h"
 #include "MuT/ASTOpParameter.h"
 #include "MuT/ASTOpScalarCurve.h"
+#include "MuT/ASTOpScalarArithmetic.h"
+#include "MuT/ASTOpConstantScalar.h"
 #include "MuT/ASTOpSwitch.h"
 #include "MuT/CodeGenerator.h"
 #include "MuT/CodeGenerator_FirstPass.h"
@@ -120,9 +122,8 @@ namespace mu
 	{
 		const NodeScalarConstant::Private& node = *Typed->GetPrivate();
 
-		Ptr<ASTOpFixed> op = new ASTOpFixed();
-		op->op.type = EOpType::SC_CONSTANT;
-		op->op.args.ScalarConstant.value = node.Value;
+		Ptr<ASTOpConstantScalar> op = new ASTOpConstantScalar;
+		op->Value = node.Value;
 
 		result.op = op;
 	}
@@ -362,43 +363,38 @@ namespace mu
 	{
 		const NodeScalarArithmeticOperation::Private& node = *Typed->GetPrivate();
 
-		Ptr<ASTOpFixed> op = new ASTOpFixed();
-		op->op.type = EOpType::SC_ARITHMETIC;
+		Ptr<ASTOpScalarArithmetic> op = new ASTOpScalarArithmetic;
 
 		switch (node.m_operation)
 		{
-		case NodeScalarArithmeticOperation::AO_ADD: op->op.args.ScalarArithmetic.operation = OP::ArithmeticArgs::ADD; break;
-		case NodeScalarArithmeticOperation::AO_SUBTRACT: op->op.args.ScalarArithmetic.operation = OP::ArithmeticArgs::SUBTRACT; break;
-		case NodeScalarArithmeticOperation::AO_MULTIPLY: op->op.args.ScalarArithmetic.operation = OP::ArithmeticArgs::MULTIPLY; break;
-		case NodeScalarArithmeticOperation::AO_DIVIDE: op->op.args.ScalarArithmetic.operation = OP::ArithmeticArgs::DIVIDE; break;
+		case NodeScalarArithmeticOperation::AO_ADD: op->Operation = OP::ArithmeticArgs::ADD; break;
+		case NodeScalarArithmeticOperation::AO_SUBTRACT: op->Operation = OP::ArithmeticArgs::SUBTRACT; break;
+		case NodeScalarArithmeticOperation::AO_MULTIPLY: op->Operation = OP::ArithmeticArgs::MULTIPLY; break;
+		case NodeScalarArithmeticOperation::AO_DIVIDE: op->Operation = OP::ArithmeticArgs::DIVIDE; break;
 		default:
 			checkf(false, TEXT("Unknown arithmetic operation."));
-			op->op.args.ScalarArithmetic.operation = OP::ArithmeticArgs::NONE;
+			op->Operation = OP::ArithmeticArgs::NONE;
 			break;
 		}
 
 		// A
 		if (Node* pA = node.m_pA.get())
 		{
-			op->SetChild(op->op.args.ScalarArithmetic.a, Generate_Generic(pA, Options));
+			op->A = Generate_Generic(pA, Options);
 		}
 		else
 		{
-			op->SetChild(op->op.args.ScalarArithmetic.a,
-				CodeGenerator::GenerateMissingScalarCode( TEXT("ScalarArithmetic A"), 1.0f, Typed->GetMessageContext() )
-			);
+			op->A = GenerateMissingScalarCode( TEXT("ScalarArithmetic A"), 1.0f, Typed->GetMessageContext() );
 		}
 
 		// B
 		if (Node* pB = node.m_pB.get())
 		{
-			op->SetChild(op->op.args.ScalarArithmetic.b, Generate_Generic(pB, Options));
+			op->B = Generate_Generic(pB, Options);
 		}
 		else
 		{
-			op->SetChild(op->op.args.ScalarArithmetic.b,
-				CodeGenerator::GenerateMissingScalarCode( TEXT("ScalarArithmetic B"), 1.0f, Typed->GetMessageContext() )
-			);
+			op->B =  GenerateMissingScalarCode( TEXT("ScalarArithmetic B"), 1.0f, Typed->GetMessageContext() );
 		}
 
 		result.op = op;
diff --git a/Engine/Plugins/Mutable/Source/MutableTools/Private/MuT/CodeGenerator_SecondPass.cpp b/Engine/Plugins/Mutable/Source/MutableTools/Private/MuT/CodeGenerator_SecondPass.cpp
--- a/Engine/Plugins/Mutable/Source/MutableTools/Private/MuT/CodeGenerator_SecondPass.cpp
+++ b/Engine/Plugins/Mutable/Source/MutableTools/Private/MuT/CodeGenerator_SecondPass.cpp
@@ -7,6 +7,8 @@
 #include "MuR/Operations.h"
 #include "MuR/Types.h"
 #include "MuT/ASTOpConstantBool.h"
+#include "MuT/ASTOpBoolAnd.h"
+#include "MuT/ASTOpBoolNot.h"
 #include "MuT/ErrorLog.h"
 #include "MuT/CodeGenerator_FirstPass.h"
 
@@ -38,13 +40,13 @@ namespace mu
 		// If this tag is already in the list of positive tags, return true as condition
 		if (posTag.Contains(tagIndex))
 		{
-			return m_opPool.Add(new ASTOpConstantBool(true));
+			return OpPool.Add(new ASTOpConstantBool(true));
 		}
 
 		// If this tag is already in the list of negative tags, return false as condition
 		if (negTag.Contains(tagIndex))
 		{
-			return m_opPool.Add(new ASTOpConstantBool(false));
+			return OpPool.Add(new ASTOpConstantBool(false));
 		}
 
 		// Cached?
@@ -103,7 +105,7 @@ namespace mu
 			if (surfCondition && surfCondition->GetOpType()==EOpType::BO_CONSTANT)
 			{
 				const ASTOpConstantBool* constOp = static_cast<const ASTOpConstantBool*>(surfCondition.get());
-				if (constOp->value == false)
+				if (constOp->bValue == false)
 				{
 					continue;
 				}
@@ -112,15 +114,14 @@ namespace mu
 			Ptr<ASTOp> fullCondition;
 			if (surfCondition)
 			{
-				Ptr<ASTOpFixed> f = new ASTOpFixed();
-				f->op.type = EOpType::BO_AND;
-				f->SetChild(f->op.args.BoolBinary.a, surface.ObjectCondition);
-				f->SetChild(f->op.args.BoolBinary.b, surfCondition);
-				fullCondition = m_opPool.Add(f);
+				Ptr<ASTOpBoolAnd> f = new ASTOpBoolAnd;
+				f->A = surface.ObjectCondition;
+				f->B = surfCondition;
+				fullCondition = OpPool.Add(f);
 			}
 			else
 			{
-				fullCondition = m_opPool.Add(surface.ObjectCondition);
+				fullCondition = OpPool.Add(surface.ObjectCondition);
 			}
 
 
@@ -130,11 +131,10 @@ namespace mu
 			}
 			else
 			{
-				Ptr<ASTOpFixed> o = new ASTOpFixed();
-				o->op.type = EOpType::BO_OR;
-				o->SetChild(o->op.args.BoolBinary.a, fullCondition);
-				o->SetChild(o->op.args.BoolBinary.b, c);
-				c = m_opPool.Add(o);
+				Ptr<ASTOpBoolAnd> o = new ASTOpBoolAnd();
+				o->A = fullCondition;
+				o->B = c;
+				c = OpPool.Add(o);
 			}
 
 			// Optimise the condition now.
@@ -158,13 +158,13 @@ namespace mu
 		// If this surface is already in the list of positive surfaces, return true as condition
 		if (posSurf.Contains(Index))
 		{
-			return m_opPool.Add(new ASTOpConstantBool(true));
+			return OpPool.Add(new ASTOpConstantBool(true));
 		}
 
 		// If this surface is already in the list of negative surfaces, return false as condition
 		if (negSurf.Contains(Index))
 		{
-			return m_opPool.Add(new ASTOpConstantBool(false));
+			return OpPool.Add(new ASTOpConstantBool(false));
 		}
 
 		Ptr<ASTOp> c;
@@ -176,7 +176,7 @@ namespace mu
 			{
 				// This could happen if a tag is in a variation but noone defines it.
 				// This surface depends on a tag that will never be active, so it will never be used.
-				return m_opPool.Add(new ASTOpConstantBool(false));
+				return OpPool.Add(new ASTOpConstantBool(false));
 			}
 
 			int32 tagIndex = int32(it - &FirstPass->Tags[0]);
@@ -205,7 +205,7 @@ namespace mu
 			{
 				const ASTOpConstantBool* constOp = static_cast<const ASTOpConstantBool*>(tagCondition.get());
 				isConstant = true;
-				constantValue = constOp->value;
+				constantValue = constOp->bValue;
 			}
 
 			if (!isConstant)
@@ -216,11 +216,10 @@ namespace mu
 				}
 				else
 				{
-					Ptr<ASTOpFixed> o = new ASTOpFixed;
-					o->op.type = EOpType::BO_AND;
-					o->SetChild(o->op.args.BoolBinary.a, tagCondition);
-					o->SetChild(o->op.args.BoolBinary.b, c);
-					c = m_opPool.Add(o);
+					Ptr<ASTOpBoolAnd> o = new ASTOpBoolAnd;
+					o->A = tagCondition;
+					o->B = c;
+					c = OpPool.Add(o);
 				}
 			}
 			else if (constantValue == true)
@@ -233,7 +232,7 @@ namespace mu
 				Ptr<ASTOpConstantBool> f = new ASTOpConstantBool(false);
 
 				// No need to evaluate anything else.
-				c = m_opPool.Add(f);
+				c = OpPool.Add(f);
 				break;
 			}
 		}
@@ -264,9 +263,7 @@ namespace mu
 			// No condition is equal to a conditional with a true constant
 			if (!tagCondition)
 			{
-				ASTOpConstantBool* ConstOp = new ASTOpConstantBool();
-				ConstOp->value = true;
-				tagCondition = ConstOp;
+				tagCondition = new ASTOpConstantBool(true);
 			}
 
 			// TODO: Optimise the tag condition here
@@ -278,27 +275,25 @@ namespace mu
 			{
 				const ASTOpConstantBool* constOp = static_cast<const ASTOpConstantBool*>(tagCondition.get());
 				isConstant = true;
-				constantValue = constOp->value;
+				constantValue = constOp->bValue;
 			}
 
 
 			if (!isConstant && tagCondition)
 			{
-				Ptr<ASTOpFixed> n = new ASTOpFixed;
-				n->op.type = EOpType::BO_NOT;
-				n->SetChild(n->op.args.BoolNot.source, tagCondition);
+				Ptr<ASTOpBoolNot> n = new ASTOpBoolNot;
+				n->A = tagCondition;
 
 				if (!c)
 				{
-					c = m_opPool.Add(n);
+					c = OpPool.Add(n);
 				}
 				else
 				{
-					Ptr<ASTOpFixed> o = new ASTOpFixed;
-					o->op.type = EOpType::BO_AND;
-					o->SetChild(o->op.args.BoolBinary.a, n);
-					o->SetChild(o->op.args.BoolBinary.b, c);
-					c = m_opPool.Add(o);
+					Ptr<ASTOpBoolAnd> o = new ASTOpBoolAnd;
+					o->A = n;
+					o->B = c;
+					c = OpPool.Add(o);
 				}
 			}
 			else if (isConstant && constantValue == true)
@@ -307,7 +302,7 @@ namespace mu
 				Ptr<ASTOpConstantBool> f = new ASTOpConstantBool(false);
 
 				// No need to evaluate anything else.
-				c = m_opPool.Add(f);
+				c = OpPool.Add(f);
 				break;
 			}
 		}
@@ -414,10 +409,9 @@ namespace mu
 					Surface.NegativeTags,
 					Empty, Empty, Empty, Empty);
 
-				Ptr<ASTOpFixed> ConditionOp = new ASTOpFixed();
-				ConditionOp->op.type = EOpType::BO_AND;
-				ConditionOp->SetChild(ConditionOp->op.args.BoolBinary.a, Surface.ObjectCondition);
-				ConditionOp->SetChild(ConditionOp->op.args.BoolBinary.b, c);
+				Ptr<ASTOpBoolAnd> ConditionOp = new ASTOpBoolAnd();
+				ConditionOp->A = Surface.ObjectCondition;
+				ConditionOp->B = c;
 
 				Surface.FinalCondition = ConditionOp;
 			}
@@ -433,11 +427,10 @@ namespace mu
 			//		Surface.Edits[EditIndex].NegativeTags,
 			//		Empty, Empty, Empty, Empty);
 
-			//	Ptr<ASTOpFixed> OpAnd = new ASTOpFixed;
-			//	OpAnd->op.type = EOpType::BO_AND;
-			//	OpAnd->SetChild(OpAnd->op.args.BoolBinary.a, Edit.Condition);
-			//	OpAnd->SetChild(OpAnd->op.args.BoolBinary.b, c);
-			//	c = m_opPool.Add(OpAnd);
+			//	Ptr<ASTOpBoolAnd> OpAnd = new ASTOpBoolAnd;
+			//	OpAnd->a = Edit.Condition;
+			//	OpAnd->b = c;
+			//	c = OpPool.Add(OpAnd);
 			//	
 			//	Edit.Condition = OpAnd;
 			//}
@@ -451,10 +444,9 @@ namespace mu
 				FirstPass->Modifiers[ModifierIndex].NegativeTags,
 				Empty, Empty, Empty, Empty);
 			
-			Ptr<ASTOpFixed> ConditionOp = new ASTOpFixed();
-			ConditionOp->op.type = EOpType::BO_AND;
-			ConditionOp->SetChild(ConditionOp->op.args.BoolBinary.a, FirstPass->Modifiers[ModifierIndex].ObjectCondition);
-			ConditionOp->SetChild(ConditionOp->op.args.BoolBinary.b, c);
+			Ptr<ASTOpBoolAnd> ConditionOp = new ASTOpBoolAnd;
+			ConditionOp->A = FirstPass->Modifiers[ModifierIndex].ObjectCondition;
+			ConditionOp->B = c;
 
 			FirstPass->Modifiers[ModifierIndex].FinalCondition = ConditionOp;
 		}
diff --git a/Engine/Plugins/Mutable/Source/MutableTools/Private/MuT/CodeGenerator_SecondPass.h b/Engine/Plugins/Mutable/Source/MutableTools/Private/MuT/CodeGenerator_SecondPass.h
--- a/Engine/Plugins/Mutable/Source/MutableTools/Private/MuT/CodeGenerator_SecondPass.h
+++ b/Engine/Plugins/Mutable/Source/MutableTools/Private/MuT/CodeGenerator_SecondPass.h
@@ -79,7 +79,7 @@ namespace mu
 
         TMap<FConditionGenerationKey,Ptr<ASTOp>> TagConditionGenerationCache;
 
-        UniqueOpPool m_opPool;
+        FUniqueOpPool OpPool;
 
         //!
         Ptr<ASTOp> GenerateTagCondition( int32 tagIndex,
diff --git a/Engine/Plugins/Mutable/Source/MutableTools/Private/MuT/CodeOptimiser.cpp b/Engine/Plugins/Mutable/Source/MutableTools/Private/MuT/CodeOptimiser.cpp
--- a/Engine/Plugins/Mutable/Source/MutableTools/Private/MuT/CodeOptimiser.cpp
+++ b/Engine/Plugins/Mutable/Source/MutableTools/Private/MuT/CodeOptimiser.cpp
@@ -9,6 +9,7 @@
 #include "MuT/ASTOpSwitch.h"
 #include "MuT/ASTOpConstantResource.h"
 #include "MuT/ASTOpConstantBool.h"
+#include "MuT/ASTOpConstantColor.h"
 #include "MuT/ASTOpMeshApplyShape.h"
 #include "MuT/ASTOpMeshBindShape.h"
 #include "MuT/ASTOpMeshClipMorphPlane.h"
@@ -567,12 +568,7 @@ namespace mu
 				MUTABLE_CPUPROFILER_SCOPE(ConstantBool);
 
 				bool value = System->GetPrivate()->BuildBool( model, LocalParams.Get(), at );
-
-				{
-					mu::Ptr<ASTOpConstantBool> constantOp = new ASTOpConstantBool();
-					constantOp->value = value;
-					Result = constantOp;
-				}
+				Result = new ASTOpConstantBool(value);
 				break;
 			}
 
@@ -584,13 +580,9 @@ namespace mu
 				ResultColor = System->GetPrivate()->BuildColour( model, LocalParams.Get(), at );
 
 				{
-					mu::Ptr<ASTOpFixed> constantOp = new ASTOpFixed();
-					constantOp->op.type = EOpType::CO_CONSTANT;
-					constantOp->op.args.ColourConstant.value[0] = ResultColor[0];
-					constantOp->op.args.ColourConstant.value[1] = ResultColor[1];
-					constantOp->op.args.ColourConstant.value[2] = ResultColor[2];
-					constantOp->op.args.ColourConstant.value[3] = ResultColor[3];
-					Result = constantOp;
+					mu::Ptr<ASTOpConstantColor> ConstantOp = new ASTOpConstantColor();
+					ConstantOp->Value = ResultColor;
+					Result = ConstantOp;
 				}
 				break;
 			}
diff --git a/Engine/Plugins/Mutable/Source/MutableTools/Private/MuT/CodeOptimiser_Logic.cpp b/Engine/Plugins/Mutable/Source/MutableTools/Private/MuT/CodeOptimiser_Logic.cpp
--- a/Engine/Plugins/Mutable/Source/MutableTools/Private/MuT/CodeOptimiser_Logic.cpp
+++ b/Engine/Plugins/Mutable/Source/MutableTools/Private/MuT/CodeOptimiser_Logic.cpp
@@ -10,997 +10,14 @@
 #include "MuT/ASTOpInstanceAdd.h"
 #include "MuT/ASTOpSwitch.h"
 #include "MuT/ASTOpMeshMerge.h"
+#include "MuT/ASTOpBoolEqualIntConst.h"
 #include "MuT/CodeOptimiser.h"
 
 
 namespace mu
 {
 
-
-    //---------------------------------------------------------------------------------------------
-    //---------------------------------------------------------------------------------------------
-    //---------------------------------------------------------------------------------------------
-//    IntExpressionPtr GetIntExpression( OP::ADDRESS at, const FProgram& program )
-//    {
-//        IntExpressionPtr pRes = new IntExpression;
-
-//        const OP& op = program.m_code[at];
-//        EOpType type = op.type;
-//        switch ( type )
-//        {
-//        case EOpType::NU_CONSTANT:
-//            pRes->type = IntExpression::IET_CONSTANT;
-//            pRes->value = op.args.IntConstant.value;
-//            break;
-
-//        case EOpType::NU_PARAMETER:
-//            pRes->type = IntExpression::IET_PARAMETER;
-//            pRes->parameter = op.args.Parameter.variable;
-//            break;
-
-//        default:
-//            //check( false );
-//            pRes->type = IntExpression::IET_UNKNOWN;
-//            break;
-//        }
-
-//        return pRes;
-//    }
-
-
-    //---------------------------------------------------------------------------------------------
-    //---------------------------------------------------------------------------------------------
-    //---------------------------------------------------------------------------------------------
-//    BoolExpressionPtr GetBoolExpression( OP::ADDRESS at,
-//                                         const FProgram& program,
-//                                         BOOL_EXPRESSION_CACHE* cache )
-//    {
-//        if (cache && (size_t)at<=cache->Values.size() && cache->Values[at])
-//        {
-//            return cache->Values[at];
-//        }
-
-//        BoolExpressionPtr pRes = new BoolExpression;
-
-//        const OP& op = program.m_code[at];
-//        EOpType type = (EOpType)op.type;
-//        switch ( type )
-//        {
-//        case EOpType::BO_CONSTANT:
-//            pRes->type = op.args.BoolConstant.value
-//                ? BoolExpression::BET_TRUE
-//                : BoolExpression::BET_FALSE;
-//            break;
-
-//        case EOpType::BO_PARAMETER:
-//            pRes->type = BoolExpression::BET_PARAMETER;
-//            pRes->parameter = op.args.Parameter.variable;
-//            break;
-
-//        case EOpType::BO_AND:
-//            pRes->type = BoolExpression::BET_AND;
-//            pRes->a = GetBoolExpression( op.args.BoolBinary.a, program, cache );
-//            pRes->b = GetBoolExpression( op.args.BoolBinary.b, program, cache );
-//            break;
-
-//        case EOpType::BO_OR:
-//            pRes->type = BoolExpression::BET_OR;
-//            pRes->a = GetBoolExpression( op.args.BoolBinary.a, program, cache );
-//            pRes->b = GetBoolExpression( op.args.BoolBinary.b, program, cache );
-//            break;
-
-//        case EOpType::BO_NOT:
-//            pRes->type = BoolExpression::BET_NOT;
-//            pRes->a = GetBoolExpression( op.args.BoolNot.source, program, cache );
-//            break;
-
-//        case EOpType::BO_EQUAL_INT_CONST:
-//            pRes->type = BoolExpression::BET_INT_EQUAL_CONST;
-//            pRes->intExp = GetIntExpression( op.args.BoolEqualScalarConst.value, program );
-//            pRes->intValue = op.args.BoolEqualScalarConst.constant;
-//            break;
-
-//        default:
-//            check( false );
-//            pRes->type = BoolExpression::BET_UNKNOWN;
-//            break;
-//        }
-
-//        if (cache)
-//        {
-//            if (cache->Values.size()<program.OpAddress.Num())
-//            {
-//                cache->Values.resize(program.OpAddress.Num());
-//            }
-//            cache->Values[at]=pRes;
-//        }
-
-//        return pRes;
-//    }
-
-
-    //---------------------------------------------------------------------------------------------
-    //---------------------------------------------------------------------------------------------
-    //---------------------------------------------------------------------------------------------
-//    void LogicCallstackGather::Apply( FProgram& program, int state )
-//    {
-//        MUTABLE_CPUPROFILER_SCOPE(LogicCallstackGather);
-//        Traverse( program.States[state].Root, program );
-//    }
-
-
-//    //---------------------------------------------------------------------------------------------
-//    void LogicCallstackGather::Visit( OP::ADDRESS at, FProgram& program )
-//    {
-//        // We traverse top down
-//        EOpType type = (EOpType)program.m_code[at].type;
-//        OP::ARGS args = program.m_code[at].args;
-
-//        switch ( type )
-//        {
-
-//        //-----------------------------------------------------------------------------------------
-//        case EOpType::NU_CONDITIONAL:
-//        case EOpType::SC_CONDITIONAL:
-//        case EOpType::CO_CONDITIONAL:
-//        case EOpType::IM_CONDITIONAL:
-//        case EOpType::ME_CONDITIONAL:
-//        case EOpType::LA_CONDITIONAL:
-//        case EOpType::IN_CONDITIONAL:
-//        {
-//            m_factsPerCallstack[at].push_back( m_facts );
-
-//            BoolExpressionPtr pCond = GetBoolExpression( args.Conditional.condition, program,
-//                                                         &m_boolExpressionCache );
-
-//            // Check if the condition has a value with the current callstack. If it does, we skip
-//            // the contradictory paths, as they will never happen.
-//            BoolExpression::TYPE result = pCond->Evaluate( m_facts );
-
-//            // Recurse the condition without adding extra facts
-//            Visit( args.Conditional.condition, program );
-
-//            if ( result != BoolExpression::BET_FALSE )
-//            {
-//                // yes branch
-//                m_facts.push_back( pCond );
-//                Visit( args.Conditional.yes, program );
-//                m_facts.pop_back();
-//            }
-
-//            if ( result != BoolExpression::BET_TRUE )
-//            {
-//                // no branch
-//                BoolExpressionPtr pNot = new BoolExpression;
-//                pNot->type = BoolExpression::BET_NOT;
-//                pNot->a = pCond;
-
-//                m_facts.push_back( pNot );
-//                Visit( args.Conditional.no, program );
-//                m_facts.pop_back();
-//            }
-
-//            break;
-//        }
-
-//        case EOpType::NU_SWITCH:
-//        case EOpType::SC_SWITCH:
-//        case EOpType::CO_SWITCH:
-//        case EOpType::IM_SWITCH:
-//        case EOpType::VO_SWITCH:
-//        case EOpType::ME_SWITCH:
-//        case EOpType::LA_SWITCH:
-//        {
-//            m_factsPerCallstack[at].push_back( m_facts );
-
-//            IntExpressionPtr pValue = GetIntExpression( args.Switch.variable, program );
-
-//            const auto& options = program.m_constantSwitches[args.Switch.options].Options;
-//            for ( size_t i=0; i<options.size(); ++i )
-//            {
-//                if ( options[i].at )
-//                {
-//                    BoolExpressionPtr pCond = new BoolExpression;
-//                    pCond->type = BoolExpression::BET_INT_EQUAL_CONST;
-//                    pCond->intExp = pValue;
-//                    pCond->intValue = args.Switch.conditions[i];
-
-//                    // Check if the condition has a value with the current callstack. If it does, we
-//                    // skip the contradictory paths, as they will never happen.
-//                    BoolExpression::TYPE result = pCond->Evaluate( m_facts );
-
-//                    if (result!=BoolExpression::BET_FALSE)
-//                    {
-//                        // We don't know if it is true or false. Add its value as a fact and
-//                        // recurse.
-//                        m_facts.push_back( pCond );
-//                        Visit( options[i].at, program );
-//                        m_facts.pop_back();
-//                    }
-//                }
-//            }
-
-//            if ( args.Switch.def )
-//            {
-//                // TODO: Add facts negating the options in the switch
-//                Visit( args.Switch.def, program );
-//            }
-//            break;
-//        }
-
-//        default:
-//            // Normal recursion
-//            Recurse( at, program );
-//            break;
-
-//        }
-//    }
-
-
-    //---------------------------------------------------------------------------------------------
-    //---------------------------------------------------------------------------------------------
-    //---------------------------------------------------------------------------------------------
-//    bool ClonelessLogicOptimiser::Apply( FProgram& program, int state )
-//    {
-//        MUTABLE_CPUPROFILER_SCOPE(ClonelessLogicOptimiser);
-
-//        // Gather all the conditions for every op
-//        m_callstacks.Apply( program, state );
-
-//        bModified = false;
-//        Traverse( program.States[state].Root, program );
-
-//        return bModified;
-//    }
-
-
-//    //---------------------------------------------------------------------------------------------
-//    OP::ADDRESS ClonelessLogicOptimiser::Visit( OP::ADDRESS at, FProgram& program )
-//    {
-//        // We traverse top down
-
-//        EOpType type = (EOpType)program.m_code[at].type;
-//        OP::ARGS args = program.m_code[at].args;
-
-//        switch ( type )
-//        {
-
-//        //-----------------------------------------------------------------------------------------
-//        case EOpType::NU_CONDITIONAL:
-//        case EOpType::SC_CONDITIONAL:
-//        case EOpType::CO_CONDITIONAL:
-//        case EOpType::IM_CONDITIONAL:
-//        case EOpType::ME_CONDITIONAL:
-//        case EOpType::LA_CONDITIONAL:
-//        case EOpType::IN_CONDITIONAL:
-//        {
-//            BOOL_EXPRESSION_CACHE localBoolExpressionCache;
-//            BoolExpressionPtr pCond = GetBoolExpression( args.Conditional.condition, program, &localBoolExpressionCache );
-
-//            // Check if the condition has a value with all the possible callstacks on this op
-//            vector< vector<BoolExpressionPtr> >& callstacks = m_callstacks.m_factsPerCallstack[at];
-//            check( callstacks.size() );
-
-//            BoolExpression::TYPE result = pCond->Evaluate( callstacks[0] );
-//            for ( size_t c=1; c<callstacks.size(); ++c )
-//            {
-//                BoolExpression::TYPE thisResult = pCond->Evaluate( callstacks[c] );
-//                if (result!=thisResult)
-//                {
-//                    result = BoolExpression::BET_UNKNOWN;
-//                    break;
-//                }
-//            }
-
-//            if ( result==BoolExpression::BET_TRUE )
-//            {
-//                bModified = true;
-//                m_callstacks.m_factsPerCallstack[at] = m_callstacks.m_factsPerCallstack[program.m_code[at].args.Conditional.yes];
-//                at = program.m_code[at].args.Conditional.yes;
-//                at = Visit( at, program );
-//            }
-//            else if ( result==BoolExpression::BET_FALSE )
-//            {
-//                bModified = true;
-//                m_callstacks.m_factsPerCallstack[at] = m_callstacks.m_factsPerCallstack[program.m_code[at].args.Conditional.no];
-//                at = program.m_code[at].args.Conditional.no;
-//                at = Visit( at, program );
-//            }
-//            else
-//            {
-//                at = Recurse( at, program );
-//            }
-
-//            break;
-//        }
-
-////        case EOpType::NU_SWITCH:
-////        case EOpType::SC_SWITCH:
-////        case EOpType::CO_SWITCH:
-////        case EOpType::IM_SWITCH:
-////        case EOpType::VO_SWITCH:
-////        case EOpType::ME_SWITCH:
-////        case EOpType::LA_SWITCH:
-////        {
-////            IntExpressionPtr pValue = GetIntExpression( args.Switch.variable, program );
-
-////            vector< vector<BoolExpressionPtr> >& callstacks = m_callstacks.m_factsPerCallstack[at];
-////            check( callstacks.size() );
-
-////            for ( int i=0;i<MUTABLE_OP_MAX_SWITCH_OPTIONS; ++i )
-////            {
-////                if ( args.Switch.values[i] )
-////                {
-////                    BoolExpressionPtr pCond = new BoolExpression;
-////                    pCond->type = BoolExpression::BET_INT_EQUAL_CONST;
-////                    pCond->intExp = pValue;
-////                    pCond->intValue = args.Switch.conditions[i];
-
-////                    // Check if the condition has a value with all the possible callstacks on this op
-////                    BoolExpression::TYPE result = pCond->Evaluate( callstacks[0] );
-////                    for ( size_t c=1; c<callstacks.size(); ++c )
-////                    {
-////                        BoolExpression::TYPE thisResult = pCond->Evaluate( callstacks[c] );
-////                        if (result!=thisResult)
-////                        {
-////                            result = BoolExpression::BET_UNKNOWN;
-////                            break;
-////                        }
-////                    }
-
-////                    if ( result==BoolExpression::BET_TRUE )
-////                    {
-////                        // This is the only possible branch.
-////                        bModified = true;
-////                        at = args.Switch.values[i];
-////                        //at = Visit( at, program );
-////                        return at;
-////                        break;
-////                    }
-////                    else if ( result==BoolExpression::BET_FALSE )
-////                    {
-////                        // This option is impossible
-////                        args.Switch.values[i] = 0;
-////                        args.Switch.conditions[i] = 0;
-////                    }
-////                    else
-////                    {
-////                        // We don't know if it is true or false. Add its value as a fact and
-////                        // recurse.
-////                        //args.Switch.values[i] = Visit( args.Switch.values[i], program );
-////                    }
-////                }
-////            }
-
-////            if ( args.Switch.def )
-////            {
-////                // TODO: Add facts negating the options in the switch
-////                //args.Switch.def = Visit( args.Switch.def, program );
-////            }
-
-////            program.m_code[at].args = args;
-////            break;
-////        }
-
-//        default:
-//            // Normal recursion
-//            at = Recurse( at, program );
-//            break;
-
-//        }
-
-//        return at;
-//    }
-
-
-//    //---------------------------------------------------------------------------------------------
-//    //---------------------------------------------------------------------------------------------
-//    //---------------------------------------------------------------------------------------------
-//    bool LogicOptimiser::Apply( FProgram& program, int state )
-//    {
-//        m_initialCodeSize = program.OpAddress.Num();
-//        bModified = false;
-//        Traverse( program.States[state].Root, program );
-
-//        return bModified;
-//    }
-
-
-//    //---------------------------------------------------------------------------------------------
-//    OP::ADDRESS LogicOptimiser::Visit( OP::ADDRESS at, FProgram& program )
-//    {
-//        // Sanity stop
-//        const int sanityFactor = 4;
-//        if ( program.OpAddress.Num()>m_initialCodeSize*sanityFactor)
-//        {
-//            return at;
-//        }
-
-//        // We traverse top down
-
-//        EOpType type = (EOpType)program.m_code[at].type;
-//        OP::ARGS args = program.m_code[at].args;
-
-//        // We only apply it to instances, for performance reasons.
-//        if (GetOpType(type)!=Instance)
-//        {
-//            return at;
-//        }
-
-//        switch ( type )
-//        {
-
-//        //-----------------------------------------------------------------------------------------
-////        case EOpType::NU_CONDITIONAL:
-////        case EOpType::SC_CONDITIONAL:
-////        case EOpType::CO_CONDITIONAL:
-////        case EOpType::IM_CONDITIONAL:
-////        case EOpType::ME_CONDITIONAL:
-////        case EOpType::LA_CONDITIONAL:
-//        case EOpType::IN_CONDITIONAL:
-//        {
-//            BoolExpressionPtr pCond = GetBoolExpression( args.Conditional.condition, program, &m_boolExpressionCache );
-
-//            // Check if the condition has a value with the current facts
-//            // TODO: This is a simple adhoc check which is also quite imprecise.
-//            BoolExpression::TYPE result = pCond->Evaluate( m_facts );
-//            if ( result==BoolExpression::BET_TRUE )
-//            {
-//                bModified = true;
-//                at = program.m_code[at].args.Conditional.yes;
-//                at = Visit( at, program );
-//            }
-//            else if ( result==BoolExpression::BET_FALSE )
-//            {
-//                bModified = true;
-//                at = program.m_code[at].args.Conditional.no;
-//                at = Visit( at, program );
-//            }
-//            else
-//            {
-//                // Recurse after modifying the current facts
-
-//                // TODO: Recurse condition too?
-
-//                // yes
-//                m_facts.push_back( pCond );
-//                args.Conditional.yes = Visit( args.Conditional.yes, program );
-//                m_facts.pop_back();
-
-//                // no
-//                BoolExpressionPtr pNot = new BoolExpression;
-//                pNot->type = BoolExpression::BET_NOT;
-//                pNot->a = pCond;
-
-//                m_facts.push_back( pNot );
-//                args.Conditional.no = Visit( args.Conditional.no, program );
-//                m_facts.pop_back();
-
-//                // update op if necessary
-//                if ( FMemory::Memcmp( &args, &program.m_code[at].args, sizeof(OP::ARGS) ) )
-//                {
-//                    OP op = program.m_code[at];
-//                    op.args = args;
-//                    at = program.AddOp( op );
-//                }
-//            }
-
-//            break;
-//        }
-
-////        case EOpType::NU_SWITCH:
-////        case EOpType::SC_SWITCH:
-////        case EOpType::CO_SWITCH:
-////        case EOpType::IM_SWITCH:
-////        case EOpType::VO_SWITCH:
-////        case EOpType::ME_SWITCH:
-////        case EOpType::LA_SWITCH:
-////        {
-////            IntExpressionPtr pValue = GetIntExpression( args.Switch.variable, program );
-
-////            for ( int i=0;i<MUTABLE_OP_MAX_SWITCH_OPTIONS; ++i )
-////            {
-////                if ( args.Switch.values[i] )
-////                {
-////                    BoolExpressionPtr pCond = new BoolExpression;
-////                    pCond->type = BoolExpression::BET_INT_EQUAL_CONST;
-////                    pCond->intExp = pValue;
-////                    pCond->intValue = args.Switch.conditions[i];
-
-////                    BoolExpression::TYPE result = pCond->Evaluate( m_facts );
-
-////                    if ( result==BoolExpression::BET_TRUE )
-////                    {
-////                        // This is the only possible branch.
-////                        bModified = true;
-////                        at = args.Switch.values[i];
-////                        at = Visit( at, program );
-////                        return at;
-////                        break;
-////                    }
-////                    else if ( result==BoolExpression::BET_FALSE )
-////                    {
-////                        // This option is impossible
-////                        args.Switch.values[i] = 0;
-////                        args.Switch.conditions[i] = 0;
-////                    }
-////                    else
-////                    {
-////                        // We don't know if it is true or false. Add its value as a fact and
-////                        // recurse.
-////                        m_facts.push_back( pCond );
-////                        args.Switch.values[i] = Visit( args.Switch.values[i], program );
-////                        m_facts.pop_back();
-////                    }
-////                }
-////            }
-
-////            if ( args.Switch.def )
-////            {
-////                // TODO: Add facts negating the options in the switch
-////                args.Switch.def = Visit( args.Switch.def, program );
-////            }
-
-////            // update op if necessary
-////            if ( FMemory::Memcmp( &args, &program.m_code[at].args, sizeof(OP::ARGS) ) )
-////            {
-////                bModified = true;
-////                OP op = program.m_code[at];
-////                op.args = args;
-////                at = program.AddOp( op );
-////            }
-////            break;
-////        }
-
-//        default:
-//            // Normal recursion
-//            at = Recurse( at, program );
-//            break;
-
-//        }
-
-//        return at;
-//    }
-
-
-    //---------------------------------------------------------------------------------------------
-    //---------------------------------------------------------------------------------------------
-    //---------------------------------------------------------------------------------------------
-//    IntExpression::IntExpression()
-//    {
-//        type = IET_UNKNOWN;
-//        parameter = 0;
-//        value = 0;
-//    }
-
-
-//    //---------------------------------------------------------------------------------------------
-//    OP::ADDRESS IntExpression::GenerateCode( FProgram& program, FProgram::OP_CACHE* cache )
-//    {
-//        OP::ADDRESS at = 0;
-
-//        switch ( type )
-//        {
-
-//        case IET_CONSTANT:
-//        {
-//            OP op;
-//            op.type = EOpType::NU_CONSTANT;
-//            op.args.IntConstant.value = value;
-//            at = program.AddOpUnique( op, cache );
-//            break;
-//        }
-
-//        case IET_PARAMETER:
-//        {
-//            OP op;
-//            op.type = EOpType::NU_PARAMETER;
-//            op.args.Parameter.variable = parameter;
-//            at = program.AddOpUnique( op, cache );
-//            break;
-//        }
-
-//        default:
-//            check( false );
-//            break;
-
-//        }
-
-//        return at;
-//    }
-
-
-//    //---------------------------------------------------------------------------------------------
-//    //---------------------------------------------------------------------------------------------
-//    //---------------------------------------------------------------------------------------------
-//    BoolExpression::BoolExpression()
-//    {
-//        type = BET_UNKNOWN;
-//        parameter = 0;
-//        intValue = 0;
-//    }
-
-
-//    //---------------------------------------------------------------------------------------------
-//    BoolExpression::TYPE BoolExpression::Evaluate( vector<BoolExpressionPtr>& facts )
-//    {
-//        TYPE result = BET_UNKNOWN;
-
-//        switch ( type )
-//        {
-//        case BET_UNKNOWN:
-//            break;
-
-//        case BET_TRUE:
-//            result = BET_TRUE;
-//            break;
-
-//        case BET_FALSE:
-//            result = BET_FALSE;
-//            break;
-
-//        case BET_PARAMETER:
-//            for ( size_t f=0; f<facts.size(); ++f )
-//            {
-//                switch ( EvaluateParameter( facts[f], parameter ) )
-//                {
-//                case BET_TRUE:
-//                    return BET_TRUE;
-//                    break;
-
-//                case BET_FALSE:
-//                    return BET_FALSE;
-//                    break;
-
-//                default:
-//                    break;
-//                }
-//            }
-//            break;
-
-//        case BET_NOT:
-//            for ( size_t f=0; f<facts.size(); ++f )
-//            {
-//                result = a->Evaluate( facts );
-//                switch ( result )
-//                {
-//                case BET_TRUE:
-//                    return BET_FALSE;
-//                    break;
-
-//                case BET_FALSE:
-//                    return BET_TRUE;
-//                    break;
-
-//                default:
-//                    break;
-//                }
-//            }
-//            break;
-
-//        case BET_AND:
-//        {
-//            TYPE resultA = BET_UNKNOWN;
-//            TYPE resultB = BET_UNKNOWN;
-//            for ( size_t f=0; f<facts.size(); ++f )
-//            {
-//                if ( resultA == BET_UNKNOWN )
-//                {
-//                    resultA = a->Evaluate( facts );
-//                }
-//                if ( resultB == BET_UNKNOWN )
-//                {
-//                    resultB = b->Evaluate( facts );
-//                }
-//            }
-
-//            if ( resultA==BET_TRUE && resultB==BET_TRUE )
-//            {
-//                result = BET_TRUE;
-//            }
-//            else if ( resultA==BET_FALSE || resultB==BET_FALSE )
-//            {
-//                result = BET_FALSE;
-//            }
-//            break;
-//        }
-
-//        case BET_OR:
-//        {
-//            TYPE resultA = BET_UNKNOWN;
-//            TYPE resultB = BET_UNKNOWN;
-//            for ( size_t f=0; f<facts.size(); ++f )
-//            {
-//                if ( resultA == BET_UNKNOWN )
-//                {
-//                    resultA = a->Evaluate( facts );
-//                }
-//                if ( resultB == BET_UNKNOWN )
-//                {
-//                    resultB = b->Evaluate( facts );
-//                }
-//            }
-
-//            if ( resultA==BET_TRUE || resultB==BET_TRUE )
-//            {
-//                result = BET_TRUE;
-//            }
-//            else if ( resultA==BET_FALSE && resultB==BET_FALSE )
-//            {
-//                result = BET_FALSE;
-//            }
-//            break;
-//        }
-
-//        case BET_INT_EQUAL_CONST:
-//            result = EvaluateIntEquality( facts, intExp, intValue );
-//            break;
-
-//        default:
-//            check( false );
-//            break;
-//        }
-
-//        return result;
-//    }
-
-
-//    //---------------------------------------------------------------------------------------------
-//    OP::ADDRESS BoolExpression::GenerateCode( FProgram& program, FProgram::OP_CACHE* cache )
-//    {
-//        if (m_generated)
-//        {
-//            return m_generated;
-//        }
-
-//        OP::ADDRESS at = 0;
-
-//        switch ( type )
-//        {
-
-//        case BET_TRUE:
-//        {
-//            OP op;
-//            op.type = EOpType::BO_CONSTANT;
-//            op.args.BoolConstant.value = true;
-//            at = program.AddOpUnique( op, cache );
-//            break;
-//        }
-
-//        case BET_FALSE:
-//        {
-//            OP op;
-//            op.type = EOpType::BO_CONSTANT;
-//            op.args.BoolConstant.value = false;
-//            at = program.AddOpUnique( op, cache );
-//            break;
-//        }
-
-//        case BET_NOT:
-//        {
-//            OP op;
-//            op.type = EOpType::BO_NOT;
-//            op.args.BoolNot.source = a->GenerateCode( program, cache );
-//            at = program.AddOpUnique( op, cache );
-//            break;
-//        }
-
-//        case BET_AND:
-//        {
-//            OP op;
-//            op.type = EOpType::BO_AND;
-//            op.args.BoolBinary.a = a->GenerateCode( program, cache );
-//            op.args.BoolBinary.b = b->GenerateCode( program, cache );
-//            at = program.AddOpUnique( op, cache );
-//            break;
-//        }
-
-//        case BET_OR:
-//        {
-//            OP op;
-//            op.type = EOpType::BO_OR;
-//            op.args.BoolBinary.a = a->GenerateCode( program, cache );
-//            op.args.BoolBinary.b = b->GenerateCode( program, cache );
-//            at = program.AddOpUnique( op, cache );
-//            break;
-//        }
-
-
-//        case BET_PARAMETER:
-//        {
-//            OP op;
-//            op.type = EOpType::BO_PARAMETER;
-//            op.args.Parameter.variable = parameter;
-//            at = program.AddOpUnique( op, cache );
-//            break;
-//        }
-
-
-//        case BET_INT_EQUAL_CONST:
-//        {
-//            OP op;
-//            op.type = EOpType::BO_EQUAL_INT_CONST;
-//            op.args.BoolEqualScalarConst.value = intExp->GenerateCode( program, cache );
-//            op.args.BoolEqualScalarConst.constant = (int16_t)intValue;
-//            at = program.AddOpUnique( op, cache );
-//            break;
-//        }
-
-
-//        default:
-//            check( false );
-//            break;
-
-//        }
-
-//        m_generated = at;
-
-//        return at;
-//    }
-
-
-//    //---------------------------------------------------------------------------------------------
-//    BoolExpression::TYPE BoolExpression::EvaluateParameter( BoolExpressionPtr fact, OP::ADDRESS parameter)
-//    {
-//        TYPE result = BET_UNKNOWN;
-
-//        switch (fact->type)
-//        {
-//        case BET_PARAMETER:
-//            if ( parameter==fact->parameter )
-//            {
-//                result = BET_TRUE;
-//            }
-//            break;
-
-//        case BET_NOT:
-//            result = EvaluateParameter( fact->a, parameter );
-//            if ( result == BET_TRUE )
-//            {
-//                result = BET_FALSE;
-//            }
-//            else if ( result == BET_FALSE )
-//            {
-//                result = BET_TRUE;
-//            }
-//            break;
-
-//        case BET_AND:
-//            result = EvaluateParameter( fact->a, parameter );
-//            if ( result == BET_UNKNOWN )
-//            {
-//                result = EvaluateParameter( fact->b, parameter );
-//            }
-//            break;
-
-//        case BET_OR:
-//        {
-//            TYPE resultA = EvaluateParameter( fact->a, parameter );
-//            TYPE resultB = EvaluateParameter( fact->b, parameter );
-//            if ( resultA==BET_TRUE && resultB==BET_TRUE )
-//            {
-//                result = BET_TRUE;
-//            }
-//            else if ( resultA==BET_FALSE && resultB==BET_FALSE )
-//            {
-//                result = BET_FALSE;
-//            }
-//            else
-//            {
-//                // TODO: Check true+false value
-//                result = BET_UNKNOWN;
-//            }
-//            break;
-//        }
-
-//        case BET_INT_EQUAL_CONST:
-//        case BET_UNKNOWN:
-//        case BET_TRUE:
-//        case BET_FALSE:
-//            break;
-
-//        default:
-//            check( false);
-//            break;
-
-//        }
-
-//        return result;
-//    }
-
-
-//    //---------------------------------------------------------------------------------------------
-//    BoolExpression::TYPE BoolExpression::EvaluateIntEquality
-//        ( const vector<BoolExpressionPtr>& facts, IntExpressionPtr intExp, int value )
-//    {
-//        TYPE result = BET_UNKNOWN;
-
-//        switch (intExp->type)
-//        {
-//        case IntExpression::IET_CONSTANT:
-//            if ( value==intExp->value )
-//            {
-//                result = BET_TRUE;
-//            }
-//            else
-//            {
-//                result = BET_FALSE;
-//            }
-//            break;
-
-//        case IntExpression::IET_PARAMETER:
-//            for ( size_t f=0; (result==BET_UNKNOWN) && f<facts.size(); ++f )
-//            {
-//                result = EvaluateIntEquality(facts[f].get(),intExp,value);
-//            }
-//            break;
-
-//        default:
-//            break;
-
-//        }
-
-//        return result;
-//    }
-
-
-//    //---------------------------------------------------------------------------------------------
-//    BoolExpression::TYPE BoolExpression::EvaluateIntEquality
-//        ( const BoolExpression* fact, IntExpressionPtr intExp, int value )
-//    {
-//        TYPE result = BET_UNKNOWN;
-
-//        if( fact->type == BET_INT_EQUAL_CONST )
-//        {
-//            if ( fact->intExp->type == IntExpression::IET_PARAMETER
-//                &&
-//                fact->intExp->parameter == intExp->parameter
-//                )
-//            {
-//                if ( fact->intValue == value )
-//                {
-//                    result = BET_TRUE;
-//                }
-//                else
-//                {
-//                    result = BET_FALSE;
-//                }
-//            }
-//        }
-
-//        else if( fact->type == BET_NOT )
-//        {
-//            if (fact->a->type==BET_INT_EQUAL_CONST)
-//            {
-//                TYPE notResult = EvaluateIntEquality( fact->a.get(), intExp, value );
-//                if (notResult==BET_TRUE)
-//                {
-//                    result = BET_FALSE;
-//                }
-//                else
-//                {
-//                    // We cannot know, because there are other values.
-//                }
-//            }
-//            else
-//            {
-//                MUTABLE_CPUPROFILER_SCOPE(Dummy);
-//            }
-//        }
-
-//        else
-//        {
-//            MUTABLE_CPUPROFILER_SCOPE(Dummy);
-//        }
-
-//        return result;
-//    }
-
-
-    //---------------------------------------------------------------------------------------------
-    //---------------------------------------------------------------------------------------------
-    //---------------------------------------------------------------------------------------------
-    bool LocalLogicOptimiserAST(ASTOpList& roots)
+	bool LocalLogicOptimiserAST(ASTOpList& roots)
     {
         MUTABLE_CPUPROFILER_SCOPE(LocalLogicOptimiserAST);
 
@@ -1129,57 +146,6 @@ namespace mu
                             }
                         }
                     }
-
-// todo, since the move to aggregated remove mask op
-//                    else if ( yes->GetOpType()==EOpType::ME_REMOVEMASK )
-//                    {
-//                        Ptr<ASTOpFixed> remove = dynamic_cast<ASTOpFixed*>(yes.get());
-
-//                        bool ended = false;
-//                        while (!ended)
-//                        {
-//                            ended = true;
-//                            auto base = remove->children[remove->op.args.MeshRemoveMask.source].child();
-//                            if ( base && base->GetOpType()==EOpType::ME_CONDITIONAL )
-//                            {
-//                                auto bottomConditional = dynamic_cast<ASTOpConditional*>(base.get());
-
-//                                // Are the two conditions exclusive?
-//                                bool conditionaAreExclusive = false;
-
-//                                ASTOpList facts;
-//                                facts.push_back(topConditional->condition.child());
-
-//                                // Check if the child condition has a value with the current facts
-//                                auto pChildCond = bottomConditional->condition.child();
-//                                ASTOp::FBoolEvalResult result = pChildCond->EvaluateBool( facts );
-//                                if ( result==ASTOp::BET_FALSE )
-//                                {
-//                                    conditionaAreExclusive = true;
-//                                }
-
-//                                if (conditionaAreExclusive)
-//                                {
-//                                    if (remove->GetParentCount()==1)
-//                                    {
-//                                        // Directly modify the instruction to skip the impossible child option
-//                                        remove->SetChild(remove->op.args.MeshRemoveMask.source, bottomConditional->no.child() );
-//                                    }
-//                                    else
-//                                    {
-//                                        // Other parents may not impose the same condition that allows
-//                                        // the optimisation.
-//                                        auto newRemove = mu::Clone<ASTOpFixed>(remove);
-//                                        newRemove->SetChild(newRemove->op.args.MeshRemoveMask.source, bottomConditional->no.child());
-//                                        topConditional->yes = newRemove;
-//                                    }
-
-//                                    modified = true;
-//                                    ended = false;
-//                                }
-//                            }
-//                        }
-//                    }
                 }
 
 
@@ -1215,8 +181,8 @@ namespace mu
                         Ptr<ASTOpSwitch> switchOp = new ASTOpSwitch();
                         switchOp->type = GetSwitchForType(GetOpDataType(topConditional->GetOpType()));
 
-						ASTOpFixed* firstCompare = static_cast<ASTOpFixed*>(topConditional->condition.child().get());
-                        switchOp->variable = firstCompare->children[firstCompare->op.args.BoolEqualScalarConst.value].child();
+						ASTOpBoolEqualIntConst* firstCompare = static_cast<ASTOpBoolEqualIntConst*>(topConditional->condition.child().get());
+                        switchOp->variable = firstCompare->Value.child();
 
                         auto current = n;
                         while(current)
@@ -1237,17 +203,15 @@ namespace mu
                                  &&
                                  conditional->condition->GetOpType()==EOpType::BO_EQUAL_INT_CONST )
                             {
-								ASTOpFixed* compare = static_cast<ASTOpFixed*>(conditional->condition.child().get());
+								ASTOpBoolEqualIntConst* compare = static_cast<ASTOpBoolEqualIntConst*>(conditional->condition.child().get());
                                 check(compare);
                                 if (compare)
                                 {
-                                    auto compareValue = compare->children[compare->op.args.BoolEqualScalarConst.value].child();
+                                    auto compareValue = compare->Value.child();
                                     check(compare);
                                     if ( compareValue == switchOp->variable.child() )
                                     {
-                                        switchOp->cases.Emplace( compare->op.args.BoolEqualScalarConst.constant,
-                                                                      switchOp,
-                                                                      conditional->yes.child() );
+                                        switchOp->cases.Emplace( compare->Constant, switchOp, conditional->yes.child() );
 
                                         current = conditional->no.child();
                                         valid = true;
@@ -1316,50 +280,6 @@ namespace mu
                             }
                         }
                     }
-
-                    // todo, since the move to aggregated remove mask op
-//                    if (caseType==EOpType::ME_REMOVEMASK)
-//                    {
-//                        // Is the source of all remove operations the default option in the switch?
-//                        auto source = topSwitch->def.child();
-//                        for (const auto& c:topSwitch->cases)
-//                        {
-//                            if ( c.branch )
-//                            {
-//                                auto typedRemove = dynamic_cast<const ASTOpFixed*>(c.branch.child().get());
-//                                check( typedRemove );
-//                                if (typedRemove->children[typedRemove->op.args.MeshRemoveMask.source].child()!=source)
-//                                {
-//                                    source = nullptr;
-//                                    break;
-//                                }
-//                            }
-//                        }
-
-//                        if (source)
-//                        {
-//                            auto newSwitch = mu::Clone<ASTOpSwitch>(topSwitch);
-//                            for (auto& c: newSwitch->cases)
-//                            {
-//                                if ( c.branch )
-//                                {
-//                                    auto typedRemove = dynamic_cast<const ASTOpFixed*>(c.branch.child().get());
-//                                    c.branch = typedRemove->children[typedRemove->op.args.MeshRemoveMask.mask].child();
-//                                }
-//                            }
-//                            newSwitch->def = nullptr;
-
-//                            Ptr<ASTOpFixed> newRemove = new ASTOpFixed();
-//                            newRemove->op.type = EOpType::ME_REMOVEMASK;
-//                            newRemove->SetChild( newRemove->op.args.MeshRemoveMask.mask, newSwitch );
-//                            newRemove->SetChild( newRemove->op.args.MeshRemoveMask.source, source);
-
-//                            ASTOp::Replace(n,newRemove);
-//                            n = newRemove;
-//                            modified = true;
-//                        }
-//                    }
-
                 }
 
                 return recurse;
diff --git a/Engine/Plugins/Mutable/Source/MutableTools/Private/MuT/CodeOptimiser_Parameters.cpp b/Engine/Plugins/Mutable/Source/MutableTools/Private/MuT/CodeOptimiser_Parameters.cpp
--- a/Engine/Plugins/Mutable/Source/MutableTools/Private/MuT/CodeOptimiser_Parameters.cpp
+++ b/Engine/Plugins/Mutable/Source/MutableTools/Private/MuT/CodeOptimiser_Parameters.cpp
@@ -16,6 +16,7 @@
 #include "MuT/ASTOpLayoutFromMesh.h"
 #include "MuT/ASTOpSwitch.h"
 #include "MuT/ASTOpParameter.h"
+#include "MuT/ASTOpConstantColor.h"
 #include "MuT/CodeOptimiser.h"
 #include "MuT/Compiler.h"
 #include "MuT/CompilerPrivate.h"
@@ -229,12 +230,8 @@ namespace mu
     {
         if ( !Mask)
         {
-            Ptr<ASTOpFixed> whiteOp = new ASTOpFixed;
-            whiteOp->op.type = EOpType::CO_CONSTANT;
-            whiteOp->op.args.ColourConstant.value[0] = 1;
-            whiteOp->op.args.ColourConstant.value[1] = 1;
-			whiteOp->op.args.ColourConstant.value[2] = 1;
-			whiteOp->op.args.ColourConstant.value[3] = 1;
+            Ptr<ASTOpConstantColor> whiteOp = new ASTOpConstantColor;
+            whiteOp->Value = FVector4f(1,1,1,1);
 
             Ptr<ASTOpImagePlainColor> wplainOp = new ASTOpImagePlainColor;
             wplainOp->Color = whiteOp;
@@ -451,12 +448,8 @@ namespace mu
 
 							const ASTOpImageLayerColor* typedYes = static_cast<const ASTOpImageLayerColor*>(TypedOp->yes.child().get());
 
-                            Ptr<ASTOpFixed> blackOp = new ASTOpFixed;
-                            blackOp->op.type = EOpType::CO_CONSTANT;
-                            blackOp->op.args.ColourConstant.value[0] = 0;
-                            blackOp->op.args.ColourConstant.value[1] = 0;
-							blackOp->op.args.ColourConstant.value[2] = 0;
-							blackOp->op.args.ColourConstant.value[3] = 1;
+                            Ptr<ASTOpConstantColor> blackOp = new ASTOpConstantColor;
+                            blackOp->Value = FVector4f(0,0,0,1);
 
                             Ptr<ASTOpImagePlainColor> plainOp = new ASTOpImagePlainColor;
                             plainOp->Color = blackOp;
@@ -556,12 +549,8 @@ namespace mu
 
                             const ASTOpImageLayerColor* typedNo = static_cast<const ASTOpImageLayerColor*>(TypedOp->no.child().get());
 
-                            Ptr<ASTOpFixed> blackOp = new ASTOpFixed;
-                            blackOp->op.type = EOpType::CO_CONSTANT;
-                            blackOp->op.args.ColourConstant.value[0] = 0;
-                            blackOp->op.args.ColourConstant.value[1] = 0;
-							blackOp->op.args.ColourConstant.value[2] = 0;
-							blackOp->op.args.ColourConstant.value[3] = 1;
+                            Ptr<ASTOpConstantColor> blackOp = new ASTOpConstantColor;
+                            blackOp->Value = FVector4f(0,0,0,1);
 
                             Ptr<ASTOpImagePlainColor> plainOp = new ASTOpImagePlainColor;
                             plainOp->Color = blackOp;
@@ -1018,12 +1007,8 @@ namespace mu
 						Ptr<ASTOp> newMaskBlock = EnsureValidMask(blockMask, blockImage);
                         maskOp->BlockImage = newMaskBlock;
 
-                        Ptr<ASTOpFixed> blackOp = new ASTOpFixed;
-                        blackOp->op.type = EOpType::CO_CONSTANT;
-                        blackOp->op.args.ColourConstant.value[0] = 0;
-                        blackOp->op.args.ColourConstant.value[1] = 0;
-                        blackOp->op.args.ColourConstant.value[2] = 0;
-						blackOp->op.args.ColourConstant.value[3] = 1;
+                        Ptr<ASTOpConstantColor> blackOp = new ASTOpConstantColor;
+                        blackOp->Value = FVector4f(0,0,0,1);
 
                         Ptr<ASTOpImagePlainColor> plainOp = new ASTOpImagePlainColor;
                         plainOp->Color = blackOp;
@@ -1073,12 +1058,8 @@ namespace mu
                         maskOp->BlockImage = newMaskBlock;
 
 
-                        Ptr<ASTOpFixed> blackOp = new ASTOpFixed;
-                        blackOp->op.type = EOpType::CO_CONSTANT;
-                        blackOp->op.args.ColourConstant.value[0] = 0;
-                        blackOp->op.args.ColourConstant.value[1] = 0;
-                        blackOp->op.args.ColourConstant.value[2] = 0;
-						blackOp->op.args.ColourConstant.value[3] = 1;
+                        Ptr<ASTOpConstantColor> blackOp = new ASTOpConstantColor;
+                        blackOp->Value = FVector4f(0,0,0,1);
 
                         Ptr<ASTOpImagePlainColor> plainOp = new ASTOpImagePlainColor;
                         plainOp->Color = blackOp;
@@ -1099,12 +1080,8 @@ namespace mu
                     {
                         blendedOp->BlockImage = blockBlended;
 
-                        Ptr<ASTOpFixed> blackOp = new ASTOpFixed;
-                        blackOp->op.type = EOpType::CO_CONSTANT;
-                        blackOp->op.args.ColourConstant.value[0] = 0;
-                        blackOp->op.args.ColourConstant.value[1] = 0;
-                        blackOp->op.args.ColourConstant.value[2] = 0;
-						blackOp->op.args.ColourConstant.value[3] = 1;
+                        Ptr<ASTOpConstantColor> blackOp = new ASTOpConstantColor;
+						blackOp->Value = FVector4f(0, 0, 0, 1);
 
                         Ptr<ASTOpImagePlainColor> plainOp = new ASTOpImagePlainColor;
                         plainOp->Color = blackOp;
@@ -1182,12 +1159,8 @@ namespace mu
 
 					Ptr<ASTOpImageCompose> maskOp = mu::Clone<ASTOpImageCompose>(At);
                     {
-                        Ptr<ASTOpFixed> blackOp = new ASTOpFixed;
-                        blackOp->op.type = EOpType::CO_CONSTANT;
-                        blackOp->op.args.ColourConstant.value[0] = 0;
-                        blackOp->op.args.ColourConstant.value[1] = 0;
-                        blackOp->op.args.ColourConstant.value[2] = 0;
-						blackOp->op.args.ColourConstant.value[3] = 1;
+                        Ptr<ASTOpConstantColor> blackOp = new ASTOpConstantColor;
+                        blackOp->Value = FVector4f(0,0,0,1);
 
                         Ptr<ASTOpImagePlainColor> plainOp = new ASTOpImagePlainColor;
 						plainOp->Color = blackOp;
@@ -1227,12 +1200,8 @@ namespace mu
 
 					Ptr<ASTOpImageCompose> maskOp = mu::Clone<ASTOpImageCompose>(At);
                     {
-                        Ptr<ASTOpFixed> blackOp = new ASTOpFixed;
-                        blackOp->op.type = EOpType::CO_CONSTANT;
-                        blackOp->op.args.ColourConstant.value[0] = 0;
-                        blackOp->op.args.ColourConstant.value[1] = 0;
-                        blackOp->op.args.ColourConstant.value[2] = 0;
-						blackOp->op.args.ColourConstant.value[3] = 1;
+                        Ptr<ASTOpConstantColor> blackOp = new ASTOpConstantColor;
+						blackOp->Value = FVector4f(0, 0, 0, 1);
 
                         Ptr<ASTOpImagePlainColor> plainOp = new ASTOpImagePlainColor;
                         plainOp->Color = blackOp;
diff --git a/Engine/Plugins/Mutable/Source/MutableTools/Private/MuT/CodeOptimiser_Semantic.cpp b/Engine/Plugins/Mutable/Source/MutableTools/Private/MuT/CodeOptimiser_Semantic.cpp
--- a/Engine/Plugins/Mutable/Source/MutableTools/Private/MuT/CodeOptimiser_Semantic.cpp
+++ b/Engine/Plugins/Mutable/Source/MutableTools/Private/MuT/CodeOptimiser_Semantic.cpp
@@ -73,301 +73,6 @@ namespace mu
     }
 
 
-    //---------------------------------------------------------------------------------------------
-    Ptr<ASTOp> ASTOpFixed::OptimiseSemantic( const FModelOptimizationOptions& options, int32 Pass) const
-    {
-        Ptr<ASTOp> at;
-
-        EOpType type = GetOpType();
-        switch ( type )
-        {
-
-        case EOpType::BO_NOT:
-        {
-			Ptr<ASTOp> sourceAt = children[op.args.BoolNot.source].child();
-            if  ( sourceAt->GetOpType() == EOpType::BO_CONSTANT )
-            {
-                Ptr<ASTOpConstantBool> newAt = new ASTOpConstantBool();
-                newAt->value = !newAt->value;
-                at=newAt;
-            }
-            break;
-        }
-
-        case EOpType::BO_AND:
-        {
-            bool changed = false;
-			Ptr<ASTOp> aAt = children[op.args.BoolBinary.a].child();
-			Ptr<ASTOp> bAt = children[op.args.BoolBinary.b].child();
-            if ( !aAt )
-            {
-                at = bAt;
-                changed = true;
-            }
-            else if ( !bAt )
-            {
-                at = aAt;
-                changed = true;
-            }
-            else if ( aAt->GetOpType() == EOpType::BO_CONSTANT )
-            {
-                if (static_cast<const ASTOpConstantBool*>(aAt.get())->value)
-                {
-                    at = bAt;
-                    changed = true;
-                }
-                else
-                {
-                    at = aAt;
-                    changed = true;
-                }
-            }
-            else if ( bAt->GetOpType() == EOpType::BO_CONSTANT )
-            {
-                if (static_cast<const ASTOpConstantBool*>(bAt.get())->value)
-                {
-                    at = aAt;
-                    changed = true;
-                }
-                else
-                {
-                    at = bAt;
-                    changed = true;
-                }
-            }
-
-            // Common cases of repeated branch in children
-            else if ( aAt->GetOpType() == type )
-            {
-				const ASTOpFixed* typedA = static_cast<const ASTOpFixed*>(aAt.get());
-                if ( typedA->children[typedA->op.args.BoolBinary.a].child()==bAt
-                     ||
-                     typedA->children[typedA->op.args.BoolBinary.b].child()==bAt )
-                {
-                    at = aAt;
-                    changed = true;
-                }
-            }
-            else if ( bAt->GetOpType() == type )
-            {
-				const ASTOpFixed* typedB = static_cast<const ASTOpFixed*>(bAt.get());
-                if ( typedB->children[typedB->op.args.BoolBinary.b].child()==aAt
-                     ||
-                     typedB->children[typedB->op.args.BoolBinary.b].child()==bAt )
-                {
-                    at = bAt;
-                    changed = true;
-                }
-            }
-
-            else if (aAt == bAt || *aAt==*bAt)
-            {
-                at = aAt;
-                changed = true;
-            }
-
-            // if it became null, it means true (neutral AND argument)
-            if (changed && !at)
-            {
-                Ptr<ASTOpConstantBool> newAt = new ASTOpConstantBool();
-                newAt->value = true;
-                at = newAt;
-            }
-            break;
-        }
-
-        case EOpType::BO_OR:
-        {
-            bool changed = false;
-			Ptr<ASTOp> aAt = children[op.args.BoolBinary.a].child();
-			Ptr<ASTOp> bAt = children[op.args.BoolBinary.b].child();
-            if ( !aAt )
-            {
-                at = bAt;
-                changed = true;
-            }
-            else if ( !bAt )
-            {
-                at = aAt;
-                changed = true;
-            }
-            else if ( aAt->GetOpType() == EOpType::BO_CONSTANT )
-            {
-                if (static_cast<const ASTOpConstantBool*>(aAt.get())->value)
-                {
-                    at = aAt;
-                    changed = true;
-                }
-                else
-                {
-                    at = bAt;
-                    changed = true;
-                }
-            }
-            else if ( bAt->GetOpType() == EOpType::BO_CONSTANT )
-            {
-                if (static_cast<const ASTOpConstantBool*>(bAt.get())->value)
-                {
-                    at = bAt;
-                    changed = true;
-                }
-                else
-                {
-                    at = aAt;
-                    changed = true;
-                }
-            }
-
-            // Common cases of repeated branch in children
-            else if ( aAt->GetOpType() == type )
-            {
-				const ASTOpFixed* typedA = static_cast<const ASTOpFixed*>(aAt.get());
-                if ( typedA->children[typedA->op.args.BoolBinary.a].child()==bAt
-                     ||
-                     typedA->children[typedA->op.args.BoolBinary.b].child()==bAt )
-                {
-                    at = aAt;
-                    changed = true;
-                }
-            }
-            else if ( bAt->GetOpType() == type )
-            {
-				const ASTOpFixed* typedB = static_cast<const ASTOpFixed*>(bAt.get());
-                if ( typedB->children[typedB->op.args.BoolBinary.b].child()==aAt
-                     ||
-                     typedB->children[typedB->op.args.BoolBinary.b].child()==bAt )
-                {
-                    at = bAt;
-                    changed = true;
-                }
-            }
-
-            else if (aAt == bAt || *aAt==*bAt)
-            {
-                at = aAt;
-                changed = true;
-            }
-
-            // if it became null, it means false (neutral OR argument)
-            if (changed && !at)
-            {
-                Ptr<ASTOpConstantBool> newAt = new ASTOpConstantBool();
-                newAt->value = false;
-                at = newAt;
-            }
-
-            break;
-        }
-
-		case EOpType::CO_SWIZZLE:
-		{
-			// Optimizations that can be applied per-channel
-			{
-				Ptr<ASTOpFixed> NewAt;
-
-				for (int32 ChannelIndex = 0; ChannelIndex < MUTABLE_OP_MAX_SWIZZLE_CHANNELS; ++ChannelIndex)
-				{
-					Ptr<ASTOp> Candidate = children[op.args.ColourSwizzle.sources[ChannelIndex]].child();
-					if (!Candidate)
-					{
-						continue;
-					}
-
-					switch (Candidate->GetOpType())
-					{
-
-						// Swizzle + swizzle = swizzle
-					case EOpType::CO_SWIZZLE:
-					{
-						if (!NewAt)
-						{
-							NewAt = mu::Clone<ASTOpFixed>(this);
-						}
-						const ASTOpFixed* TypedCandidate = static_cast<const ASTOpFixed*>(Candidate.get());
-						int32 CandidateChannel = op.args.ColourSwizzle.sourceChannels[ChannelIndex];
-
-						NewAt->SetChild(NewAt->op.args.ColourSwizzle.sources[ChannelIndex], TypedCandidate->children[TypedCandidate->op.args.ColourSwizzle.sources[CandidateChannel]]);
-						NewAt->op.args.ColourSwizzle.sourceChannels[ChannelIndex] = TypedCandidate->op.args.ColourSwizzle.sourceChannels[CandidateChannel];
-						break;
-					}
-
-					default:
-						break;
-
-					}
-				}
-
-				at = NewAt;
-			}
-
-			// Not optimized yet?
-			if (!at)
-			{
-				// Optimizations that depend on all channels.
-				bool bAllChannelsSameType = true;
-				EOpType AllChannelsType = EOpType::NONE;
-				for (int32 ChannelIndex = 0; ChannelIndex < MUTABLE_OP_MAX_SWIZZLE_CHANNELS; ++ChannelIndex)
-				{
-					Ptr<ASTOp> Candidate = children[op.args.ColourSwizzle.sources[ChannelIndex]].child();
-					if (!Candidate)
-					{
-						continue;
-					}
-
-					if (ChannelIndex == 0)
-					{
-						AllChannelsType = Candidate->GetOpType();
-					}
-					else if (Candidate->GetOpType()!=AllChannelsType)
-					{
-						bAllChannelsSameType = false;
-						break;
-					}
-				}
-
-				if (bAllChannelsSameType)
-				{
-					switch (AllChannelsType)
-					{
-					case EOpType::CO_FROMSCALARS:
-					{
-						// We can remove the swizzle and replace it with a new FromScalars actually swizzling the inputs.
-
-						Ptr<ASTOpFixed> NewAt = new ASTOpFixed();
-						NewAt->op.type = EOpType::CO_FROMSCALARS;
-
-						for (int32 ChannelIndex = 0; ChannelIndex < MUTABLE_OP_MAX_SWIZZLE_CHANNELS; ++ChannelIndex)
-						{
-							const ASTOp* SelectedSourceGeneric = children[op.args.ColourSwizzle.sources[ChannelIndex]].child().get();
-							if (SelectedSourceGeneric)
-							{
-								const ASTOpFixed* SelectedSource = static_cast<const ASTOpFixed*>(SelectedSourceGeneric);
-								int32 SelectedChannel = op.args.ColourSwizzle.sourceChannels[ChannelIndex];
-								Ptr<ASTOp> SelectedFloatInput = SelectedSource->children[SelectedSource->op.args.ColourFromScalars.v[SelectedChannel]].child();
-								NewAt->SetChild(NewAt->op.args.ColourFromScalars.v[ChannelIndex], SelectedFloatInput);
-							}
-						}
-
-						at = NewAt;
-						break;
-					}
-
-					default: break;
-					}
-				}
-			}
-			
-			break;
-		}
-
-        default:
-            break;
-        }
-
-        return at;
-    }
-
-
     //---------------------------------------------------------------------------------------------
     //---------------------------------------------------------------------------------------------
     //---------------------------------------------------------------------------------------------
diff --git a/Engine/Plugins/Mutable/Source/MutableTools/Private/MuT/Compiler.cpp b/Engine/Plugins/Mutable/Source/MutableTools/Private/MuT/Compiler.cpp
--- a/Engine/Plugins/Mutable/Source/MutableTools/Private/MuT/Compiler.cpp
+++ b/Engine/Plugins/Mutable/Source/MutableTools/Private/MuT/Compiler.cpp
@@ -97,13 +97,11 @@ namespace mu
 		{ false,	{ 0, 0, 0, 0, 0, 0, 0, 0, 0, 0, 0, 0, 0, 0, 0, 0, 0, 0, 0, 0, 0, 0, 0, 0, 0, 0, 0, 0, 0 }		},	// IN_SWITCH
 		{ false,	{ 0, 0, 0, 0, 0, 0, 0, 0, 0, 0, 0, 0, 0, 0, 0, 0, 0, 0, 0, 0, 0, 0, 0, 0, 0, 0, 0, 0, 0 }		},	// ED_SWITCH
 
-		{ false,	{ 0, 0, 0, 0, 0, 0, 0, 0, 0, 0, 0, 0, 0, 0, 0, 0, 0, 0, 0, 0, 0, 0, 0, 0, 0, 0, 0, 0, 0 }		},	// BO_LESS
 		{ false,	{ 0, 0, 0, 0, 0, 0, 0, 0, 0, 0, 0, 0, 0, 0, 0, 0, 0, 0, 0, 0, 0, 0, 0, 0, 0, 0, 0, 0, 0 }		},	// BO_EQUAL_SC_CONST
 		{ false,	{ 0, 0, 0, 0, 0, 0, 0, 0, 0, 0, 0, 0, 0, 0, 0, 0, 0, 0, 0, 0, 0, 0, 0, 0, 0, 0, 0, 0, 0 }		},	// BO_AND
 		{ false,	{ 0, 0, 0, 0, 0, 0, 0, 0, 0, 0, 0, 0, 0, 0, 0, 0, 0, 0, 0, 0, 0, 0, 0, 0, 0, 0, 0, 0, 0 }		},	// BO_OR
 		{ false,	{ 0, 0, 0, 0, 0, 0, 0, 0, 0, 0, 0, 0, 0, 0, 0, 0, 0, 0, 0, 0, 0, 0, 0, 0, 0, 0, 0, 0, 0 }		},	// BO_NOT
 
-		{ true,		{ 0, 0, 0, 0, 0, 0, 0, 0, 0, 0, 0, 0, 0, 0, 0, 0, 0, 0, 0, 0, 0, 0, 0, 0, 0, 0, 0, 0, 0 }		},	// SC_MULTIPLYADD
 		{ true,		{ 0, 0, 0, 0, 0, 0, 0, 0, 0, 0, 0, 0, 0, 0, 0, 0, 0, 0, 0, 0, 0, 0, 0, 0, 0, 0, 0, 0, 0 }		},	// SC_ARITHMETIC
 		{ true,		{ 0, 0, 0, 0, 0, 0, 0, 0, 0, 0, 0, 0, 0, 0, 0, 0, 0, 0, 0, 0, 0, 0, 0, 0, 0, 0, 0, 0, 0 }		},	// SC_CURVE
 
diff --git a/Engine/Plugins/Mutable/Source/MutableTools/Private/MuT/DataPacker.cpp b/Engine/Plugins/Mutable/Source/MutableTools/Private/MuT/DataPacker.cpp
--- a/Engine/Plugins/Mutable/Source/MutableTools/Private/MuT/DataPacker.cpp
+++ b/Engine/Plugins/Mutable/Source/MutableTools/Private/MuT/DataPacker.cpp
@@ -16,6 +16,7 @@
 #include "MuR/Ptr.h"
 #include "MuR/RefCounted.h"
 #include "MuT/ASTOpConstantResource.h"
+#include "MuT/ASTOpConstantColor.h"
 #include "MuT/ASTOpImageCompose.h"
 #include "MuT/ASTOpImageLayer.h"
 #include "MuT/ASTOpImageLayerColor.h"
@@ -769,12 +770,8 @@ namespace mu
 					// It is more efficient to just have an instruction for it instead, to avoid the overhead
 					// of data loading. 
 					// Warning This eliminates the mips. \TODO: Add support for mips in plaincolour instruction?
-					Ptr<ASTOpFixed> NewColor = new ASTOpFixed;
-					NewColor->op.type = EOpType::CO_CONSTANT;
-					NewColor->op.args.ColourConstant.value[0] = PlainColor[0];
-					NewColor->op.args.ColourConstant.value[1] = PlainColor[1];
-					NewColor->op.args.ColourConstant.value[2] = PlainColor[2];
-					NewColor->op.args.ColourConstant.value[3] = PlainColor[3];
+					Ptr<ASTOpConstantColor> NewColor = new ASTOpConstantColor;
+					NewColor->Value = PlainColor;
 
 					Ptr<ASTOpImagePlainColor> NewPlain = new ASTOpImagePlainColor;
 					NewPlain->Color = NewColor;
diff --git a/Engine/Plugins/Mutable/Source/MutableTools/Private/MuT/ErrorLog.cpp b/Engine/Plugins/Mutable/Source/MutableTools/Private/MuT/ErrorLog.cpp
--- a/Engine/Plugins/Mutable/Source/MutableTools/Private/MuT/ErrorLog.cpp
+++ b/Engine/Plugins/Mutable/Source/MutableTools/Private/MuT/ErrorLog.cpp
@@ -243,13 +243,11 @@ namespace mu
 		TEXT("IN_SWITCH        "),
 		TEXT("ED_SWITCH        "),
 
-		TEXT("BO_LESS          "),
 		TEXT("BO_EQUAL_SC_CONST"),
 		TEXT("BO_AND           "),
 		TEXT("BO_OR            "),
 		TEXT("BO_NOT           "),
 		
-		TEXT("SC_MULTIPLYADD   "),
 		TEXT("SC_ARITHMETIC    "),
 		TEXT("SC_CURVE         "),
 		
