[mutable] Fix compilation after CL merges.
#rnx

diff --git a/Engine/Plugins/Mutable/Source/CustomizableObjectEditor/Private/MuCOE/SMutableCodeViewer.cpp b/Engine/Plugins/Mutable/Source/CustomizableObjectEditor/Private/MuCOE/SMutableCodeViewer.cpp
--- a/Engine/Plugins/Mutable/Source/CustomizableObjectEditor/Private/MuCOE/SMutableCodeViewer.cpp
+++ b/Engine/Plugins/Mutable/Source/CustomizableObjectEditor/Private/MuCOE/SMutableCodeViewer.cpp
@@ -1572,15 +1572,6 @@ void SMutableCodeViewer::GenerateElementRecursive(const int32& InStateIndex, mu:
 		
 		break;
 	}
-
-	case mu::EOpType::SC_MULTIPLYADD:
-	{
-		mu::OP::ScalarMultiplyAddArgs Args = InProgram.GetOpArgs<mu::OP::ScalarMultiplyAddArgs>(InParentAddress);
-		AddOpFunc(Args.factor0, TEXT("factor0 "));
-		AddOpFunc(Args.factor1, TEXT("factor1 "));
-		AddOpFunc(Args.add, TEXT("add "));
-		break;
-	}
 		
 	case mu::EOpType::SC_CURVE:
 	{
@@ -1593,9 +1584,9 @@ void SMutableCodeViewer::GenerateElementRecursive(const int32& InStateIndex, mu:
 	case mu::EOpType::CO_SAMPLEIMAGE:
 	{
 		mu::OP::ColourSampleImageArgs Args = InProgram.GetOpArgs<mu::OP::ColourSampleImageArgs>(InParentAddress);
-		AddOpFunc(Args.image, TEXT("image "));
-		AddOpFunc(Args.x, TEXT("x "));
-		AddOpFunc(Args.y, TEXT("y "));
+		AddOpFunc(Args.Image, TEXT("image "));
+		AddOpFunc(Args.X, TEXT("x "));
+		AddOpFunc(Args.Y, TEXT("y "));
 		break;
 	}
 		
