[Mutable] New Customizable Object Compile API.

Unified functionalities of ConditionalAutoCompile and CompileCustomizableObjectSynchronously into a single new Compile function.

- Soft deprecated ConditionalAutoCompile.
- Soft deprecated CompileCustomizableObjectSynchronously.
- Exposed "Compile Only Selected".
- Customizable Object editor now displays "Compiling..." when a compilation is in the queue.
- Fixed Customizable Objects not loading in Standalone.
- Moved Compilation On Screen Warnings/Errors to the Customizable Object module. Now they also appear in Standalone.

diff --git a/Engine/Plugins/Mutable/Source/CustomizableObject/Internal/MuCO/CustomizableObjectCompilerTypes.h b/Engine/Plugins/Mutable/Source/CustomizableObject/Internal/MuCO/CustomizableObjectCompilerTypes.h
--- a/Engine/Plugins/Mutable/Source/CustomizableObject/Internal/MuCO/CustomizableObjectCompilerTypes.h
+++ b/Engine/Plugins/Mutable/Source/CustomizableObject/Internal/MuCO/CustomizableObjectCompilerTypes.h
@@ -79,9 +79,6 @@ struct FCompilationOptions
 
 	// Used to enable AnimBp override physics mainipualtion.  
 	bool bAnimBpPhysicsManipulationEnabled = false;
-
-	// Used to reduce the number of notifications when compiling objects
-	bool bSilentCompilation = true;
 	
 	/** Force a very big number on the mips to skip during compilation. Useful to debug special cooks of the data. */
 	bool bForceLargeLODBias = false;
@@ -135,49 +132,3 @@ enum class ECompilationResultPrivate : uint8
 	Warnings, // Only warnings.
 };
 
-
-#if WITH_EDITOR
-struct CUSTOMIZABLEOBJECT_API FCompilationRequest
-{
-	FCompilationRequest(UCustomizableObject& CustomizableObject, bool bAsync);
-
-	UCustomizableObject* GetCustomizableObject();
-
-	FCompilationOptions& GetCompileOptions();
-
-	bool IsAsyncCompilation() const;
-
-	void SetDerivedDataCachePolicy(UE::DerivedData::ECachePolicy InCachePolicy);
-	UE::DerivedData::ECachePolicy GetDerivedDataCachePolicy() const;
-
-	void BuildDerivedDataCacheKey();
-	UE::DerivedData::FCacheKey GetDerivedDataCacheKey() const;
-
-	void SetCompilationState(ECompilationStatePrivate InState, ECompilationResultPrivate InResult);
-
-	ECompilationStatePrivate GetCompilationState() const;
-	ECompilationResultPrivate GetCompilationResult() const;
-
-	TArray<FText>& GetWarnings();
-	TArray<FText>& GetErrors();
-	
-	bool operator==(const FCompilationRequest& Other) const;
-
-private:
-
-	TWeakObjectPtr<UCustomizableObject> CustomizableObject;
-
-	FCompilationOptions Options;
-
-	ECompilationStatePrivate State = ECompilationStatePrivate::None;
-	ECompilationResultPrivate Result = ECompilationResultPrivate::Unknown;
-
-	bool bAsync = true;
-
-	UE::DerivedData::ECachePolicy DDCPolicy = UE::DerivedData::ECachePolicy::None;
-	UE::DerivedData::FCacheKey DDCKey;
-	
-	TArray<FText> Warnings;
-	TArray<FText> Errors;
-};
-#endif
diff --git a/Engine/Plugins/Mutable/Source/CustomizableObject/Internal/MuCO/CustomizableObjectInstancePrivate.h b/Engine/Plugins/Mutable/Source/CustomizableObject/Internal/MuCO/CustomizableObjectInstancePrivate.h
--- a/Engine/Plugins/Mutable/Source/CustomizableObject/Internal/MuCO/CustomizableObjectInstancePrivate.h
+++ b/Engine/Plugins/Mutable/Source/CustomizableObject/Internal/MuCO/CustomizableObjectInstancePrivate.h
@@ -345,7 +345,10 @@ public:
 	FCustomizableObjectInstanceDescriptor& GetDescriptor() const;
 
 	void SetDescriptor(const FCustomizableObjectInstanceDescriptor& InDescriptor);
-	
+
+	/** Return true if the instance is not locked and if it's compiled. */
+	bool CanUpdateInstance() const;
+
 	UCustomizableObjectInstance* GetPublic() const;
 
 	// If any components are using this instance, they will store the min of their distances to the player here every frame for LOD purposes
diff --git a/Engine/Plugins/Mutable/Source/CustomizableObject/Internal/MuCO/CustomizableObjectPrivate.h b/Engine/Plugins/Mutable/Source/CustomizableObject/Internal/MuCO/CustomizableObjectPrivate.h
--- a/Engine/Plugins/Mutable/Source/CustomizableObject/Internal/MuCO/CustomizableObjectPrivate.h
+++ b/Engine/Plugins/Mutable/Source/CustomizableObject/Internal/MuCO/CustomizableObjectPrivate.h
@@ -1425,6 +1425,51 @@ public:
 };
 
 
+#if WITH_EDITOR
+struct CUSTOMIZABLEOBJECT_API FCompilationRequest
+{
+	FCompilationRequest(UCustomizableObject& CustomizableObject);
+
+	UCustomizableObject* GetCustomizableObject();
+	
+	void SetDerivedDataCachePolicy(UE::DerivedData::ECachePolicy InCachePolicy);
+	UE::DerivedData::ECachePolicy GetDerivedDataCachePolicy() const;
+
+	void BuildDerivedDataCacheKey();
+	UE::DerivedData::FCacheKey GetDerivedDataCacheKey() const;
+
+	void SetCompilationState(ECompilationStatePrivate InState, ECompilationResultPrivate InResult);
+
+	ECompilationStatePrivate GetCompilationState() const;
+	ECompilationResultPrivate GetCompilationResult() const;
+	
+	bool operator==(const FCompilationRequest& Other) const;
+
+private:
+	TWeakObjectPtr<UCustomizableObject> CustomizableObject;
+
+	ECompilationStatePrivate State = ECompilationStatePrivate::None;
+	ECompilationResultPrivate Result = ECompilationResultPrivate::Unknown;
+	
+	UE::DerivedData::ECachePolicy DDCPolicy = UE::DerivedData::ECachePolicy::None;
+	UE::DerivedData::FCacheKey DDCKey;
+
+public:
+	FCompilationOptions Options;
+
+	TArray<FText> Warnings;
+	TArray<FText> Errors;
+
+	FCompileDelegate Callback;
+
+	bool bAsync = true;
+	bool bSkipIfCompiled = false;
+	bool bSkipIfOutOfDate = false;
+	bool bSilentCompilation = true;
+};
+#endif
+
+
 UCLASS()
 class CUSTOMIZABLEOBJECT_API UCustomizableObjectPrivate : public UObject
 {
@@ -1643,7 +1688,6 @@ public:
 	/** Cook requests. */
 	TArray<TSharedRef<FCompilationRequest>> CompileRequests;
 	
-	ECompilationStatePrivate CompilationState = ECompilationStatePrivate::None;
 	ECompilationResultPrivate CompilationResult = ECompilationResultPrivate::Unknown;
 	
 	FPostCompileDelegate PostCompileDelegate;
diff --git a/Engine/Plugins/Mutable/Source/CustomizableObject/Internal/MuCO/CustomizableObjectSystemPrivate.h b/Engine/Plugins/Mutable/Source/CustomizableObject/Internal/MuCO/CustomizableObjectSystemPrivate.h
--- a/Engine/Plugins/Mutable/Source/CustomizableObject/Internal/MuCO/CustomizableObjectSystemPrivate.h
+++ b/Engine/Plugins/Mutable/Source/CustomizableObject/Internal/MuCO/CustomizableObjectSystemPrivate.h
@@ -791,7 +791,7 @@ public:
 
 	// Handle to the registered TickDelegate.
 	FTSTicker::FDelegateHandle TickDelegateHandle;
-	FTickerDelegate TickDelegate;
+	FTSTicker::FDelegateHandle TickWarningsDelegateHandle;
 
 	/** Change the current status of Mutable. Enabling/Disabling core features.	
 	 * Disabling Mutable will turn off compilation, generation, and streaming and will remove the system ticker. */
diff --git a/Engine/Plugins/Mutable/Source/CustomizableObject/Internal/MuCO/ICustomizableObjectEditorModule.h b/Engine/Plugins/Mutable/Source/CustomizableObject/Internal/MuCO/ICustomizableObjectEditorModule.h
--- a/Engine/Plugins/Mutable/Source/CustomizableObject/Internal/MuCO/ICustomizableObjectEditorModule.h
+++ b/Engine/Plugins/Mutable/Source/CustomizableObject/Internal/MuCO/ICustomizableObjectEditorModule.h
@@ -3,7 +3,9 @@
 #pragma once
 
 #include "Modules/ModuleManager.h"
+
 #include "Templates/SharedPointer.h"
+#include "MuCO/CustomizableObject.h"
 
 class FCustomizableObjectEditorLogger;
 class ICustomizableObjectEditor;
@@ -88,11 +90,13 @@ public:
 	 */
 	virtual void BakeCustomizableObjectInstance(UCustomizableObjectInstance* InTargetInstance, const FBakingConfiguration& InBakingConfig) = 0;
 
+	// TODO GMT Remove from this interface (public header). Should only be a FCustomizableObjectEditorModule function. Also move FCompilationRequest to COE module.
 	/** Request for a given customizable object to be compiled. Async compile requests will be queued and processed sequentially.
 	 * @param InCompilationRequest - Request to compile an object.
 	 * @param bForceRequest - Queue request even if already in the pending list. */
 	virtual void CompileCustomizableObject(const TSharedRef<FCompilationRequest>& InCompilationRequest, bool bForceRequest = false) = 0;
-	virtual void CompileCustomizableObjects(const TArray<TSharedRef<FCompilationRequest>>& InCompilationRequests, bool bForceRequests = false) = 0;
+
+	virtual void CompileCustomizableObject(UCustomizableObject& Object, const FCompileParams& Params) = 0;
 
 	virtual int32 Tick(bool bBlocking) = 0;
 
@@ -111,5 +115,7 @@ public:
 	virtual void BackwardsCompatibleFixup(UEdGraph& Graph, int32 CustomizableObjectCustomVersion) = 0;
 	
 	virtual void PostBackwardsCompatibleFixup(UEdGraph& Graph) = 0;
+
+	virtual bool IsCompiling(const UCustomizableObject& Object) const = 0;
 };
  
diff --git a/Engine/Plugins/Mutable/Source/CustomizableObject/Private/MuCO/CustomizableObject.cpp b/Engine/Plugins/Mutable/Source/CustomizableObject/Private/MuCO/CustomizableObject.cpp
--- a/Engine/Plugins/Mutable/Source/CustomizableObject/Private/MuCO/CustomizableObject.cpp
+++ b/Engine/Plugins/Mutable/Source/CustomizableObject/Private/MuCO/CustomizableObject.cpp
@@ -468,7 +468,7 @@ void UCustomizableObject::BeginCacheForCookedPlatformData(const ITargetPlatform*
 	}
 
 	const TSharedRef<FCompilationRequest>* CompileRequest = GetPrivate()->CompileRequests.FindByPredicate(
-		[&TargetPlatform](const TSharedPtr<FCompilationRequest>& Request) { return Request->GetCompileOptions().TargetPlatform == TargetPlatform; });
+		[&TargetPlatform](const TSharedPtr<FCompilationRequest>& Request) { return Request->Options.TargetPlatform == TargetPlatform; });
 	
 	if (CompileRequest)
 	{
@@ -488,7 +488,7 @@ bool UCustomizableObject::IsCachedCookedPlatformDataLoaded(const ITargetPlatform
 	}
 
 	const TSharedRef<FCompilationRequest>* CompileRequest = GetPrivate()->CompileRequests.FindByPredicate(
-		[&TargetPlatform](const TSharedRef<FCompilationRequest>& Request) { return Request->GetCompileOptions().TargetPlatform == TargetPlatform; });
+		[&TargetPlatform](const TSharedRef<FCompilationRequest>& Request) { return Request->Options.TargetPlatform == TargetPlatform; });
 	
 	if (CompileRequest)
 	{
@@ -727,13 +727,13 @@ void UCustomizableObjectPrivate::CompileForTargetPlatform(UCustomizableObject& C
 
 	const bool bAsync = CVarMutableAsyncCook.GetValueOnAnyThread();
 
-	TSharedRef<FCompilationRequest> CompileRequest = MakeShared<FCompilationRequest>(CustomizableObject, bAsync);
-	FCompilationOptions& Options = CompileRequest->GetCompileOptions();
-	Options.OptimizationLevel = UE_MUTABLE_MAX_OPTIMIZATION;	// Force max optimization when packaging.
-	Options.TextureCompression = ECustomizableObjectTextureCompression::HighQuality;
-	Options.bIsCooking = true;
-	Options.bUseBulkData = CVarMutableUseBulkData.GetValueOnAnyThread();
-	Options.TargetPlatform = &TargetPlatform;
+	TSharedRef<FCompilationRequest> CompileRequest = MakeShared<FCompilationRequest>(CustomizableObject);
+	CompileRequest->bAsync = bAsync;
+	CompileRequest->Options.OptimizationLevel = UE_MUTABLE_MAX_OPTIMIZATION; // Force max optimization when packaging.
+	CompileRequest->Options.TextureCompression = ECustomizableObjectTextureCompression::HighQuality;
+	CompileRequest->Options.bIsCooking = true;
+	CompileRequest->Options.bUseBulkData = CVarMutableUseBulkData.GetValueOnAnyThread();
+	CompileRequest->Options.TargetPlatform = &TargetPlatform;
 
 	const int32 DDCUsage = CVarMutableDerivedDataCacheUsage.GetValueOnAnyThread();
 	UE::DerivedData::ECachePolicy DefaultCachePolicy = UE::DerivedData::ECachePolicy::None;
@@ -791,7 +791,6 @@ bool UCustomizableObject::ConditionalAutoCompile()
 	// Don't compile if Mutable or AutoCompile is disabled.
 	if (!System->IsActive() || !System->IsAutoCompileEnabled())
 	{
-		System->AddUncompiledCOWarning(*this);
 		return false;
 	}
 
@@ -804,8 +803,9 @@ bool UCustomizableObject::ConditionalAutoCompile()
 		}
 		else
 		{
-			TSharedRef<FCompilationRequest> CompileRequest = MakeShared<FCompilationRequest>(*this, true);
-			CompileRequest->GetCompileOptions().bSilentCompilation = true;
+			TSharedRef<FCompilationRequest> CompileRequest = MakeShared<FCompilationRequest>(*this);
+			CompileRequest->bSilentCompilation = true;
+			
 			EditorModule->CompileCustomizableObject(CompileRequest);
 		}
 	}
@@ -935,15 +935,21 @@ UCustomizableObjectPrivate* UCustomizableObject::GetPrivate()
 bool UCustomizableObject::IsCompiled() const
 {
 #if WITH_EDITOR
-	const bool bIsCompiled = Private->GetModel() != nullptr && Private->GetModel()->IsValid();
+	const bool bIsCompiled = GetPrivate()->GetModel() != nullptr && GetPrivate()->GetModel()->IsValid();
 #else
-	const bool bIsCompiled = Private->GetModel() != nullptr;
+	const bool bIsCompiled = GetPrivate()->GetModel() != nullptr;
 #endif
 
 	return bIsCompiled;
 }
 
 
+bool UCustomizableObject::IsLoading() const
+{
+	return GetPrivate()->Status.Get() == FCustomizableObjectStatusTypes::EState::Loading;
+}
+
+
 void UCustomizableObjectPrivate::AddUncompiledCOWarning(const FString& AdditionalLoggingInfo)
 {
 	// Send a warning (on-screen notification, log error, and in-editor notification)
@@ -1909,6 +1915,23 @@ bool UCustomizableObject::IsParameterMultidimensional(const FString& InParameter
 }
 
 
+void UCustomizableObject::Compile(const FCompileParams& Params)
+{
+	if (ICustomizableObjectEditorModule* EditorModule = ICustomizableObjectEditorModule::Get())
+	{
+		EditorModule->CompileCustomizableObject(*this, Params);
+	}
+	else
+	{
+		FCompileCallbackParams CallbackParams;
+		CallbackParams.bRequestFailed = true;
+		CallbackParams.bCompiled = IsCompiled();
+
+		Params.Callback.ExecuteIfBound(CallbackParams);
+	}
+}
+
+
 bool UCustomizableObject::IsParameterMultidimensional(const int32& InParamIndex) const
 {
 	return GetPrivate()->IsParameterMultidimensional(InParamIndex);
@@ -3698,11 +3721,10 @@ FArchive& operator<<(FArchive& Ar, FMutableRefSkeletalMeshData& Data)
 #endif 
 
 #if WITH_EDITOR
-FCompilationRequest::FCompilationRequest(UCustomizableObject& InCustomizableObject, bool bAsyncCompile)
+FCompilationRequest::FCompilationRequest(UCustomizableObject& InCustomizableObject)
 {
 	CustomizableObject = &InCustomizableObject;
 	Options = InCustomizableObject.GetPrivate()->GetCompileOptions();
-	bAsync = bAsyncCompile;
 	DDCPolicy = UE::DerivedData::ECachePolicy::None;
 }
 
@@ -3713,18 +3735,6 @@ UCustomizableObject* FCompilationRequest::GetCustomizableObject()
 }
 
 
-FCompilationOptions& FCompilationRequest::GetCompileOptions()
-{
-	return Options;
-}
-
-
-bool FCompilationRequest::IsAsyncCompilation() const
-{
-	return bAsync;
-}
-
-
 void FCompilationRequest::SetDerivedDataCachePolicy(UE::DerivedData::ECachePolicy InCachePolicy)
 {
 	DDCPolicy = InCachePolicy;
@@ -3772,18 +3782,6 @@ ECompilationResultPrivate FCompilationRequest::GetCompilationResult() const
 }
 
 
-TArray<FText>& FCompilationRequest::GetWarnings()
-{
-	return Warnings;
-}
-
-
-TArray<FText>& FCompilationRequest::GetErrors()
-{
-	return Errors;
-}
-
-
 bool FCompilationRequest::operator==(const FCompilationRequest& Other) const
 {
 	return CustomizableObject == Other.CustomizableObject && Options.TargetPlatform == Other.Options.TargetPlatform;
diff --git a/Engine/Plugins/Mutable/Source/CustomizableObject/Private/MuCO/CustomizableObjectInstance.cpp b/Engine/Plugins/Mutable/Source/CustomizableObject/Private/MuCO/CustomizableObjectInstance.cpp
--- a/Engine/Plugins/Mutable/Source/CustomizableObject/Private/MuCO/CustomizableObjectInstance.cpp
+++ b/Engine/Plugins/Mutable/Source/CustomizableObject/Private/MuCO/CustomizableObjectInstance.cpp
@@ -990,18 +990,29 @@ void UCustomizableInstancePrivate::PostEditChangePropertyWithoutEditor()
 }
 
 
-bool UCustomizableObjectInstance::CanUpdateInstance() const
+bool UCustomizableInstancePrivate::CanUpdateInstance() const
 {
-	UCustomizableObject* CustomizableObject = GetCustomizableObject();
-	if (!CustomizableObject)
+	UCustomizableObject* Object = GetPublic()->GetCustomizableObject();
+	if (!Object)
 	{
 		return false;
 	}
 
 #if WITH_EDITOR
-	return CustomizableObject->ConditionalAutoCompile();
+	if (Object->GetPrivate()->IsLocked())
+	{
+		return false;
+	}
+
+	if (!Object->IsCompiled())
+	{
+		return false;
+	}
+
+	return true;
+	
 #else
-	return CustomizableObject->IsCompiled();
+	return Object->IsCompiled();
 #endif
 }
 
@@ -1045,7 +1056,26 @@ void UCustomizableObjectInstance::UpdateSkeletalMeshAsyncResult(FInstanceUpdateN
 
 void UCustomizableInstancePrivate::TickUpdateCloseCustomizableObjects(UCustomizableObjectInstance& Public, FMutableInstanceUpdateMap& InOutRequestedUpdates)
 {
-	if (!Public.CanUpdateInstance())
+	UCustomizableObject* Object = Public.GetCustomizableObject();
+	if (!Object)
+	{
+		return;
+	}
+
+#if WITH_EDITOR
+	if (!Object->IsCompiled() &&
+		Object->GetPrivate()->CompilationResult != ECompilationResultPrivate::Errors) // Avoid constantly retry failed compilations.
+	{
+		if (ICustomizableObjectEditorModule* EditorModule = ICustomizableObjectEditorModule::Get())
+		{
+			TSharedRef<FCompilationRequest> CompileRequest = MakeShared<FCompilationRequest>(*Object);
+			CompileRequest->bSilentCompilation = true;
+			EditorModule->CompileCustomizableObject(CompileRequest);
+		}
+	}
+#endif
+		
+	if (!CanUpdateInstance())
 	{
 		return;
 	}
@@ -1102,7 +1132,7 @@ void UCustomizableInstancePrivate::UpdateInstanceIfNotGenerated(UCustomizableObj
 		return;
 	}
 
-	if (!Public.CanUpdateInstance())
+	if (!CanUpdateInstance())
 	{
 		return;
 	}
@@ -7241,7 +7271,7 @@ void UCustomizableObjectInstance::SetRequestedLODs(const TMap<FName, uint8>& InM
 	// First LOD Available =	First available lod per platform and quality (scalability) settings.	Per Skeletal Mesh.
 	// Requested LODs =			What the renderer wants to render.										Per Skeletal Mesh Component.
 
-	if (!CanUpdateInstance())
+	if (!GetPrivate()->CanUpdateInstance())
 	{
 		return;
 	}
@@ -7667,6 +7697,12 @@ TSet<UAssetUserData*> UCustomizableObjectInstance::GetMergedAssetUserData(int32
 }
 
 
+bool UCustomizableObjectInstance::CanUpdateInstance() const
+{
+	return GetPrivate()->CanUpdateInstance();
+}
+
+
 #if WITH_EDITORONLY_DATA
 
 void CalculateBonesToRemove(const FSkeletalMeshLODRenderData& LODResource, const FReferenceSkeleton& RefSkeleton, TArray<FBoneReference>& OutBonesToRemove)
diff --git a/Engine/Plugins/Mutable/Source/CustomizableObject/Private/MuCO/CustomizableObjectSystem.cpp b/Engine/Plugins/Mutable/Source/CustomizableObject/Private/MuCO/CustomizableObjectSystem.cpp
--- a/Engine/Plugins/Mutable/Source/CustomizableObject/Private/MuCO/CustomizableObjectSystem.cpp
+++ b/Engine/Plugins/Mutable/Source/CustomizableObject/Private/MuCO/CustomizableObjectSystem.cpp
@@ -174,6 +174,44 @@ static void CVarMutableSinkFunction()
 static FAutoConsoleVariableSink CVarMutableSink(FConsoleCommandDelegate::CreateStatic(&CVarMutableSinkFunction));
 
 
+constexpr float OnScreenWarningsTickerTime = 1.0f;
+
+
+int64 GetOnScreenMessageKey(const UCustomizableObject& Object, TMap<TWeakObjectPtr<const UCustomizableObject>, uint64>& KeyMap)
+{
+	uint64 Key;
+	if (uint64* Result = KeyMap.Find(&Object))
+	{
+		Key = *Result;
+	}
+	else
+	{
+		Key = 0;
+		while (GEngine->OnScreenDebugMessageExists(Key))
+		{
+			++Key;
+		}
+		
+		KeyMap.Add(TWeakObjectPtr(&Object), Key);
+	}
+
+	return Key;
+}
+
+
+void RemoveUnusedOnScreenMessages(TMap<TWeakObjectPtr<const UCustomizableObject>, uint64>& KeyMap)
+{
+	for (TMap<TWeakObjectPtr<const UCustomizableObject>, uint64>::TIterator It = KeyMap.CreateIterator(); It; ++It)
+	{
+		if (!It->Key.IsValid())
+		{
+			GEngine->RemoveOnScreenDebugMessage(It->Value);
+			It.RemoveCurrent();
+		}
+	}
+}
+
+
 FUpdateContextPrivate::FUpdateContextPrivate(UCustomizableObjectInstance& InInstance, const FCustomizableObjectInstanceDescriptor& Descriptor)
 {
 	check(IsInGameThread());
@@ -216,6 +254,126 @@ FUpdateContextPrivate::FUpdateContextPrivate(UCustomizableObjectInstance& InInst
 }
 
 
+bool TickWarnings(float DeltaTime)
+{
+	static TMap<TWeakObjectPtr<const UCustomizableObject>, uint64> KeysNotCompiled;
+	static TMap<TWeakObjectPtr<const UCustomizableObject>, uint64> KeysNotOptimized;
+	static TMap<TWeakObjectPtr<const UCustomizableObject>, uint64> KeysOutOfDate;
+	
+	TSet<const UCustomizableObject*> Objects;
+	
+	for (TObjectIterator<UCustomizableObjectInstanceUsage> CustomizableObjectInstanceUsage; CustomizableObjectInstanceUsage; ++CustomizableObjectInstanceUsage)
+	{
+		if (!IsValid(*CustomizableObjectInstanceUsage) || CustomizableObjectInstanceUsage->IsTemplate())
+		{
+			continue;
+		}
+		
+		const UCustomizableObjectInstance* Instance = CustomizableObjectInstanceUsage->GetCustomizableObjectInstance();
+		if (!Instance)
+		{
+			continue;
+		}
+
+		const UCustomizableObject* Object = Cast<UCustomizableObject>(Instance->GetCustomizableObject());
+		if (!Object)
+		{
+			continue;
+		}
+		
+		const USkeletalMeshComponent* Parent = Cast<USkeletalMeshComponent>(CustomizableObjectInstanceUsage->GetAttachParent());
+		if (!Parent)
+		{
+			continue;
+		}
+
+		const UWorld* World = Parent->GetWorld();
+		if (!World)
+		{
+			continue;
+		}
+
+		if (World->WorldType != EWorldType::PIE &&
+			World->WorldType != EWorldType::Game)
+		{
+			continue;
+		}
+
+		Objects.Add(Object);
+	}
+
+	for (const UCustomizableObject* Object : Objects)
+	{
+		constexpr float DisplayTime = OnScreenWarningsTickerTime * 2.0f;
+		
+		if (!Object->IsLoading() && !Object->IsCompiled())
+		{
+			const uint64 Key = GetOnScreenMessageKey(*Object, KeysNotCompiled);
+			FString Msg = FString::Printf(TEXT("Customizable Object [%s] not compiled."), *Object->GetName());
+			GEngine->AddOnScreenDebugMessage(Key, DisplayTime, FColor::Red, Msg);
+		}
+		else if (uint64* Key = KeysNotCompiled.Find(TWeakObjectPtr(Object)))
+		{
+			GEngine->RemoveOnScreenDebugMessage(*Key);
+		}
+
+		const UModelResources* ModelResources = Object->GetPrivate()->GetModelResources();
+		if (!ModelResources)
+		{
+			continue;
+		}
+
+		if (Object->GetPrivate()->Status.Get() != FCustomizableObjectStatus::EState::ModelLoaded)
+		{
+			continue;
+		}
+
+#if WITH_EDITORONLY_DATA
+		if (!ModelResources->bIsCompiledWithOptimization)
+		{
+			const uint64 Key = GetOnScreenMessageKey(*Object, KeysNotOptimized);
+			FString Msg = FString::Printf(TEXT("Customizable Object [%s] was compiled without optimization."), *Object->GetName());
+			GEngine->AddOnScreenDebugMessage(Key, DisplayTime, FColor::Yellow, Msg);
+		}
+		else if (uint64* Key = KeysNotOptimized.Find(TWeakObjectPtr(Object)))
+		{
+			GEngine->RemoveOnScreenDebugMessage(*Key);
+		}
+#endif
+
+#if WITH_EDITOR
+		TArray<FName> OutOfDatePackages;
+		TArray<FName> AddedPackages;
+		TArray<FName> RemovedPackages;
+		bool bReleaseVersion;
+		if (Object->GetPrivate()->IsCompilationOutOfDate(true, OutOfDatePackages, AddedPackages, RemovedPackages, bReleaseVersion))
+		{
+			const uint64 Key = GetOnScreenMessageKey(*Object, KeysOutOfDate);
+			
+			if (!GEngine->OnScreenDebugMessageExists(Key))
+			{
+				UE_LOG(LogMutable, Display, TEXT("Customizable Object [%s] compilation out of date. Changes since last compilation:"), *Object->GetName());
+				PrintParticipatingPackagesDiff(OutOfDatePackages, AddedPackages, RemovedPackages, bReleaseVersion);
+			}
+			
+			FString Msg = FString::Printf(TEXT("Customizable Object [%s] compilation out of date. See the Output Log for more information."), *Object->GetName());
+			GEngine->AddOnScreenDebugMessage(Key, DisplayTime, FColor::Yellow, Msg);
+		}
+		else if (uint64* Key = KeysOutOfDate.Find(TWeakObjectPtr(Object)))
+		{
+			GEngine->RemoveOnScreenDebugMessage(*Key);
+		}
+#endif
+	}
+
+	RemoveUnusedOnScreenMessages(KeysNotCompiled);
+	RemoveUnusedOnScreenMessages(KeysNotOptimized);
+	RemoveUnusedOnScreenMessages(KeysOutOfDate);
+
+	return true;
+}
+
+
 bool FUpdateContextPrivate::IsContextValid() const
 {
 	return bValid;
@@ -506,8 +664,6 @@ UCustomizableObjectSystem* UCustomizableObjectSystem::GetInstance()
 		UCustomizableObjectSystemPrivate::SSystem->AddToRoot();
 		checkf(!GUObjectArray.IsDisregardForGC(UCustomizableObjectSystemPrivate::SSystem), TEXT("Mutable was initialized too early in the UE4 init process, for instance, in the constructor of a default UObject."));
 		UCustomizableObjectSystemPrivate::SSystem->InitSystem();
-
-		//FCoreUObjectDelegates::PurgePendingReleaseSkeletalMesh.AddUObject(UCustomizableObjectSystemPrivate::SSystem, &UCustomizableObjectSystem::PurgePendingReleaseSkeletalMesh);
 	}
 
 	return UCustomizableObjectSystemPrivate::SSystem;
@@ -645,6 +801,8 @@ void UCustomizableObjectSystem::BeginDestroy()
 		{
 			FTSTicker::GetCoreTicker().RemoveTicker(Private->TickDelegateHandle);			
 		}
+
+		FTSTicker::GetCoreTicker().RemoveTicker(Private->TickWarningsDelegateHandle);
 #endif // !UE_SERVER
 		
 		// Complete pending taskgraph tasks
@@ -1254,7 +1412,7 @@ EUpdateRequired UCustomizableObjectSystemPrivate::IsUpdateRequired(const UCustom
 	UCustomizableObjectSystem* System = UCustomizableObjectSystem::GetInstance();
 	const UCustomizableInstancePrivate* const Private = Instance.GetPrivate();
 	
-	if (!Instance.CanUpdateInstance())
+	if (!Instance.GetPrivate()->CanUpdateInstance())
 	{
 		return EUpdateRequired::NoUpdate;
 	}
@@ -1375,7 +1533,7 @@ void UCustomizableObjectSystemPrivate::EnqueueUpdateSkeletalMesh(const TSharedRe
 		return;
 	}
 	
-	if (!Instance->CanUpdateInstance())
+	if (!Instance->GetPrivate()->CanUpdateInstance())
 	{
 		Context->UpdateResult = EUpdateResult::Error;
 		FinishUpdateGlobal(Context);
@@ -4396,8 +4554,7 @@ void UCustomizableObjectSystem::SetImagePixelFormatOverride(const mu::FImageOper
 void UCustomizableObjectSystem::AddUncompiledCOWarning(const UCustomizableObject& InObject, FString const* OptionalLogInfo)
 {
 	FString Msg;
-	Msg += FString::Printf(TEXT("Warning: Customizable Object [%s] not compiled."), *InObject.GetName());
-	GEngine->AddOnScreenDebugMessage((uint64)((PTRINT)&InObject), 10.0f, FColor::Red, Msg);
+	Msg += FString::Printf(TEXT("Warning: Customizable Object [%s] not loaded or compiled."), *InObject.GetName());
 
 #if WITH_EDITOR
 	// Mutable will spam these warnings constantly due to the tick and LOD manager checking for instances to update with every tick. Send only one message per CO in the editor.
@@ -4414,7 +4571,7 @@ void UCustomizableObjectSystem::AddUncompiledCOWarning(const UCustomizableObject
 
 	if (!GetPrivate()->UncompiledCustomizableObjectsNotificationPtr.IsValid())
 	{
-		FNotificationInfo Info(FText::FromString("Uncompiled Customizable Object/s found. Please, check the Message Log - Mutable for more information."));
+		FNotificationInfo Info(FText::FromString("Customizable Object/s not loaded or compiled. Please, check the Message Log - Mutable for more information."));
 		Info.bFireAndForget = true;
 		Info.bUseThrobber = true;
 		Info.FadeOutDuration = 1.0f;
@@ -4424,12 +4581,12 @@ void UCustomizableObjectSystem::AddUncompiledCOWarning(const UCustomizableObject
 	}
 
 	const FString ErrorString = FString::Printf(
-		TEXT("Customizable Object [%s] not compiled.  Compile via the editor or via code before instancing.  %s"),
+		TEXT("Customizable Object [%s] not loaded or not compiled. Compile via the editor or via code before instancing.  %s"),
 		*InObject.GetName(), OptionalLogInfo ? **OptionalLogInfo : TEXT(""));
 
 #else // !WITH_EDITOR
 	const FString ErrorString = FString::Printf(
-		TEXT("Customizable Object [%s] not compiled.  This is not an Editor build, so this is an unrecoverable bad state; could be due to code or a cook failure.  %s"),
+		TEXT("Customizable Object [%s] not loaded or compiled. This is not an Editor build, so this is an unrecoverable bad state; could be due to code or a cook failure.  %s"),
 		*InObject.GetName(), OptionalLogInfo ? **OptionalLogInfo : TEXT(""));
 #endif
 
@@ -4589,9 +4746,13 @@ void UCustomizableObjectSystemPrivate::OnMutableEnabledChanged(IConsoleVariable*
 		{
 			if (!SystemPrivate->TickDelegateHandle.IsValid())
 			{
-				SystemPrivate->TickDelegate = FTickerDelegate::CreateUObject(System, &UCustomizableObjectSystem::Tick);
-				SystemPrivate->TickDelegateHandle = FTSTicker::GetCoreTicker().AddTicker(SystemPrivate->TickDelegate, 0.f);
-			}			
+				SystemPrivate->TickDelegateHandle = FTSTicker::GetCoreTicker().AddTicker(FTickerDelegate::CreateUObject(System, &UCustomizableObjectSystem::Tick), 0.f);
+			}
+		}
+
+		if (!SystemPrivate->TickWarningsDelegateHandle.IsValid())
+		{
+			SystemPrivate->TickWarningsDelegateHandle = FTSTicker::GetCoreTicker().AddTicker(FTickerDelegate::CreateStatic(&TickWarnings), 0.f);
 		}
 #endif // !UE_SERVER
 	}
diff --git a/Engine/Plugins/Mutable/Source/CustomizableObject/Public/MuCO/CustomizableObject.h b/Engine/Plugins/Mutable/Source/CustomizableObject/Public/MuCO/CustomizableObject.h
--- a/Engine/Plugins/Mutable/Source/CustomizableObject/Public/MuCO/CustomizableObject.h
+++ b/Engine/Plugins/Mutable/Source/CustomizableObject/Public/MuCO/CustomizableObject.h
@@ -129,6 +129,73 @@ struct FMutableLODSettings
 };
 
 
+USTRUCT(BlueprintType)
+struct FCompileCallbackParams
+{
+	GENERATED_BODY()
+
+	/** The compile request has failed. See output log for more details. False if the request has been skipped. */
+	UPROPERTY(BlueprintReadOnly, Category = Compile)
+	bool bRequestFailed = false;
+
+	/** The compile request has finished with warnings. False if the request has been skipped. */
+	UPROPERTY(BlueprintReadOnly, Category = Compile)
+	bool bWarnings = false;
+
+	/** The compile request has finished with errors. False if the request has been skipped. */
+	UPROPERTY(BlueprintReadOnly, Category = Compile)
+	bool bErrors = false;
+	
+	/** The compilation has been skipped due to Compile Parameters. */
+	UPROPERTY(BlueprintReadOnly, Category = Compile)
+	bool bSkipped = false;
+
+	/** The Customizable Object is compiled. */
+	UPROPERTY(BlueprintReadOnly, Category = Compile)
+	bool bCompiled = false;
+};
+
+
+DECLARE_DYNAMIC_DELEGATE_OneParam(FCompileDelegate, const FCompileCallbackParams&, Params);
+
+
+USTRUCT(BlueprintType)
+struct FCompileParams
+{
+	GENERATED_BODY()
+
+	/** If true, skip the compilation if the Customizable Object is already compiled. */
+	UPROPERTY(BlueprintReadWrite, Category = Compile)
+	bool bSkipIfCompiled = false;
+
+	/** If true, skip the compilation if the Customizable Object or its contents are out of date (child Customizable Objects, Data Tables, referenced Materials...). */
+	UPROPERTY(BlueprintReadWrite, Category = Compile)
+	bool bSkipIfOutOfDate = true;
+	
+	/** Gathers all asset references used in this Customizable Object. Marks the object as modified.*/
+	UPROPERTY(BlueprintReadWrite, Category = Compile)
+	bool bGatherReferneces = false;
+
+	/** If true, compile the Customizable Object asynchronously. If synchronously and the Customizable Object is not ready the compilation will fail. */
+	UPROPERTY(BlueprintReadWrite, Category = Compile)
+	bool bAsync = true;
+
+	/** If assigned, compile the Customizable Object only the selected parameter of the given Instance. */
+	UPROPERTY(BlueprintReadWrite, Category = Compile)
+	TObjectPtr<UCustomizableObjectInstance> CompileOnlySelectedInstance = nullptr;
+	
+	UPROPERTY(BlueprintReadWrite, Category = Compile)
+	ECustomizableObjectOptimizationLevel OptimizationLevel = ECustomizableObjectOptimizationLevel::Maximum;
+
+	UPROPERTY(BlueprintReadWrite, Category = Compile)
+	ECustomizableObjectTextureCompression TextureCompression = ECustomizableObjectTextureCompression::Fast;
+
+	/** Called once the compilation request has finished (successfully or not). */
+	UPROPERTY(BlueprintReadWrite, Category = Compile)
+	FCompileDelegate Callback;
+};
+
+
 UCLASS( BlueprintType, config=Engine )
 class CUSTOMIZABLEOBJECT_API UCustomizableObject : public UObject
 {
@@ -386,9 +453,15 @@ public:
 	  * @return True if the parameter is multidimensional and false if it is not. */
 	UFUNCTION(BlueprintCallable, Category = CustomizableObject)
 	bool IsParameterMultidimensional(UPARAM(DisplayName = "Parameter Name") const FString& InParameterName) const;
-
+	
+	/** Compile the Customizable Object. */
+	UFUNCTION(BlueprintCallable, Category = CustomizableObject)
+	void Compile(const FCompileParams& Params);
+	
 #if WITH_EDITOR
-	/** Compile the object if Automatic Compilation is enabled and has not been already compiled.
+	/** DEPRECATED. Use Compile instead.
+	  * 
+	  *	Compile the object if Automatic Compilation is enabled and has not been already compiled.
 	  * Automatic compilation can be enabled/disabled in the Mutable's Plugin Settings.
 	  * @return true if compiled */
 	bool ConditionalAutoCompile();
@@ -507,10 +580,15 @@ public:
 	UFUNCTION(BlueprintCallable, Category = CustomizableObject)
 	UCustomizableObjectInstance* CreateInstance();
 	
-	/** Check if the CustomizableObject asset has been compiled. This will always be true in a packaged game, but it could be false in the editor. */
+	/** Check if the CustomizableObject asset has been compiled. This will always be true in a packaged game, but it could be false in the editor.
+	 *  It may return false due to the Customizable Object still being loaded. */
 	UFUNCTION(BlueprintCallable, Category = CustomizableObject)
 	bool IsCompiled() const;
 
+	/** Return true if the Customizable Object is still being loaded. It may take a few frames to load the Customizable Object. */
+	UFUNCTION(BlueprintCallable, Category = CustomizableObject)
+	bool IsLoading() const;
+	
 #if WITH_EDITOR
 	/** Return true if this Customizable Object has references to a parent Customizable Object. Only root Customizable Objects will return false. */
 	UFUNCTION(BlueprintCallable, Category = CustomizableObject)
diff --git a/Engine/Plugins/Mutable/Source/CustomizableObject/Public/MuCO/CustomizableObjectInstance.h b/Engine/Plugins/Mutable/Source/CustomizableObject/Public/MuCO/CustomizableObjectInstance.h
--- a/Engine/Plugins/Mutable/Source/CustomizableObject/Public/MuCO/CustomizableObjectInstance.h
+++ b/Engine/Plugins/Mutable/Source/CustomizableObject/Public/MuCO/CustomizableObjectInstance.h
@@ -324,8 +324,8 @@ public:
 	  * It requires that the CustomizableObject had the bEnableAssetUserDataMerge set to true during compilation. */
 	UFUNCTION(BlueprintCallable, Category = CustomizableObjectInstance)
 	TSet<UAssetUserData*> GetMergedAssetUserData(int32 ComponentIndex) const;
-	
-	/** Return true if the instance is not locked and if it's compiled. */
+
+	UE_DEPRECATED(5.6, "Use UCustomizableObject::IsLoading or UCustomizableObject::IsCompiling instead.")
 	bool CanUpdateInstance() const;
 
 	/** Generate the instance with the current parameters and update all the components Skeletal Meshes asynchronously. */
diff --git a/Engine/Plugins/Mutable/Source/CustomizableObject/Public/MuCO/CustomizableObjectSystem.h b/Engine/Plugins/Mutable/Source/CustomizableObject/Public/MuCO/CustomizableObjectSystem.h
--- a/Engine/Plugins/Mutable/Source/CustomizableObject/Public/MuCO/CustomizableObjectSystem.h
+++ b/Engine/Plugins/Mutable/Source/CustomizableObject/Public/MuCO/CustomizableObjectSystem.h
@@ -28,8 +28,6 @@ struct FFrame;
 struct FGuid;
 
 
-constexpr uint64 KEY_OFFSET_COMPILATION_OUT_OF_DATE = 1;
-
 extern TAutoConsoleVariable<bool> CVarClearWorkingMemoryOnUpdateEnd;
 
 extern TAutoConsoleVariable<bool> CVarReuseImagesBetweenInstances;
@@ -47,7 +45,16 @@ extern TAutoConsoleVariable<bool> CVarEnableReleaseMeshResources;
 extern TAutoConsoleVariable<bool> CVarFixLowPriorityTasksOverlap;
 
 
-UENUM()
+UENUM(BlueprintType)
+enum class ECustomizableObjectOptimizationLevel : uint8
+{
+	None = 0, 
+	Minimal UE_DEPRECATED(5.6, "Converted internally to None. Use None instead.") = 1,
+	Maximum = 2
+};
+
+
+UENUM(BlueprintType)
 enum class ECustomizableObjectTextureCompression : uint8
 {
 	// Don't use texture compression
diff --git a/Engine/Plugins/Mutable/Source/CustomizableObjectEditor/Private/MuCOE/AssetDefinition_CustomizableObject.cpp b/Engine/Plugins/Mutable/Source/CustomizableObjectEditor/Private/MuCOE/AssetDefinition_CustomizableObject.cpp
--- a/Engine/Plugins/Mutable/Source/CustomizableObjectEditor/Private/MuCOE/AssetDefinition_CustomizableObject.cpp
+++ b/Engine/Plugins/Mutable/Source/CustomizableObjectEditor/Private/MuCOE/AssetDefinition_CustomizableObject.cpp
@@ -202,8 +202,8 @@ namespace MenuExtension_CustomizableObject
 			{
 				if (Object)
 				{
-					TSharedRef<FCompilationRequest> NewRequest = MakeShared<FCompilationRequest>(*Object, true);
-					NewRequest->GetCompileOptions().bSilentCompilation = true;
+					TSharedRef<FCompilationRequest> NewRequest = MakeShared<FCompilationRequest>(*Object);
+					NewRequest->bSilentCompilation = true;
 					CompilationRequests.Add(NewRequest);
 				}
 			}
@@ -245,15 +245,18 @@ namespace MenuExtension_CustomizableObject
 					UCustomizableObject* Object = Cast<UCustomizableObject>(MutablePrivate::LoadObject(Asset));
 					if (Object)
 					{
-						TSharedRef<FCompilationRequest> NewRequest = MakeShared<FCompilationRequest>(*Object, true);
-						NewRequest->GetCompileOptions().bSilentCompilation = true;
+						TSharedRef<FCompilationRequest> NewRequest = MakeShared<FCompilationRequest>(*Object);
+						NewRequest->bSilentCompilation = true;
 						CompilationRequests.Add(NewRequest);
 					}
 				}
 			}
 		}
-
-		ICustomizableObjectEditorModule::GetChecked().CompileCustomizableObjects(CompilationRequests);
+		
+		for (const TSharedRef<FCompilationRequest>& CompilationRequest : CompilationRequests)
+		{
+			ICustomizableObjectEditorModule::GetChecked().CompileCustomizableObject(CompilationRequest);
+		}
 	}
 
 	
diff --git a/Engine/Plugins/Mutable/Source/CustomizableObjectEditor/Private/MuCOE/CustomizableObjectCompileRunnable.cpp b/Engine/Plugins/Mutable/Source/CustomizableObjectEditor/Private/MuCOE/CustomizableObjectCompileRunnable.cpp
--- a/Engine/Plugins/Mutable/Source/CustomizableObjectEditor/Private/MuCOE/CustomizableObjectCompileRunnable.cpp
+++ b/Engine/Plugins/Mutable/Source/CustomizableObjectEditor/Private/MuCOE/CustomizableObjectCompileRunnable.cpp
@@ -337,7 +337,7 @@ FCustomizableObjectSaveDDRunnable::FCustomizableObjectSaveDDRunnable(const TShar
 	ModelStreamables = InModelStreamables;
 	PlatformData.ModelResources = TStrongObjectPtr<UModelResources>(&InModelResources);
 
-	Options = InRequest->GetCompileOptions();
+	Options = InRequest->Options;
 
 	DDCKey = InRequest->GetDerivedDataCacheKey();
 	DefaultDDCPolicy = InRequest->GetDerivedDataCachePolicy();
diff --git a/Engine/Plugins/Mutable/Source/CustomizableObjectEditor/Private/MuCOE/CustomizableObjectCompiler.cpp b/Engine/Plugins/Mutable/Source/CustomizableObjectEditor/Private/MuCOE/CustomizableObjectCompiler.cpp
--- a/Engine/Plugins/Mutable/Source/CustomizableObjectEditor/Private/MuCOE/CustomizableObjectCompiler.cpp
+++ b/Engine/Plugins/Mutable/Source/CustomizableObjectEditor/Private/MuCOE/CustomizableObjectCompiler.cpp
@@ -135,30 +135,57 @@ void FCustomizableObjectCompiler::PreloadingReferencerAssetsCallback(bool bAsync
 
 void FCustomizableObjectCompiler::Compile(const TSharedRef<FCompilationRequest>& InCompileRequest)
 {
-	TRACE_BEGIN_REGION(UE_MUTABLE_COMPILE_REGION);
+	MUTABLE_CPUPROFILER_SCOPE(FCustomizableObjectCompiler::Compile)
 	
 	check(IsInGameThread());
-	check(!CurrentRequest);
-	
-	CurrentRequest = InCompileRequest.ToSharedPtr();
-	CurrentObject = InCompileRequest->GetCustomizableObject();
-	CurrentOptions = InCompileRequest->GetCompileOptions();
 
-	if (!CurrentObject)
+	UCustomizableObject* Object = InCompileRequest->GetCustomizableObject();
+	if (!Object)
 	{
 		UE_LOG(LogMutable, Warning, TEXT("Failed to compile Customizable Object. Object is missing."));
-		CompleteRequest(ECompilationStatePrivate::Completed, ECompilationResultPrivate::Errors);
+		
+		FCompileCallbackParams Params;
+		Params.bRequestFailed = true;
+    	
+		InCompileRequest->Callback.ExecuteIfBound(Params);
 		return;
 	}
-
-	if (CurrentObject->GetPrivate()->CompilationState == ECompilationStatePrivate::InProgress)
+	
+	if (InCompileRequest->bSkipIfCompiled &&
+		Object->IsCompiled())
 	{
-		UE_LOG(LogMutable, Warning, TEXT("Failed to compile Customizable Object [%s]. Object already being compiled."), *CurrentObject->GetName());
-		CurrentObject = nullptr; // Someone else is compiling the CO. Invalidate the CurrentObject pointer to avoid changing the state of the ongoing compilation.
-		CompleteRequest(ECompilationStatePrivate::Completed, ECompilationResultPrivate::Errors);
+		FCompileCallbackParams Params;
+		Params.bSkipped = true;
+    	
+		InCompileRequest->Callback.ExecuteIfBound(Params);
 		return;
 	}
 
+	if (InCompileRequest->bSkipIfOutOfDate)
+	{
+		TArray<FName> OutOfDatePackages;
+		TArray<FName> AddedPackages;
+		TArray<FName> RemovedPackages;
+		bool bReleaseVersion;
+		if (!Object->GetPrivate()->IsCompilationOutOfDate(false, OutOfDatePackages, AddedPackages, RemovedPackages, bReleaseVersion))
+		{
+			FCompileCallbackParams Params;
+			Params.bSkipped = true;
+			Params.bCompiled = Object->IsCompiled();
+
+			InCompileRequest->Callback.ExecuteIfBound(Params);
+			return;
+		}
+	}
+
+	TRACE_BEGIN_REGION(UE_MUTABLE_COMPILE_REGION);
+
+	check(!CurrentRequest);
+
+	CurrentRequest = InCompileRequest.ToSharedPtr();
+	CurrentObject = Object;
+	CurrentOptions = InCompileRequest->Options;
+	
 	if (!UCustomizableObjectSystem::IsActive())
 	{
 		UE_LOG(LogMutable, Warning, TEXT("Failed to compile Customizable Object [%s]. Mutable is disabled. To enable it set the CVar Mutable.Enabled to true."), *CurrentObject->GetName());
@@ -186,7 +213,7 @@ void FCustomizableObjectCompiler::Compile(const TSharedRef<FCompilationRequest>&
 
 	UCustomizableObjectSystem* System = UCustomizableObjectSystem::GetInstanceChecked();
 
-	if (!CurrentRequest->IsAsyncCompilation())
+	if (!CurrentRequest->bAsync)
 	{
 		// Sync compilation. Force finish all pending updates and async compilations
 		System->GetPrivate()->BlockTillAllRequestsFinished();
@@ -272,7 +299,7 @@ void FCustomizableObjectCompiler::Compile(const TSharedRef<FCompilationRequest>&
 	
 	bool bAssetsLoaded = true;
 
-	const bool bAsync = InCompileRequest->IsAsyncCompilation();
+	const bool bAsync = InCompileRequest->bAsync;
 	if (ArrayAssetToStream.Num() > 0)
 	{
 		// Customizations are marked as editoronly on load and are not packaged into the runtime game by default.
@@ -300,10 +327,50 @@ void FCustomizableObjectCompiler::Compile(const TSharedRef<FCompilationRequest>&
 }
 
 
-void FCustomizableObjectCompiler::Compile(const TArray<TSharedRef<FCompilationRequest>>& InCompileRequests)
+void FCustomizableObjectCompiler::EnqueueCompileRequest(const TSharedRef<FCompilationRequest>& CompileRequest, bool bForceRequests)
 {
-	NumCompilationRequests += InCompileRequests.Num();
-	CompileRequests.Append(InCompileRequests);
+	check(IsInGameThread());
+
+	const UCustomizableObject* CustomizableObject = CompileRequest->GetCustomizableObject();
+	if (!CustomizableObject)
+	{
+		FCompileCallbackParams Params;
+		Params.bRequestFailed = true;
+
+		CompileRequest->Callback.ExecuteIfBound(Params);
+		return;
+	}
+	
+	if (IsRunningGame())
+	{
+		FCompileCallbackParams Params;
+		Params.bRequestFailed = true;
+		Params.bCompiled = CompileRequest->GetCustomizableObject()->IsCompiled();
+
+		CompileRequest->Callback.ExecuteIfBound(Params);
+		return;
+	}
+		
+	if (!CompileRequest->bAsync)
+	{
+		FCustomizableObjectCompiler SyncCompiler;
+		SyncCompiler.Compile(CompileRequest);
+	}
+	else
+	{
+		if (!bForceRequests && (CustomizableObject->GetPrivate()->IsLocked() || IsRequestQueued(CompileRequest)))
+		{
+			FCompileCallbackParams Params;
+			Params.bRequestFailed = true;
+			Params.bCompiled = CompileRequest->GetCustomizableObject()->IsCompiled();
+
+			CompileRequest->Callback.ExecuteIfBound(Params);
+			return;
+		}
+
+		++NumCompilationRequests;
+		CompileRequests.Add(CompileRequest);
+	}
 }
 
 
@@ -335,11 +402,21 @@ void FCustomizableObjectCompiler::TickCook(float DeltaTime, bool bCookCompete)
 
 bool FCustomizableObjectCompiler::IsRequestQueued(const TSharedRef<FCompilationRequest>& InCompileRequest) const
 {
-	return (CurrentRequest == InCompileRequest) || 
-		(CompileRequests.ContainsByPredicate([&InCompileRequest](const TSharedRef<FCompilationRequest>& Other)
+	return CurrentRequest == InCompileRequest || 
+		CompileRequests.ContainsByPredicate([&InCompileRequest](const TSharedRef<FCompilationRequest>& Other)
 		{
 			return InCompileRequest.Get() == Other.Get(); // Compare the content of the request not the ref
-		}));
+		});
+}
+
+
+bool FCustomizableObjectCompiler::IsRequestQueued(const UCustomizableObject& Object) const
+{
+	return (CurrentRequest && CurrentRequest->GetCustomizableObject() == &Object) || 
+		CompileRequests.ContainsByPredicate([ObjectPtr = &Object](const TSharedRef<FCompilationRequest>& Other)
+		{
+			return ObjectPtr == Other->GetCustomizableObject();
+		});
 }
 
 
@@ -702,7 +779,6 @@ void FCustomizableObjectCompiler::SetCompilationState(ECompilationStatePrivate S
 
 	if (CurrentObject)
 	{
-		CurrentObject->GetPrivate()->CompilationState = State;
 		CurrentObject->GetPrivate()->CompilationResult = Result;
 	}
 }
@@ -1002,7 +1078,7 @@ void FCustomizableObjectCompiler::CompleteRequest(ECompilationStatePrivate State
 	}
 
 	// Copy warnings and errors to the request
-	CompilationLogsContainer.GetMessages(CurrentRequest->GetWarnings(), CurrentRequest->GetErrors());
+	CompilationLogsContainer.GetMessages(CurrentRequest->Warnings, CurrentRequest->Errors);
 	
 	// Clear Messages
 	CompilationLogsContainer.ClearMessageCounters();
@@ -1013,6 +1089,19 @@ void FCustomizableObjectCompiler::CompleteRequest(ECompilationStatePrivate State
 		GEngine->ForceGarbageCollection();
 	}
 
+	FCompileCallbackParams Params;
+	Params.bErrors = Result == ECompilationResultPrivate::Errors;
+	Params.bWarnings = Result == ECompilationResultPrivate::Warnings;
+	
+	if (CurrentObject)
+	{
+		Params.bCompiled = CurrentObject->IsCompiled();
+		Params.bErrors |= !CurrentRequest->Errors.IsEmpty();
+		Params.bWarnings |= !CurrentRequest->Warnings.IsEmpty();
+	}
+		
+	CurrentRequest->Callback.ExecuteIfBound(Params);
+	
 	// Request completed, reset pointers and state
 	CurrentObject = nullptr;
 	CurrentRequest.Reset();
@@ -1041,8 +1130,20 @@ bool FCustomizableObjectCompiler::TryPopCompileRequest()
 	for (TArray<TSharedRef<FCompilationRequest>>::TIterator It = CompileRequests.CreateIterator(); It; ++It)
 	{
 		TSharedRef<FCompilationRequest>& CompileRequest = *It;
+
+		UCustomizableObject* Object = CompileRequest->GetCustomizableObject();
+		if (!Object)
+		{
+			It.RemoveCurrent();
+
+			FCompileCallbackParams Params;
+			Params.bRequestFailed = true;
+    	
+			CompileRequest->Callback.ExecuteIfBound(Params);
+			continue;
+		}
 		
-		switch (CompileRequest->GetCustomizableObject()->GetPrivate()->Status.Get())
+		switch (Object->GetPrivate()->Status.Get())
 		{
 		case FCustomizableObjectStatusTypes::EState::Loading:
 			// Wait.
@@ -1433,7 +1534,7 @@ void FCustomizableObjectCompiler::AddCompileNotification(const FText& Compilatio
 	FCustomizableObjectEditorLogger::CreateLog(Text)
 	.SubText(CompilationStep)
 	.Category(ELoggerCategory::Compilation)
-	.Notification(!CurrentOptions.bSilentCompilation)
+	.Notification(!CurrentRequest->bSilentCompilation)
 	.CustomNotification()
 	.FixNotification()
 	.Log();
@@ -1481,7 +1582,7 @@ void FCustomizableObjectCompiler::NotifyCompilationErrors() const
 	FCustomizableObjectEditorLogger::CreateLog(Message)
 	.Category(ELoggerCategory::Compilation)
 	.Severity(Severity)
-	.Notification(!CurrentOptions.bSilentCompilation || !NoWarningsOrErrors)
+	.Notification(!CurrentRequest->bSilentCompilation || !NoWarningsOrErrors)
 	.CustomNotification()
 	.Log();
 }
diff --git a/Engine/Plugins/Mutable/Source/CustomizableObjectEditor/Private/MuCOE/CustomizableObjectCompiler.h b/Engine/Plugins/Mutable/Source/CustomizableObjectEditor/Private/MuCOE/CustomizableObjectCompiler.h
--- a/Engine/Plugins/Mutable/Source/CustomizableObjectEditor/Private/MuCOE/CustomizableObjectCompiler.h
+++ b/Engine/Plugins/Mutable/Source/CustomizableObjectEditor/Private/MuCOE/CustomizableObjectCompiler.h
@@ -77,10 +77,17 @@ public:
 	 *
 	 * Loads assets which reference Object's package asynchronously before calling ProcessChildObjectsRecursively. */
 	void Compile(const TSharedRef<FCompilationRequest>& InCompileRequest);
-	void Compile(const TArray<TSharedRef<FCompilationRequest>>& InCompileRequests);
-	
+
+	/** Enqueue a new compile request.
+	 * @param bForceRequests enqueue even if the request is already in the queue. See FCompilationRequest::operator==(...). */
+	void EnqueueCompileRequest(const TSharedRef<FCompilationRequest>& CompileRequest, bool bForceRequests);
+
+	/** Queued or in progress. */
 	bool IsRequestQueued(const TSharedRef<FCompilationRequest>& InCompileRequest) const;
 
+	/** Queued or in progress. */
+	bool IsRequestQueued(const UCustomizableObject& Object) const;
+	
 	/** FSerializableObject interface */
 	virtual void AddReferencedObjects(FReferenceCollector& Collector) override;
 	virtual FString GetReferencerName() const override
@@ -157,5 +164,5 @@ private:
 
 	/** Compilation start time in seconds. */
 	double CompilationStartTime = 0;
-
 };
+
diff --git a/Engine/Plugins/Mutable/Source/CustomizableObjectEditor/Private/MuCOE/CustomizableObjectEditor.cpp b/Engine/Plugins/Mutable/Source/CustomizableObjectEditor/Private/MuCOE/CustomizableObjectEditor.cpp
--- a/Engine/Plugins/Mutable/Source/CustomizableObjectEditor/Private/MuCOE/CustomizableObjectEditor.cpp
+++ b/Engine/Plugins/Mutable/Source/CustomizableObjectEditor/Private/MuCOE/CustomizableObjectEditor.cpp
@@ -4,6 +4,8 @@
 
 #include "AssetViewerSettings.h"
 #include "ContentBrowserModule.h"
+#include "CustomizableObjectCompiler.h"
+#include "CustomizableObjectEditorModule.h"
 #include "CustomizableObjectEditorPerformanceAnalyzer.h"
 #include "Animation/DebugSkelMeshComponent.h"
 #include "AssetRegistry/ARFilter.h"
@@ -1448,26 +1450,15 @@ void FCustomizableObjectEditor::CompileObject(bool bOnlySelectedParameters, bool
 		return;
 	}
 
-	TSharedRef<FCompilationRequest> CompileRequest = MakeShared<FCompilationRequest>(*CustomizableObject, true);
-
-	FCompilationOptions& CompilationOptions = CompileRequest->GetCompileOptions();
-	CompilationOptions.bSilentCompilation = false;
-	CompilationOptions.bGatherReferences = bGatherReferences;
+	TSharedRef<FCompilationRequest> CompileRequest = MakeShared<FCompilationRequest>(*CustomizableObject);
+	CompileRequest->bSilentCompilation = false;
+	CompileRequest->Options.bGatherReferences = bGatherReferences;
 
 	if (bOnlySelectedParameters)
 	{
-		UCustomizableObjectInstance* Instance = GetPreviewInstance();
-		check(Instance);
-		
-		const TArray<FCustomizableObjectIntParameterValue>& IntParameters = Instance->GetPrivate()->GetDescriptor().GetIntParameters();
-		CompilationOptions.ParamNamesToSelectedOptions.Reserve(IntParameters.Num());
-
-		for (const FCustomizableObjectIntParameterValue& IntParam : IntParameters)
-		{
-			CompilationOptions.ParamNamesToSelectedOptions.Add(IntParam.ParameterName, IntParam.ParameterValueName);
-		}
+		CompileRequest->Options.ParamNamesToSelectedOptions = GetCompileOnlySelectedParameters(*GetPreviewInstance());
 	}
-
+	
 	ICustomizableObjectEditorModule::GetChecked().CompileCustomizableObject(CompileRequest);
 }
 
@@ -2293,7 +2284,11 @@ void FCustomizableObjectEditor::OnCustomizableObjectStatusChanged(FCustomizableO
 		}
 		else if (CurrentState == FCustomizableObjectStatusTypes::EState::NoModel)
 		{
-			CustomizableObject->ConditionalAutoCompile();
+			TSharedRef<FCompilationRequest> CompileRequest = MakeShared<FCompilationRequest>(*CustomizableObject);
+			CompileRequest->bSkipIfOutOfDate = true;
+			CompileRequest->bSilentCompilation = false;
+			
+			ICustomizableObjectEditorModule::GetChecked().CompileCustomizableObject(CompileRequest);
 		}
 	}
 }
diff --git a/Engine/Plugins/Mutable/Source/CustomizableObjectEditor/Private/MuCOE/CustomizableObjectEditorFunctionLibrary.cpp b/Engine/Plugins/Mutable/Source/CustomizableObjectEditor/Private/MuCOE/CustomizableObjectEditorFunctionLibrary.cpp
--- a/Engine/Plugins/Mutable/Source/CustomizableObjectEditor/Private/MuCOE/CustomizableObjectEditorFunctionLibrary.cpp
+++ b/Engine/Plugins/Mutable/Source/CustomizableObjectEditor/Private/MuCOE/CustomizableObjectEditorFunctionLibrary.cpp
@@ -4,10 +4,12 @@
 
 #include "AssetToolsModule.h"
 #include "ContentBrowserModule.h"
+#include "CustomizableObjectEditorModule.h"
 #include "IContentBrowserSingleton.h"
 #include "AssetRegistry/AssetRegistryModule.h"
 #include "MuCOE/CustomizableObjectCompiler.h"
 #include "MuCO/CustomizableObject.h"
+#include "MuCO/CustomizableObjectSystemPrivate.h"
 #include "MuCO/ICustomizableObjectEditorModule.h"
 #include "MuCOE/CustomizableObjectFactory.h"
 #include "UObject/SavePackage.h"
@@ -27,29 +29,13 @@ ECustomizableObjectCompilationState UCustomizableObjectEditorFunctionLibrary::Co
 
 	const double StartTime = FPlatformTime::Seconds();
 
-	TSharedRef<FCompilationRequest> CompileRequest = MakeShared<FCompilationRequest>(*CustomizableObject, false);
-	FCompilationOptions& Options = CompileRequest->GetCompileOptions();
-
-	switch (InOptimizationLevel)
-	{
-	case ECustomizableObjectOptimizationLevel::None:
-PRAGMA_DISABLE_DEPRECATION_WARNINGS
-	case ECustomizableObjectOptimizationLevel::Minimal:
-PRAGMA_ENABLE_DEPRECATION_WARNINGS
-		Options.OptimizationLevel = 0;
-		break;
-		
-	case ECustomizableObjectOptimizationLevel::Maximum:
-		Options.OptimizationLevel = UE_MUTABLE_MAX_OPTIMIZATION;
-		break;
-
-	default:
-		unimplemented();
-	}
+	TSharedRef<FCompilationRequest> CompileRequest = MakeShared<FCompilationRequest>(*CustomizableObject);
+	CompileRequest->bAsync = false;
+	CompileRequest->bSilentCompilation = false;
+	CompileRequest->Options.OptimizationLevel = ConvertOptimizationLevel(InOptimizationLevel); 
+	CompileRequest->Options.TextureCompression = InTextureCompression;
+	CompileRequest->Options.bGatherReferences = bGatherReferences;
 	
-	Options.TextureCompression = InTextureCompression;
-	Options.bSilentCompilation = false;
-	Options.bGatherReferences = bGatherReferences;
 	ICustomizableObjectEditorModule::GetChecked().CompileCustomizableObject(CompileRequest);
 
 	check(CompileRequest->GetCompilationState() == ECompilationStatePrivate::Completed);
diff --git a/Engine/Plugins/Mutable/Source/CustomizableObjectEditor/Private/MuCOE/CustomizableObjectEditorModule.cpp b/Engine/Plugins/Mutable/Source/CustomizableObjectEditor/Private/MuCOE/CustomizableObjectEditorModule.cpp
--- a/Engine/Plugins/Mutable/Source/CustomizableObjectEditor/Private/MuCOE/CustomizableObjectEditorModule.cpp
+++ b/Engine/Plugins/Mutable/Source/CustomizableObjectEditor/Private/MuCOE/CustomizableObjectEditorModule.cpp
@@ -81,6 +81,7 @@
 #include "MuCOE/GraphTraversal.h"
 #include "MuCOE/CustomizableObjectInstanceBaker.h"
 #include "Editor.h"
+#include "MuCO/CustomizableObjectInstancePrivate.h"
 #include "MuCO/CustomizableObjectSystemPrivate.h"
 #include "MuCOE/CustomizableObjectGraph.h"
 #include "Nodes/CustomizableObjectNodeComponentMeshDetails.h"
@@ -104,101 +105,9 @@ const FName CustomizableObjectDebuggerAppIdentifier = FName(TEXT("CustomizableOb
 /** Max timespan in days before a Saved/MutableStreamedDataEditor file is deleted. */
 constexpr int32 MaxAccessTimespan = 30;
 
-constexpr float ShowOnScreenCompileWarningsTickerTime = 1.0f;
 
+IMPLEMENT_MODULE( FCustomizableObjectEditorModule, CustomizableObjectEditor );
 
-void ShowOnScreenCompileWarnings()
-{
-	TSet<const UCustomizableObject*> Objects;
-	
-	for (TObjectIterator<UCustomizableObjectInstanceUsage> CustomizableObjectInstanceUsage; CustomizableObjectInstanceUsage; ++CustomizableObjectInstanceUsage)
-	{
-		if (!IsValid(*CustomizableObjectInstanceUsage) || CustomizableObjectInstanceUsage->IsTemplate())
-		{
-			continue;
-		}
-		
-		const UCustomizableObjectInstance* Instance = CustomizableObjectInstanceUsage->GetCustomizableObjectInstance();
-		if (!Instance)
-		{
-			continue;
-		}
-
-		const UCustomizableObject* Object = Cast<UCustomizableObject>(Instance->GetCustomizableObject());
-		if (!Object)
-		{
-			continue;
-		}
-		
-		const USkeletalMeshComponent* Parent = Cast<USkeletalMeshComponent>(CustomizableObjectInstanceUsage->GetAttachParent());
-		if (!Parent)
-		{
-			continue;
-		}
-
-		const UWorld* World = Parent->GetWorld();
-		if (!World)
-		{
-			continue;
-		}
-
-		if (World->WorldType != EWorldType::PIE)
-		{
-			continue;
-		}
-
-		Objects.Add(Object);
-	}
-
-	for (const UCustomizableObject* Object : Objects)
-	{
-		if (Object->GetPrivate()->Status.Get() != FCustomizableObjectStatus::EState::ModelLoaded)
-		{
-			continue;
-		}
-	
-		const UModelResources* ModelResources = Object->GetPrivate()->GetModelResources();
-		if (!ModelResources)
-		{
-			continue;
-		}
-
-		// Show a warning if the compilation was not done with optimizations.
-		const uint64 KeyCompiledWithOptimization = reinterpret_cast<uint64>(Object);
-
-		if (!ModelResources->bIsCompiledWithOptimization)
-		{
-			FString Msg = FString::Printf(TEXT("Customizable Object [%s] was compiled without optimization."), *Object->GetName());
-			GEngine->AddOnScreenDebugMessage(KeyCompiledWithOptimization, ShowOnScreenCompileWarningsTickerTime * 2.0f, FColor::Yellow, Msg);
-		}
-		else
-		{
-			GEngine->RemoveOnScreenDebugMessage(KeyCompiledWithOptimization);
-		}
-		
-		const uint64 KeyCompiledOutOfDate = reinterpret_cast<uint64>(Object) + KEY_OFFSET_COMPILATION_OUT_OF_DATE; // Offset added to avoid collision with bIsCompiledWithOptimization warning
-		TArray<FName> OutOfDatePackages;
-		TArray<FName> AddedPackages;
-		TArray<FName> RemovedPackages;
-		bool bReleaseVersion;
-		if (Object->GetPrivate()->IsCompilationOutOfDate(true, OutOfDatePackages, AddedPackages, RemovedPackages, bReleaseVersion))
-		{
-			FString Msg = FString::Printf(TEXT("Customizable Object [%s] compilation out of date. See the Output Log for more information."), *Object->GetName());
-			GEngine->AddOnScreenDebugMessage(KeyCompiledOutOfDate, ShowOnScreenCompileWarningsTickerTime * 2.0f, FColor::Yellow, Msg);
-			
-			if (!GEngine->OnScreenDebugMessageExists(KeyCompiledOutOfDate))
-			{
-				UE_LOG(LogMutable, Display, TEXT("Customizable Object [%s] compilation out of date. Changes since last compilation:"), *Object->GetName());
-
-				PrintParticipatingPackagesDiff(OutOfDatePackages, AddedPackages, RemovedPackages, bReleaseVersion);
-			}
-		}
-		else
-		{
-			GEngine->RemoveOnScreenDebugMessage(KeyCompiledOutOfDate);
-		}
-	}
-}
 
 void DeleteUnusedMutableStreamedDataEditorFiles()
 {
@@ -229,10 +138,6 @@ void DeleteUnusedMutableStreamedDataEditorFiles()
 	}
 }
 
-
-IMPLEMENT_MODULE( FCustomizableObjectEditorModule, CustomizableObjectEditor );
-
-
 static void* CustomMalloc(std::size_t Size_t, uint32_t Alignment)
 {
 	return FMemory::Malloc(Size_t, Alignment);
@@ -245,6 +150,42 @@ static void CustomFree(void* mem)
 }
 
 
+int32 ConvertOptimizationLevel(ECustomizableObjectOptimizationLevel OptimizationLevel)
+{
+	switch (OptimizationLevel)
+	{
+	case ECustomizableObjectOptimizationLevel::None:
+		PRAGMA_DISABLE_DEPRECATION_WARNINGS
+			case ECustomizableObjectOptimizationLevel::Minimal:
+				PRAGMA_ENABLE_DEPRECATION_WARNINGS
+						return 0;
+
+	case ECustomizableObjectOptimizationLevel::Maximum:
+		return UE_MUTABLE_MAX_OPTIMIZATION;
+	
+	default:
+		unimplemented();
+		return 0;
+	}
+}
+
+
+TMap<FString, FString> GetCompileOnlySelectedParameters(const UCustomizableObjectInstance& Instance)
+{
+	TMap<FString, FString> Parameters;
+	
+	const TArray<FCustomizableObjectIntParameterValue>& IntParameters = Instance.GetPrivate()->GetDescriptor().GetIntParameters();
+	Parameters.Reserve(IntParameters.Num());
+
+	for (const FCustomizableObjectIntParameterValue& IntParam : IntParameters)
+	{
+		Parameters.Add(IntParam.ParameterName, IntParam.ParameterValueName);
+	}
+
+	return Parameters;
+}
+
+
 void FCustomizableObjectEditorModule::StartupModule()
 {
 	// Delete unused local compiled data
@@ -317,13 +258,7 @@ void FCustomizableObjectEditorModule::StartupModule()
 		TEXT("mutable.OpenCOIE"),
 		TEXT("Looks for a Customizable Object Instance within the player pawn and opens its Customizable Object Instance Editor. Specify slot ID to control which component is edited."),
 		FConsoleCommandWithArgsDelegate::CreateStatic(&FCustomizableObjectEditorModule::OpenCOIE));
-
-	WarningsTickerHandle = FTSTicker::GetCoreTicker().AddTicker(TEXT("ShowOnScreenCompileWarnings"), ShowOnScreenCompileWarningsTickerTime, [](float)
-	{
-		ShowOnScreenCompileWarnings();
-		return true;
-	});
-
+	
 	FEditorDelegates::PreBeginPIE.AddRaw(this, &FCustomizableObjectEditorModule::OnPreBeginPIE);
 }
 
@@ -359,8 +294,6 @@ void FCustomizableObjectEditorModule::ShutdownModule()
 	FCoreDelegates::OnPostEngineInit.RemoveAll(this);
 
 	FCustomizableObjectEditorStyle::Shutdown();
-
-	FTSTicker::GetCoreTicker().RemoveTicker(WarningsTickerHandle);
 }
 
 
@@ -792,49 +725,39 @@ void FCustomizableObjectEditorModule::PostBackwardsCompatibleFixup(UEdGraph& Gra
 
 void FCustomizableObjectEditorModule::CompileCustomizableObject(const TSharedRef<FCompilationRequest>& InCompilationRequest, bool bForceRequest)
 {
-	if (IsRunningGame())
-	{
-		return;
-	}
-
-	CompileCustomizableObjects({ InCompilationRequest }, bForceRequest);
+	Compiler.EnqueueCompileRequest(InCompilationRequest, bForceRequest);
 }
 
-void FCustomizableObjectEditorModule::CompileCustomizableObjects(const TArray<TSharedRef<FCompilationRequest>>& InCompilationRequests, bool bForceRequests)
-{
-	check(IsInGameThread());
 
-	if (IsRunningGame())
+void FCustomizableObjectEditorModule::CompileCustomizableObject(UCustomizableObject& Object, const FCompileParams& Params)
+{
+	UCustomizableObjectSystem* System = UCustomizableObjectSystem::GetInstance();
+	if (!System)
 	{
+		FCompileCallbackParams CallbackParams;
+		CallbackParams.bRequestFailed = true;
+		CallbackParams.bCompiled = Object.IsCompiled();
+		
+		Params.Callback.ExecuteIfBound(CallbackParams);
 		return;
 	}
 
-	TArray<TSharedRef<FCompilationRequest>> FilteredAsyncRequests;
-	FilteredAsyncRequests.Reserve(InCompilationRequests.Num());
-
-	for (const TSharedRef<FCompilationRequest>& Request : InCompilationRequests)
+	TSharedRef<FCompilationRequest> CompileRequest = MakeShared<FCompilationRequest>(Object);
+	CompileRequest->bAsync = Params.bAsync;
+	CompileRequest->bSkipIfCompiled = Params.bSkipIfCompiled;
+	CompileRequest->bSkipIfOutOfDate = Params.bSkipIfOutOfDate;
+	CompileRequest->bSilentCompilation = true;
+	CompileRequest->Options.TextureCompression = Params.TextureCompression;
+	CompileRequest->Options.OptimizationLevel = ConvertOptimizationLevel(Params.OptimizationLevel);
+	CompileRequest->Options.bGatherReferences = Params.bGatherReferneces;
+	CompileRequest->Callback = Params.Callback;
+	
+	if (Params.CompileOnlySelectedInstance)
 	{
-		const UCustomizableObject* CustomizableObject = Request->GetCustomizableObject();
-		if (!CustomizableObject)
-		{
-			continue;
-		}
-		
-		if (!Request->IsAsyncCompilation())
-		{
-			FCustomizableObjectCompiler* SyncCompiler = new FCustomizableObjectCompiler();
-			SyncCompiler->Compile(Request);
-			delete SyncCompiler;
-		}
-
-		else if (bForceRequests ||
-			(!CustomizableObject->GetPrivate()->IsLocked() && !Compiler.IsRequestQueued(Request)))
-		{
-			FilteredAsyncRequests.Add(Request);
-		}
+		CompileRequest->Options.ParamNamesToSelectedOptions = GetCompileOnlySelectedParameters(*Params.CompileOnlySelectedInstance);
 	}
-
-	Compiler.Compile(FilteredAsyncRequests);
+	
+	CompileCustomizableObject(CompileRequest);
 }
 
 
@@ -858,6 +781,12 @@ int32 FCustomizableObjectEditorModule::GetNumCompileRequests()
 }
 
 
+bool FCustomizableObjectEditorModule::IsCompiling(const UCustomizableObject& Object) const
+{
+	return Object.GetPrivate()->IsLocked() || Compiler.IsRequestQueued(Object);
+}
+
+
 void FCustomizableObjectEditorModule::OnPreBeginPIE(const bool bIsSimulatingInEditor)
 {
 	if (IsRunningGame() || !UCustomizableObjectSystem::IsActive())
@@ -896,8 +825,9 @@ void FCustomizableObjectEditorModule::OnPreBeginPIE(const bool bIsSimulatingInEd
 		}
 
 		// Add uncompiled objects to the objects to cook list
-		TSharedRef<FCompilationRequest> NewRequest = MakeShared<FCompilationRequest>(*Object, true);
-		NewRequest->GetCompileOptions().bSilentCompilation = true;
+		TSharedRef<FCompilationRequest> NewRequest = MakeShared<FCompilationRequest>(*Object);
+		NewRequest->bSilentCompilation = true;
+
 		Requests.Add(NewRequest);
 	}
 
@@ -906,7 +836,10 @@ void FCustomizableObjectEditorModule::OnPreBeginPIE(const bool bIsSimulatingInEd
 		const FText Msg = FText::FromString(TEXT("Warning: one or more Customizable Objects used in PIE are uncompiled.\n\nDo you want to compile them?"));
 		if (FMessageDialog::Open(EAppMsgType::OkCancel, Msg) == EAppReturnType::Ok)
 		{
-			CompileCustomizableObjects(Requests);
+			for (const TSharedRef<FCompilationRequest>& Request : Requests)
+			{
+				CompileCustomizableObject(Request);
+			}
 		}
 	}
 }
diff --git a/Engine/Plugins/Mutable/Source/CustomizableObjectEditor/Private/MuCOE/CustomizableObjectEditorModule.h b/Engine/Plugins/Mutable/Source/CustomizableObjectEditor/Private/MuCOE/CustomizableObjectEditorModule.h
--- a/Engine/Plugins/Mutable/Source/CustomizableObjectEditor/Private/MuCOE/CustomizableObjectEditorModule.h
+++ b/Engine/Plugins/Mutable/Source/CustomizableObjectEditor/Private/MuCOE/CustomizableObjectEditorModule.h
@@ -24,6 +24,11 @@ class FName;
   * - UDataTable: Data Tables used by Composite Data Tables are indirectly referenced by the UStruct and filtered by path. */
 void GetReferencingPackages(const UCustomizableObject& Object, TArray<FAssetData>& ObjectNames);
 
+/** Convert the Optimization Level enum to the int32 internal format.  */
+int32 ConvertOptimizationLevel(ECustomizableObjectOptimizationLevel OptimizationLevel);
+
+TMap<FString, FString> GetCompileOnlySelectedParameters(const UCustomizableObjectInstance& Instance);
+
 
 /**
  * StaticMesh editor module
@@ -48,19 +53,14 @@ public:
 	virtual TMap<FName, FGuid> GetParticipatingObjects(const UCustomizableObject* Object, bool bLoadObjects, const FCompilationOptions* Options = nullptr) const override;
 	virtual void BackwardsCompatibleFixup(UEdGraph& Graph, int32 CustomizableObjectCustomVersion) override;
 	virtual void PostBackwardsCompatibleFixup(UEdGraph& Graph) override;
-
-	/** Request for a given customizable object to be compiled. Async compile requests will be queued and processed sequentially. 
-	 * @param InCompilationRequest - Request to compile an object. 
-	 * @param bForceRequest - Queue request even if already in the pending list. */
+	virtual void CancelCompileRequests() override;
+	virtual int32 GetNumCompileRequests() override;
+	virtual bool IsCompiling(const UCustomizableObject& Object) const override;
 	virtual void CompileCustomizableObject(const TSharedRef<FCompilationRequest>& InCompilationRequest, bool bForceRequest = false) override;
-	virtual void CompileCustomizableObjects(const TArray<TSharedRef<FCompilationRequest>>& InCompilationRequests, bool bForceRequests = false) override;
+	virtual void CompileCustomizableObject(UCustomizableObject& Object, const FCompileParams& Params) override;
 
 	virtual int32 Tick(bool bBlocking) override;
 
-	virtual void CancelCompileRequests() override;
-
-	virtual int32 GetNumCompileRequests() override;
-
 	virtual TSharedPtr<FExtensibilityManager> GetCustomizableObjectEditorToolBarExtensibilityManager() override { return CustomizableObjectEditor_ToolBarExtensibilityManager; }
 	virtual TSharedPtr<FExtensibilityManager> GetCustomizableObjectEditorMenuExtensibilityManager() override { return CustomizableObjectEditor_MenuExtensibilityManager; }
 
@@ -80,9 +80,7 @@ private:
 
 	// Command to look for Customizable Object Instance in the player pawn of the current world and open its Customizable Object Instance Editor
 	IConsoleCommand* LaunchCOIECommand = nullptr;
-
-	FTSTicker::FDelegateHandle WarningsTickerHandle;
-
+	
 	static void OpenCOIE(const TArray<FString>& Arguments);
 
 	// Used to ask the user if they want to recompile uncompiled PIE COs
diff --git a/Engine/Plugins/Mutable/Source/CustomizableObjectEditor/Private/MuCOE/SCustomizableObjectEditorViewport.cpp b/Engine/Plugins/Mutable/Source/CustomizableObjectEditor/Private/MuCOE/SCustomizableObjectEditorViewport.cpp
--- a/Engine/Plugins/Mutable/Source/CustomizableObjectEditor/Private/MuCOE/SCustomizableObjectEditorViewport.cpp
+++ b/Engine/Plugins/Mutable/Source/CustomizableObjectEditor/Private/MuCOE/SCustomizableObjectEditorViewport.cpp
@@ -28,6 +28,7 @@
 #include "MuCO/CustomizableObjectPrivate.h"
 #include "MuCO/CustomizableObjectSystem.h"
 #include "MuCO/LoadUtils.h"
+#include "MuCO/ICustomizableObjectEditorModule.h"
 #include "Widgets/Input/STextComboBox.h"
 
 class UCustomizableObject;
@@ -179,8 +180,9 @@ FText SCustomizableObjectEditorViewport::GetCompileErrorOverlayText() const
 	{
 		return LOCTEXT("Loading", "Loading...");
 	}
-	
-	if (Object->GetPrivate()->CompilationState == ECompilationStatePrivate::InProgress)
+
+	ICustomizableObjectEditorModule& Module = ICustomizableObjectEditorModule::GetChecked();
+	if (Module.IsCompiling(*Object))
 	{
 		return LOCTEXT("Compiling", "Compiling...");
 	}
diff --git a/Engine/Plugins/Mutable/Source/CustomizableObjectEditor/Public/MuCOE/CustomizableObjectEditorFunctionLibrary.h b/Engine/Plugins/Mutable/Source/CustomizableObjectEditor/Public/MuCOE/CustomizableObjectEditorFunctionLibrary.h
--- a/Engine/Plugins/Mutable/Source/CustomizableObjectEditor/Public/MuCOE/CustomizableObjectEditorFunctionLibrary.h
+++ b/Engine/Plugins/Mutable/Source/CustomizableObjectEditor/Public/MuCOE/CustomizableObjectEditorFunctionLibrary.h
@@ -7,14 +7,6 @@
 #include "CustomizableObjectEditorFunctionLibrary.generated.h"
 
 // This mirrors the logic in CustomizableObjectEditor.cpp
-UENUM(BlueprintType)
-enum class ECustomizableObjectOptimizationLevel : uint8
-{
-	None = 0, 
-	Minimal UE_DEPRECATED(5.6, "Converted internally to None. Use None instead.") = 1,
-	Maximum = 2
-};
-
 
 UENUM()
 enum class ECustomizableObjectCompilationState : uint8
@@ -59,7 +51,8 @@ class CUSTOMIZABLEOBJECTEDITOR_API UCustomizableObjectEditorFunctionLibrary : pu
 	GENERATED_BODY()
 
 public:
-	/** 
+	/** DEPRECATED. Use UCustomizableObject::Compile instead.
+	 * 
 	 *	Synchronously compiles the provided CustomizableObject, LogMutable will contain intermittent updates on
 	 *	progress.
 	 * 
diff --git a/Engine/Plugins/Mutable/Source/MutableValidation/Private/MuV/AssetValidator_CustomizableObjects.cpp b/Engine/Plugins/Mutable/Source/MutableValidation/Private/MuV/AssetValidator_CustomizableObjects.cpp
--- a/Engine/Plugins/Mutable/Source/MutableValidation/Private/MuV/AssetValidator_CustomizableObjects.cpp
+++ b/Engine/Plugins/Mutable/Source/MutableValidation/Private/MuV/AssetValidator_CustomizableObjects.cpp
@@ -145,13 +145,14 @@ EDataValidationResult UAssetValidator_CustomizableObjects::IsCustomizableObjectV
 	for	(const FCompilationOptions& Options : CompilationOptionsToTest)
 	{
 		// Run Sync compilation -> Warning : Potentially long operation -------------
-		TSharedRef<FCompilationRequest> CompileRequest = MakeShared<FCompilationRequest>(*RootObject, false);
+		TSharedRef<FCompilationRequest> CompileRequest = MakeShared<FCompilationRequest>(*RootObject);
+		CompileRequest->bAsync = false;
 		ICustomizableObjectEditorModule::GetChecked().CompileCustomizableObject(CompileRequest);
 		// --------------------------------------------------------------------------
 		
 		// Get compilation errors and warnings
-		const TArray<FText>& CompilationErrors = CompileRequest->GetErrors();
-		const TArray<FText>& CompilationWarnings = CompileRequest->GetWarnings();
+		const TArray<FText>& CompilationErrors = CompileRequest->Errors;
+		const TArray<FText>& CompilationWarnings = CompileRequest->Warnings;
 		
 		// Cache the messages returned by the compiler
 		for (const FText& FoundError : CompilationErrors)
diff --git a/Engine/Plugins/Mutable/Source/MutableValidation/Private/MuV/CustomizableObjectCompilationUtility.cpp b/Engine/Plugins/Mutable/Source/MutableValidation/Private/MuV/CustomizableObjectCompilationUtility.cpp
--- a/Engine/Plugins/Mutable/Source/MutableValidation/Private/MuV/CustomizableObjectCompilationUtility.cpp
+++ b/Engine/Plugins/Mutable/Source/MutableValidation/Private/MuV/CustomizableObjectCompilationUtility.cpp
@@ -89,8 +89,11 @@ bool FCustomizableObjectCompilationUtility::CompileCustomizableObject(UCustomiza
 		mu::FGlobalMemoryCounter::Zero();
 		
 		UE_LOG(LogMutable,Display,TEXT("Compiling Customizable Object..."));
-		TSharedRef<FCompilationRequest> CompileRequest = MakeShared<FCompilationRequest>(*InCustomizableObject, true);
-		CompileRequest->GetCompileOptions() = CompilationOptions;
+		
+		TSharedRef<FCompilationRequest> CompileRequest = MakeShared<FCompilationRequest>(*InCustomizableObject);
+		CompileRequest->bSilentCompilation = false;
+		CompileRequest->Options = CompilationOptions;
+		
 		ICustomizableObjectEditorModule::GetChecked().CompileCustomizableObject(CompileRequest);
 		
 		// Wait while the compilation takes place
diff --git a/Engine/Plugins/Mutable/Source/MutableValidation/Private/MuV/ValidationUtils.cpp b/Engine/Plugins/Mutable/Source/MutableValidation/Private/MuV/ValidationUtils.cpp
--- a/Engine/Plugins/Mutable/Source/MutableValidation/Private/MuV/ValidationUtils.cpp
+++ b/Engine/Plugins/Mutable/Source/MutableValidation/Private/MuV/ValidationUtils.cpp
@@ -65,7 +65,6 @@ FCompilationOptions GetCompilationOptionsForBenchmarking (UCustomizableObject& R
 {
 	// Override some configurations that may have been changed by the user
 	FCompilationOptions CISCompilationOptions = ReferenceCustomizableObject.GetPrivate()->GetCompileOptions();
-	CISCompilationOptions.bSilentCompilation = false;										// totally optional
 	CISCompilationOptions.OptimizationLevel = CustomizableObjectBenchmarkingUtils::GetOptimizationLevelForBenchmarking();
 	CISCompilationOptions.TextureCompression = ECustomizableObjectTextureCompression::Fast;	// Does not affect instance update speed but does compilation
 
