[mutable] Cleanup code generator local state. Moved it to struct passed up and down in the graph recursion.
#rnx

diff --git a/Engine/Plugins/Mutable/Source/MutableTools/Private/MuT/CodeGenerator.cpp b/Engine/Plugins/Mutable/Source/MutableTools/Private/MuT/CodeGenerator.cpp
--- a/Engine/Plugins/Mutable/Source/MutableTools/Private/MuT/CodeGenerator.cpp
+++ b/Engine/Plugins/Mutable/Source/MutableTools/Private/MuT/CodeGenerator.cpp
@@ -105,29 +105,26 @@ namespace mu
 	//---------------------------------------------------------------------------------------------
 	//---------------------------------------------------------------------------------------------
 	//---------------------------------------------------------------------------------------------
-	CodeGenerator::CodeGenerator(CompilerOptions::Private* options)
+	CodeGenerator::CodeGenerator(CompilerOptions::Private* Options)
 		: GenerateMeshConstantPipe(TEXT("GenerateMeshConstantPipe"))
 	{
-		CompilerOptions = options;
+		CompilerOptions = Options;
 
 		// Create the message log
 		ErrorLog = MakeShared<FErrorLog>();
-
-		// Add the parent at the top of the hierarchy
-		CurrentParents.Add(FParentKey());
 	}
 
 	//---------------------------------------------------------------------------------------------
-	void CodeGenerator::GenerateRoot(const Ptr<const Node> pNode)
+	void CodeGenerator::GenerateRoot(const Ptr<const Node> InNode)
 	{
 		MUTABLE_CPUPROFILER_SCOPE(Generate);
 
 		// First pass
-		FirstPass.Generate(ErrorLog, pNode.get(), CompilerOptions->bIgnoreStates, this);
+		FirstPass.Generate(ErrorLog, InNode.get(), CompilerOptions->bIgnoreStates, this);
 
 		// Second pass
 		SecondPassGenerator SecondPass(&FirstPass, CompilerOptions);
-		bool bSuccess = SecondPass.Generate(ErrorLog, pNode.get());
+		bool bSuccess = SecondPass.Generate(ErrorLog, InNode.get());
 		if (!bSuccess)
 		{
 			return;
@@ -138,17 +135,18 @@ namespace mu
 			MUTABLE_CPUPROFILER_SCOPE(MainPass);
 
             int32 CurrentStateIndex = 0;
-			for (const TPair<FObjectState, const Node*>& s : FirstPass.States)
+			for (const FObjectState& State : FirstPass.States)
 			{
 				MUTABLE_CPUPROFILER_SCOPE(MainPassState);
 
 				FGenericGenerationOptions Options;
 				Options.State = CurrentStateIndex;
 
-				Ptr<ASTOp> stateRoot = Generate_Generic(pNode, Options);
-				States.Emplace(s.Key, stateRoot);
+				FGenericGenerationResult Result;
+				GenerateGeneric(Options, Result, InNode.get());
 
-				AdditionalComponents.Empty();
+				Ptr<ASTOp> StateRoot = Result.Op;
+				States.Emplace(State, StateRoot);
 
 				++CurrentStateIndex;
 			}
@@ -156,120 +154,137 @@ namespace mu
 	}
 
 
-	Ptr<ASTOp> CodeGenerator::Generate_Generic(const Ptr<const Node> pNode, const FGenericGenerationOptions& Options )
+void CodeGenerator::GenerateGeneric(const FGenericGenerationOptions& Options, FGenericGenerationResult& OutResult, const Ptr<const Node> InNode )
 	{
-		if (!pNode)
+		if (!InNode)
 		{
-			return nullptr;
+			return;
 		}
 
         // Type-specific generation
-		if (pNode->GetType()->IsA(NodeScalar::GetStaticType()))
+		if (InNode->GetType()->IsA(NodeObject::GetStaticType()))
+		{
+			const NodeObject* ObjectNode = static_cast<const NodeObject*>(InNode.get());
+			FObjectGenerationOptions ObjectOptions;
+			ObjectOptions.ActiveTags = Options.ActiveTags;
+			ObjectOptions.bIsImage = Options.bIsImage;
+			ObjectOptions.State = Options.State;
+			FObjectGenerationResult ObjectResult;
+			GenerateObject(ObjectOptions, ObjectResult, ObjectNode);
+			OutResult.Op = ObjectResult.Op;
+			return;
+		}
+
+		else if (InNode->GetType()->IsA(NodeScalar::GetStaticType()))
 		{
-			const NodeScalar* ScalarNode = static_cast<const NodeScalar*>(pNode.get());
+			const NodeScalar* ScalarNode = static_cast<const NodeScalar*>(InNode.get());
 			FScalarGenerationResult ScalarResult;
 			GenerateScalar(ScalarResult, Options, ScalarNode);
-			return ScalarResult.op;
+			OutResult.Op = ScalarResult.op;
+			return;
 		}
 
-		else if (pNode->GetType()->IsA(NodeColour::GetStaticType()))
+		else if (InNode->GetType()->IsA(NodeColour::GetStaticType()))
 		{
-			const NodeColour* ColorNode = static_cast<const NodeColour*>(pNode.get());
+			const NodeColour* ColorNode = static_cast<const NodeColour*>(InNode.get());
 			FColorGenerationResult Result;
 			GenerateColor(Result, Options, ColorNode);
-			return Result.op;
+			OutResult.Op = Result.op;
+			return;
 		}
 
-		else if (pNode->GetType()->IsA(NodeProjector::GetStaticType()))
+		else if (InNode->GetType()->IsA(NodeProjector::GetStaticType()))
 		{
-			const NodeProjector* projNode = static_cast<const NodeProjector*>(pNode.get());
+			const NodeProjector* projNode = static_cast<const NodeProjector*>(InNode.get());
 			FProjectorGenerationResult ProjResult;
 			GenerateProjector(ProjResult, Options, projNode);
-			return ProjResult.op;
+			OutResult.Op = ProjResult.op;
+			return;
 		}
 
-		else if (pNode->GetType()->IsA(NodeSurfaceNew::GetStaticType()))
+		else if (InNode->GetType()->IsA(NodeSurfaceNew::GetStaticType()))
 		{
-			const NodeSurfaceNew* surfNode = static_cast<const NodeSurfaceNew*>(pNode.get());
+			const NodeSurfaceNew* surfNode = static_cast<const NodeSurfaceNew*>(InNode.get());
 
 			// This happens only if we generate a node graph that has a NodeSurfaceNew at the root.
-			FSurfaceGenerationResult surfResult;
+			FSurfaceGenerationResult SurfResult;
 			FSurfaceGenerationOptions SurfaceOptions(Options);
-			GenerateSurface(surfResult, SurfaceOptions, surfNode);
-			return surfResult.surfaceOp;
+			GenerateSurface(SurfResult, SurfaceOptions, surfNode);
+			OutResult.Op = SurfResult.surfaceOp;
+			return;
 		}
 
-		else if (pNode->GetType()->IsA(NodeSurfaceVariation::GetStaticType()))
+		else if (InNode->GetType()->IsA(NodeSurfaceVariation::GetStaticType()))
 		{
 			// This happens only if we generate a node graph that has a NodeSurfaceVariation at the root.
-			return nullptr;
+			return;
 		}
 
-		else if (pNode->GetType()->IsA(NodeSurfaceSwitch::GetStaticType()))
+		else if (InNode->GetType()->IsA(NodeSurfaceSwitch::GetStaticType()))
 		{
 			// This happens only if we generate a node graph that has a NodeSurfaceSwitch at the root.
-			return nullptr;
+			return;
 		}
 
-		else if (pNode->GetType()->IsA(NodeModifier::GetStaticType()))
+		else if (InNode->GetType()->IsA(NodeModifier::GetStaticType()))
 		{
 			// This happens only if we generate a node graph that has a modifier at the root.
-			return nullptr;
+			return;
 		}
 
-		else if (pNode->GetType()->IsA(NodeComponent::GetStaticType()))
+		else if (InNode->GetType()->IsA(NodeComponent::GetStaticType()))
 		{
-			const NodeComponent* ComponentNode = static_cast<const NodeComponent*>(pNode.get());
-			FComponentGenerationOptions ComponentOptions( Options, nullptr );
-			FGenericGenerationResult Result;
-			GenerateComponent(ComponentOptions, Result, ComponentNode);
-			return Result.op;
+			const NodeComponent* ComponentNode = static_cast<const NodeComponent*>(InNode.get());
+			FComponentGenerationOptions ComponentOptions(Options, nullptr);
+			GenerateComponent(ComponentOptions, OutResult, ComponentNode);
+			return;
+		}
+
+		else
+		{
+			// Unsupported node.
+			check(false);
 		}
 
+	}
 
-		Ptr<ASTOp> ResultOp;
+
+	void CodeGenerator::GenerateObject(const FObjectGenerationOptions& InOptions, FObjectGenerationResult& OutResult, const NodeObject* InUntypedNode)
+	{
+		if (!InUntypedNode)
+		{
+			OutResult = {};
+			return;
+		}
 
 		// See if it was already generated
+		FGeneratedObjectCacheKey Key;
+		Key.Node = InUntypedNode;
+		Key.Options = InOptions;
+		FGeneratedObjectsMap::ValueType* Found = GeneratedObjects.Find(Key);
+		if (Found)
+		{
+			OutResult = *Found;
+			return;
+		}
 
-		FGeneratedCacheKey Key;
-		Key.Node = pNode;
-		Key.Options = Options;
-		FGeneratedGenericNodesMap::ValueType* it = GeneratedGenericNodes.Find(Key);
-		if (it)
+		// Generate for each different type of node
+		const FNodeType* Type = InUntypedNode->GetType();
+		if (Type == NodeObjectNew::GetStaticType())
+		{
+			GenerateObject_New(InOptions, OutResult, static_cast<const NodeObjectNew*>(InUntypedNode));
+		}
+		else if (Type == NodeObjectGroup::GetStaticType())
 		{
-			ResultOp = it->op;
+			GenerateObject_Group(InOptions, OutResult, static_cast<const NodeObjectGroup*>(InUntypedNode));
 		}
 		else
 		{
-			FGenericGenerationResult Result;
-
-			// Generate for each different type of node
-			if (pNode->GetType() == NodeObjectNew::GetStaticType())
-			{
-				Generate_ObjectNew(Options, Result, static_cast<const NodeObjectNew*>(pNode.get()));
-			}
-			else if (pNode->GetType()==NodeObjectGroup::GetStaticType())
-			{
-				Generate_ObjectGroup(Options, Result, static_cast<const NodeObjectGroup*>(pNode.get()));
-			}
-			else
-			{
-				check(false);
-			}
-
-			ResultOp = Result.op;
-			GeneratedGenericNodes.Add(Key, Result);
+			check(false);
 		}
 
-		// debug: expensive check of all code generation
-//        if (ResultOp)
-//        {
-//            ASTOpList roots;
-//            roots.push_back(ResultOp);
-//            ASTOp::FullAssert(roots);
-//        }
-
-		return ResultOp;
+		// Cache the result
+		GeneratedObjects.Add(Key, OutResult);
 	}
 
 
@@ -526,7 +541,6 @@ namespace mu
 	}
 
 
-
 	//---------------------------------------------------------------------------------------------
 	void CodeGenerator::GenerateComponent(const FComponentGenerationOptions& InOptions, FGenericGenerationResult& OutResult, const NodeComponent* InUntypedNode)
 	{
@@ -540,7 +554,7 @@ namespace mu
 		FGeneratedComponentCacheKey Key;
 		Key.Node = InUntypedNode;
 		Key.Options = InOptions;
-		GeneratedComponentMap::ValueType* it = GeneratedComponents.Find(Key);
+		FGeneratedComponentMap::ValueType* it = GeneratedComponents.Find(Key);
 		if (it)
 		{
 			OutResult = *it;
@@ -557,7 +571,7 @@ namespace mu
 		{
 			// Nothing to do because it is all preprocessed in the first code generator stage
 			//GenerateComponent_Edit(InOptions, OutResult, static_cast<const NodeComponentEdit*>(InUntypedNode));
-			OutResult.op = InOptions.BaseInstance;
+			OutResult.Op = InOptions.BaseInstance;
 		}
 		else if (Type == NodeComponentSwitch::GetStaticType())
 		{
@@ -588,13 +602,11 @@ namespace mu
 		{
 			if (const NodeLOD* LODNode = node.LODs[LODIndex].get())
 			{
-				CurrentParents.Last().Lod = LODIndex;
-
 				FLODGenerationOptions LODOptions(Options, LODIndex, InNode );
 				FGenericGenerationResult LODResult;
 				Generate_LOD(LODOptions, LODResult, LODNode);
 
-				LODsOp->lods.Emplace(LODsOp, LODResult.op);
+				LODsOp->lods.Emplace(LODsOp, LODResult.Op);
 			}
 		}
 
@@ -604,7 +616,7 @@ namespace mu
 		InstanceOp->value = LODsOp;
 		InstanceOp->ExternalId = node.Id;
 
-		Result.op = InstanceOp;
+		Result.Op = InstanceOp;
 
 		// Add a conditional if this component has conditions
 		for (const FirstPassGenerator::FComponent& Component: FirstPass.Components)
@@ -624,10 +636,10 @@ namespace mu
 				Ptr<ASTOpConditional> IfOp = new ASTOpConditional();
 				IfOp->type = EOpType::IN_CONDITIONAL;
 				IfOp->no = Options.BaseInstance;
-				IfOp->yes = Result.op;
+				IfOp->yes = Result.Op;
 				IfOp->condition = ConditionOp;
 
-				Result.op = IfOp;
+				Result.Op = IfOp;
 			}
 		}
 	}
@@ -640,7 +652,7 @@ namespace mu
 		if (Node->Options.Num() == 0)
 		{
 			// No options in the switch!
-			Result.op = Options.BaseInstance;
+			Result.Op = Options.BaseInstance;
 			return;
 		}
 
@@ -650,7 +662,9 @@ namespace mu
 		// Variable value
 		if (Node->Parameter)
 		{
-			Op->variable = Generate_Generic(Node->Parameter.get(), Options);
+			FScalarGenerationResult ParamResult;
+			GenerateScalar(ParamResult, Options, Node->Parameter.get());
+			Op->variable = ParamResult.op;
 		}
 		else
 		{
@@ -667,7 +681,7 @@ namespace mu
 			{
 				FGenericGenerationResult BaseResult;
 				GenerateComponent(Options, BaseResult, Node->Options[OptionIndex].get());
-				Branch = BaseResult.op;
+				Branch = BaseResult.Op;
 			}
 			else
 			{
@@ -678,7 +692,7 @@ namespace mu
 			Op->cases.Emplace(OptionIndex, Op, Branch);
 		}
 
-		Result.op = Op;
+		Result.Op = Op;
 	}
 
 
@@ -692,7 +706,7 @@ namespace mu
 			FGenericGenerationResult BranchResults;
 
 			GenerateComponent(Options, BranchResults, Node->DefaultComponent.get());
-			CurrentMeshOp = BranchResults.op;
+			CurrentMeshOp = BranchResults.Op;
 		}
 
 		// Process variations in reverse order, since conditionals are built bottom-up.
@@ -725,7 +739,7 @@ namespace mu
 				FGenericGenerationResult BranchResults;
 				GenerateComponent(Options, BranchResults, Node->Variations[VariationIndex].Component.get());
 
-				VariationMeshOp = BranchResults.op;
+				VariationMeshOp = BranchResults.Op;
 			}
 
 			Ptr<ASTOpConditional> Conditional = new ASTOpConditional;
@@ -737,7 +751,7 @@ namespace mu
 			CurrentMeshOp = Conditional;
 		}
 
-		Result.op = CurrentMeshOp;
+		Result.Op = CurrentMeshOp;
 	}
 
 
@@ -1696,53 +1710,41 @@ namespace mu
             lastCompOp = iop;
         }
 
-        Result.op = lastCompOp;
+        Result.Op = lastCompOp;
     }
 
 
     //---------------------------------------------------------------------------------------------
    // Ptr<ASTOp> CodeGenerator::Visit( const NodeObjectNew::Private& node )
-	void CodeGenerator::Generate_ObjectNew(const FGenericGenerationOptions& Options, FGenericGenerationResult& Result, const NodeObjectNew* InNode)
+	void CodeGenerator::GenerateObject_New(const FObjectGenerationOptions& Options, FObjectGenerationResult& OutResult, const NodeObjectNew* InNode)
 	{
 		MUTABLE_CPUPROFILER_SCOPE(NodeObjectNew);
 
 		// There is always at least a null parent
-		bool bIsChildObject = CurrentParents.Num() > 1;
+		bool bIsChildObject = Options.ParentObjectNode != nullptr;
 
 		// Add this object as current parent
-        CurrentParents.Add( FParentKey() );
-        CurrentParents.Last().ObjectNode = InNode;
+		FObjectGenerationOptions ChildOptions = Options;
+		ChildOptions.ParentObjectNode = InNode;
 
         // Parse the child objects first, which will accumulate operations in the patching lists
-        for ( int32 t=0; t< InNode->Children.Num(); ++t )
+        for ( int32 ChildIndex=0; ChildIndex < InNode->Children.Num(); ++ChildIndex)
         {
-            if ( const NodeObject* pChildNode = InNode->Children[t].get() )
+            if ( const NodeObject* pChildNode = InNode->Children[ChildIndex].get() )
             {
-                Ptr<ASTOp> paramOp;
-
-                // If there are parent objects, the condition of this object depends on the
-                // condition of the parent object
-                if ( CurrentObject.Num() )
-                {
-                    paramOp = CurrentObject.Last().Condition;
-                }
-                else
-                {
-                    // In case there is no group node, we generate a constant true condition
-                    // This condition will be overwritten by the group nodes.
-                    Ptr<ASTOpConstantBool> op = new ASTOpConstantBool();
-                    op->bValue = true;
-                    paramOp = op;
-                }
-
-				FObjectGenerationData data;
-                data.Condition = paramOp;
-                CurrentObject.Add( data );
+				// If there are parent objects, the condition of this object depends on the condition of the parent object
+				FObjectGenerationOptions ThisChildOptions = ChildOptions;
+				if (!ThisChildOptions.CurrentObjectCondition)
+				{
+					// In case there is no group node, we generate a constant true condition
+					// This condition will be overwritten by the group nodes.
+					ThisChildOptions.CurrentObjectCondition = new ASTOpConstantBool(true);
+				}
 
                 // This op is ignored: everything is stored as patches to apply to the parent when it is compiled.
-                Generate_Generic( pChildNode, Options );
-
-                CurrentObject.Pop();
+				FObjectGenerationResult ThisResult;
+                GenerateObject(ThisChildOptions, ThisResult, pChildNode);
+				OutResult.AdditionalComponents.Append(ThisResult.AdditionalComponents);
             }
         }
 
@@ -1756,15 +1758,15 @@ namespace mu
 		}
 
 		// Add the components in this node
-        for ( int32 t=0; t< InNode->Components.Num(); ++t )
+        for ( int32 ComponentIndex=0; ComponentIndex < InNode->Components.Num(); ++ComponentIndex)
         {
-			const NodeComponent* ComponentNode = InNode->Components[t].get();
+			const NodeComponent* ComponentNode = InNode->Components[ComponentIndex].get();
             if (ComponentNode)
             {
 				FComponentGenerationOptions ComponentOptions( Options, LastCompOp );
 				FGenericGenerationResult ComponentResult;
 				GenerateComponent(ComponentOptions, ComponentResult, ComponentNode);
-				LastCompOp = ComponentResult.op;
+				LastCompOp = ComponentResult.Op;
             }
         }
 
@@ -1776,12 +1778,12 @@ namespace mu
 		}
 
 		// Add the components from child objects
-		FAdditionalComponentKey ThisKey;
-		ThisKey.ObjectNode = CurrentParents.Last().ObjectNode;
-		TArray<FAdditionalComponentData>* ThisAdditionalComponents = AdditionalComponents.Find(ThisKey);
+		FObjectGenerationResult::FAdditionalComponentKey ThisKey;
+		ThisKey.ObjectNode = InNode;
+		TArray<FObjectGenerationResult::FAdditionalComponentData>* ThisAdditionalComponents = OutResult.AdditionalComponents.Find(ThisKey);
 		if (LastCompOp && ThisAdditionalComponents)
 		{
-			for (const FAdditionalComponentData& Additional : *ThisAdditionalComponents)
+			for (const FObjectGenerationResult::FAdditionalComponentData& Additional : *ThisAdditionalComponents)
 			{
 				check(Additional.PlaceholderOp);
 				ASTOp::Replace(Additional.PlaceholderOp, LastCompOp);
@@ -1790,23 +1792,17 @@ namespace mu
 		}
 
 		// Store this chain of components for use in parent objects if necessary
-		// 2 is because there must be a parent and there is always a null element as well.
 		if (LastCompOp && bIsChildObject)
 		{
-			// TODO: Directly to the root object?
-			const FParentKey& ParentObjectKey = CurrentParents[CurrentParents.Num() - 2];
-			FAdditionalComponentKey ParentKey;
-			ParentKey.ObjectNode = ParentObjectKey.ObjectNode;
+			FObjectGenerationResult::FAdditionalComponentKey ParentKey;
+			ParentKey.ObjectNode = Options.ParentObjectNode;
 
-			FAdditionalComponentData Data;
+			FObjectGenerationResult::FAdditionalComponentData Data;
 			Data.ComponentOp = LastCompOp;
 			Data.PlaceholderOp = PlaceholderOp;
-			AdditionalComponents.FindOrAdd(ParentKey).Add(Data);
+			OutResult.AdditionalComponents.FindOrAdd(ParentKey).Add(Data);
 		}
 
-
-        Ptr<ASTOp> RootOp = LastCompOp;
-
 		// Add an ASTOpAddExtensionData for each connected ExtensionData node
 		for (const NodeObjectNew::FNamedExtensionDataNode& NamedNode : InNode->ExtensionDataNodes)
 		{
@@ -1829,15 +1825,14 @@ namespace mu
 			}
 
 			FConditionalExtensionDataOp& SavedOp = ConditionalExtensionDataOps.AddDefaulted_GetRef();
-			if (CurrentObject.Num() > 0)
-			{
-				SavedOp.Condition = CurrentObject.Last().Condition;
-			}
+			SavedOp.Condition = Options.CurrentObjectCondition;
 			SavedOp.ExtensionDataOp = ChildResult.Op;
 			SavedOp.ExtensionDataName = NamedNode.Name;
 		}
 
-		if (CurrentObject.Num() == 0)
+		Ptr<ASTOp> RootOp = LastCompOp;
+
+		if (!Options.ParentObjectNode)
 		{
 			for (const FConditionalExtensionDataOp& SavedOp : ConditionalExtensionDataOps)
 			{
@@ -1863,63 +1858,52 @@ namespace mu
 			}
 		}
 
-        CurrentParents.Pop();
-
-        Result.op = RootOp;
+		OutResult.Op = RootOp;
     }
 
 
-    //---------------------------------------------------------------------------------------------
-    //Ptr<ASTOp> CodeGenerator::Visit( const NodeObjectGroup::Private& node )
-	void CodeGenerator::Generate_ObjectGroup(const FGenericGenerationOptions& Options, FGenericGenerationResult& Result, const NodeObjectGroup* InNode)
+	void CodeGenerator::GenerateObject_Group(const FObjectGenerationOptions& Options, FObjectGenerationResult& OutResult, const NodeObjectGroup* InNode)
 	{
 		const NodeObjectGroup::Private& node = *InNode->GetPrivate();
 
-		TArray<FString> usedNames;
+		TArray<FString> UsedNames;
 
         // Parse the child objects first, which will accumulate operations in the patching lists
-        for ( int32 t=0; t<node.m_children.Num(); ++t )
+        for ( int32 ChildIndex=0; ChildIndex <node.m_children.Num(); ++ChildIndex)
         {
-            if ( const NodeObject* pChildNode = node.m_children[t].get() )
+            if ( const NodeObject* pChildNode = node.m_children[ChildIndex].get() )
             {
 				// Look for the child condition in the first pass
-                Ptr<ASTOp> conditionOp;
-				bool found = false;
-                for( int32 i = 0; !found && i != FirstPass.Objects.Num(); i++ )
+                Ptr<ASTOp> ConditionOp;
+				bool bFound = false;
+                for( int32 ObjectIndex = 0; !bFound && ObjectIndex<FirstPass.Objects.Num(); ObjectIndex++ )
 				{
-					FirstPassGenerator::FObject& it = FirstPass.Objects[i];
-					if (it.Node == pChildNode)
+					FirstPassGenerator::FObject& Candidate = FirstPass.Objects[ObjectIndex];
+					if (Candidate.Node == pChildNode)
 					{
-						found = true;
-						conditionOp = it.Condition;
+						bFound = true;
+						ConditionOp = Candidate.Condition;
 					}
 				}
 
-                // \todo
-                // TrunkStraight_02_BranchTop crash
-                // It may happen with partial compilations?
-                // check(found);
-
-				FObjectGenerationData data;
-                data.Condition = conditionOp;
-                CurrentObject.Add( data );
-
-                // This op is ignored: everything is stored as patches to apply to the parent when
-                // it is compiled.
-                Generate_Generic( pChildNode, Options );
+				FObjectGenerationOptions ChildOptions = Options;
+				ChildOptions.CurrentObjectCondition = ConditionOp;
 
-                CurrentObject.Pop();
+                // The result op is ignored: everything is stored as data to apply when the parent is compiled.
+				FObjectGenerationResult ChildResult;
+                GenerateObject( Options, ChildResult, pChildNode );
+				OutResult.AdditionalComponents.Append(ChildResult.AdditionalComponents);
 
 				// Check for duplicated child names
-				FString strChildName = pChildNode->GetName();
-				if (usedNames.Contains(strChildName))
+				FString ChildName = pChildNode->GetName();
+				if (UsedNames.Contains(ChildName))
 				{
-					FString Msg = FString::Printf(TEXT("Object group has more than one children with the same name [%s]."), *strChildName );
+					FString Msg = FString::Printf(TEXT("Object group has more than one children with the same name [%s]."), *ChildName );
 					ErrorLog->Add(Msg, ELMT_WARNING, InNode->GetMessageContext());
 				}
 				else
 				{
-					usedNames.Add(strChildName);
+					UsedNames.Add(ChildName);
 				}
             }
         }
diff --git a/Engine/Plugins/Mutable/Source/MutableTools/Private/MuT/CodeGenerator.h b/Engine/Plugins/Mutable/Source/MutableTools/Private/MuT/CodeGenerator.h
--- a/Engine/Plugins/Mutable/Source/MutableTools/Private/MuT/CodeGenerator.h
+++ b/Engine/Plugins/Mutable/Source/MutableTools/Private/MuT/CodeGenerator.h
@@ -112,14 +112,13 @@ namespace mu
 
     public:
 
-        CodeGenerator( CompilerOptions::Private* options );
+        CodeGenerator( CompilerOptions::Private* );
 
         //! Data will be stored in States
         void GenerateRoot( const Ptr<const Node> );
 
 	public:
 
-		// Generic top-level nodes
 		struct FGenericGenerationOptions
 		{
 			friend FORCEINLINE uint32 GetTypeHash(const FGenericGenerationOptions& InKey)
@@ -139,6 +138,14 @@ namespace mu
 			bool bIsImage = false;
 		};
 
+		struct FObjectGenerationOptions : public FGenericGenerationOptions
+		{
+			const NodeObjectNew* ParentObjectNode = nullptr;
+
+			/** Condition that enables a specific object. */
+			Ptr<ASTOp> CurrentObjectCondition;
+		};
+
 		struct FComponentGenerationOptions : public FGenericGenerationOptions
 		{
 			FComponentGenerationOptions(const FGenericGenerationOptions& BaseOptions, const Ptr<ASTOp>& InBaseInstance )
@@ -190,7 +197,7 @@ namespace mu
 
 		struct FGenericGenerationResult
 		{
-			Ptr<ASTOp> op;
+			Ptr<ASTOp> Op;
 		};
 
 		struct FGeneratedCacheKey
@@ -209,28 +216,81 @@ namespace mu
 			FORCEINLINE bool operator==(const FGeneratedCacheKey& Other) const = default;
 		};
 
-		typedef TMap<FGeneratedCacheKey, FGenericGenerationResult> FGeneratedGenericNodesMap;
-		FGeneratedGenericNodesMap GeneratedGenericNodes;
 
-		Ptr<ASTOp> Generate_Generic(const Ptr<const Node>, const FGenericGenerationOptions& );
+		void GenerateGeneric(const FGenericGenerationOptions& Options, FGenericGenerationResult& OutResult, const Ptr<const Node> InNode);
 		void Generate_LOD(const FLODGenerationOptions&, FGenericGenerationResult&, const NodeLOD*);
-		void Generate_ObjectNew(const FGenericGenerationOptions&, FGenericGenerationResult&, const NodeObjectNew*);
-		void Generate_ObjectGroup(const FGenericGenerationOptions&, FGenericGenerationResult&, const NodeObjectGroup*);
+
+
+		struct FGeneratedObjectCacheKey
+		{
+			Ptr<const Node> Node;
+			FObjectGenerationOptions Options;
+
+			friend FORCEINLINE uint32 GetTypeHash(const FGeneratedObjectCacheKey& InKey)
+			{
+				uint32 KeyHash = 0;
+				KeyHash = HashCombineFast(KeyHash, ::GetTypeHash(InKey.Node.get()));
+				KeyHash = HashCombineFast(KeyHash, GetTypeHash(InKey.Options));
+				return KeyHash;
+			}
+
+			FORCEINLINE bool operator==(const FGeneratedObjectCacheKey& Other) const
+			{
+				return Node == Other.Node && Options == Other.Options;
+			}
+		};
+
+		struct FObjectGenerationResult : public FGenericGenerationResult
+		{
+			/** List of additional components to add to an object that come from child objects.
+			 * The index is the object and lod that should receive the components.
+			 */
+			struct FAdditionalComponentKey
+			{
+				const NodeObjectNew* ObjectNode = nullptr;
+
+				FORCEINLINE bool operator==(const FAdditionalComponentKey& Other) const
+				{
+					return ObjectNode == Other.ObjectNode;
+				}
+
+				friend FORCEINLINE uint32 GetTypeHash(const FAdditionalComponentKey& InKey)
+				{
+					uint32 KeyHash = ::GetTypeHash(InKey.ObjectNode);
+					return KeyHash;
+				}
+			};
+
+			struct FAdditionalComponentData
+			{
+				Ptr<ASTOp> ComponentOp;
+				Ptr<ASTOp> PlaceholderOp;
+			};
+
+			TMap< FAdditionalComponentKey, TArray<FAdditionalComponentData> > AdditionalComponents;
+		};
+
+		typedef TMap<FGeneratedObjectCacheKey, FObjectGenerationResult> FGeneratedObjectsMap;
+		FGeneratedObjectsMap GeneratedObjects;
+
+		void GenerateObject(const FObjectGenerationOptions&, FObjectGenerationResult&, const NodeObject*);
+		void GenerateObject_New(const FObjectGenerationOptions&, FObjectGenerationResult&, const NodeObjectNew*);
+		void GenerateObject_Group(const FObjectGenerationOptions&, FObjectGenerationResult&, const NodeObjectGroup*);
 
     public:
 
-        //! Settings
-        CompilerOptions::Private* CompilerOptions = nullptr;
+        /** */
+        const CompilerOptions::Private* CompilerOptions = nullptr;
 
-		//!
+		/** */
 		FirstPassGenerator FirstPass;
 
-        //!
+        /** Used to accumulate all generation error messages. */
         TSharedPtr<FErrorLog> ErrorLog;
 
-        //! After the entire code generation this contains the information about all the states
-        typedef TArray< TPair<FObjectState, Ptr<ASTOp>> > StateList;
-        StateList States;
+        /** After the entire code generation this contains the information about all the states. */
+        typedef TArray< TPair<FObjectState, Ptr<ASTOp>> > FStateList;
+        FStateList States;
 
     private:
 
@@ -253,60 +313,10 @@ namespace mu
 			TSet<uint32> UniqueVertexIDGroups;
 		};
 		FGenerateMeshConstantState GenerateMeshConstantState;
-		UE::Tasks::FPipe GenerateMeshConstantPipe;
-
-		/** List of image resources for every image formats that have been generated so far as palceholders for missing images. */
-		TSharedPtr<FImage> MissingImage[size_t(EImageFormat::Count)];
-
-		struct FParentKey
-		{
-			const NodeObjectNew* ObjectNode = nullptr;
-			int32 Lod = -1;
-		};
-
-		// TODO: complete move of this top-down state to method arguments.
-		TArray< FParentKey > CurrentParents;
-
-        /** List of additional components to add to an object that come from child objects.
-        * The index is the object and lod that should receive the components.
-		*/
-        struct FAdditionalComponentKey
-        {
-			FAdditionalComponentKey()
-            {
-				ObjectNode = nullptr;
-            }
-
-            const NodeObjectNew* ObjectNode;
-
-			FORCEINLINE bool operator==(const FAdditionalComponentKey& Other) const
-			{
-				return ObjectNode == Other.ObjectNode;
-			}
-
-			friend FORCEINLINE uint32 GetTypeHash(const FAdditionalComponentKey& InKey)
-			{
-				uint32 KeyHash = 0;
-				KeyHash = HashCombineFast(KeyHash, ::GetTypeHash(InKey.ObjectNode));
-				return KeyHash;
-			}
-		};
-
-		struct FAdditionalComponentData
-		{
-			Ptr<ASTOp> ComponentOp;
-			Ptr<ASTOp> PlaceholderOp;
-		};
-
-        TMap< FAdditionalComponentKey, TArray<FAdditionalComponentData> > AdditionalComponents;
 
+		/** Any task accessing GenerateMeshConstantState should run in this pipe. */
+		UE::Tasks::FPipe GenerateMeshConstantPipe;
 
-        struct FObjectGenerationData
-        {
-            // Condition that enables a specific object
-            Ptr<ASTOp> Condition;
-        };
-		TArray<FObjectGenerationData> CurrentObject;
 
 		/** The key for generated tables is made of the source table and a parameter name. */
 		struct FTableCacheKey
@@ -358,8 +368,8 @@ namespace mu
 			}
 		};
 
-		typedef TMap<FGeneratedComponentCacheKey, FGenericGenerationResult> GeneratedComponentMap;
-		GeneratedComponentMap GeneratedComponents;
+		typedef TMap<FGeneratedComponentCacheKey, FGenericGenerationResult> FGeneratedComponentMap;
+		FGeneratedComponentMap GeneratedComponents;
 
 		void GenerateComponent(const FComponentGenerationOptions&, FGenericGenerationResult&, const NodeComponent*);
 		void GenerateComponent_New(const FComponentGenerationOptions&, FGenericGenerationResult&, const NodeComponentNew*);
@@ -439,7 +449,7 @@ namespace mu
 			{
 			}
 
-			NodePtrConst Node;
+			Ptr<const Node> Node;
 			FImageGenerationOptions Options;
 
 			friend FORCEINLINE uint32 GetTypeHash(const FGeneratedImageCacheKey& InKey)
@@ -453,8 +463,8 @@ namespace mu
 			FORCEINLINE bool operator==(const FGeneratedImageCacheKey& Other) const = default;
 		};
 
-		typedef TMap<FGeneratedImageCacheKey, FImageGenerationResult> GeneratedImagesMap;
-		GeneratedImagesMap GeneratedImages;
+		typedef TMap<FGeneratedImageCacheKey, FImageGenerationResult> FGeneratedImagesMap;
+		FGeneratedImagesMap GeneratedImages;
 
 		void GenerateImage(const FImageGenerationOptions&, FImageGenerationResult& result, const NodeImagePtrConst& node);
 		void GenerateImage_Constant(const FImageGenerationOptions&, FImageGenerationResult&, const NodeImageConstant*);
@@ -481,22 +491,22 @@ namespace mu
 		void GenerateImage_NormalComposite(const FImageGenerationOptions&, FImageGenerationResult&, const NodeImageNormalComposite*);
 		void GenerateImage_Transform(const FImageGenerationOptions&, FImageGenerationResult&, const NodeImageTransform*);
 
-		//!
+		/** */
 		TSharedPtr<FImage> GenerateMissingImage(EImageFormat);
 
-		//!
+		/** */
 		Ptr<ASTOp> GenerateMissingImageCode(const TCHAR* strWhere, EImageFormat, const void* errorContext, const FImageGenerationOptions& Options);
 
-		//!
+		/** */
 		Ptr<ASTOp> GeneratePlainImageCode(const FVector4f& Color, const FImageGenerationOptions& Options);
 
-		//!
+		/** */
 		Ptr<ASTOp> GenerateImageFormat(Ptr<ASTOp>, EImageFormat);
 
-		//!
+		/** */
 		Ptr<ASTOp> GenerateImageUncompressed(Ptr<ASTOp>);
 
-		//!
+		/** */
 		Ptr<ASTOp> GenerateImageSize(Ptr<ASTOp>, UE::Math::TIntVector2<int32>);
 
 		/** Evaluate if the image to generate is big enough to be split in separate operations and tiled afterwards. */
diff --git a/Engine/Plugins/Mutable/Source/MutableTools/Private/MuT/CodeGenerator_Color.cpp b/Engine/Plugins/Mutable/Source/MutableTools/Private/MuT/CodeGenerator_Color.cpp
--- a/Engine/Plugins/Mutable/Source/MutableTools/Private/MuT/CodeGenerator_Color.cpp
+++ b/Engine/Plugins/Mutable/Source/MutableTools/Private/MuT/CodeGenerator_Color.cpp
@@ -199,7 +199,9 @@ namespace mu
 		// Variable value
 		if (node.Parameter)
 		{
-			op->variable = Generate_Generic(node.Parameter.get(), Options);
+			FScalarGenerationResult ParamResult;
+			GenerateScalar(ParamResult, Options, node.Parameter.get());
+			op->variable = ParamResult.op;
 		}
 		else
 		{
@@ -208,19 +210,21 @@ namespace mu
 		}
 
 		// Options
-		for (int32 t = 0; t < node.Options.Num(); ++t)
+		for (int32 OptionIndex = 0; OptionIndex < node.Options.Num(); ++OptionIndex)
 		{
 			Ptr<ASTOp> branch;
-			if (node.Options[t])
+			if (node.Options[OptionIndex])
 			{
-				branch = Generate_Generic(node.Options[t].get(), Options);
+				FColorGenerationResult ParamResult;
+				GenerateColor(ParamResult, Options, node.Options[OptionIndex].get());
+				branch = ParamResult.op;
 			}
 			else
 			{
 				// This argument is required
 				branch = GenerateMissingColourCode(TEXT("Switch option"), Typed->GetMessageContext());
 			}
-			op->cases.Emplace((int16)t, op, branch);
+			op->cases.Emplace((int16)OptionIndex, op, branch);
 		}
 
 		result.op = op;
diff --git a/Engine/Plugins/Mutable/Source/MutableTools/Private/MuT/CodeGenerator_ExtensionData.cpp b/Engine/Plugins/Mutable/Source/MutableTools/Private/MuT/CodeGenerator_ExtensionData.cpp
--- a/Engine/Plugins/Mutable/Source/MutableTools/Private/MuT/CodeGenerator_ExtensionData.cpp
+++ b/Engine/Plugins/Mutable/Source/MutableTools/Private/MuT/CodeGenerator_ExtensionData.cpp
@@ -95,10 +95,12 @@ class Node;
 		Ptr<ASTOpSwitch> Op = new ASTOpSwitch;
 		Op->type = EOpType::ED_SWITCH;
 
-		// Variable name
+		// Variable
 		if (Switch->Parameter)
 		{
-			Op->variable = Generate_Generic(Switch->Parameter.get(), Options);
+			FScalarGenerationResult ParamResult; 
+			GenerateScalar(ParamResult, Options, Switch->Parameter.get());
+			Op->variable = ParamResult.op;
 		}
 		else
 		{
diff --git a/Engine/Plugins/Mutable/Source/MutableTools/Private/MuT/CodeGenerator_FirstPass.cpp b/Engine/Plugins/Mutable/Source/MutableTools/Private/MuT/CodeGenerator_FirstPass.cpp
--- a/Engine/Plugins/Mutable/Source/MutableTools/Private/MuT/CodeGenerator_FirstPass.cpp
+++ b/Engine/Plugins/Mutable/Source/MutableTools/Private/MuT/CodeGenerator_FirstPass.cpp
@@ -129,9 +129,9 @@ namespace mu
 
         if ( States.IsEmpty() )
         {
-            FObjectState data;
-            data.Name = "Default";
-            States.Emplace( data, Root );
+            FObjectState Data;
+            Data.Name = "Default";
+            States.Add( Data );
         }
 	}
 
@@ -189,6 +189,7 @@ namespace mu
 		}
 		else
 		{
+			// This node type is not supported in this pass.
 			check(false);
 		}
 	}
@@ -277,33 +278,33 @@ namespace mu
 
         case NodeSurfaceVariation::VariationType::State:
         {
-            int32 stateCount = States.Num();
+            int32 StateCount = States.Num();
 
             // Default
             {
                 // Store the states for the default branch here
-                StateCondition defaultStates;
+                StateCondition DefaultStates;
                 {
 					StateCondition AllTrue;
-					AllTrue.Init(true,stateCount);
-                    defaultStates = CurrentStateCondition.Last().IsEmpty()
+					AllTrue.Init(true, StateCount);
+					DefaultStates = CurrentStateCondition.Last().IsEmpty()
                             ? AllTrue
                             : CurrentStateCondition.Last();
 
                     for (const NodeSurfaceVariation::FVariation& v: InNode->Variations)
                     {
-                        for( size_t s=0; s<stateCount; ++s )
+                        for( int32 StateIndex=0; StateIndex < StateCount; ++StateIndex)
                         {
-                            if (States[s].Key.Name==v.Tag)
+                            if (States[StateIndex].Name==v.Tag)
                             {
                                 // Remove this state from the default options, since it has its own variation
-                                defaultStates[s] = false;
+								DefaultStates[StateIndex] = false;
                             }
                         }
                     }
                 }
 
-                CurrentStateCondition.Add(defaultStates);
+                CurrentStateCondition.Add(DefaultStates);
 
                 for (const Ptr<NodeSurface>& n : InNode->DefaultSurfaces)
                 {
@@ -322,11 +323,11 @@ namespace mu
             {
                 // Store the states for this variation here
 				StateCondition variationStates;
-				variationStates.Init(false,stateCount);
+				variationStates.Init(false,StateCount);
 
-                for( int32 StateIndex=0; StateIndex<stateCount; ++StateIndex)
+                for( int32 StateIndex=0; StateIndex< StateCount; ++StateIndex)
                 {
-                    if (States[StateIndex].Key.Name==v.Tag)
+                    if (States[StateIndex].Name==v.Tag)
                     {
                         variationStates[StateIndex] = true;
                     }
@@ -569,7 +570,7 @@ namespace mu
         // Accumulate the model states
         for ( const FObjectState& s: InNode->States )
         {
-            States.Emplace( s, InNode );
+            States.Add( s );
 
             if ( s.RuntimeParams.Num() > MUTABLE_MAX_RUNTIME_PARAMETERS_PER_STATE )
             {
diff --git a/Engine/Plugins/Mutable/Source/MutableTools/Private/MuT/CodeGenerator_FirstPass.h b/Engine/Plugins/Mutable/Source/MutableTools/Private/MuT/CodeGenerator_FirstPass.h
--- a/Engine/Plugins/Mutable/Source/MutableTools/Private/MuT/CodeGenerator_FirstPass.h
+++ b/Engine/Plugins/Mutable/Source/MutableTools/Private/MuT/CodeGenerator_FirstPass.h
@@ -184,8 +184,8 @@ namespace mu
         };
         TArray<FTag> Tags;
 
-        /** Accumulate the model states found while generating code, with their generated root nodes. */
-        typedef TArray< TPair<FObjectState, const Node*> > StateList;
+        /** Accumulate the model states found while generating code. */
+        typedef TArray< FObjectState > StateList;
         StateList States;
 
 		/** FParameters added for every node. */
diff --git a/Engine/Plugins/Mutable/Source/MutableTools/Private/MuT/CodeGenerator_Image.cpp b/Engine/Plugins/Mutable/Source/MutableTools/Private/MuT/CodeGenerator_Image.cpp
--- a/Engine/Plugins/Mutable/Source/MutableTools/Private/MuT/CodeGenerator_Image.cpp
+++ b/Engine/Plugins/Mutable/Source/MutableTools/Private/MuT/CodeGenerator_Image.cpp
@@ -119,7 +119,7 @@ namespace mu
 
 		// See if it was already generated
 		FGeneratedImageCacheKey key(Options,Untyped);
-		GeneratedImagesMap::ValueType* CachedPtr = GeneratedImages.Find(key);
+		FGeneratedImagesMap::ValueType* CachedPtr = GeneratedImages.Find(key);
 		if (CachedPtr)
 		{
 			Result = *CachedPtr;
@@ -698,9 +698,11 @@ namespace mu
         Ptr<ASTOpImageInterpolate> op = new ASTOpImageInterpolate();
 
         // Factor
-        if ( Node* Factor = node.m_pFactor.get() )
+        if ( NodeScalar* Factor = node.m_pFactor.get() )
         {
-            op->Factor = Generate_Generic(Factor, Options);
+			FScalarGenerationResult ParamResult;
+			GenerateScalar(ParamResult, Options, Factor);
+			op->Factor = ParamResult.op;
         }
         else
         {
@@ -929,10 +931,12 @@ namespace mu
 
 
         // Factor
-        if ( Node* Factor = node.m_pFactor.get() )
+        if ( NodeScalar* Factor = node.m_pFactor.get() )
         {
-            op->Factor = Generate_Generic(Factor, Options);
-        }
+			FScalarGenerationResult ParamResult;
+			GenerateScalar(ParamResult, Options, Factor);
+			op->Factor = ParamResult.op;
+		}
         else
         {
             // This argument is required
@@ -1062,10 +1066,12 @@ namespace mu
         op->Base = BaseOp;
 
         // Threshold
-        Ptr<ASTOp> ThresholdOp = 0;
-        if ( Node* pScalar = node.m_pThreshold.get() )
+        Ptr<ASTOp> ThresholdOp;
+        if ( NodeScalar* ThresholdNode = node.m_pThreshold.get() )
         {
-			ThresholdOp = Generate_Generic( pScalar, Options);
+			FScalarGenerationResult ParamResult;
+			GenerateScalar(ParamResult, Options, ThresholdNode);
+			ThresholdOp = ParamResult.op;
         }
         else
         {
@@ -1085,7 +1091,7 @@ namespace mu
 
 		MUTABLE_CPUPROFILER_SCOPE(NodeImageResize);
 
-        Ptr<ASTOp> at = 0;
+        Ptr<ASTOp> at;
 
         // Source image
         Ptr<ASTOp> base;
@@ -1158,10 +1164,12 @@ namespace mu
 		const NodeImagePlainColour::Private& node = *InNode->GetPrivate();
 
 		// Source colour
-        Ptr<ASTOp> base = 0;
+        Ptr<ASTOp> base;
         if ( node.m_pColour )
         {
-            base = Generate_Generic( node.m_pColour.get(), Options);
+			FColorGenerationResult ColorResult;
+            GenerateColor( ColorResult, Options, node.m_pColour.get() );
+			base = ColorResult.op;
         }
         else
         {
@@ -1222,8 +1230,7 @@ namespace mu
         if (node.Options.Num() == 0)
 		{
 			// No options in the switch!
-            Ptr<ASTOp> missingOp = GenerateMissingImageCode(TEXT("Switch option"), EImageFormat::RGBA_UByte, InNode->GetMessageContext(), Options);
-			Result.op = missingOp;
+			Result.op = GenerateMissingImageCode(TEXT("Switch option"), EImageFormat::RGBA_UByte, InNode->GetMessageContext(), Options);
 			return;
 		}
 
@@ -1233,7 +1240,9 @@ namespace mu
 		// Variable value
 		if ( node.m_pParameter )
 		{
-            op->variable = Generate_Generic( node.m_pParameter.get(), Options);
+			FScalarGenerationResult ParamResult;
+			GenerateScalar(ParamResult, Options, node.m_pParameter.get());
+			op->variable = ParamResult.op;
 		}
 		else
 		{
@@ -1242,14 +1251,14 @@ namespace mu
 		}
 
 		// Options
-        for ( std::size_t t=0; t< node.Options.Num(); ++t )
+        for ( int32 OptionIndex=0; OptionIndex < node.Options.Num(); ++OptionIndex)
         {
             Ptr<ASTOp> branch;
 
-            if (node.Options[t])
+            if (node.Options[OptionIndex])
             {
 				FImageGenerationResult BaseResult;
-				GenerateImage(Options, BaseResult, node.Options[t]);
+				GenerateImage(Options, BaseResult, node.Options[OptionIndex]);
 				branch = BaseResult.op;
 			}
             else
@@ -1258,7 +1267,7 @@ namespace mu
                 branch = GenerateMissingImageCode(TEXT("Switch option"), EImageFormat::RGBA_UByte, InNode->GetMessageContext(), Options);
             }
 
-            op->cases.Emplace((int16_t)t,op,branch);
+            op->cases.Emplace(int16(OptionIndex),op,branch);
         }
 
         Ptr<ASTOp> switchAt = op;
@@ -1278,8 +1287,10 @@ namespace mu
         // Condition
         if ( node.Parameter )
         {
-            op->condition = Generate_Generic( node.Parameter.get(), Options);
-        }
+			FBoolGenerationResult ParamResult;
+			GenerateBool(ParamResult, Options, node.Parameter.get());
+			op->condition = ParamResult.op;
+		}
         else
         {
             // This argument is required
@@ -1565,7 +1576,9 @@ namespace mu
 			if (node.m_pAngleFadeEnd) pPropsNode->Y = node.m_pAngleFadeEnd;
 			else pPropsNode->Y = pDefaultFade;
 
-			ImageRasterOp->angleFadeProperties = Generate_Generic(pPropsNode, Options);
+			FColorGenerationResult ParamResult;
+			GenerateColor(ParamResult, Options, pPropsNode);
+			ImageRasterOp->angleFadeProperties = ParamResult.op;
 		}
 
         // Seam correction operations
@@ -1863,92 +1876,88 @@ namespace mu
 
 	TSharedPtr<FImage> CodeGenerator::GenerateMissingImage(EImageFormat Format)
 	{
-		// Create the image node if it hasn't been created yet.
-		if (!MissingImage[SIZE_T(Format)])
-		{
-			// Make a checkered debug image
-			const FImageSize Size(16, 16);
+		MUTABLE_CPUPROFILER_SCOPE(GenerateMissingImage);
 
-			TSharedPtr<FImage> GeneratedImage = MakeShared<FImage>(Size[0], Size[1], 1, Format, EInitializationType::NotInitialized);
+		// Make a checkered debug image
+		const FImageSize Size(16, 16);
 
-			switch (Format)
-			{
-			case EImageFormat::L_UByte:
+		TSharedPtr<FImage> GeneratedImage = MakeShared<FImage>(Size[0], Size[1], 1, Format, EInitializationType::NotInitialized);
+
+		switch (Format)
+		{
+		case EImageFormat::L_UByte:
+		{
+			uint8* DataPtr = GeneratedImage->GetLODData(0);
+			for (int32 P = 0; P < Size[0]*Size[1]; ++P)
 			{
-				uint8* DataPtr = GeneratedImage->GetLODData(0);
-				for (int32 P = 0; P < Size[0]*Size[1]; ++P)
+				if ((P + P/Size[0]) % 2)
 				{
-					if ((P + P/Size[0]) % 2)
-					{
-						DataPtr[0] = 255;
-					}
-					else
-					{
-						DataPtr[0] = 64;
-					}
-
-					DataPtr++;
+					DataPtr[0] = 255;
+				}
+				else
+				{
+					DataPtr[0] = 64;
 				}
-				break;
+
+				DataPtr++;
 			}
-			case EImageFormat::RGB_UByte:
+			break;
+		}
+		case EImageFormat::RGB_UByte:
+		{
+			uint8* DataPtr = GeneratedImage->GetLODData(0);
+			for (int32 P = 0; P < Size[0]*Size[1]; ++P)
 			{
-				uint8* DataPtr = GeneratedImage->GetLODData(0);
-				for (int32 P = 0; P < Size[0]*Size[1]; ++P)
+				if ((P + P/Size[0]) % 2)
 				{
-					if ((P + P/Size[0]) % 2)
-					{
-						DataPtr[0] = 255;
-						DataPtr[1] = 255;
-						DataPtr[2] = 64;
-					}
-					else
-					{
-						DataPtr[0] = 64;
-						DataPtr[1] = 64;
-						DataPtr[2] = 255;
-					}
-
-					DataPtr += 3;
+					DataPtr[0] = 255;
+					DataPtr[1] = 255;
+					DataPtr[2] = 64;
 				}
-				break;
+				else
+				{
+					DataPtr[0] = 64;
+					DataPtr[1] = 64;
+					DataPtr[2] = 255;
+				}
+
+				DataPtr += 3;
 			}
-			case EImageFormat::BGRA_UByte:
-			case EImageFormat::RGBA_UByte:
+			break;
+		}
+		case EImageFormat::BGRA_UByte:
+		case EImageFormat::RGBA_UByte:
+		{
+			uint8* DataPtr = GeneratedImage->GetLODData(0);
+			for (int32 P = 0; P < Size[0]*Size[1]; ++P)
 			{
-				uint8* DataPtr = GeneratedImage->GetLODData(0);
-				for (int32 P = 0; P < Size[0]*Size[1]; ++P)
+				if ((P + P/Size[0]) % 2)
 				{
-					if ((P + P/Size[0]) % 2)
-					{
-						DataPtr[0] = 255;
-						DataPtr[1] = 255;
-						DataPtr[2] = 64;
-						DataPtr[3] = 255;
-					}
-					else
-					{
-						DataPtr[0] = 64;
-						DataPtr[1] = 64;
-						DataPtr[2] = 255;
-						DataPtr[3] = 128;
-					}
-
-					DataPtr += 4;
+					DataPtr[0] = 255;
+					DataPtr[1] = 255;
+					DataPtr[2] = 64;
+					DataPtr[3] = 255;
+				}
+				else
+				{
+					DataPtr[0] = 64;
+					DataPtr[1] = 64;
+					DataPtr[2] = 255;
+					DataPtr[3] = 128;
 				}
-				break;
-			}
-
-			default:
-				check( false );
-				break;
 
+				DataPtr += 4;
 			}
+			break;
+		}
+
+		default:
+			check( false );
+			break;
 
-			MissingImage[(SIZE_T)Format] = GeneratedImage;
 		}
 
-		return MissingImage[(SIZE_T)Format];
+		return GeneratedImage;
 	}
 
 
diff --git a/Engine/Plugins/Mutable/Source/MutableTools/Private/MuT/CodeGenerator_Mesh.cpp b/Engine/Plugins/Mutable/Source/MutableTools/Private/MuT/CodeGenerator_Mesh.cpp
--- a/Engine/Plugins/Mutable/Source/MutableTools/Private/MuT/CodeGenerator_Mesh.cpp
+++ b/Engine/Plugins/Mutable/Source/MutableTools/Private/MuT/CodeGenerator_Mesh.cpp
@@ -618,7 +618,9 @@ namespace mu
         // Factor
         if ( node.Factor )
         {
-			FactorOp = Generate_Generic( node.Factor.get(), StaticOptions );
+			FScalarGenerationResult ParamResult;
+			GenerateScalar(ParamResult, StaticOptions, node.Factor.get());
+			FactorOp = ParamResult.op;
         }
         else
         {
@@ -892,8 +894,10 @@ namespace mu
 		Ptr<ASTOp> Variable;
         if ( node.m_pParameter )
         {
-			Variable = Generate_Generic( node.m_pParameter.get(), StaticOptions);
-        }
+			FScalarGenerationResult ParamResult;
+			GenerateScalar(ParamResult, StaticOptions, node.m_pParameter.get());
+			Variable = ParamResult.op;
+		}
         else
         {
             // This argument is required
diff --git a/Engine/Plugins/Mutable/Source/MutableTools/Private/MuT/CodeGenerator_Scalar.cpp b/Engine/Plugins/Mutable/Source/MutableTools/Private/MuT/CodeGenerator_Scalar.cpp
--- a/Engine/Plugins/Mutable/Source/MutableTools/Private/MuT/CodeGenerator_Scalar.cpp
+++ b/Engine/Plugins/Mutable/Source/MutableTools/Private/MuT/CodeGenerator_Scalar.cpp
@@ -342,10 +342,11 @@ namespace mu
 	{
 		Ptr<ASTOpScalarCurve> op = new ASTOpScalarCurve();
 
-		// T
-		if (Node* pA = Typed->CurveSampleValue.get())
+		if (NodeScalar* pA = Typed->CurveSampleValue.get())
 		{
-			op->time = Generate_Generic(pA, Options);
+			FScalarGenerationResult ParamResult;
+			GenerateScalar(ParamResult, Options, pA);
+			op->time = ParamResult.op;
 		}
 		else
 		{
@@ -378,9 +379,12 @@ namespace mu
 		}
 
 		// A
-		if (Node* pA = node.m_pA.get())
+		if (NodeScalar* pA = node.m_pA.get())
 		{
-			op->A = Generate_Generic(pA, Options);
+			FScalarGenerationResult ParamResult;
+			GenerateScalar(ParamResult, Options, pA);
+
+			op->A = ParamResult.op;
 		}
 		else
 		{
@@ -388,9 +392,12 @@ namespace mu
 		}
 
 		// B
-		if (Node* pB = node.m_pB.get())
+		if (NodeScalar* pB = node.m_pB.get())
 		{
-			op->B = Generate_Generic(pB, Options);
+			FScalarGenerationResult ParamResult;
+			GenerateScalar(ParamResult, Options, pB);
+
+			op->B = ParamResult.op;
 		}
 		else
 		{
@@ -412,7 +419,11 @@ namespace mu
 				NodeScalarConstantPtr pCell = new NodeScalarConstant();
 				float scalar = node.Table->GetPrivate()->Rows[row].Values[colIndex].Scalar;
 				pCell->SetValue(scalar);
-				return Generate_Generic(pCell, Options);
+
+				FScalarGenerationResult ParamResult;
+				GenerateScalar(ParamResult, Options, pCell);
+
+				return ParamResult.op;
 			});
 
 		result.op = Op;
@@ -431,9 +442,10 @@ namespace mu
 		pNode->SetValue(value);
 
 		FGenericGenerationOptions Options;
-		Ptr<ASTOp> result = Generate_Generic(pNode, Options);
+		FScalarGenerationResult ParamResult;
+		GenerateScalar(ParamResult, Options, pNode.get());
 
-		return result;
+		return ParamResult.op;
 	}
 
 }
\ No newline at end of file
diff --git a/Engine/Plugins/Mutable/Source/MutableTools/Private/MuT/CodeOptimiser_GPU.cpp b/Engine/Plugins/Mutable/Source/MutableTools/Private/MuT/CodeOptimiser_GPU.cpp
deleted file mode 100644
diff --git a/Engine/Plugins/Mutable/Source/MutableTools/Private/MuT/CodeOptimiser_GPU.h b/Engine/Plugins/Mutable/Source/MutableTools/Private/MuT/CodeOptimiser_GPU.h
deleted file mode 100644
