[mutable] Fixed order of the Mutable Projector nodes "Projector Snap to Bone" elements.
- The component name will appear first and then the bone selection dropdown

#rnx

diff --git a/Engine/Plugins/Mutable/Source/CustomizableObjectEditor/Private/MuCOE/Nodes/CustomizableObjectNodeModifierRemoveMesh.cpp b/Engine/Plugins/Mutable/Source/CustomizableObjectEditor/Private/MuCOE/Nodes/CustomizableObjectNodeModifierRemoveMesh.cpp
--- a/Engine/Plugins/Mutable/Source/CustomizableObjectEditor/Private/MuCOE/Nodes/CustomizableObjectNodeModifierRemoveMesh.cpp
+++ b/Engine/Plugins/Mutable/Source/CustomizableObjectEditor/Private/MuCOE/Nodes/CustomizableObjectNodeModifierRemoveMesh.cpp
@@ -46,7 +46,9 @@ void UCustomizableObjectNodeModifierRemoveMesh::PinConnectionListChanged(UEdGrap
 FText UCustomizableObjectNodeModifierRemoveMesh::GetTooltipText() const
 {
 	return LOCTEXT("Remove_Mesh_Tooltip",
-	"Removes the faces of a material that are defined only by vertexes shared by the material and the input mesh.It also removes any vertex\nand edge that only define deleted faces, they are not left dangling. If the mesh removed covers at least all the faces included in one or\nmore layout blocs, those blocs are removed, freeing final texture layout space.");
+	"Removes the faces of a mesh section that are defined only by the vertexes shared by said mesh section and the input mesh. \n"
+	"It also removes any vertex and edge that only define deleted faces. \n"
+	"If the removed mesh covers all the faces included in one or more layout blocks those blocks get removed, freeing layout space in the final texture.");
 }
 
 bool UCustomizableObjectNodeModifierRemoveMesh::IsSingleOutputNode() const
diff --git a/Engine/Plugins/Mutable/Source/CustomizableObjectEditor/Private/MuCOE/Nodes/CustomizableObjectNodeProjectorConstant.h b/Engine/Plugins/Mutable/Source/CustomizableObjectEditor/Private/MuCOE/Nodes/CustomizableObjectNodeProjectorConstant.h
--- a/Engine/Plugins/Mutable/Source/CustomizableObjectEditor/Private/MuCOE/Nodes/CustomizableObjectNodeProjectorConstant.h
+++ b/Engine/Plugins/Mutable/Source/CustomizableObjectEditor/Private/MuCOE/Nodes/CustomizableObjectNodeProjectorConstant.h
@@ -51,15 +51,15 @@ public:
 	FVector BoneComboBoxUpDirection;
 
 	// UObject interface.
-	void PostEditChangeProperty(FPropertyChangedEvent& PropertyChangedEvent) override;
+	virtual void PostEditChangeProperty(FPropertyChangedEvent& PropertyChangedEvent) override;
 
 	// Begin EdGraphNode interface
-	FText GetNodeTitle(ENodeTitleType::Type TitleType) const override;
-	FLinearColor GetNodeTitleColor() const override;
-	FText GetTooltipText() const override;
+	virtual FText GetNodeTitle(ENodeTitleType::Type TitleType) const override;
+	virtual FLinearColor GetNodeTitleColor() const override;
+	virtual FText GetTooltipText() const override;
 
 	// UCustomizableObjectNode interface
-	void AllocateDefaultPins(UCustomizableObjectNodeRemapPins* RemapPins) override;
+	virtual void AllocateDefaultPins(UCustomizableObjectNodeRemapPins* RemapPins) override;
 	virtual void BackwardsCompatibleFixup(int32 CustomizableObjectCustomVersion) override;
 
 	// Own interface
diff --git a/Engine/Plugins/Mutable/Source/CustomizableObjectEditor/Private/MuCOE/Nodes/CustomizableObjectNodeProjectorParameter.h b/Engine/Plugins/Mutable/Source/CustomizableObjectEditor/Private/MuCOE/Nodes/CustomizableObjectNodeProjectorParameter.h
--- a/Engine/Plugins/Mutable/Source/CustomizableObjectEditor/Private/MuCOE/Nodes/CustomizableObjectNodeProjectorParameter.h
+++ b/Engine/Plugins/Mutable/Source/CustomizableObjectEditor/Private/MuCOE/Nodes/CustomizableObjectNodeProjectorParameter.h
@@ -51,17 +51,17 @@ public:
 	FVector BoneComboBoxUpDirection = FVector::ZeroVector;
 
 	// UObject interface.
-	void PostEditChangeProperty(FPropertyChangedEvent& PropertyChangedEvent) override;
+	virtual void PostEditChangeProperty(FPropertyChangedEvent& PropertyChangedEvent) override;
 
 	// Begin EdGraphNode interface
-	FText GetNodeTitle(ENodeTitleType::Type TitleType) const override;
-	FLinearColor GetNodeTitleColor() const override;
-	FText GetTooltipText() const override;
+	virtual FText GetNodeTitle(ENodeTitleType::Type TitleType) const override;
+	virtual FLinearColor GetNodeTitleColor() const override;
+	virtual FText GetTooltipText() const override;
 	virtual bool GetCanRenameNode() const override { return true; }
 	virtual void OnRenameNode(const FString& NewName) override;
 
 	// UCustomizableObjectNode interface
-	void AllocateDefaultPins(UCustomizableObjectNodeRemapPins* RemapPins) override;
+	virtual void AllocateDefaultPins(UCustomizableObjectNodeRemapPins* RemapPins) override;
 	virtual bool IsAffectedByLOD() const override { return false; }
 	virtual void BackwardsCompatibleFixup(int32 CustomizableObjectCustomVersion) override;
 
diff --git a/Engine/Plugins/Mutable/Source/CustomizableObjectEditor/Private/MuCOE/Nodes/CustomizableObjectNodeProjectorParameterDetails.cpp b/Engine/Plugins/Mutable/Source/CustomizableObjectEditor/Private/MuCOE/Nodes/CustomizableObjectNodeProjectorParameterDetails.cpp
--- a/Engine/Plugins/Mutable/Source/CustomizableObjectEditor/Private/MuCOE/Nodes/CustomizableObjectNodeProjectorParameterDetails.cpp
+++ b/Engine/Plugins/Mutable/Source/CustomizableObjectEditor/Private/MuCOE/Nodes/CustomizableObjectNodeProjectorParameterDetails.cpp
@@ -31,9 +31,17 @@ void FCustomizableObjectNodeProjectorParameterDetails::CustomizeDetails( IDetail
 {
 	FCustomizableObjectNodeDetails::CustomizeDetails(DetailBuilder);
 
+	// Set the order in which the categories will be displayed. This is required because we are editing some of them and that makes them change order
+	DetailBuilder.SortCategories([](const TMap<FName, IDetailCategoryBuilder*>& CategoryMap)
+	{
+		(*CategoryMap.FindChecked(FName("Clipboard"))).SetSortOrder(0);
+		(*CategoryMap.FindChecked(FName("CustomizableObject"))).SetSortOrder(1);
+		(*CategoryMap.FindChecked(FName("ProjectorSnapToBone"))).SetSortOrder(2);
+		(*CategoryMap.FindChecked(FName("UI"))).SetSortOrder(3);
+	});
+
 	NodeConstant = nullptr;
 	NodeParameter = nullptr;
-	DetailBuilderPtr = &DetailBuilder;
 
 	TSharedPtr<const IDetailsView> DetailsView = DetailBuilder.GetDetailsViewSharedPtr();
 	if (DetailsView && DetailsView->GetSelectedObjects().Num())
@@ -48,9 +56,9 @@ void FCustomizableObjectNodeProjectorParameterDetails::CustomizeDetails( IDetail
 		}
 	}
 
+	
 	IDetailCategoryBuilder& BlocksCategory = DetailBuilder.EditCategory( "Clipboard" );
-	//BlocksCategory.CategoryIcon( "ActorClassIcon.CustomizableObject" );
-
+	
 	if ((NodeConstant != nullptr) || (NodeParameter != nullptr))
 	{
 		BlocksCategory.AddCustomRow( LOCTEXT("FCustomizableObjectNodeProjectorParameterDetails", "Projector Data") )
@@ -91,91 +99,35 @@ void FCustomizableObjectNodeProjectorParameterDetails::CustomizeDetails( IDetail
 	const FSimpleDelegate OnReferenceSkeletonComponentChangedDelegate =
 		FSimpleDelegate::CreateSP(this, &FCustomizableObjectNodeProjectorParameterDetails::OnReferenceSkeletonComponentChanged);
 	ReferenceSkeletonIndexPropertyHandle->SetOnPropertyValueChanged(OnReferenceSkeletonComponentChangedDelegate);
-
-	IDetailCategoryBuilder& SnapToBoneCategory = DetailBuilder.EditCategory("ProjectorSnapToBone");
-	DetailBuilder.HideProperty("ProjectorBone");
-
-	SkeletalMesh = nullptr;
-	TSharedPtr<FString> BoneToSelect;
-
-	UCustomizableObject* CustomizableObject = nullptr;
-	FName ProjectorBoneName;
-	FName ReferenceComponent;
-
-	if (NodeConstant != nullptr)
-	{
-		CustomizableObject = Cast<UCustomizableObject>(NodeConstant->GetCustomizableObjectGraph()->GetOuter());
-		ProjectorBoneName = NodeConstant->ProjectorBone;
-		ReferenceComponent = NodeConstant->ReferenceSkeletonComponent;
-	}
-	else if (NodeParameter != nullptr)
-	{
-		CustomizableObject = Cast<UCustomizableObject>(NodeParameter->GetCustomizableObjectGraph()->GetOuter());
-		ProjectorBoneName = NodeParameter->ProjectorBone;
-		ReferenceComponent = NodeParameter->ReferenceSkeletonComponent;
-	}
-
-	//UCustomizableObject* CustomizableObject = Cast<UCustomizableObject>(Node->GetCustomizableObjectGraph()->GetOuter());
-
-	if (CustomizableObject)
-	{
-		SkeletalMesh = CustomizableObject->GetComponentMeshReferenceSkeletalMesh(ReferenceComponent);
-	}
-
-	if (SkeletalMesh)
-	{
-		BoneComboOptions.Empty();
-			
-		for (int32 i = 0; i < SkeletalMesh->GetRefSkeleton().GetRawBoneNum(); ++i)
-		{
-			FName BoneName = SkeletalMesh->GetRefSkeleton().GetBoneName(i);
-			BoneComboOptions.Add(MakeShareable(new FString(BoneName.ToString())));
-
-			if (BoneName == ProjectorBoneName)
-			{
-				BoneToSelect = BoneComboOptions.Last();
-			}
-		}
-
-        BoneComboOptions.Sort(CompareNames);
-
-		// Add them to the parent combo box
-		TSharedRef<IPropertyHandle> BoneProperty = DetailBuilder.GetProperty("ProjectorBone");
-
-		SnapToBoneCategory.AddCustomRow(LOCTEXT("UCustomizableObjectNodeProjectorParameter", "Projector Bone"))
-		[
-			SNew(SProperty, BoneProperty)
-			.ShouldDisplayName(false)
-			.CustomWidget()
-			[
-				SNew(SBorder)
-				.BorderImage(UE_MUTABLE_GET_BRUSH("NoBorder"))
-				.Padding(FMargin(0.0f, 0.0f, 10.0f, 0.0f))
-				.HAlign(EHorizontalAlignment::HAlign_Fill)
-				[
-					SNew(SHorizontalBox)
-					+SHorizontalBox::Slot().HAlign(HAlign_Left).VAlign(VAlign_Center)
-					[
-						SNew(STextBlock)
-						.Text( LOCTEXT("FCustomizableObjectNodeRemoveMeshDetails", "Projector Bone"))
-						.Font(IDetailLayoutBuilder::GetDetailFont())
-					]
-					+SHorizontalBox::Slot().HAlign(HAlign_Fill)
-					[
-						SNew(STextComboBox)
-						.OptionsSource(&BoneComboOptions)
-						.InitiallySelectedItem(BoneToSelect)
-						.OnSelectionChanged(this, &FCustomizableObjectNodeProjectorParameterDetails::OnBoneComboBoxSelectionChanged, BoneProperty)
-						.Font(IDetailLayoutBuilder::GetDetailFont())
-					]
-				]
-			]
-		];
-	}
+	
+	CaptureContextData();
+	CacheSkeletalMeshBoneNames();
+
+	// Add them to the parent combo box
+	TSharedRef<IPropertyHandle> BoneProperty = DetailBuilder.GetProperty("ProjectorBone");
+
+	DetailBuilder.EditDefaultProperty(BoneProperty)
+	->CustomWidget()
+	.Visibility(TAttribute<EVisibility>(this, &FCustomizableObjectNodeProjectorParameterDetails::ShouldBoneDropdownBeVisible))
+	.NameContent()
+	[
+		SNew(STextBlock)
+		.Text( LOCTEXT("FCustomizableObjectNodeRemoveMeshDetails", "Projector Bone"))
+		.Font(IDetailLayoutBuilder::GetDetailFont())
+
+	].ValueContent()
+	[
+		SAssignNew(BoneSelectionComboBox, STextComboBox)
+		.OptionsSource(&BoneComboOptions)
+		.InitiallySelectedItem(BoneToSelect)
+		.OnSelectionChanged(this, &FCustomizableObjectNodeProjectorParameterDetails::OnBoneComboBoxSelectionChanged, BoneProperty)
+		.Font(IDetailLayoutBuilder::GetDetailFont())
+	];
+	
 }
 
 
-FReply FCustomizableObjectNodeProjectorParameterDetails::OnProjectorCopyPressed()
+FReply FCustomizableObjectNodeProjectorParameterDetails::OnProjectorCopyPressed() const
 {
 	FString ExportedText;
 
@@ -196,7 +148,7 @@ FReply FCustomizableObjectNodeProjectorParameterDetails::OnProjectorCopyPressed(
 }
 
 
-FReply FCustomizableObjectNodeProjectorParameterDetails::OnProjectorPastePressed()
+FReply FCustomizableObjectNodeProjectorParameterDetails::OnProjectorPastePressed() const
 {
 	FScopedTransaction Transaction(LOCTEXT("PasteTransform", "Paste Transform"));
 	
@@ -231,11 +183,67 @@ FReply FCustomizableObjectNodeProjectorParameterDetails::OnProjectorPastePressed
 	return FReply::Handled();
 }
 
-void FCustomizableObjectNodeProjectorParameterDetails::OnBoneComboBoxSelectionChanged(TSharedPtr<FString> Selection, ESelectInfo::Type SelectInfo, TSharedRef<IPropertyHandle> BoneProperty)
+
+void FCustomizableObjectNodeProjectorParameterDetails::CaptureContextData()
+{
+	UCustomizableObject* CustomizableObject = nullptr;
+	FName ReferenceComponent;
+
+	if (NodeConstant != nullptr)
+	{
+		CustomizableObject = Cast<UCustomizableObject>(NodeConstant->GetCustomizableObjectGraph()->GetOuter());
+		ProjectorBoneName = NodeConstant->ProjectorBone;
+		ReferenceComponent = NodeConstant->ReferenceSkeletonComponent;
+	}
+	else if (NodeParameter != nullptr)
+	{
+		CustomizableObject = Cast<UCustomizableObject>(NodeParameter->GetCustomizableObjectGraph()->GetOuter());
+		ProjectorBoneName = NodeParameter->ProjectorBone;
+		ReferenceComponent = NodeParameter->ReferenceSkeletonComponent;
+	}
+
+	SkeletalMesh = nullptr;
+	if (CustomizableObject)
+	{
+		SkeletalMesh = CustomizableObject->GetComponentMeshReferenceSkeletalMesh(ReferenceComponent);
+	}
+}
+
+
+void FCustomizableObjectNodeProjectorParameterDetails::CacheSkeletalMeshBoneNames()
+{
+	BoneComboOptions.Empty();
+	
+	BoneToSelect.Reset();
+	if (BoneSelectionComboBox)
+	{
+		// Ensure the combo box displayed element is reset when changing the data being pointed at by "BoneToSelect"
+		BoneSelectionComboBox->ClearSelection();
+	}
+	
+	if (SkeletalMesh)
+	{
+		for (int32 i = 0; i < SkeletalMesh->GetRefSkeleton().GetRawBoneNum(); ++i)
+		{
+			FName BoneName = SkeletalMesh->GetRefSkeleton().GetBoneName(i);
+			BoneComboOptions.Add(MakeShareable(new FString(BoneName.ToString())));
+	
+			if (BoneName == ProjectorBoneName)
+			{
+				BoneToSelect = BoneComboOptions.Last();
+			}
+		}
+	
+		BoneComboOptions.Sort(CompareNames);
+	}
+}
+
+
+void FCustomizableObjectNodeProjectorParameterDetails::OnBoneComboBoxSelectionChanged(TSharedPtr<FString> InSelection, ESelectInfo::Type InSelectInfo, TSharedRef<IPropertyHandle> InBoneProperty)
 {
 	for (int OptionIndex = 0; OptionIndex < BoneComboOptions.Num(); ++OptionIndex)
 	{
-		if (BoneComboOptions[OptionIndex] == Selection)
+		if (BoneComboOptions[OptionIndex] == InSelection)
 		{
 			FVector Location = FVector::ZeroVector;
 			FVector Direction = FVector::ForwardVector;
@@ -243,7 +251,7 @@ void FCustomizableObjectNodeProjectorParameterDetails::OnBoneComboBoxSelectionCh
 			if (SkeletalMesh)
 			{
 				const TArray<FTransform>& BoneArray = SkeletalMesh->GetRefSkeleton().GetRefBonePose();
-				int32 ParentIndex = SkeletalMesh->GetRefSkeleton().FindBoneIndex(FName(**Selection));
+				int32 ParentIndex = SkeletalMesh->GetRefSkeleton().FindBoneIndex(FName(**InSelection));
 
 				FVector ChildLocation = FVector::ForwardVector;
 
@@ -286,7 +294,7 @@ void FCustomizableObjectNodeProjectorParameterDetails::OnBoneComboBoxSelectionCh
 				}
 			}
 
-			BoneProperty->SetValue(*BoneComboOptions[OptionIndex].Get());
+			InBoneProperty->SetValue(*BoneComboOptions[OptionIndex].Get());
 
 			return;
 		}
@@ -304,8 +312,22 @@ void FCustomizableObjectNodeProjectorParameterDetails::OnReferenceSkeletonCompon
 	{
 		NodeParameter->ProjectorBone = FName();
 	}
+	
+	CaptureContextData();
+	CacheSkeletalMeshBoneNames();
+}
+
+
+EVisibility FCustomizableObjectNodeProjectorParameterDetails::ShouldBoneDropdownBeVisible() const
+{
+	if (BoneComboOptions.Num() > 0)
+	{
+		return EVisibility::Visible;
+	}
 
-	DetailBuilderPtr->ForceRefreshDetails();
+	return EVisibility::Collapsed;
 }
 
+
+
 #undef LOCTEXT_NAMESPACE
diff --git a/Engine/Plugins/Mutable/Source/CustomizableObjectEditor/Private/MuCOE/Nodes/CustomizableObjectNodeProjectorParameterDetails.h b/Engine/Plugins/Mutable/Source/CustomizableObjectEditor/Private/MuCOE/Nodes/CustomizableObjectNodeProjectorParameterDetails.h
--- a/Engine/Plugins/Mutable/Source/CustomizableObjectEditor/Private/MuCOE/Nodes/CustomizableObjectNodeProjectorParameterDetails.h
+++ b/Engine/Plugins/Mutable/Source/CustomizableObjectEditor/Private/MuCOE/Nodes/CustomizableObjectNodeProjectorParameterDetails.h
@@ -3,7 +3,6 @@
 #pragma once
 
 #include "CustomizableObjectNodeDetails.h"
-#include "IDetailCustomization.h"
 
 class FReply;
 namespace ESelectInfo { enum Type : int; }
@@ -11,29 +10,43 @@ namespace ESelectInfo { enum Type : int; }
 class FString;
 class IDetailLayoutBuilder;
 class IPropertyHandle;
+struct EVisibility;
 
 
 class FCustomizableObjectNodeProjectorParameterDetails : public FCustomizableObjectNodeDetails
 {
 public:
 	/** Makes a new instance of this detail layout class for a specific detail view requesting it */
-	static TSharedRef<IDetailCustomization> MakeInstance();
+	static TSharedRef<class IDetailCustomization> MakeInstance();
 
 	/** ILayoutDetails interface */
 	virtual void CustomizeDetails( IDetailLayoutBuilder& DetailBuilder ) override;
 
 
 private:
+
+	/** The combo box that we use to display a list of bone names for the current Skeletal mesh */
+	TSharedPtr<class STextComboBox> BoneSelectionComboBox = nullptr;
+	
+	FName ProjectorBoneName;
+	TSharedPtr<FString> BoneToSelect;
 	TArray< TSharedPtr<FString> > BoneComboOptions;
 
-	class UCustomizableObjectNodeProjectorConstant* NodeConstant;
-	class UCustomizableObjectNodeProjectorParameter* NodeParameter;
-	FReply OnProjectorCopyPressed();
-	FReply OnProjectorPastePressed();
+	class UCustomizableObjectNodeProjectorConstant* NodeConstant = nullptr;
+	class UCustomizableObjectNodeProjectorParameter* NodeParameter = nullptr;
+
+	FReply OnProjectorCopyPressed() const;
+	FReply OnProjectorPastePressed() const;
+
+	class USkeletalMesh* SkeletalMesh = nullptr;
 
-	class USkeletalMesh* SkeletalMesh;
-	class IDetailLayoutBuilder* DetailBuilderPtr;
+	/** Cache the Skeletal Mesh of the component being pointed at and the Projector Parameter name of the node whose descriptor this class edits */
+	void CaptureContextData();
 
-	void OnBoneComboBoxSelectionChanged(TSharedPtr<FString> Selection, ESelectInfo::Type SelectInfo, TSharedRef<IPropertyHandle> BoneProperty);
+	/** Cache the names of the bones present in the cached Skeletal Mesh */
+	void CacheSkeletalMeshBoneNames();
+	
+	void OnBoneComboBoxSelectionChanged(TSharedPtr<FString> InSelection, ESelectInfo::Type InSelectInfo, TSharedRef<IPropertyHandle> InBoneProperty);
 	void OnReferenceSkeletonComponentChanged();
+	EVisibility ShouldBoneDropdownBeVisible() const;
 };
