[mutable] Fix for missing clothing and real-time morph data in meshes coming from SkeletalMesh nodes.
#rnx

diff --git a/Engine/Plugins/Mutable/Source/CustomizableObjectEditor/Private/MuCOE/GenerateMutableSource/GenerateMutableSourceMesh.cpp b/Engine/Plugins/Mutable/Source/CustomizableObjectEditor/Private/MuCOE/GenerateMutableSource/GenerateMutableSourceMesh.cpp
--- a/Engine/Plugins/Mutable/Source/CustomizableObjectEditor/Private/MuCOE/GenerateMutableSource/GenerateMutableSourceMesh.cpp
+++ b/Engine/Plugins/Mutable/Source/CustomizableObjectEditor/Private/MuCOE/GenerateMutableSource/GenerateMutableSourceMesh.cpp
@@ -3629,34 +3629,15 @@ mu::Ptr<mu::NodeMesh> GenerateMutableSourceMesh(const UEdGraphPin* Pin,
 				{
 					MeshData.bHasVertexColors = TypedNodeSkel->SkeletalMesh->GetHasVertexColors();	
 					MeshData.NumTexCoordChannels = ImportedModel->LODModels[LODIndex].NumTexCoords;
-					MeshData.MaxBoneIndexTypeSizeBytes = MutableMesh->GetBoneMap().Num() > 256 ? 2 : 1;
+					MeshData.MaxBoneIndexTypeSizeBytes = ImportedModel->LODModels[LODIndex].RequiredBones.Num() > 256 ? 2 : 1;
 					MeshData.MaxNumBonesPerVertex = ImportedModel->LODModels[LODIndex].GetMaxBoneInfluences();
-					
 
 					// When mesh data is combined we will get an upper and lower bound of the number of triangles.
 					MeshData.MaxNumTriangles = ImportedModel->LODModels[LODIndex].Sections[SectionIndex].NumTriangles;
 					MeshData.MinNumTriangles = ImportedModel->LODModels[LODIndex].Sections[SectionIndex].NumTriangles;
-				}
-
-				// Find if the mesh has realtime morphs and clothing looking at the generated mesh buffers.
-				{
-					const mu::FMeshBufferSet& MeshSet = MutableMesh->GetVertexBuffers();
-
-					int32 MorphIndexBuffer, MorphIndexChannel;
-					MeshSet.FindChannel(mu::EMeshBufferSemantic::Other, 0, &MorphIndexBuffer, &MorphIndexChannel);
-
-					int32 MorphResourceBuffer, MorphResourceChannel;
-					MeshSet.FindChannel(mu::EMeshBufferSemantic::Other, 1, &MorphResourceBuffer, &MorphResourceChannel);
-
-					MeshData.bHasRealTimeMorphs = MorphIndexBuffer >= 0 && MorphResourceBuffer >= 0; 
-
-					int32 ClothIndexBuffer, ClothIndexChannel;
-					MeshSet.FindChannel(mu::EMeshBufferSemantic::Other, 2, &ClothIndexBuffer, &ClothIndexChannel);
-
-					int32 ClothResourceBuffer, ClothResourceChannel;
-					MeshSet.FindChannel(mu::EMeshBufferSemantic::Other, 3, &ClothResourceBuffer, &ClothResourceChannel);
 
-					MeshData.bHasClothing = ClothIndexBuffer >= 0 && ClothResourceBuffer >= 0; 
+					MeshData.bHasRealTimeMorphs = GenerationContext.CompilationContext.Options.bRealTimeMorphTargetsEnabled && TypedNodeSkel->SkeletalMesh->GetMorphTargets().Num();
+					MeshData.bHasClothing = GenerationContext.CompilationContext.Options.bClothingEnabled && ImportedModel->LODModels[LODIndex].HasClothData();
 				}
 			}
 
