[mutable] Backout 39674175 because of bugs in the sample.
#rnx

diff --git a/Engine/Plugins/Mutable/Source/MutableTools/Private/MuT/CodeGenerator.cpp b/Engine/Plugins/Mutable/Source/MutableTools/Private/MuT/CodeGenerator.cpp
--- a/Engine/Plugins/Mutable/Source/MutableTools/Private/MuT/CodeGenerator.cpp
+++ b/Engine/Plugins/Mutable/Source/MutableTools/Private/MuT/CodeGenerator.cpp
@@ -105,26 +105,29 @@ namespace mu
 	//---------------------------------------------------------------------------------------------
 	//---------------------------------------------------------------------------------------------
 	//---------------------------------------------------------------------------------------------
-	CodeGenerator::CodeGenerator(CompilerOptions::Private* Options)
+	CodeGenerator::CodeGenerator(CompilerOptions::Private* options)
 		: GenerateMeshConstantPipe(TEXT("GenerateMeshConstantPipe"))
 	{
-		CompilerOptions = Options;
+		CompilerOptions = options;
 
 		// Create the message log
 		ErrorLog = MakeShared<FErrorLog>();
+
+		// Add the parent at the top of the hierarchy
+		CurrentParents.Add(FParentKey());
 	}
 
 	//---------------------------------------------------------------------------------------------
-	void CodeGenerator::GenerateRoot(const Ptr<const Node> InNode)
+	void CodeGenerator::GenerateRoot(const Ptr<const Node> pNode)
 	{
 		MUTABLE_CPUPROFILER_SCOPE(Generate);
 
 		// First pass
-		FirstPass.Generate(ErrorLog, InNode.get(), CompilerOptions->bIgnoreStates, this);
+		FirstPass.Generate(ErrorLog, pNode.get(), CompilerOptions->bIgnoreStates, this);
 
 		// Second pass
 		SecondPassGenerator SecondPass(&FirstPass, CompilerOptions);
-		bool bSuccess = SecondPass.Generate(ErrorLog, InNode.get());
+		bool bSuccess = SecondPass.Generate(ErrorLog, pNode.get());
 		if (!bSuccess)
 		{
 			return;
@@ -135,18 +138,17 @@ namespace mu
 			MUTABLE_CPUPROFILER_SCOPE(MainPass);
 
             int32 CurrentStateIndex = 0;
-			for (const FObjectState& State : FirstPass.States)
+			for (const TPair<FObjectState, const Node*>& s : FirstPass.States)
 			{
 				MUTABLE_CPUPROFILER_SCOPE(MainPassState);
 
 				FGenericGenerationOptions Options;
 				Options.State = CurrentStateIndex;
 
-				FGenericGenerationResult Result;
-				GenerateGeneric(Options, Result, InNode.get());
+				Ptr<ASTOp> stateRoot = Generate_Generic(pNode, Options);
+				States.Emplace(s.Key, stateRoot);
 
-				Ptr<ASTOp> StateRoot = Result.Op;
-				States.Emplace(State, StateRoot);
+				AdditionalComponents.Empty();
 
 				++CurrentStateIndex;
 			}
@@ -154,137 +156,120 @@ namespace mu
 	}
 
 
-void CodeGenerator::GenerateGeneric(const FGenericGenerationOptions& Options, FGenericGenerationResult& OutResult, const Ptr<const Node> InNode )
+	Ptr<ASTOp> CodeGenerator::Generate_Generic(const Ptr<const Node> pNode, const FGenericGenerationOptions& Options )
 	{
-		if (!InNode)
+		if (!pNode)
 		{
-			return;
+			return nullptr;
 		}
 
         // Type-specific generation
-		if (InNode->GetType()->IsA(NodeObject::GetStaticType()))
-		{
-			const NodeObject* ObjectNode = static_cast<const NodeObject*>(InNode.get());
-			FObjectGenerationOptions ObjectOptions;
-			ObjectOptions.ActiveTags = Options.ActiveTags;
-			ObjectOptions.bIsImage = Options.bIsImage;
-			ObjectOptions.State = Options.State;
-			FObjectGenerationResult ObjectResult;
-			GenerateObject(ObjectOptions, ObjectResult, ObjectNode);
-			OutResult.Op = ObjectResult.Op;
-			return;
-		}
-
-		else if (InNode->GetType()->IsA(NodeScalar::GetStaticType()))
+		if (pNode->GetType()->IsA(NodeScalar::GetStaticType()))
 		{
-			const NodeScalar* ScalarNode = static_cast<const NodeScalar*>(InNode.get());
+			const NodeScalar* ScalarNode = static_cast<const NodeScalar*>(pNode.get());
 			FScalarGenerationResult ScalarResult;
 			GenerateScalar(ScalarResult, Options, ScalarNode);
-			OutResult.Op = ScalarResult.op;
-			return;
+			return ScalarResult.op;
 		}
 
-		else if (InNode->GetType()->IsA(NodeColour::GetStaticType()))
+		else if (pNode->GetType()->IsA(NodeColour::GetStaticType()))
 		{
-			const NodeColour* ColorNode = static_cast<const NodeColour*>(InNode.get());
+			const NodeColour* ColorNode = static_cast<const NodeColour*>(pNode.get());
 			FColorGenerationResult Result;
 			GenerateColor(Result, Options, ColorNode);
-			OutResult.Op = Result.op;
-			return;
+			return Result.op;
 		}
 
-		else if (InNode->GetType()->IsA(NodeProjector::GetStaticType()))
+		else if (pNode->GetType()->IsA(NodeProjector::GetStaticType()))
 		{
-			const NodeProjector* projNode = static_cast<const NodeProjector*>(InNode.get());
+			const NodeProjector* projNode = static_cast<const NodeProjector*>(pNode.get());
 			FProjectorGenerationResult ProjResult;
 			GenerateProjector(ProjResult, Options, projNode);
-			OutResult.Op = ProjResult.op;
-			return;
+			return ProjResult.op;
 		}
 
-		else if (InNode->GetType()->IsA(NodeSurfaceNew::GetStaticType()))
+		else if (pNode->GetType()->IsA(NodeSurfaceNew::GetStaticType()))
 		{
-			const NodeSurfaceNew* surfNode = static_cast<const NodeSurfaceNew*>(InNode.get());
+			const NodeSurfaceNew* surfNode = static_cast<const NodeSurfaceNew*>(pNode.get());
 
 			// This happens only if we generate a node graph that has a NodeSurfaceNew at the root.
-			FSurfaceGenerationResult SurfResult;
+			FSurfaceGenerationResult surfResult;
 			FSurfaceGenerationOptions SurfaceOptions(Options);
-			GenerateSurface(SurfResult, SurfaceOptions, surfNode);
-			OutResult.Op = SurfResult.surfaceOp;
-			return;
+			GenerateSurface(surfResult, SurfaceOptions, surfNode);
+			return surfResult.surfaceOp;
 		}
 
-		else if (InNode->GetType()->IsA(NodeSurfaceVariation::GetStaticType()))
+		else if (pNode->GetType()->IsA(NodeSurfaceVariation::GetStaticType()))
 		{
 			// This happens only if we generate a node graph that has a NodeSurfaceVariation at the root.
-			return;
+			return nullptr;
 		}
 
-		else if (InNode->GetType()->IsA(NodeSurfaceSwitch::GetStaticType()))
+		else if (pNode->GetType()->IsA(NodeSurfaceSwitch::GetStaticType()))
 		{
 			// This happens only if we generate a node graph that has a NodeSurfaceSwitch at the root.
-			return;
+			return nullptr;
 		}
 
-		else if (InNode->GetType()->IsA(NodeModifier::GetStaticType()))
+		else if (pNode->GetType()->IsA(NodeModifier::GetStaticType()))
 		{
 			// This happens only if we generate a node graph that has a modifier at the root.
-			return;
+			return nullptr;
 		}
 
-		else if (InNode->GetType()->IsA(NodeComponent::GetStaticType()))
+		else if (pNode->GetType()->IsA(NodeComponent::GetStaticType()))
 		{
-			const NodeComponent* ComponentNode = static_cast<const NodeComponent*>(InNode.get());
-			FComponentGenerationOptions ComponentOptions(Options, nullptr);
-			GenerateComponent(ComponentOptions, OutResult, ComponentNode);
-			return;
-		}
-
-		else
-		{
-			// Unsupported node.
-			check(false);
+			const NodeComponent* ComponentNode = static_cast<const NodeComponent*>(pNode.get());
+			FComponentGenerationOptions ComponentOptions( Options, nullptr );
+			FGenericGenerationResult Result;
+			GenerateComponent(ComponentOptions, Result, ComponentNode);
+			return Result.op;
 		}
 
-	}
 
-
-	void CodeGenerator::GenerateObject(const FObjectGenerationOptions& InOptions, FObjectGenerationResult& OutResult, const NodeObject* InUntypedNode)
-	{
-		if (!InUntypedNode)
-		{
-			OutResult = {};
-			return;
-		}
+		Ptr<ASTOp> ResultOp;
 
 		// See if it was already generated
-		FGeneratedObjectCacheKey Key;
-		Key.Node = InUntypedNode;
-		Key.Options = InOptions;
-		FGeneratedObjectsMap::ValueType* Found = GeneratedObjects.Find(Key);
-		if (Found)
-		{
-			OutResult = *Found;
-			return;
-		}
 
-		// Generate for each different type of node
-		const FNodeType* Type = InUntypedNode->GetType();
-		if (Type == NodeObjectNew::GetStaticType())
-		{
-			GenerateObject_New(InOptions, OutResult, static_cast<const NodeObjectNew*>(InUntypedNode));
-		}
-		else if (Type == NodeObjectGroup::GetStaticType())
+		FGeneratedCacheKey Key;
+		Key.Node = pNode;
+		Key.Options = Options;
+		FGeneratedGenericNodesMap::ValueType* it = GeneratedGenericNodes.Find(Key);
+		if (it)
 		{
-			GenerateObject_Group(InOptions, OutResult, static_cast<const NodeObjectGroup*>(InUntypedNode));
+			ResultOp = it->op;
 		}
 		else
 		{
-			check(false);
+			FGenericGenerationResult Result;
+
+			// Generate for each different type of node
+			if (pNode->GetType() == NodeObjectNew::GetStaticType())
+			{
+				Generate_ObjectNew(Options, Result, static_cast<const NodeObjectNew*>(pNode.get()));
+			}
+			else if (pNode->GetType()==NodeObjectGroup::GetStaticType())
+			{
+				Generate_ObjectGroup(Options, Result, static_cast<const NodeObjectGroup*>(pNode.get()));
+			}
+			else
+			{
+				check(false);
+			}
+
+			ResultOp = Result.op;
+			GeneratedGenericNodes.Add(Key, Result);
 		}
 
-		// Cache the result
-		GeneratedObjects.Add(Key, OutResult);
+		// debug: expensive check of all code generation
+//        if (ResultOp)
+//        {
+//            ASTOpList roots;
+//            roots.push_back(ResultOp);
+//            ASTOp::FullAssert(roots);
+//        }
+
+		return ResultOp;
 	}
 
 
@@ -541,6 +526,7 @@ void CodeGenerator::GenerateGeneric(const FGenericGenerationOptions& Options, FG
 	}
 
 
+
 	//---------------------------------------------------------------------------------------------
 	void CodeGenerator::GenerateComponent(const FComponentGenerationOptions& InOptions, FGenericGenerationResult& OutResult, const NodeComponent* InUntypedNode)
 	{
@@ -554,7 +540,7 @@ void CodeGenerator::GenerateGeneric(const FGenericGenerationOptions& Options, FG
 		FGeneratedComponentCacheKey Key;
 		Key.Node = InUntypedNode;
 		Key.Options = InOptions;
-		FGeneratedComponentMap::ValueType* it = GeneratedComponents.Find(Key);
+		GeneratedComponentMap::ValueType* it = GeneratedComponents.Find(Key);
 		if (it)
 		{
 			OutResult = *it;
@@ -571,7 +557,7 @@ void CodeGenerator::GenerateGeneric(const FGenericGenerationOptions& Options, FG
 		{
 			// Nothing to do because it is all preprocessed in the first code generator stage
 			//GenerateComponent_Edit(InOptions, OutResult, static_cast<const NodeComponentEdit*>(InUntypedNode));
-			OutResult.Op = InOptions.BaseInstance;
+			OutResult.op = InOptions.BaseInstance;
 		}
 		else if (Type == NodeComponentSwitch::GetStaticType())
 		{
@@ -602,11 +588,13 @@ void CodeGenerator::GenerateGeneric(const FGenericGenerationOptions& Options, FG
 		{
 			if (const NodeLOD* LODNode = node.LODs[LODIndex].get())
 			{
+				CurrentParents.Last().Lod = LODIndex;
+
 				FLODGenerationOptions LODOptions(Options, LODIndex, InNode );
 				FGenericGenerationResult LODResult;
 				Generate_LOD(LODOptions, LODResult, LODNode);
 
-				LODsOp->lods.Emplace(LODsOp, LODResult.Op);
+				LODsOp->lods.Emplace(LODsOp, LODResult.op);
 			}
 		}
 
@@ -616,7 +604,7 @@ void CodeGenerator::GenerateGeneric(const FGenericGenerationOptions& Options, FG
 		InstanceOp->value = LODsOp;
 		InstanceOp->ExternalId = node.Id;
 
-		Result.Op = InstanceOp;
+		Result.op = InstanceOp;
 
 		// Add a conditional if this component has conditions
 		for (const FirstPassGenerator::FComponent& Component: FirstPass.Components)
@@ -636,10 +624,10 @@ void CodeGenerator::GenerateGeneric(const FGenericGenerationOptions& Options, FG
 				Ptr<ASTOpConditional> IfOp = new ASTOpConditional();
 				IfOp->type = EOpType::IN_CONDITIONAL;
 				IfOp->no = Options.BaseInstance;
-				IfOp->yes = Result.Op;
+				IfOp->yes = Result.op;
 				IfOp->condition = ConditionOp;
 
-				Result.Op = IfOp;
+				Result.op = IfOp;
 			}
 		}
 	}
@@ -652,7 +640,7 @@ void CodeGenerator::GenerateGeneric(const FGenericGenerationOptions& Options, FG
 		if (Node->Options.Num() == 0)
 		{
 			// No options in the switch!
-			Result.Op = Options.BaseInstance;
+			Result.op = Options.BaseInstance;
 			return;
 		}
 
@@ -662,9 +650,7 @@ void CodeGenerator::GenerateGeneric(const FGenericGenerationOptions& Options, FG
 		// Variable value
 		if (Node->Parameter)
 		{
-			FScalarGenerationResult ParamResult;
-			GenerateScalar(ParamResult, Options, Node->Parameter.get());
-			Op->variable = ParamResult.op;
+			Op->variable = Generate_Generic(Node->Parameter.get(), Options);
 		}
 		else
 		{
@@ -681,7 +667,7 @@ void CodeGenerator::GenerateGeneric(const FGenericGenerationOptions& Options, FG
 			{
 				FGenericGenerationResult BaseResult;
 				GenerateComponent(Options, BaseResult, Node->Options[OptionIndex].get());
-				Branch = BaseResult.Op;
+				Branch = BaseResult.op;
 			}
 			else
 			{
@@ -692,7 +678,7 @@ void CodeGenerator::GenerateGeneric(const FGenericGenerationOptions& Options, FG
 			Op->cases.Emplace(OptionIndex, Op, Branch);
 		}
 
-		Result.Op = Op;
+		Result.op = Op;
 	}
 
 
@@ -706,7 +692,7 @@ void CodeGenerator::GenerateGeneric(const FGenericGenerationOptions& Options, FG
 			FGenericGenerationResult BranchResults;
 
 			GenerateComponent(Options, BranchResults, Node->DefaultComponent.get());
-			CurrentMeshOp = BranchResults.Op;
+			CurrentMeshOp = BranchResults.op;
 		}
 
 		// Process variations in reverse order, since conditionals are built bottom-up.
@@ -739,7 +725,7 @@ void CodeGenerator::GenerateGeneric(const FGenericGenerationOptions& Options, FG
 				FGenericGenerationResult BranchResults;
 				GenerateComponent(Options, BranchResults, Node->Variations[VariationIndex].Component.get());
 
-				VariationMeshOp = BranchResults.Op;
+				VariationMeshOp = BranchResults.op;
 			}
 
 			Ptr<ASTOpConditional> Conditional = new ASTOpConditional;
@@ -751,7 +737,7 @@ void CodeGenerator::GenerateGeneric(const FGenericGenerationOptions& Options, FG
 			CurrentMeshOp = Conditional;
 		}
 
-		Result.Op = CurrentMeshOp;
+		Result.op = CurrentMeshOp;
 	}
 
 
@@ -1710,41 +1696,53 @@ void CodeGenerator::GenerateGeneric(const FGenericGenerationOptions& Options, FG
             lastCompOp = iop;
         }
 
-        Result.Op = lastCompOp;
+        Result.op = lastCompOp;
     }
 
 
     //---------------------------------------------------------------------------------------------
    // Ptr<ASTOp> CodeGenerator::Visit( const NodeObjectNew::Private& node )
-	void CodeGenerator::GenerateObject_New(const FObjectGenerationOptions& Options, FObjectGenerationResult& OutResult, const NodeObjectNew* InNode)
+	void CodeGenerator::Generate_ObjectNew(const FGenericGenerationOptions& Options, FGenericGenerationResult& Result, const NodeObjectNew* InNode)
 	{
 		MUTABLE_CPUPROFILER_SCOPE(NodeObjectNew);
 
 		// There is always at least a null parent
-		bool bIsChildObject = Options.ParentObjectNode != nullptr;
+		bool bIsChildObject = CurrentParents.Num() > 1;
 
 		// Add this object as current parent
-		FObjectGenerationOptions ChildOptions = Options;
-		ChildOptions.ParentObjectNode = InNode;
+        CurrentParents.Add( FParentKey() );
+        CurrentParents.Last().ObjectNode = InNode;
 
         // Parse the child objects first, which will accumulate operations in the patching lists
-        for ( int32 ChildIndex=0; ChildIndex < InNode->Children.Num(); ++ChildIndex)
+        for ( int32 t=0; t< InNode->Children.Num(); ++t )
         {
-            if ( const NodeObject* pChildNode = InNode->Children[ChildIndex].get() )
+            if ( const NodeObject* pChildNode = InNode->Children[t].get() )
             {
-				// If there are parent objects, the condition of this object depends on the condition of the parent object
-				FObjectGenerationOptions ThisChildOptions = ChildOptions;
-				if (!ThisChildOptions.CurrentObjectCondition)
-				{
-					// In case there is no group node, we generate a constant true condition
-					// This condition will be overwritten by the group nodes.
-					ThisChildOptions.CurrentObjectCondition = new ASTOpConstantBool(true);
-				}
+                Ptr<ASTOp> paramOp;
+
+                // If there are parent objects, the condition of this object depends on the
+                // condition of the parent object
+                if ( CurrentObject.Num() )
+                {
+                    paramOp = CurrentObject.Last().Condition;
+                }
+                else
+                {
+                    // In case there is no group node, we generate a constant true condition
+                    // This condition will be overwritten by the group nodes.
+                    Ptr<ASTOpConstantBool> op = new ASTOpConstantBool();
+                    op->bValue = true;
+                    paramOp = op;
+                }
+
+				FObjectGenerationData data;
+                data.Condition = paramOp;
+                CurrentObject.Add( data );
 
                 // This op is ignored: everything is stored as patches to apply to the parent when it is compiled.
-				FObjectGenerationResult ThisResult;
-                GenerateObject(ThisChildOptions, ThisResult, pChildNode);
-				OutResult.AdditionalComponents.Append(ThisResult.AdditionalComponents);
+                Generate_Generic( pChildNode, Options );
+
+                CurrentObject.Pop();
             }
         }
 
@@ -1758,15 +1756,15 @@ void CodeGenerator::GenerateGeneric(const FGenericGenerationOptions& Options, FG
 		}
 
 		// Add the components in this node
-        for ( int32 ComponentIndex=0; ComponentIndex < InNode->Components.Num(); ++ComponentIndex)
+        for ( int32 t=0; t< InNode->Components.Num(); ++t )
         {
-			const NodeComponent* ComponentNode = InNode->Components[ComponentIndex].get();
+			const NodeComponent* ComponentNode = InNode->Components[t].get();
             if (ComponentNode)
             {
 				FComponentGenerationOptions ComponentOptions( Options, LastCompOp );
 				FGenericGenerationResult ComponentResult;
 				GenerateComponent(ComponentOptions, ComponentResult, ComponentNode);
-				LastCompOp = ComponentResult.Op;
+				LastCompOp = ComponentResult.op;
             }
         }
 
@@ -1778,12 +1776,12 @@ void CodeGenerator::GenerateGeneric(const FGenericGenerationOptions& Options, FG
 		}
 
 		// Add the components from child objects
-		FObjectGenerationResult::FAdditionalComponentKey ThisKey;
-		ThisKey.ObjectNode = InNode;
-		TArray<FObjectGenerationResult::FAdditionalComponentData>* ThisAdditionalComponents = OutResult.AdditionalComponents.Find(ThisKey);
+		FAdditionalComponentKey ThisKey;
+		ThisKey.ObjectNode = CurrentParents.Last().ObjectNode;
+		TArray<FAdditionalComponentData>* ThisAdditionalComponents = AdditionalComponents.Find(ThisKey);
 		if (LastCompOp && ThisAdditionalComponents)
 		{
-			for (const FObjectGenerationResult::FAdditionalComponentData& Additional : *ThisAdditionalComponents)
+			for (const FAdditionalComponentData& Additional : *ThisAdditionalComponents)
 			{
 				check(Additional.PlaceholderOp);
 				ASTOp::Replace(Additional.PlaceholderOp, LastCompOp);
@@ -1792,17 +1790,23 @@ void CodeGenerator::GenerateGeneric(const FGenericGenerationOptions& Options, FG
 		}
 
 		// Store this chain of components for use in parent objects if necessary
+		// 2 is because there must be a parent and there is always a null element as well.
 		if (LastCompOp && bIsChildObject)
 		{
-			FObjectGenerationResult::FAdditionalComponentKey ParentKey;
-			ParentKey.ObjectNode = Options.ParentObjectNode;
+			// TODO: Directly to the root object?
+			const FParentKey& ParentObjectKey = CurrentParents[CurrentParents.Num() - 2];
+			FAdditionalComponentKey ParentKey;
+			ParentKey.ObjectNode = ParentObjectKey.ObjectNode;
 
-			FObjectGenerationResult::FAdditionalComponentData Data;
+			FAdditionalComponentData Data;
 			Data.ComponentOp = LastCompOp;
 			Data.PlaceholderOp = PlaceholderOp;
-			OutResult.AdditionalComponents.FindOrAdd(ParentKey).Add(Data);
+			AdditionalComponents.FindOrAdd(ParentKey).Add(Data);
 		}
 
+
+        Ptr<ASTOp> RootOp = LastCompOp;
+
 		// Add an ASTOpAddExtensionData for each connected ExtensionData node
 		for (const NodeObjectNew::FNamedExtensionDataNode& NamedNode : InNode->ExtensionDataNodes)
 		{
@@ -1825,14 +1829,15 @@ void CodeGenerator::GenerateGeneric(const FGenericGenerationOptions& Options, FG
 			}
 
 			FConditionalExtensionDataOp& SavedOp = ConditionalExtensionDataOps.AddDefaulted_GetRef();
-			SavedOp.Condition = Options.CurrentObjectCondition;
+			if (CurrentObject.Num() > 0)
+			{
+				SavedOp.Condition = CurrentObject.Last().Condition;
+			}
 			SavedOp.ExtensionDataOp = ChildResult.Op;
 			SavedOp.ExtensionDataName = NamedNode.Name;
 		}
 
-		Ptr<ASTOp> RootOp = LastCompOp;
-
-		if (!Options.ParentObjectNode)
+		if (CurrentObject.Num() == 0)
 		{
 			for (const FConditionalExtensionDataOp& SavedOp : ConditionalExtensionDataOps)
 			{
@@ -1858,52 +1863,63 @@ void CodeGenerator::GenerateGeneric(const FGenericGenerationOptions& Options, FG
 			}
 		}
 
-		OutResult.Op = RootOp;
+        CurrentParents.Pop();
+
+        Result.op = RootOp;
     }
 
 
-	void CodeGenerator::GenerateObject_Group(const FObjectGenerationOptions& Options, FObjectGenerationResult& OutResult, const NodeObjectGroup* InNode)
+    //---------------------------------------------------------------------------------------------
+    //Ptr<ASTOp> CodeGenerator::Visit( const NodeObjectGroup::Private& node )
+	void CodeGenerator::Generate_ObjectGroup(const FGenericGenerationOptions& Options, FGenericGenerationResult& Result, const NodeObjectGroup* InNode)
 	{
 		const NodeObjectGroup::Private& node = *InNode->GetPrivate();
 
-		TArray<FString> UsedNames;
+		TArray<FString> usedNames;
 
         // Parse the child objects first, which will accumulate operations in the patching lists
-        for ( int32 ChildIndex=0; ChildIndex <node.m_children.Num(); ++ChildIndex)
+        for ( int32 t=0; t<node.m_children.Num(); ++t )
         {
-            if ( const NodeObject* pChildNode = node.m_children[ChildIndex].get() )
+            if ( const NodeObject* pChildNode = node.m_children[t].get() )
             {
 				// Look for the child condition in the first pass
-                Ptr<ASTOp> ConditionOp;
-				bool bFound = false;
-                for( int32 ObjectIndex = 0; !bFound && ObjectIndex<FirstPass.Objects.Num(); ObjectIndex++ )
+                Ptr<ASTOp> conditionOp;
+				bool found = false;
+                for( int32 i = 0; !found && i != FirstPass.Objects.Num(); i++ )
 				{
-					FirstPassGenerator::FObject& Candidate = FirstPass.Objects[ObjectIndex];
-					if (Candidate.Node == pChildNode)
+					FirstPassGenerator::FObject& it = FirstPass.Objects[i];
+					if (it.Node == pChildNode)
 					{
-						bFound = true;
-						ConditionOp = Candidate.Condition;
+						found = true;
+						conditionOp = it.Condition;
 					}
 				}
 
-				FObjectGenerationOptions ChildOptions = Options;
-				ChildOptions.CurrentObjectCondition = ConditionOp;
+                // \todo
+                // TrunkStraight_02_BranchTop crash
+                // It may happen with partial compilations?
+                // check(found);
+
+				FObjectGenerationData data;
+                data.Condition = conditionOp;
+                CurrentObject.Add( data );
+
+                // This op is ignored: everything is stored as patches to apply to the parent when
+                // it is compiled.
+                Generate_Generic( pChildNode, Options );
 
-                // The result op is ignored: everything is stored as data to apply when the parent is compiled.
-				FObjectGenerationResult ChildResult;
-                GenerateObject( Options, ChildResult, pChildNode );
-				OutResult.AdditionalComponents.Append(ChildResult.AdditionalComponents);
+                CurrentObject.Pop();
 
 				// Check for duplicated child names
-				FString ChildName = pChildNode->GetName();
-				if (UsedNames.Contains(ChildName))
+				FString strChildName = pChildNode->GetName();
+				if (usedNames.Contains(strChildName))
 				{
-					FString Msg = FString::Printf(TEXT("Object group has more than one children with the same name [%s]."), *ChildName );
+					FString Msg = FString::Printf(TEXT("Object group has more than one children with the same name [%s]."), *strChildName );
 					ErrorLog->Add(Msg, ELMT_WARNING, InNode->GetMessageContext());
 				}
 				else
 				{
-					UsedNames.Add(ChildName);
+					usedNames.Add(strChildName);
 				}
             }
         }
diff --git a/Engine/Plugins/Mutable/Source/MutableTools/Private/MuT/CodeGenerator.h b/Engine/Plugins/Mutable/Source/MutableTools/Private/MuT/CodeGenerator.h
--- a/Engine/Plugins/Mutable/Source/MutableTools/Private/MuT/CodeGenerator.h
+++ b/Engine/Plugins/Mutable/Source/MutableTools/Private/MuT/CodeGenerator.h
@@ -112,13 +112,14 @@ namespace mu
 
     public:
 
-        CodeGenerator( CompilerOptions::Private* );
+        CodeGenerator( CompilerOptions::Private* options );
 
         //! Data will be stored in States
         void GenerateRoot( const Ptr<const Node> );
 
 	public:
 
+		// Generic top-level nodes
 		struct FGenericGenerationOptions
 		{
 			friend FORCEINLINE uint32 GetTypeHash(const FGenericGenerationOptions& InKey)
@@ -138,14 +139,6 @@ namespace mu
 			bool bIsImage = false;
 		};
 
-		struct FObjectGenerationOptions : public FGenericGenerationOptions
-		{
-			const NodeObjectNew* ParentObjectNode = nullptr;
-
-			/** Condition that enables a specific object. */
-			Ptr<ASTOp> CurrentObjectCondition;
-		};
-
 		struct FComponentGenerationOptions : public FGenericGenerationOptions
 		{
 			FComponentGenerationOptions(const FGenericGenerationOptions& BaseOptions, const Ptr<ASTOp>& InBaseInstance )
@@ -197,7 +190,7 @@ namespace mu
 
 		struct FGenericGenerationResult
 		{
-			Ptr<ASTOp> Op;
+			Ptr<ASTOp> op;
 		};
 
 		struct FGeneratedCacheKey
@@ -216,81 +209,28 @@ namespace mu
 			FORCEINLINE bool operator==(const FGeneratedCacheKey& Other) const = default;
 		};
 
+		typedef TMap<FGeneratedCacheKey, FGenericGenerationResult> FGeneratedGenericNodesMap;
+		FGeneratedGenericNodesMap GeneratedGenericNodes;
 
-		void GenerateGeneric(const FGenericGenerationOptions& Options, FGenericGenerationResult& OutResult, const Ptr<const Node> InNode);
+		Ptr<ASTOp> Generate_Generic(const Ptr<const Node>, const FGenericGenerationOptions& );
 		void Generate_LOD(const FLODGenerationOptions&, FGenericGenerationResult&, const NodeLOD*);
-
-
-		struct FGeneratedObjectCacheKey
-		{
-			Ptr<const Node> Node;
-			FObjectGenerationOptions Options;
-
-			friend FORCEINLINE uint32 GetTypeHash(const FGeneratedObjectCacheKey& InKey)
-			{
-				uint32 KeyHash = 0;
-				KeyHash = HashCombineFast(KeyHash, ::GetTypeHash(InKey.Node.get()));
-				KeyHash = HashCombineFast(KeyHash, GetTypeHash(InKey.Options));
-				return KeyHash;
-			}
-
-			FORCEINLINE bool operator==(const FGeneratedObjectCacheKey& Other) const
-			{
-				return Node == Other.Node && Options == Other.Options;
-			}
-		};
-
-		struct FObjectGenerationResult : public FGenericGenerationResult
-		{
-			/** List of additional components to add to an object that come from child objects.
-			 * The index is the object and lod that should receive the components.
-			 */
-			struct FAdditionalComponentKey
-			{
-				const NodeObjectNew* ObjectNode = nullptr;
-
-				FORCEINLINE bool operator==(const FAdditionalComponentKey& Other) const
-				{
-					return ObjectNode == Other.ObjectNode;
-				}
-
-				friend FORCEINLINE uint32 GetTypeHash(const FAdditionalComponentKey& InKey)
-				{
-					uint32 KeyHash = ::GetTypeHash(InKey.ObjectNode);
-					return KeyHash;
-				}
-			};
-
-			struct FAdditionalComponentData
-			{
-				Ptr<ASTOp> ComponentOp;
-				Ptr<ASTOp> PlaceholderOp;
-			};
-
-			TMap< FAdditionalComponentKey, TArray<FAdditionalComponentData> > AdditionalComponents;
-		};
-
-		typedef TMap<FGeneratedObjectCacheKey, FObjectGenerationResult> FGeneratedObjectsMap;
-		FGeneratedObjectsMap GeneratedObjects;
-
-		void GenerateObject(const FObjectGenerationOptions&, FObjectGenerationResult&, const NodeObject*);
-		void GenerateObject_New(const FObjectGenerationOptions&, FObjectGenerationResult&, const NodeObjectNew*);
-		void GenerateObject_Group(const FObjectGenerationOptions&, FObjectGenerationResult&, const NodeObjectGroup*);
+		void Generate_ObjectNew(const FGenericGenerationOptions&, FGenericGenerationResult&, const NodeObjectNew*);
+		void Generate_ObjectGroup(const FGenericGenerationOptions&, FGenericGenerationResult&, const NodeObjectGroup*);
 
     public:
 
-        /** */
-        const CompilerOptions::Private* CompilerOptions = nullptr;
+        //! Settings
+        CompilerOptions::Private* CompilerOptions = nullptr;
 
-		/** */
+		//!
 		FirstPassGenerator FirstPass;
 
-        /** Used to accumulate all generation error messages. */
+        //!
         TSharedPtr<FErrorLog> ErrorLog;
 
-        /** After the entire code generation this contains the information about all the states. */
-        typedef TArray< TPair<FObjectState, Ptr<ASTOp>> > FStateList;
-        FStateList States;
+        //! After the entire code generation this contains the information about all the states
+        typedef TArray< TPair<FObjectState, Ptr<ASTOp>> > StateList;
+        StateList States;
 
     private:
 
@@ -313,10 +253,60 @@ namespace mu
 			TSet<uint32> UniqueVertexIDGroups;
 		};
 		FGenerateMeshConstantState GenerateMeshConstantState;
-
-		/** Any task accessing GenerateMeshConstantState should run in this pipe. */
 		UE::Tasks::FPipe GenerateMeshConstantPipe;
 
+		/** List of image resources for every image formats that have been generated so far as palceholders for missing images. */
+		TSharedPtr<FImage> MissingImage[size_t(EImageFormat::Count)];
+
+		struct FParentKey
+		{
+			const NodeObjectNew* ObjectNode = nullptr;
+			int32 Lod = -1;
+		};
+
+		// TODO: complete move of this top-down state to method arguments.
+		TArray< FParentKey > CurrentParents;
+
+        /** List of additional components to add to an object that come from child objects.
+        * The index is the object and lod that should receive the components.
+		*/
+        struct FAdditionalComponentKey
+        {
+			FAdditionalComponentKey()
+            {
+				ObjectNode = nullptr;
+            }
+
+            const NodeObjectNew* ObjectNode;
+
+			FORCEINLINE bool operator==(const FAdditionalComponentKey& Other) const
+			{
+				return ObjectNode == Other.ObjectNode;
+			}
+
+			friend FORCEINLINE uint32 GetTypeHash(const FAdditionalComponentKey& InKey)
+			{
+				uint32 KeyHash = 0;
+				KeyHash = HashCombineFast(KeyHash, ::GetTypeHash(InKey.ObjectNode));
+				return KeyHash;
+			}
+		};
+
+		struct FAdditionalComponentData
+		{
+			Ptr<ASTOp> ComponentOp;
+			Ptr<ASTOp> PlaceholderOp;
+		};
+
+        TMap< FAdditionalComponentKey, TArray<FAdditionalComponentData> > AdditionalComponents;
+
+
+        struct FObjectGenerationData
+        {
+            // Condition that enables a specific object
+            Ptr<ASTOp> Condition;
+        };
+		TArray<FObjectGenerationData> CurrentObject;
 
 		/** The key for generated tables is made of the source table and a parameter name. */
 		struct FTableCacheKey
@@ -368,8 +358,8 @@ namespace mu
 			}
 		};
 
-		typedef TMap<FGeneratedComponentCacheKey, FGenericGenerationResult> FGeneratedComponentMap;
-		FGeneratedComponentMap GeneratedComponents;
+		typedef TMap<FGeneratedComponentCacheKey, FGenericGenerationResult> GeneratedComponentMap;
+		GeneratedComponentMap GeneratedComponents;
 
 		void GenerateComponent(const FComponentGenerationOptions&, FGenericGenerationResult&, const NodeComponent*);
 		void GenerateComponent_New(const FComponentGenerationOptions&, FGenericGenerationResult&, const NodeComponentNew*);
@@ -449,7 +439,7 @@ namespace mu
 			{
 			}
 
-			Ptr<const Node> Node;
+			NodePtrConst Node;
 			FImageGenerationOptions Options;
 
 			friend FORCEINLINE uint32 GetTypeHash(const FGeneratedImageCacheKey& InKey)
@@ -463,8 +453,8 @@ namespace mu
 			FORCEINLINE bool operator==(const FGeneratedImageCacheKey& Other) const = default;
 		};
 
-		typedef TMap<FGeneratedImageCacheKey, FImageGenerationResult> FGeneratedImagesMap;
-		FGeneratedImagesMap GeneratedImages;
+		typedef TMap<FGeneratedImageCacheKey, FImageGenerationResult> GeneratedImagesMap;
+		GeneratedImagesMap GeneratedImages;
 
 		void GenerateImage(const FImageGenerationOptions&, FImageGenerationResult& result, const NodeImagePtrConst& node);
 		void GenerateImage_Constant(const FImageGenerationOptions&, FImageGenerationResult&, const NodeImageConstant*);
@@ -491,22 +481,22 @@ namespace mu
 		void GenerateImage_NormalComposite(const FImageGenerationOptions&, FImageGenerationResult&, const NodeImageNormalComposite*);
 		void GenerateImage_Transform(const FImageGenerationOptions&, FImageGenerationResult&, const NodeImageTransform*);
 
-		/** */
+		//!
 		TSharedPtr<FImage> GenerateMissingImage(EImageFormat);
 
-		/** */
+		//!
 		Ptr<ASTOp> GenerateMissingImageCode(const TCHAR* strWhere, EImageFormat, const void* errorContext, const FImageGenerationOptions& Options);
 
-		/** */
+		//!
 		Ptr<ASTOp> GeneratePlainImageCode(const FVector4f& Color, const FImageGenerationOptions& Options);
 
-		/** */
+		//!
 		Ptr<ASTOp> GenerateImageFormat(Ptr<ASTOp>, EImageFormat);
 
-		/** */
+		//!
 		Ptr<ASTOp> GenerateImageUncompressed(Ptr<ASTOp>);
 
-		/** */
+		//!
 		Ptr<ASTOp> GenerateImageSize(Ptr<ASTOp>, UE::Math::TIntVector2<int32>);
 
 		/** Evaluate if the image to generate is big enough to be split in separate operations and tiled afterwards. */
diff --git a/Engine/Plugins/Mutable/Source/MutableTools/Private/MuT/CodeGenerator_Color.cpp b/Engine/Plugins/Mutable/Source/MutableTools/Private/MuT/CodeGenerator_Color.cpp
--- a/Engine/Plugins/Mutable/Source/MutableTools/Private/MuT/CodeGenerator_Color.cpp
+++ b/Engine/Plugins/Mutable/Source/MutableTools/Private/MuT/CodeGenerator_Color.cpp
@@ -199,9 +199,7 @@ namespace mu
 		// Variable value
 		if (node.Parameter)
 		{
-			FScalarGenerationResult ParamResult;
-			GenerateScalar(ParamResult, Options, node.Parameter.get());
-			op->variable = ParamResult.op;
+			op->variable = Generate_Generic(node.Parameter.get(), Options);
 		}
 		else
 		{
@@ -210,21 +208,19 @@ namespace mu
 		}
 
 		// Options
-		for (int32 OptionIndex = 0; OptionIndex < node.Options.Num(); ++OptionIndex)
+		for (int32 t = 0; t < node.Options.Num(); ++t)
 		{
 			Ptr<ASTOp> branch;
-			if (node.Options[OptionIndex])
+			if (node.Options[t])
 			{
-				FColorGenerationResult ParamResult;
-				GenerateColor(ParamResult, Options, node.Options[OptionIndex].get());
-				branch = ParamResult.op;
+				branch = Generate_Generic(node.Options[t].get(), Options);
 			}
 			else
 			{
 				// This argument is required
 				branch = GenerateMissingColourCode(TEXT("Switch option"), Typed->GetMessageContext());
 			}
-			op->cases.Emplace((int16)OptionIndex, op, branch);
+			op->cases.Emplace((int16)t, op, branch);
 		}
 
 		result.op = op;
diff --git a/Engine/Plugins/Mutable/Source/MutableTools/Private/MuT/CodeGenerator_ExtensionData.cpp b/Engine/Plugins/Mutable/Source/MutableTools/Private/MuT/CodeGenerator_ExtensionData.cpp
--- a/Engine/Plugins/Mutable/Source/MutableTools/Private/MuT/CodeGenerator_ExtensionData.cpp
+++ b/Engine/Plugins/Mutable/Source/MutableTools/Private/MuT/CodeGenerator_ExtensionData.cpp
@@ -95,12 +95,10 @@ class Node;
 		Ptr<ASTOpSwitch> Op = new ASTOpSwitch;
 		Op->type = EOpType::ED_SWITCH;
 
-		// Variable
+		// Variable name
 		if (Switch->Parameter)
 		{
-			FScalarGenerationResult ParamResult; 
-			GenerateScalar(ParamResult, Options, Switch->Parameter.get());
-			Op->variable = ParamResult.op;
+			Op->variable = Generate_Generic(Switch->Parameter.get(), Options);
 		}
 		else
 		{
diff --git a/Engine/Plugins/Mutable/Source/MutableTools/Private/MuT/CodeGenerator_FirstPass.cpp b/Engine/Plugins/Mutable/Source/MutableTools/Private/MuT/CodeGenerator_FirstPass.cpp
--- a/Engine/Plugins/Mutable/Source/MutableTools/Private/MuT/CodeGenerator_FirstPass.cpp
+++ b/Engine/Plugins/Mutable/Source/MutableTools/Private/MuT/CodeGenerator_FirstPass.cpp
@@ -129,9 +129,9 @@ namespace mu
 
         if ( States.IsEmpty() )
         {
-            FObjectState Data;
-            Data.Name = "Default";
-            States.Add( Data );
+            FObjectState data;
+            data.Name = "Default";
+            States.Emplace( data, Root );
         }
 	}
 
@@ -189,7 +189,6 @@ namespace mu
 		}
 		else
 		{
-			// This node type is not supported in this pass.
 			check(false);
 		}
 	}
@@ -278,33 +277,33 @@ namespace mu
 
         case NodeSurfaceVariation::VariationType::State:
         {
-            int32 StateCount = States.Num();
+            int32 stateCount = States.Num();
 
             // Default
             {
                 // Store the states for the default branch here
-                StateCondition DefaultStates;
+                StateCondition defaultStates;
                 {
 					StateCondition AllTrue;
-					AllTrue.Init(true, StateCount);
-					DefaultStates = CurrentStateCondition.Last().IsEmpty()
+					AllTrue.Init(true,stateCount);
+                    defaultStates = CurrentStateCondition.Last().IsEmpty()
                             ? AllTrue
                             : CurrentStateCondition.Last();
 
                     for (const NodeSurfaceVariation::FVariation& v: InNode->Variations)
                     {
-                        for( int32 StateIndex=0; StateIndex < StateCount; ++StateIndex)
+                        for( size_t s=0; s<stateCount; ++s )
                         {
-                            if (States[StateIndex].Name==v.Tag)
+                            if (States[s].Key.Name==v.Tag)
                             {
                                 // Remove this state from the default options, since it has its own variation
-								DefaultStates[StateIndex] = false;
+                                defaultStates[s] = false;
                             }
                         }
                     }
                 }
 
-                CurrentStateCondition.Add(DefaultStates);
+                CurrentStateCondition.Add(defaultStates);
 
                 for (const Ptr<NodeSurface>& n : InNode->DefaultSurfaces)
                 {
@@ -323,11 +322,11 @@ namespace mu
             {
                 // Store the states for this variation here
 				StateCondition variationStates;
-				variationStates.Init(false,StateCount);
+				variationStates.Init(false,stateCount);
 
-                for( int32 StateIndex=0; StateIndex< StateCount; ++StateIndex)
+                for( int32 StateIndex=0; StateIndex<stateCount; ++StateIndex)
                 {
-                    if (States[StateIndex].Name==v.Tag)
+                    if (States[StateIndex].Key.Name==v.Tag)
                     {
                         variationStates[StateIndex] = true;
                     }
@@ -570,7 +569,7 @@ namespace mu
         // Accumulate the model states
         for ( const FObjectState& s: InNode->States )
         {
-            States.Add( s );
+            States.Emplace( s, InNode );
 
             if ( s.RuntimeParams.Num() > MUTABLE_MAX_RUNTIME_PARAMETERS_PER_STATE )
             {
diff --git a/Engine/Plugins/Mutable/Source/MutableTools/Private/MuT/CodeGenerator_FirstPass.h b/Engine/Plugins/Mutable/Source/MutableTools/Private/MuT/CodeGenerator_FirstPass.h
--- a/Engine/Plugins/Mutable/Source/MutableTools/Private/MuT/CodeGenerator_FirstPass.h
+++ b/Engine/Plugins/Mutable/Source/MutableTools/Private/MuT/CodeGenerator_FirstPass.h
@@ -184,8 +184,8 @@ namespace mu
         };
         TArray<FTag> Tags;
 
-        /** Accumulate the model states found while generating code. */
-        typedef TArray< FObjectState > StateList;
+        /** Accumulate the model states found while generating code, with their generated root nodes. */
+        typedef TArray< TPair<FObjectState, const Node*> > StateList;
         StateList States;
 
 		/** FParameters added for every node. */
diff --git a/Engine/Plugins/Mutable/Source/MutableTools/Private/MuT/CodeGenerator_Image.cpp b/Engine/Plugins/Mutable/Source/MutableTools/Private/MuT/CodeGenerator_Image.cpp
--- a/Engine/Plugins/Mutable/Source/MutableTools/Private/MuT/CodeGenerator_Image.cpp
+++ b/Engine/Plugins/Mutable/Source/MutableTools/Private/MuT/CodeGenerator_Image.cpp
@@ -119,7 +119,7 @@ namespace mu
 
 		// See if it was already generated
 		FGeneratedImageCacheKey key(Options,Untyped);
-		FGeneratedImagesMap::ValueType* CachedPtr = GeneratedImages.Find(key);
+		GeneratedImagesMap::ValueType* CachedPtr = GeneratedImages.Find(key);
 		if (CachedPtr)
 		{
 			Result = *CachedPtr;
@@ -698,11 +698,9 @@ namespace mu
         Ptr<ASTOpImageInterpolate> op = new ASTOpImageInterpolate();
 
         // Factor
-        if ( NodeScalar* Factor = node.m_pFactor.get() )
+        if ( Node* Factor = node.m_pFactor.get() )
         {
-			FScalarGenerationResult ParamResult;
-			GenerateScalar(ParamResult, Options, Factor);
-			op->Factor = ParamResult.op;
+            op->Factor = Generate_Generic(Factor, Options);
         }
         else
         {
@@ -931,12 +929,10 @@ namespace mu
 
 
         // Factor
-        if ( NodeScalar* Factor = node.m_pFactor.get() )
+        if ( Node* Factor = node.m_pFactor.get() )
         {
-			FScalarGenerationResult ParamResult;
-			GenerateScalar(ParamResult, Options, Factor);
-			op->Factor = ParamResult.op;
-		}
+            op->Factor = Generate_Generic(Factor, Options);
+        }
         else
         {
             // This argument is required
@@ -1066,12 +1062,10 @@ namespace mu
         op->Base = BaseOp;
 
         // Threshold
-        Ptr<ASTOp> ThresholdOp;
-        if ( NodeScalar* ThresholdNode = node.m_pThreshold.get() )
+        Ptr<ASTOp> ThresholdOp = 0;
+        if ( Node* pScalar = node.m_pThreshold.get() )
         {
-			FScalarGenerationResult ParamResult;
-			GenerateScalar(ParamResult, Options, ThresholdNode);
-			ThresholdOp = ParamResult.op;
+			ThresholdOp = Generate_Generic( pScalar, Options);
         }
         else
         {
@@ -1091,7 +1085,7 @@ namespace mu
 
 		MUTABLE_CPUPROFILER_SCOPE(NodeImageResize);
 
-        Ptr<ASTOp> at;
+        Ptr<ASTOp> at = 0;
 
         // Source image
         Ptr<ASTOp> base;
@@ -1164,12 +1158,10 @@ namespace mu
 		const NodeImagePlainColour::Private& node = *InNode->GetPrivate();
 
 		// Source colour
-        Ptr<ASTOp> base;
+        Ptr<ASTOp> base = 0;
         if ( node.m_pColour )
         {
-			FColorGenerationResult ColorResult;
-            GenerateColor( ColorResult, Options, node.m_pColour.get() );
-			base = ColorResult.op;
+            base = Generate_Generic( node.m_pColour.get(), Options);
         }
         else
         {
@@ -1230,7 +1222,8 @@ namespace mu
         if (node.Options.Num() == 0)
 		{
 			// No options in the switch!
-			Result.op = GenerateMissingImageCode(TEXT("Switch option"), EImageFormat::RGBA_UByte, InNode->GetMessageContext(), Options);
+            Ptr<ASTOp> missingOp = GenerateMissingImageCode(TEXT("Switch option"), EImageFormat::RGBA_UByte, InNode->GetMessageContext(), Options);
+			Result.op = missingOp;
 			return;
 		}
 
@@ -1240,9 +1233,7 @@ namespace mu
 		// Variable value
 		if ( node.m_pParameter )
 		{
-			FScalarGenerationResult ParamResult;
-			GenerateScalar(ParamResult, Options, node.m_pParameter.get());
-			op->variable = ParamResult.op;
+            op->variable = Generate_Generic( node.m_pParameter.get(), Options);
 		}
 		else
 		{
@@ -1251,14 +1242,14 @@ namespace mu
 		}
 
 		// Options
-        for ( int32 OptionIndex=0; OptionIndex < node.Options.Num(); ++OptionIndex)
+        for ( std::size_t t=0; t< node.Options.Num(); ++t )
         {
             Ptr<ASTOp> branch;
 
-            if (node.Options[OptionIndex])
+            if (node.Options[t])
             {
 				FImageGenerationResult BaseResult;
-				GenerateImage(Options, BaseResult, node.Options[OptionIndex]);
+				GenerateImage(Options, BaseResult, node.Options[t]);
 				branch = BaseResult.op;
 			}
             else
@@ -1267,7 +1258,7 @@ namespace mu
                 branch = GenerateMissingImageCode(TEXT("Switch option"), EImageFormat::RGBA_UByte, InNode->GetMessageContext(), Options);
             }
 
-            op->cases.Emplace(int16(OptionIndex),op,branch);
+            op->cases.Emplace((int16_t)t,op,branch);
         }
 
         Ptr<ASTOp> switchAt = op;
@@ -1287,10 +1278,8 @@ namespace mu
         // Condition
         if ( node.Parameter )
         {
-			FBoolGenerationResult ParamResult;
-			GenerateBool(ParamResult, Options, node.Parameter.get());
-			op->condition = ParamResult.op;
-		}
+            op->condition = Generate_Generic( node.Parameter.get(), Options);
+        }
         else
         {
             // This argument is required
@@ -1576,9 +1565,7 @@ namespace mu
 			if (node.m_pAngleFadeEnd) pPropsNode->Y = node.m_pAngleFadeEnd;
 			else pPropsNode->Y = pDefaultFade;
 
-			FColorGenerationResult ParamResult;
-			GenerateColor(ParamResult, Options, pPropsNode);
-			ImageRasterOp->angleFadeProperties = ParamResult.op;
+			ImageRasterOp->angleFadeProperties = Generate_Generic(pPropsNode, Options);
 		}
 
         // Seam correction operations
@@ -1876,88 +1863,92 @@ namespace mu
 
 	TSharedPtr<FImage> CodeGenerator::GenerateMissingImage(EImageFormat Format)
 	{
-		MUTABLE_CPUPROFILER_SCOPE(GenerateMissingImage);
-
-		// Make a checkered debug image
-		const FImageSize Size(16, 16);
+		// Create the image node if it hasn't been created yet.
+		if (!MissingImage[SIZE_T(Format)])
+		{
+			// Make a checkered debug image
+			const FImageSize Size(16, 16);
 
-		TSharedPtr<FImage> GeneratedImage = MakeShared<FImage>(Size[0], Size[1], 1, Format, EInitializationType::NotInitialized);
+			TSharedPtr<FImage> GeneratedImage = MakeShared<FImage>(Size[0], Size[1], 1, Format, EInitializationType::NotInitialized);
 
-		switch (Format)
-		{
-		case EImageFormat::L_UByte:
-		{
-			uint8* DataPtr = GeneratedImage->GetLODData(0);
-			for (int32 P = 0; P < Size[0]*Size[1]; ++P)
+			switch (Format)
 			{
-				if ((P + P/Size[0]) % 2)
-				{
-					DataPtr[0] = 255;
-				}
-				else
+			case EImageFormat::L_UByte:
+			{
+				uint8* DataPtr = GeneratedImage->GetLODData(0);
+				for (int32 P = 0; P < Size[0]*Size[1]; ++P)
 				{
-					DataPtr[0] = 64;
-				}
+					if ((P + P/Size[0]) % 2)
+					{
+						DataPtr[0] = 255;
+					}
+					else
+					{
+						DataPtr[0] = 64;
+					}
 
-				DataPtr++;
+					DataPtr++;
+				}
+				break;
 			}
-			break;
-		}
-		case EImageFormat::RGB_UByte:
-		{
-			uint8* DataPtr = GeneratedImage->GetLODData(0);
-			for (int32 P = 0; P < Size[0]*Size[1]; ++P)
+			case EImageFormat::RGB_UByte:
 			{
-				if ((P + P/Size[0]) % 2)
+				uint8* DataPtr = GeneratedImage->GetLODData(0);
+				for (int32 P = 0; P < Size[0]*Size[1]; ++P)
 				{
-					DataPtr[0] = 255;
-					DataPtr[1] = 255;
-					DataPtr[2] = 64;
-				}
-				else
-				{
-					DataPtr[0] = 64;
-					DataPtr[1] = 64;
-					DataPtr[2] = 255;
-				}
+					if ((P + P/Size[0]) % 2)
+					{
+						DataPtr[0] = 255;
+						DataPtr[1] = 255;
+						DataPtr[2] = 64;
+					}
+					else
+					{
+						DataPtr[0] = 64;
+						DataPtr[1] = 64;
+						DataPtr[2] = 255;
+					}
 
-				DataPtr += 3;
+					DataPtr += 3;
+				}
+				break;
 			}
-			break;
-		}
-		case EImageFormat::BGRA_UByte:
-		case EImageFormat::RGBA_UByte:
-		{
-			uint8* DataPtr = GeneratedImage->GetLODData(0);
-			for (int32 P = 0; P < Size[0]*Size[1]; ++P)
+			case EImageFormat::BGRA_UByte:
+			case EImageFormat::RGBA_UByte:
 			{
-				if ((P + P/Size[0]) % 2)
-				{
-					DataPtr[0] = 255;
-					DataPtr[1] = 255;
-					DataPtr[2] = 64;
-					DataPtr[3] = 255;
-				}
-				else
+				uint8* DataPtr = GeneratedImage->GetLODData(0);
+				for (int32 P = 0; P < Size[0]*Size[1]; ++P)
 				{
-					DataPtr[0] = 64;
-					DataPtr[1] = 64;
-					DataPtr[2] = 255;
-					DataPtr[3] = 128;
-				}
+					if ((P + P/Size[0]) % 2)
+					{
+						DataPtr[0] = 255;
+						DataPtr[1] = 255;
+						DataPtr[2] = 64;
+						DataPtr[3] = 255;
+					}
+					else
+					{
+						DataPtr[0] = 64;
+						DataPtr[1] = 64;
+						DataPtr[2] = 255;
+						DataPtr[3] = 128;
+					}
 
-				DataPtr += 4;
+					DataPtr += 4;
+				}
+				break;
 			}
-			break;
-		}
 
-		default:
-			check( false );
-			break;
+			default:
+				check( false );
+				break;
+
+			}
 
+			MissingImage[(SIZE_T)Format] = GeneratedImage;
 		}
 
-		return GeneratedImage;
+		return MissingImage[(SIZE_T)Format];
 	}
 
 
diff --git a/Engine/Plugins/Mutable/Source/MutableTools/Private/MuT/CodeGenerator_Mesh.cpp b/Engine/Plugins/Mutable/Source/MutableTools/Private/MuT/CodeGenerator_Mesh.cpp
--- a/Engine/Plugins/Mutable/Source/MutableTools/Private/MuT/CodeGenerator_Mesh.cpp
+++ b/Engine/Plugins/Mutable/Source/MutableTools/Private/MuT/CodeGenerator_Mesh.cpp
@@ -618,9 +618,7 @@ namespace mu
         // Factor
         if ( node.Factor )
         {
-			FScalarGenerationResult ParamResult;
-			GenerateScalar(ParamResult, StaticOptions, node.Factor.get());
-			FactorOp = ParamResult.op;
+			FactorOp = Generate_Generic( node.Factor.get(), StaticOptions );
         }
         else
         {
@@ -894,10 +892,8 @@ namespace mu
 		Ptr<ASTOp> Variable;
         if ( node.m_pParameter )
         {
-			FScalarGenerationResult ParamResult;
-			GenerateScalar(ParamResult, StaticOptions, node.m_pParameter.get());
-			Variable = ParamResult.op;
-		}
+			Variable = Generate_Generic( node.m_pParameter.get(), StaticOptions);
+        }
         else
         {
             // This argument is required
diff --git a/Engine/Plugins/Mutable/Source/MutableTools/Private/MuT/CodeGenerator_Scalar.cpp b/Engine/Plugins/Mutable/Source/MutableTools/Private/MuT/CodeGenerator_Scalar.cpp
--- a/Engine/Plugins/Mutable/Source/MutableTools/Private/MuT/CodeGenerator_Scalar.cpp
+++ b/Engine/Plugins/Mutable/Source/MutableTools/Private/MuT/CodeGenerator_Scalar.cpp
@@ -342,11 +342,10 @@ namespace mu
 	{
 		Ptr<ASTOpScalarCurve> op = new ASTOpScalarCurve();
 
-		if (NodeScalar* pA = Typed->CurveSampleValue.get())
+		// T
+		if (Node* pA = Typed->CurveSampleValue.get())
 		{
-			FScalarGenerationResult ParamResult;
-			GenerateScalar(ParamResult, Options, pA);
-			op->time = ParamResult.op;
+			op->time = Generate_Generic(pA, Options);
 		}
 		else
 		{
@@ -379,12 +378,9 @@ namespace mu
 		}
 
 		// A
-		if (NodeScalar* pA = node.m_pA.get())
+		if (Node* pA = node.m_pA.get())
 		{
-			FScalarGenerationResult ParamResult;
-			GenerateScalar(ParamResult, Options, pA);
-
-			op->A = ParamResult.op;
+			op->A = Generate_Generic(pA, Options);
 		}
 		else
 		{
@@ -392,12 +388,9 @@ namespace mu
 		}
 
 		// B
-		if (NodeScalar* pB = node.m_pB.get())
+		if (Node* pB = node.m_pB.get())
 		{
-			FScalarGenerationResult ParamResult;
-			GenerateScalar(ParamResult, Options, pB);
-
-			op->B = ParamResult.op;
+			op->B = Generate_Generic(pB, Options);
 		}
 		else
 		{
@@ -419,11 +412,7 @@ namespace mu
 				NodeScalarConstantPtr pCell = new NodeScalarConstant();
 				float scalar = node.Table->GetPrivate()->Rows[row].Values[colIndex].Scalar;
 				pCell->SetValue(scalar);
-
-				FScalarGenerationResult ParamResult;
-				GenerateScalar(ParamResult, Options, pCell);
-
-				return ParamResult.op;
+				return Generate_Generic(pCell, Options);
 			});
 
 		result.op = Op;
@@ -442,10 +431,9 @@ namespace mu
 		pNode->SetValue(value);
 
 		FGenericGenerationOptions Options;
-		FScalarGenerationResult ParamResult;
-		GenerateScalar(ParamResult, Options, pNode.get());
+		Ptr<ASTOp> result = Generate_Generic(pNode, Options);
 
-		return ParamResult.op;
+		return result;
 	}
 
 }
\ No newline at end of file
