[Mutable] Unify GeneratedBulkDataFiles output for editor and cooked builds
Store Clothing and real-time morph data in separate files.

#rnx

diff --git a/Engine/Plugins/Mutable/Source/CustomizableObject/Internal/MuCO/CustomizableObjectPrivate.h b/Engine/Plugins/Mutable/Source/CustomizableObject/Internal/MuCO/CustomizableObjectPrivate.h
--- a/Engine/Plugins/Mutable/Source/CustomizableObject/Internal/MuCO/CustomizableObjectPrivate.h
+++ b/Engine/Plugins/Mutable/Source/CustomizableObject/Internal/MuCO/CustomizableObjectPrivate.h
@@ -879,6 +879,8 @@ namespace MutablePrivate
 		Model,
 		RealTimeMorph,
 		Clothing,
+
+		DataTypeCount
 	};
 
 
@@ -1020,14 +1022,14 @@ namespace MutablePrivate
 
 	struct CUSTOMIZABLEOBJECT_API FMutableCachedPlatformData
 	{
-		/** Serialized mu::FModel */
-		TArray64<uint8> ModelData;
-
-		/** Serialized UModelResources */
+		/** mu::Model */
+		TSharedPtr<mu::FModel, ESPMode::ThreadSafe> Model;
+		
+		/** UModelResources */
 		TStrongObjectPtr<UModelResources> ModelResources;
 
 		/** Streamable resources info such as files and offsets. */
-		TSharedPtr<FModelStreamableBulkData> ModelStreamables;
+		TSharedPtr<FModelStreamableBulkData> ModelStreamableBulkData;
 
 		/** Struct containing map of RomId to RomBytes. */
 		FModelStreamableData ModelStreamableData;
@@ -1244,11 +1246,6 @@ public:
 	bool bAllowClothingPhysicsEditsPropagation = true;
 
 #if WITH_EDITORONLY_DATA
-	UPROPERTY()
-	TArray<FMorphTargetVertexData> EditorOnlyMorphTargetReconstructionData;
-
-	UPROPERTY()
-	TArray<FCustomizableObjectMeshToMeshVertData> EditorOnlyClothingMeshToMeshVertData;
 
 	// Stores what param names use a certain table as a table can be used from multiple table nodes, useful for partial compilations to restrict params
 	UPROPERTY()
@@ -1545,8 +1542,13 @@ public:
 	UCustomizableObject* GetPublic() const;
 
 #if WITH_EDITOR
+
+	// Compose folder name where the data is stored
+	static FString GetCompiledDataFolderPath();
+	static FString GetDataTypeExtension(MutablePrivate::EStreamableDataType DataType);
+
 	/** Compose file name. */
-	FString GetCompiledDataFileName(bool bIsModel, const ITargetPlatform* InTargetPlatform = nullptr, bool bIsDiskStreamer = false);
+	FString GetCompiledDataFileName(const ITargetPlatform* InTargetPlatform = nullptr, bool bIsDiskStreamer = false);
 
 	/** DDC helpers. BuildDerivedDataKey is expensive, try to cache it as much as possible. */
 	TArray<uint8> BuildDerivedDataKey(FCompilationOptions Options);
@@ -1578,8 +1580,7 @@ public:
 	// Add a profile that stores the values of the parameters used by the CustomInstance.
 	FReply AddNewParameterProfile(FString Name, class UCustomizableObjectInstance& CustomInstance);
 
-	// Compose folder name where the data is stored
-	static FString GetCompiledDataFolderPath();
+
 
 	/** Generic Load methods to read compiled data */
 	bool LoadModelResources(FArchive& Ar, const ITargetPlatform* InTargetPlatform, bool bSkipEditorOnlyData = false);
diff --git a/Engine/Plugins/Mutable/Source/CustomizableObject/Private/MuCO/CustomizableObject.cpp b/Engine/Plugins/Mutable/Source/CustomizableObject/Private/MuCO/CustomizableObject.cpp
--- a/Engine/Plugins/Mutable/Source/CustomizableObject/Private/MuCO/CustomizableObject.cpp
+++ b/Engine/Plugins/Mutable/Source/CustomizableObject/Private/MuCO/CustomizableObject.cpp
@@ -289,11 +289,9 @@ bool UCustomizableObjectPrivate::TryLoadCompiledCookDataForPlatform(const ITarge
 		return false;
 	}
 
+	SetModel(PlatformData->Model, GenerateIdentifier(*GetPublic()));
 	SetModelResources(PlatformData->ModelResources.Get(), true);
-	SetModelStreamableBulkData(PlatformData->ModelStreamables, true);
-
-	FMemoryReaderView ModelReader(PlatformData->ModelData);
-	LoadModel(ModelReader);
+	SetModelStreamableBulkData(PlatformData->ModelStreamableBulkData, true);
 
 	return GetModel() != nullptr;
 }
@@ -634,77 +632,79 @@ void UCustomizableObjectPrivate::LoadCompiledDataFromDisk()
 	const ITargetPlatform* RunningPlatform = TargetPlatformManager.GetRunningTargetPlatform();
 	check(RunningPlatform);
 
+	IFileManager& FileManager = IFileManager::Get();
+	TArray<TUniquePtr<IFileHandle>> FileHandles;
+
 	// Compose Folder Name
-	const FString FolderPath = GetCompiledDataFolderPath();
+	const FString FullFileName = GetCompiledDataFileName(RunningPlatform);
 
-	// Compose File Names
-	const FString ModelFileName = FolderPath + GetCompiledDataFileName(true, RunningPlatform);
-	const FString StreamableFileName = FolderPath + GetCompiledDataFileName(false, RunningPlatform);
+	FGuid CompiledVersionId;
 
-	IFileManager& FileManager = IFileManager::Get();
-	if (FileManager.FileExists(*ModelFileName) && FileManager.FileExists(*StreamableFileName))
+	bool bHasCompiledData = true;
+	for (int32 DataType = 0; DataType < static_cast<int32>(MutablePrivate::EStreamableDataType::DataTypeCount); ++DataType)
 	{
-		// Check CompiledData
-		TUniquePtr<IFileHandle> CompiledDataFileHandle( FPlatformFileManager::Get().GetPlatformFile().OpenRead(*ModelFileName) );
-		TUniquePtr<IFileHandle> StreamableDataFileHandle( FPlatformFileManager::Get().GetPlatformFile().OpenRead(*StreamableFileName) );
+		const FString FilePath = FullFileName + UCustomizableObjectPrivate::GetDataTypeExtension(static_cast<MutablePrivate::EStreamableDataType>(DataType));
+		if (IFileHandle* FileHandle = FPlatformFileManager::Get().GetPlatformFile().OpenRead(*FilePath))
+		{
+			FileHandles.Emplace(FileHandle);
+			
+			TArray<uint8> HeaderBytes;
 
-		MutableCompiledDataStreamHeader CompiledDataHeader;
-		MutableCompiledDataStreamHeader StreamableDataHeader;
+			const int32 HeaderSize = sizeof(MutableCompiledDataStreamHeader);
+			HeaderBytes.SetNum(HeaderSize);
+			FileHandle->Read(HeaderBytes.GetData(), HeaderSize);
+			
+			FMemoryReader AuxMemoryReader(HeaderBytes);
+			MutableCompiledDataStreamHeader DataTypeHeader;
+			AuxMemoryReader << DataTypeHeader;
 
-		int32 HeaderSize = sizeof(MutableCompiledDataStreamHeader);
-		TArray<uint8> HeaderBytes;
-		HeaderBytes.SetNum(HeaderSize);
+			if(DataTypeHeader.InternalVersion != CurrentSupportedVersion
+				|| (!IsRunningGame() && DataTypeHeader.VersionId != GetVersionId()) 
+				|| (CompiledVersionId.IsValid() && CompiledVersionId != DataTypeHeader.VersionId))
+			{
+				bHasCompiledData = false;
+				break;
+			}
 
-		{
-			CompiledDataFileHandle->Read(HeaderBytes.GetData(), HeaderSize);
-			FMemoryReader AuxMemoryReader(HeaderBytes);
-			AuxMemoryReader << CompiledDataHeader;
+			CompiledVersionId = DataTypeHeader.VersionId;
 		}
+		else
 		{
-			StreamableDataFileHandle->Read(HeaderBytes.GetData(), HeaderSize);
-			FMemoryReader AuxMemoryReader(HeaderBytes);
-			AuxMemoryReader << StreamableDataHeader;
+			bHasCompiledData = false;
+			break;
 		}
+	}
 
-		if (CompiledDataHeader.InternalVersion == UCustomizableObjectPrivate::CurrentSupportedVersion
-			&&
-			CompiledDataHeader.InternalVersion == StreamableDataHeader.InternalVersion 
-			&&
-			CompiledDataHeader.VersionId == StreamableDataHeader.VersionId)
-		{
-			if (IsRunningGame() || CompiledDataHeader.VersionId == GetVersionId())
-			{ 
-				int64 CompiledDataSize = CompiledDataFileHandle->Size() - HeaderSize;
-				TArray64<uint8> CompiledDataBytes;
-				CompiledDataBytes.SetNumUninitialized(CompiledDataSize);
+	if (bHasCompiledData)
+	{
+		TArray64<uint8> CompiledDataBytes;
 
-				CompiledDataFileHandle->Seek(HeaderSize);
-				CompiledDataFileHandle->Read(CompiledDataBytes.GetData(), CompiledDataSize);
+		const int64 CompiledDataSize = FileHandles[0]->Size() - FileHandles[0]->Tell();
+		CompiledDataBytes.SetNumUninitialized(CompiledDataSize);
+		FileHandles[0]->Read(CompiledDataBytes.GetData(), CompiledDataSize);
 
-				FMemoryReaderView MemoryReader(CompiledDataBytes);
-				
-				if (LoadModelResources(MemoryReader, RunningPlatform))
+		FMemoryReaderView MemoryReader(CompiledDataBytes);
+
+		if (LoadModelResources(MemoryReader, RunningPlatform))
+		{
+			TArray<FName> OutOfDatePackages;
+			TArray<FName> AddedPackages;
+			TArray<FName> RemovedPackages;
+			bool bReleaseVersion;
+			const bool bOutOfDate = IsCompilationOutOfDate(false, OutOfDatePackages, AddedPackages, RemovedPackages, bReleaseVersion);
+			if (!bOutOfDate)
+			{
+				LoadModelStreamableBulk(MemoryReader, /* bIsCooking */false);
+				LoadModel(MemoryReader);
+			}
+			else
+			{
+				if (OutOfDatePackages.Num())
 				{
-					TArray<FName> OutOfDatePackages;
-					TArray<FName> AddedPackages;
-					TArray<FName> RemovedPackages;
-					bool bReleaseVersion;
-					const bool bOutOfDate = IsCompilationOutOfDate(false, OutOfDatePackages, AddedPackages, RemovedPackages, bReleaseVersion);
-					if (!bOutOfDate)
-					{
-						LoadModelStreamableBulk(MemoryReader, /* bIsCooking */false);
-						LoadModel(MemoryReader);
-					}
-					else
-					{
-						if (OutOfDatePackages.Num())
-						{
-							UE_LOG(LogMutable, Display, TEXT("Invalidating compiled data due to changes in %s."), *OutOfDatePackages[0].ToString());
-						}
-						
-						PrintParticipatingPackagesDiff(OutOfDatePackages, AddedPackages, RemovedPackages, bReleaseVersion);
-					}
+					UE_LOG(LogMutable, Display, TEXT("Invalidating compiled data due to changes in %s."), *OutOfDatePackages[0].ToString());
 				}
+
+				PrintParticipatingPackagesDiff(OutOfDatePackages, AddedPackages, RemovedPackages, bReleaseVersion);
 			}
 		}
 	}
@@ -855,12 +855,30 @@ FString UCustomizableObjectPrivate::GetCompiledDataFolderPath()
 }
 
 
-FString UCustomizableObjectPrivate::GetCompiledDataFileName(bool bIsModel, const ITargetPlatform* InTargetPlatform, bool bIsDiskStreamer)
+FString UCustomizableObjectPrivate::GetCompiledDataFileName(const ITargetPlatform* InTargetPlatform, bool bIsDiskStreamer)
 {
+	const FString FilePath = GetCompiledDataFolderPath();
 	const FString PlatformName = InTargetPlatform ? InTargetPlatform->PlatformName() : FPlatformProperties::PlatformName();
 	const FString FileIdentifier = bIsDiskStreamer ? Identifier.ToString() : GenerateIdentifier(*GetPublic()).ToString();
-	const FString Extension = bIsModel ? TEXT("_M.mut") : TEXT("_S.mut");
-	return PlatformName + FileIdentifier + Extension;
+
+	return FilePath + PlatformName + FileIdentifier;
+}
+
+
+FString UCustomizableObjectPrivate::GetDataTypeExtension(MutablePrivate::EStreamableDataType DataType)
+{
+	switch (DataType)
+	{
+	case MutablePrivate::EStreamableDataType::None:	return TEXT(".mut");
+	case MutablePrivate::EStreamableDataType::Model: return TEXT("_M.mut");
+	case MutablePrivate::EStreamableDataType::RealTimeMorph: return TEXT("_RTM.mut");
+	case MutablePrivate::EStreamableDataType::Clothing:	return TEXT("_C.mut");
+	default:
+	{
+		unimplemented();
+		return s_EmptyString;
+	}
+	}
 }
 
 
@@ -3148,43 +3166,35 @@ void UCustomizableObjectPrivate::LoadCompiledDataFromDDC(FCompilationOptions Opt
 	// Check if all the requested buffers were found.
 	if (!ModelBytesDDC.IsNull() && !ModelResourcesBytesDDC.IsNull() && !BulkDataFilesBytesDDC.IsNull()  && !ModelStreamablesBytesDDC.IsNull())
 	{
+		MutablePrivate::FMutableCachedPlatformData CachedData;
+
 		// Load the compiled data to validate it.
 		FMemoryReaderView ModelResourcesReader(ModelResourcesBytesDDC.GetView());
 		if (LoadModelResources(ModelResourcesReader, Options.TargetPlatform, Options.bIsCooking))
 		{
-			UModelResources& LocalModelResources = *GetModelResources(Options.bIsCooking);
-			LocalModelResources.bIsStoredInDDC = true;
-			LocalModelResources.DDCKey = CacheKey;
-			LocalModelResources.DDCDefaultPolicy = ECachePolicy::Default;
+			CachedData.ModelResources = TStrongObjectPtr<UModelResources>(GetModelResources(Options.bIsCooking));
+			CachedData.ModelResources->bIsStoredInDDC = true;
+			CachedData.ModelResources->DDCKey = CacheKey;
+			CachedData.ModelResources->DDCDefaultPolicy = ECachePolicy::Default;
 
 			FMemoryReaderView ModelStreamablesReader(ModelStreamablesBytesDDC.GetView());
 			LoadModelStreamableBulk(ModelStreamablesReader, Options.bIsCooking);
+			CachedData.ModelStreamableBulkData = GetModelStreamableBulkData(Options.bIsCooking);
 
 			FMemoryReaderView ModelReader(ModelBytesDDC.GetView());
 			LoadModel(ModelReader);
+			CachedData.Model = GetModel();
 		}
 
-		TSharedPtr<mu::FModel> Model = GetModel();
-		TSharedPtr<FModelStreamableBulkData> ModelStreamables = GetModelStreamableBulkData(Options.bIsCooking);
-		UModelResources* LocalModelResources = GetModelResources(Options.bIsCooking);
-
 		// Cache CookedPlatfomData. 
-		if (Options.bIsCooking && Model && LocalModelResources && ModelStreamables)
+		if (Options.bIsCooking && CachedData.Model && CachedData.ModelResources && CachedData.ModelStreamableBulkData)
 		{
 			// Sync cache cooked platform data
 			// TODO UE-220138: Sync -> Async
 			MUTABLE_CPUPROFILER_SCOPE(CacheCookedPlatformData);
 
-			MutablePrivate::FMutableCachedPlatformData CachedData;
-			
-			// Cache Model, ModelResources and ModelStreamables
-			CachedData.ModelData.Append(reinterpret_cast<const uint8*>(ModelBytesDDC.GetData()), ModelBytesDDC.GetSize());
-			
-			CachedData.ModelResources = TStrongObjectPtr<UModelResources>(LocalModelResources);
 			CachedData.ModelResources->InitCookData();
 
-			CachedData.ModelStreamables = ModelStreamables;
-
 			// Value Id to file mapping to reconstruct the cached data
 			TMap<FValueId, MutablePrivate::FFile> ValueIdToFile;
 
@@ -3222,7 +3232,7 @@ void UCustomizableObjectPrivate::LoadCompiledDataFromDDC(FCompilationOptions Opt
 
 			FRequestOwner RequestOwner(EPriority::Blocking);
 			GetCache().Get(MakeArrayView(&Request, 1), RequestOwner,
-				[&CachedData, &ValueIdToFile, ModelStreamables](FCacheGetResponse&& Response)
+				[&CachedData, &ValueIdToFile](FCacheGetResponse&& Response)
 				{
 					MUTABLE_CPUPROFILER_SCOPE(CacheBulkDataFromDDC);
 
@@ -3290,11 +3300,11 @@ void UCustomizableObjectPrivate::LoadCompiledDataFromDDC(FCompilationOptions Opt
 			if (CVarMutableUseBulkData.GetValueOnAnyThread())
 			{
 				const uint32 NumBulkDataFilesPerBucket = MAX_uint8;
-				MutablePrivate::GenerateBulkDataFilesListWithFileLimit(Model, *ModelStreamables.Get(), NumBulkDataFilesPerBucket, CachedData.BulkDataFiles);
+				MutablePrivate::GenerateBulkDataFilesListWithFileLimit(CachedData.Model, *CachedData.ModelStreamableBulkData.Get(), NumBulkDataFilesPerBucket, CachedData.BulkDataFiles);
 			}
 			else
 			{
-				MutablePrivate::GenerateBulkDataFilesListWithSizeLimit(Model, *ModelStreamables.Get(), Options.TargetPlatform, Options.PackagedDataBytesLimit, CachedData.BulkDataFiles);
+				MutablePrivate::GenerateBulkDataFilesListWithSizeLimit(CachedData.Model, *CachedData.ModelStreamableBulkData.Get(), Options.TargetPlatform, Options.PackagedDataBytesLimit, CachedData.BulkDataFiles);
 			}
 
 			MutablePrivate::FMutableCachedPlatformData& CachedPlatformData = CachedPlatformsData.Add(Options.TargetPlatform->PlatformName(), {});
diff --git a/Engine/Plugins/Mutable/Source/CustomizableObject/Private/MuCO/CustomizableObjectInstance.cpp b/Engine/Plugins/Mutable/Source/CustomizableObject/Private/MuCO/CustomizableObjectInstance.cpp
--- a/Engine/Plugins/Mutable/Source/CustomizableObject/Private/MuCO/CustomizableObjectInstance.cpp
+++ b/Engine/Plugins/Mutable/Source/CustomizableObject/Private/MuCO/CustomizableObjectInstance.cpp
@@ -58,6 +58,7 @@
 #include "MuCO/CustomizableObjectResourceDataTypes.h"
 #include "MuCO/LoadUtils.h"
 #include "Serialization/BulkData.h"
+#include "HAL/PlatformFileManager.h"
 
 #include UE_INLINE_GENERATED_CPP_BY_NAME(CustomizableObjectInstance)
 
@@ -5789,90 +5790,22 @@ UE::Tasks::FTask UCustomizableInstancePrivate::LoadAdditionalAssetsAndData(const
 	if (bIsDataBlocksStreamNeeded && bUpdateMeshes)
 	{
 #if WITH_EDITOR
-		// On editor the data is always loaded, load directly form the ModelResources.
-		if (bMorphTargetsEnabled)
-		{
-			MUTABLE_CPUPROFILER_SCOPE(RealTimeMorphStreamingEditor);
-			for (uint32 BlockId : RealTimeMorphStreamableBlocksToStream)
-			{	
-				const FRealTimeMorphStreamable& RealTimeMorphStreamable = ModelStreamableBulkData->RealTimeMorphStreamables[BlockId];
-			
-				const FMutableStreamableBlock& Block = RealTimeMorphStreamable.Block;
-
-				FInstanceUpdateData::FMorphTargetMeshData& MeshData = 
-						OperationData->InstanceUpdateData.RealTimeMorphTargetMeshData.FindOrAdd(BlockId);
-				
-				MeshData.NameResolutionMap = RealTimeMorphStreamable.NameResolutionMap;
-
-				const TArray<FMorphTargetVertexData>& SourceData = ModelResources.EditorOnlyMorphTargetReconstructionData;	
-				const uint32 NumElems = RealTimeMorphStreamable.Size / sizeof(FMorphTargetVertexData);
-				const uint32 OffsetInElems = Block.Offset / sizeof(FMorphTargetVertexData);
-				MeshData.Data.SetNumUninitialized(NumElems);
-
-				check(SourceData.Num()*sizeof(FMorphTargetVertexData) >= Block.Offset + RealTimeMorphStreamable.Size);
-				FMemory::Memcpy(MeshData.Data.GetData(), SourceData.GetData() + OffsetInElems, RealTimeMorphStreamable.Size);
-			}
-		}
-		{
-			MUTABLE_CPUPROFILER_SCOPE(ClothingStreamingEditor);
-			for (uint32 BlockId : ClothingStreamableBlocksToStream)
-			{	
-				const FClothingStreamable& ClothingStreamable = ModelStreamableBulkData->ClothingStreamables[BlockId]; 
-			
-				const FMutableStreamableBlock& Block = ClothingStreamable.Block;
-
-				FInstanceUpdateData::FClothingMeshData& MeshData = 
-						OperationData->InstanceUpdateData.ClothingMeshData.FindOrAdd(BlockId);
-				
-				MeshData.ClothingAssetIndex = ClothingStreamable.ClothingAssetIndex;
-				MeshData.ClothingAssetLOD = ClothingStreamable.ClothingAssetLOD;
-				MeshData.PhysicsAssetIndex = ClothingStreamable.PhysicsAssetIndex;
-
-				const TArray<FCustomizableObjectMeshToMeshVertData>& SourceData = ModelResources.EditorOnlyClothingMeshToMeshVertData;
-				const uint32 NumElems = ClothingStreamable.Size / sizeof(FCustomizableObjectMeshToMeshVertData);
-				const uint32 OffsetInElems = Block.Offset / sizeof(FCustomizableObjectMeshToMeshVertData);
-				MeshData.Data.SetNumUninitialized(NumElems);
-
-				check(SourceData.Num()*sizeof(FCustomizableObjectMeshToMeshVertData) >= Block.Offset + ClothingStreamable.Size);
-				FMemory::Memcpy(MeshData.Data.GetData(), SourceData.GetData() + OffsetInElems, ClothingStreamable.Size);
-			}
-		}
-#else	
-		MUTABLE_CPUPROFILER_SCOPE(RealTimeMorphStreaming);
-		struct FBlockReadInfo
-		{
-			uint64 Offset;
-			IAsyncReadFileHandle* FileHandle;
-			TArrayView<uint8> AllocatedMemoryView;
-			uint32 FileId;
-		};
+		const int32 BulkDataFileOffset = sizeof(MutableCompiledDataStreamHeader);
+		const bool bUseFBulkData = false;
 
-		const bool bUseFBulkData = !ModelStreamableBulkData->StreamableBulkData.IsEmpty();
-
-		TArray<FBlockReadInfo> BlockReadInfos;
-		BlockReadInfos.Reserve(16);
-		
-		const UCustomizableObjectBulk* BulkData = CustomizableObject->GetPrivate()->GetStreamableBulkData();
-		if (!bUseFBulkData)
-		{
-			UE_CLOG(!BulkData, LogMutable, Error, TEXT("BulkData object for CustomizableObject [%s] not found."),
-				*CustomizableObject->GetFName().ToString());
-			check(BulkData);
-		}
+		const FString FullFileName = CustomizableObject->GetPrivate()->GetCompiledDataFileName(nullptr, true);
 
 		TArray<uint32> OpenFilesIds;
-
-		auto OpenOrGetFileHandleForBlock = [&OpenFilesIds, &OpenFileHandles, BulkData](const FMutableStreamableBlock& Block) 
+		auto OpenOrGetFileHandleForBlock = [&OpenFilesIds, &OpenFileHandles, FullFileName](const FMutableStreamableBlock& Block, MutablePrivate::EStreamableDataType DataType)
 			-> IAsyncReadFileHandle*
 		{
 			int32 FileHandleIndex = OpenFilesIds.Find(Block.FileId);
-			if (FileHandleIndex == INDEX_NONE && BulkData)
+			if (FileHandleIndex == INDEX_NONE)
 			{
-				TUniquePtr<IAsyncReadFileHandle> ReadFileHandle = BulkData->OpenFileAsyncRead(Block.FileId, Block.Flags);
+				IAsyncReadFileHandle* ReadFileHandle = FPlatformFileManager::Get().GetPlatformFile().OpenAsyncRead(*(FullFileName + UCustomizableObjectPrivate::GetDataTypeExtension(DataType)));
+				OpenFileHandles.Emplace(ReadFileHandle);
 
-				OpenFileHandles.Emplace(MoveTemp(ReadFileHandle));
 				FileHandleIndex = OpenFilesIds.Add(Block.FileId);
-				
 				check(OpenFileHandles.Num() == OpenFilesIds.Num());
 			}
 
@@ -5883,6 +5816,24 @@ UE::Tasks::FTask UCustomizableInstancePrivate::LoadAdditionalAssetsAndData(const
 
 			return nullptr;
 		};
+#else
+		const int32 BulkDataFileOffset = 0;
+		const bool bUseFBulkData = !ModelStreamableBulkData->StreamableBulkData.IsEmpty();
+#endif
+
+
+		MUTABLE_CPUPROFILER_SCOPE(RealTimeMorphStreaming);
+
+		struct FBlockReadInfo
+		{
+			uint64 Offset;
+			IAsyncReadFileHandle* FileHandle;
+			TArrayView<uint8> AllocatedMemoryView;
+			uint32 FileId;
+		};
+
+		TArray<FBlockReadInfo> BlockReadInfos;
+		BlockReadInfos.Reserve(16);
 
 		if (bMorphTargetsEnabled)
 		{
@@ -5911,15 +5862,15 @@ UE::Tasks::FTask UCustomizableInstancePrivate::LoadAdditionalAssetsAndData(const
 
 				ReadDestData.Data.SetNumUninitialized(NumElems);
 
+#if WITH_EDITOR
+				IAsyncReadFileHandle* FileHandle = OpenOrGetFileHandleForBlock(Block, MutablePrivate::EStreamableDataType::RealTimeMorph);
+#else
 				IAsyncReadFileHandle* FileHandle = nullptr;
-				if (!bUseFBulkData)
-				{	
-					FileHandle = OpenOrGetFileHandleForBlock(Block);
-				}
+#endif
 
 				BlockReadInfos.Emplace(FBlockReadInfo
-				{ 
-					Block.Offset,
+				{
+					BulkDataFileOffset + Block.Offset,
 					FileHandle,
 					MakeArrayView(reinterpret_cast<uint8*>(ReadDestData.Data.GetData()), ReadDestData.Data.Num()*sizeof(FMorphTargetVertexData)),
 					Block.FileId
@@ -5953,22 +5904,21 @@ UE::Tasks::FTask UCustomizableInstancePrivate::LoadAdditionalAssetsAndData(const
 
 			ReadDestData.Data.SetNumUninitialized(NumElems);
 
+#if WITH_EDITOR
+			IAsyncReadFileHandle* FileHandle = OpenOrGetFileHandleForBlock(Block, MutablePrivate::EStreamableDataType::Clothing);
+#else
 			IAsyncReadFileHandle* FileHandle = nullptr;
-			if (!bUseFBulkData)
-			{
-				FileHandle = OpenOrGetFileHandleForBlock(Block);
-			}
+#endif
 
 			BlockReadInfos.Emplace(FBlockReadInfo
-			{ 
-				Block.Offset,
+			{
+				BulkDataFileOffset + Block.Offset,
 				FileHandle,
 				MakeArrayView(reinterpret_cast<uint8*>(ReadDestData.Data.GetData()), ReadDestData.Data.Num()*sizeof(FCustomizableObjectMeshToMeshVertData)),
 				Block.FileId
 			});
 		}
 
-
 		for (const FBlockReadInfo& BlockReadInfo : BlockReadInfos)
 		{
 			if (bUseFBulkData)
@@ -5983,7 +5933,7 @@ UE::Tasks::FTask UCustomizableInstancePrivate::LoadAdditionalAssetsAndData(const
 					]() -> TUniquePtr<IBulkDataIORequest>
 					{
 						MUTABLE_CPUPROFILER_SCOPE(CustomizableInstanceLoadBlocksAsyncRead_Request);
-						FBulkDataIORequestCallBack IOCallback = 
+						FBulkDataIORequestCallBack IOCallback =
 							[OwnedOperationData, ReadDataReadyEvent, FileId = Block.FileId](bool bWasCancelled, IBulkDataIORequest*) mutable
 							{
 								if (bWasCancelled)
@@ -6035,10 +5985,7 @@ UE::Tasks::FTask UCustomizableInstancePrivate::LoadAdditionalAssetsAndData(const
 					},
 					bEnableHighPriorityLoading ? UE::Tasks::ETaskPriority::High : UE::Tasks::ETaskPriority::Normal));
 			}
-		}	
-
-#endif
-
+		}
 	}
 	
 	if (AssetsToStream.Num() > 0 || OpenFileHandles.Num() > 0 || BulkReadRequestTasks.Num() > 0)
diff --git a/Engine/Plugins/Mutable/Source/CustomizableObject/Private/MuCO/UnrealMutableModelDiskStreamer.cpp b/Engine/Plugins/Mutable/Source/CustomizableObject/Private/MuCO/UnrealMutableModelDiskStreamer.cpp
--- a/Engine/Plugins/Mutable/Source/CustomizableObject/Private/MuCO/UnrealMutableModelDiskStreamer.cpp
+++ b/Engine/Plugins/Mutable/Source/CustomizableObject/Private/MuCO/UnrealMutableModelDiskStreamer.cpp
@@ -108,9 +108,8 @@ bool FUnrealMutableModelBulkReader::PrepareStreamingForObject(UCustomizableObjec
 		}
 
 #if WITH_EDITOR
-		FString FolderPath = CustomizableObject->GetPrivate()->GetCompiledDataFolderPath();
-		FString FullFileName = FolderPath + CustomizableObject->GetPrivate()->GetCompiledDataFileName(false, nullptr, true);
-		NewData.BulkFilePrefix = *FullFileName;
+		FString FullFileName = CustomizableObject->GetPrivate()->GetCompiledDataFileName(nullptr, true);
+		NewData.BulkFilePrefix = *FullFileName + UCustomizableObjectPrivate::GetDataTypeExtension(MutablePrivate::EStreamableDataType::Model);
 
 		NewData.bIsStoredInDDC = ModelResources.bIsStoredInDDC;
 		NewData.DDCKey = ModelResources.DDCKey;
diff --git a/Engine/Plugins/Mutable/Source/CustomizableObjectEditor/Private/MuCOE/CustomizableObjectCompileRunnable.cpp b/Engine/Plugins/Mutable/Source/CustomizableObjectEditor/Private/MuCOE/CustomizableObjectCompileRunnable.cpp
--- a/Engine/Plugins/Mutable/Source/CustomizableObjectEditor/Private/MuCOE/CustomizableObjectCompileRunnable.cpp
+++ b/Engine/Plugins/Mutable/Source/CustomizableObjectEditor/Private/MuCOE/CustomizableObjectCompileRunnable.cpp
@@ -327,16 +327,12 @@ const TArray<FCustomizableObjectCompileRunnable::FError>& FCustomizableObjectCom
 
 
 FCustomizableObjectSaveDDRunnable::FCustomizableObjectSaveDDRunnable(const TSharedPtr<FCompilationRequest>& InRequest,
-	TSharedPtr<mu::FModel> InModel,
-	UModelResources& InModelResources,
-	TSharedPtr<FModelStreamableBulkData> InModelStreamables)
+	TSharedPtr<MutablePrivate::FMutableCachedPlatformData>& InPlatformData)
 {
-	MUTABLE_CPUPROFILER_SCOPE(FCustomizableObjectSaveDDRunnable::FCustomizableObjectSaveDDRunnable)
-
-	Model = InModel;
-	ModelStreamables = InModelStreamables;
-	PlatformData.ModelResources = TStrongObjectPtr<UModelResources>(&InModelResources);
+	MUTABLE_CPUPROFILER_SCOPE(FCustomizableObjectSaveDDRunnable::FCustomizableObjectSaveDDRunnable);
 
+	PlatformData = InPlatformData;
+	
 	Options = InRequest->Options;
 
 	DDCKey = InRequest->GetDerivedDataCacheKey();
@@ -352,48 +348,13 @@ FCustomizableObjectSaveDDRunnable::FCustomizableObjectSaveDDRunnable(const TShar
 	{
 		FMemoryWriter64 MemoryWriter(ModelResourcesData);
 		FObjectAndNameAsStringProxyArchive ObjectWriter(MemoryWriter, true);
-		InModelResources.Serialize(ObjectWriter);
-	}
-
-	// Cache Morphs and Clothing. 
-	{
-		// Do a copy of the Morph and Clothing Data generated at compile time. Only needed when cooking.
-		// Morphs
-		static_assert(TCanBulkSerialize<FMorphTargetVertexData>::Value);
-		TArray<FMorphTargetVertexData>& MorphVertexData = InModelResources.EditorOnlyMorphTargetReconstructionData;
-
-		for (const TPair<uint32, FRealTimeMorphStreamable>& MorphStreamable : InModelStreamables->RealTimeMorphStreamables)
-		{
-			PlatformData.MorphStreamableData.Set(
-					MorphStreamable.Key,
-					reinterpret_cast<uint8*>(MorphVertexData.GetData()) + MorphStreamable.Value.Block.Offset,
-					MorphStreamable.Value.Size);
-		}
-
-		static_assert(TCanBulkSerialize<FCustomizableObjectMeshToMeshVertData>::Value);
-		TArray<FCustomizableObjectMeshToMeshVertData>& ClothingVertexData = InModelResources.EditorOnlyClothingMeshToMeshVertData;
-		
-		for (const TPair<uint32, FClothingStreamable>& ClothingStreamable : InModelStreamables->ClothingStreamables)
-		{
-			PlatformData.ClothingStreamableData.Set(
-					ClothingStreamable.Key,
-					reinterpret_cast<uint8*>(ClothingVertexData.GetData()) + ClothingStreamable.Value.Block.Offset,
-					ClothingStreamable.Value.Size);
-		}
-
-		if (Options.bIsCooking)
-		{
-			InModelResources.EditorOnlyMorphTargetReconstructionData.Empty();
-			InModelResources.EditorOnlyClothingMeshToMeshVertData.Empty();
-		}
+		PlatformData->ModelResources->Serialize(ObjectWriter);
 	}
 
 	if (!Options.bIsCooking)
 	{
 		// We will be saving all compilation data in two separate files, write CO Data
-		FolderPath = CustomizableObject->GetPrivate()->GetCompiledDataFolderPath();
-		CompileDataFullFileName = FolderPath + CustomizableObject->GetPrivate()->GetCompiledDataFileName(true, Options.TargetPlatform);
-		StreamableDataFullFileName = FolderPath + CustomizableObject->GetPrivate()->GetCompiledDataFileName(false, Options.TargetPlatform);
+		FullFileName = CustomizableObject->GetPrivate()->GetCompiledDataFileName(Options.TargetPlatform);
 	}
 }
 
@@ -401,7 +362,7 @@ uint32 FCustomizableObjectSaveDDRunnable::Run()
 {
 	MUTABLE_CPUPROFILER_SCOPE(FCustomizableObjectSaveDDRunnable::Run)
 
-	if (Model)
+	if (PlatformData->Model)
 	{
 		CachePlatformData();
 
@@ -415,7 +376,7 @@ uint32 FCustomizableObjectSaveDDRunnable::Run()
 
 		if (!Options.bIsCooking && !bStoredSuccessfully)
 		{
-			StoreCachedPlatformDataToDisk(bStoredSuccessfully);
+			StoreCachedPlatformDataToDisk();
 		}
 	}
 
@@ -441,7 +402,7 @@ void FCustomizableObjectSaveDDRunnable::CachePlatformData()
 {
 	MUTABLE_CPUPROFILER_SCOPE(CachePlatformData);
 
-	if (!Model || !ModelStreamables)
+	if (!PlatformData->Model || !PlatformData->ModelStreamableBulkData)
 	{
 		check(false);
 		return;
@@ -449,28 +410,26 @@ void FCustomizableObjectSaveDDRunnable::CachePlatformData()
 
 	// Cache ModelStreamables
 	{
-		PlatformData.ModelStreamables = ModelStreamables;
-
 		// Generate list of files and update streamable blocks ids and offsets
 		if (Options.bUseBulkData)
 		{
-			MutablePrivate::GenerateBulkDataFilesListWithFileLimit(Model, *PlatformData.ModelStreamables.Get(), MAX_uint8, PlatformData.BulkDataFiles);
+			MutablePrivate::GenerateBulkDataFilesListWithFileLimit(PlatformData->Model, *PlatformData->ModelStreamableBulkData.Get(), MAX_uint8, PlatformData->BulkDataFiles);
 		}
 		else
 		{
 			const uint64 PackageDataBytesLimit = Options.bIsCooking ? Options.PackagedDataBytesLimit : MAX_uint64;
-			MutablePrivate::GenerateBulkDataFilesListWithSizeLimit(Model, *PlatformData.ModelStreamables.Get(), Options.TargetPlatform, PackageDataBytesLimit, PlatformData.BulkDataFiles);
+			MutablePrivate::GenerateBulkDataFilesListWithSizeLimit(PlatformData->Model, *PlatformData->ModelStreamableBulkData.Get(), Options.TargetPlatform, PackageDataBytesLimit, PlatformData->BulkDataFiles);
 		}
 	}
 
 	// Cache Model and Model Roms
 	{
-		FMemoryWriter64 ModelMemoryWriter(PlatformData.ModelData);
-		FUnrealMutableModelBulkWriterCook Streamer(&ModelMemoryWriter, &PlatformData.ModelStreamableData);
+		FMemoryWriter64 ModelMemoryWriter(ModelData);
+		FUnrealMutableModelBulkWriterCook Streamer(&ModelMemoryWriter, &PlatformData->ModelStreamableData);
 
 		// Serialize mu::FModel and streamable resources 
 		constexpr bool bDropData = true;
-		mu::FModel::Serialise(Model.Get(), Streamer, bDropData);
+		mu::FModel::Serialise(PlatformData->Model.Get(), Streamer, bDropData);
 	}
 }
 
@@ -481,7 +440,7 @@ void FCustomizableObjectSaveDDRunnable::StoreCachedPlatformDataInDDC(bool& bStor
 
 	using namespace UE::DerivedData;
 
-	check(Model.Get() != nullptr);
+	check(PlatformData->Model.Get() != nullptr);
 	check(DDCKey.Hash.IsZero() == false);
 
 	bStoredSuccessfully = false;
@@ -494,10 +453,10 @@ void FCustomizableObjectSaveDDRunnable::StoreCachedPlatformDataInDDC(bool& bStor
 		MUTABLE_CPUPROFILER_SCOPE(SerializeModelStreamables);
 
 		// ModelStreamable  will be modified for the DDC record. Modify a copy
-		FModelStreamableBulkData ModelStreamablesDDC = *ModelStreamables.Get();
+		FModelStreamableBulkData ModelStreamablesDDC = *PlatformData->ModelStreamableBulkData.Get();
 
 		// Generate list of files and update streamable blocks ids and offsets
-		MutablePrivate::GenerateBulkDataFilesListWithFileLimit(Model, ModelStreamablesDDC, MAX_int16, BulkDataFilesDDC);
+		MutablePrivate::GenerateBulkDataFilesListWithFileLimit(PlatformData->Model, ModelStreamablesDDC, MAX_int16, BulkDataFilesDDC);
 
 		TArray64<uint8> ModelStreamablesBytesDDC;
 		FMemoryWriter64 MemoryWriterDDC(ModelStreamablesBytesDDC);
@@ -535,7 +494,7 @@ void FCustomizableObjectSaveDDRunnable::StoreCachedPlatformDataInDDC(bool& bStor
 			};
 
 		constexpr bool bDropData = false;
-		MutablePrivate::SerializeBulkDataFiles(PlatformData, BulkDataFilesDDC, WriteBulkDataDDC, bDropData);
+		MutablePrivate::SerializeBulkDataFiles(*PlatformData.Get(), BulkDataFilesDDC, WriteBulkDataDDC, bDropData);
 	}
 
 
@@ -564,7 +523,7 @@ void FCustomizableObjectSaveDDRunnable::StoreCachedPlatformDataInDDC(bool& bStor
 	{
 		MUTABLE_CPUPROFILER_SCOPE(SerializeModelForDDC);
 
-		const FValue ModelValue = FValue::Compress(FSharedBuffer::MakeView(PlatformData.ModelData.GetData(), PlatformData.ModelData.Num()));
+		const FValue ModelValue = FValue::Compress(FSharedBuffer::MakeView(ModelData.GetData(), ModelData.Num()));
 		RecordBuilder.AddValue(MutablePrivate::GetDerivedDataModelId(), ModelValue);
 	}
 
@@ -588,116 +547,87 @@ void FCustomizableObjectSaveDDRunnable::StoreCachedPlatformDataInDDC(bool& bStor
 }
 
 
-void FCustomizableObjectSaveDDRunnable::StoreCachedPlatformDataToDisk(bool& bStoredSuccessfully)
+void FCustomizableObjectSaveDDRunnable::StoreCachedPlatformDataToDisk()
 {
 	MUTABLE_CPUPROFILER_SCOPE(StoreCachedPlatformDataToDisk);
 
-	check(Model.Get() != nullptr);
+	check(PlatformData->Model.Get() != nullptr);
 	check(!Options.bIsCooking);
 
-	bStoredSuccessfully = false;
-
 	// Create folder...
 	IFileManager& FileManager = IFileManager::Get();
-	FileManager.MakeDirectory(*FolderPath, true);
+	FileManager.MakeDirectory(*UCustomizableObjectPrivate::GetCompiledDataFolderPath(), true);
 
-	// Delete files...
-	
-	
+	// Delete files and create new memory writers for each streamable data type.
 	bool bFilesDeleted = true;
-	if (FileManager.FileExists(*CompileDataFullFileName)
-		&& !FileManager.Delete(*CompileDataFullFileName, true, false, true))
-	{
-		UE_LOG(LogMutable, Error, TEXT("Failed to delete compiled data in file [%s]."), *CompileDataFullFileName);
-		bFilesDeleted = false;
-	}
+	TArray<TUniquePtr<FArchive>> MemoryWriters;
 
-	if (FileManager.FileExists(*StreamableDataFullFileName)
-		&& !FileManager.Delete(*StreamableDataFullFileName, true, false, true))
+	const int32 NumDataTypes = static_cast<int32>(MutablePrivate::EStreamableDataType::DataTypeCount);
+	for (int32 DataType = 0; DataType < NumDataTypes; ++DataType)
 	{
-		UE_LOG(LogMutable, Error, TEXT("Failed to delete streamed data in file [%s]."), *StreamableDataFullFileName);
-		bFilesDeleted = false;
-	}
+		const FString FilePath = FullFileName + UCustomizableObjectPrivate::GetDataTypeExtension(static_cast<MutablePrivate::EStreamableDataType>(DataType));
+		if (FileManager.FileExists(*FilePath)
+			&& !FileManager.Delete(*FilePath, true, false, true))
+		{
+			UE_LOG(LogMutable, Error, TEXT("Failed to delete file for data type [%d]."), DataType);
+			bFilesDeleted = false;
+			break;
+		}
 
-	if (!bFilesDeleted)
-	{
-		// Couldn't delete files. Delete model and return.
-		Model.Reset();
-		return;
+		if (FArchive* MemoryWriter = FileManager.CreateFileWriter(*FilePath))
+		{
+			*MemoryWriter << CustomizableObjectHeader;
+
+			MemoryWriters.Emplace(MemoryWriter);
+		}
 	}
 
-	// Serialize Streamable resources
+	if (bFilesDeleted)
 	{
-		// Create file writer
-		TUniquePtr<FArchive> StreamableMemoryWriter(FileManager.CreateFileWriter(*StreamableDataFullFileName));
-		check(StreamableMemoryWriter);
-
-		// Serailize headers to validate data
-		*StreamableMemoryWriter << CustomizableObjectHeader;
-
-		const auto WriteBulkDataToDisk = [&StreamableMemoryWriter](MutablePrivate::FFile& File, TArray64<uint8>& FileBulkData, uint32 FileIndex)
+		// Serialize Streamable resources
+		const auto WriteBulkDataToDisk = [&MemoryWriters](MutablePrivate::FFile& File, TArray64<uint8>& FileBulkData, uint32 FileIndex)
 			{
-				switch (File.DataType)
-				{
-				case MutablePrivate::EStreamableDataType::Model:
-				{
-					StreamableMemoryWriter->Serialize(FileBulkData.GetData(), FileBulkData.Num() * sizeof(uint8));
-					break;
-				}
-				case MutablePrivate::EStreamableDataType::RealTimeMorph:
-				{
-					// TODO: Store clothing streamable to disk. UE-222777
-					break;
-				}
-				case MutablePrivate::EStreamableDataType::Clothing:
+				const int32 DataTypeIndex = static_cast<int32>(File.DataType);
+				if (ensure(MemoryWriters.IsValidIndex(DataTypeIndex)))
 				{
-					// TODO: Store clothing streamable to disk. UE-222777
-					break;
-				}
-				default:
-					break;
+					MemoryWriters[DataTypeIndex]->Serialize(FileBulkData.GetData(), FileBulkData.Num() * sizeof(uint8));
 				}
-				
 			};
 
 		// Serialize streamable resources into a single file and fix offsets
 		constexpr bool bDropData = true;
-		MutablePrivate::SerializeBulkDataFiles(PlatformData, PlatformData.BulkDataFiles, WriteBulkDataToDisk, bDropData);
-		StreamableMemoryWriter->Flush();
-		StreamableMemoryWriter->Close();
-
-		PlatformData.MorphStreamableData.Data.Empty();
-		PlatformData.ClothingStreamableData.Data.Empty();
-	}
-
-	// Serialize Model and ModelResources. Store after SerializeBulkDataFiles fixes the HashToStreamableFiles offsets.
-	{
-		// Create file writer
-		TUniquePtr<FArchive> ModelMemoryWriter(FileManager.CreateFileWriter(*CompileDataFullFileName));
-		check(ModelMemoryWriter);
+		MutablePrivate::SerializeBulkDataFiles(*PlatformData.Get(), PlatformData->BulkDataFiles, WriteBulkDataToDisk, bDropData);
 
-		// Serailize headers to validate data
-		*ModelMemoryWriter << CustomizableObjectHeader;
-
-		ModelMemoryWriter->Serialize(ModelResourcesData.GetData(), ModelResourcesData.Num() * sizeof(uint8));
 
+		// Serialize Model and ModelResources. Store after SerializeBulkDataFiles fixes the HashToStreamableFiles offsets.
+		if (ensure(MemoryWriters.IsValidIndex(0)))
 		{
+			MemoryWriters[0]->Serialize(ModelResourcesData.GetData(), ModelResourcesData.Num() * sizeof(uint8));
+
 			// ModelMemoryWriter (Writer to disk) doesn't handle FNames properly. Serialize them ModelStreamables in two steps.
 			TArray64<uint8> ModelStreamablesBytes;
 			FMemoryWriter64 ModelStreamablesMemoryWriter(ModelStreamablesBytes);
-			ModelStreamablesMemoryWriter << *PlatformData.ModelStreamables.Get();
-			ModelMemoryWriter->Serialize(ModelStreamablesBytes.GetData(), ModelStreamablesBytes.Num() * sizeof(uint8));
-		}
-
-		ModelMemoryWriter->Serialize(PlatformData.ModelData.GetData(), PlatformData.ModelData.Num() * sizeof(uint8));
-		
-		ModelMemoryWriter->Flush();
-		ModelMemoryWriter->Close();
+			ModelStreamablesMemoryWriter << *PlatformData->ModelStreamableBulkData.Get();
 
-		PlatformData.ModelData.Empty();
+			MemoryWriters[0]->Serialize(ModelStreamablesBytes.GetData(), ModelStreamablesBytes.Num() * sizeof(uint8));
+			MemoryWriters[0]->Serialize(ModelData.GetData(), ModelData.Num() * sizeof(uint8));
+		}
+	}
+	else
+	{
+		// Couldn't delete files. Delete model to invalidate compilation.
+		PlatformData->Model.Reset();
+	}
+	
+	// Write to disk and close the files
+	for (TUniquePtr<FArchive>& MemoryWriter : MemoryWriters)
+	{
+		if (MemoryWriter)
+		{
+			MemoryWriter->Flush();
+			MemoryWriter->Close();
+		}
 	}
-
-	bStoredSuccessfully = true;
 }
 
 
diff --git a/Engine/Plugins/Mutable/Source/CustomizableObjectEditor/Private/MuCOE/CustomizableObjectCompileRunnable.h b/Engine/Plugins/Mutable/Source/CustomizableObjectEditor/Private/MuCOE/CustomizableObjectCompileRunnable.h
--- a/Engine/Plugins/Mutable/Source/CustomizableObjectEditor/Private/MuCOE/CustomizableObjectCompileRunnable.h
+++ b/Engine/Plugins/Mutable/Source/CustomizableObjectEditor/Private/MuCOE/CustomizableObjectCompileRunnable.h
@@ -85,7 +85,7 @@ class FCustomizableObjectSaveDDRunnable : public FRunnable
 {
 public:
 
-	FCustomizableObjectSaveDDRunnable(const TSharedPtr<FCompilationRequest>& InRequest, TSharedPtr<mu::FModel> InModel, UModelResources& ModelResources, TSharedPtr<FModelStreamableBulkData> ModelStreamables);
+	FCustomizableObjectSaveDDRunnable(const TSharedPtr<FCompilationRequest>& InRequest, TSharedPtr<MutablePrivate::FMutableCachedPlatformData>& InPlatformData);
 
 	// FRunnable interface
 	uint32 Run() override;
@@ -102,7 +102,7 @@ private:
 
 	void StoreCachedPlatformDataInDDC(bool& bStoredSuccessfully);
 
-	void StoreCachedPlatformDataToDisk(bool& bStoredSuccessfully);
+	void StoreCachedPlatformDataToDisk();
 
 	FCompilationOptions Options;
 
@@ -110,10 +110,8 @@ private:
 
 	FString CustomizableObjectName;
 
-	// Paths used to save files to disk
-	FString FolderPath;
-	FString CompileDataFullFileName;
-	FString StreamableDataFullFileName;
+	// Used to save files to disk
+	FString FullFileName;
 
 	UE::DerivedData::FCacheKey DDCKey;
 	UE::DerivedData::ECachePolicy DefaultDDCPolicy;
@@ -123,12 +121,11 @@ private:
 
 public:
 
-	TSharedPtr<mu::FModel, ESPMode::ThreadSafe> Model;
-	TSharedPtr<FModelStreamableBulkData> ModelStreamables;
+	TArray64<uint8> ModelData;
 	TArray64<uint8> ModelResourcesData;
 
 	// Cached platform data
-	MutablePrivate::FMutableCachedPlatformData PlatformData;
+	TSharedPtr<MutablePrivate::FMutableCachedPlatformData> PlatformData;
 
 	// DDC Helpers
 	TArray<MutablePrivate::FFile> BulkDataFilesDDC;
diff --git a/Engine/Plugins/Mutable/Source/CustomizableObjectEditor/Private/MuCOE/CustomizableObjectCompiler.cpp b/Engine/Plugins/Mutable/Source/CustomizableObjectEditor/Private/MuCOE/CustomizableObjectCompiler.cpp
--- a/Engine/Plugins/Mutable/Source/CustomizableObjectEditor/Private/MuCOE/CustomizableObjectCompiler.cpp
+++ b/Engine/Plugins/Mutable/Source/CustomizableObjectEditor/Private/MuCOE/CustomizableObjectCompiler.cpp
@@ -243,6 +243,9 @@ void FCustomizableObjectCompiler::Compile(const TSharedRef<FCompilationRequest>&
 
 	CompilationStartTime = FPlatformTime::Seconds();
 
+	// Platform data to cache all compiled resources
+	PlatformData = MakeShared<MutablePrivate::FMutableCachedPlatformData>();
+
 	// Now that we know for sure that the CO is locked and there are no pending updates of instances using the CO,
 	// destroy any live update instances, as they become invalid when recompiling the CO
 	for (TObjectIterator<UCustomizableObjectInstance> It; It; ++It)
@@ -424,7 +427,6 @@ void FCustomizableObjectCompiler::AddReferencedObjects(FReferenceCollector& Coll
 {
 	Collector.AddReferencedObjects(ArrayGCProtect);
 	Collector.AddReferencedObject(CurrentObject);
-	Collector.AddReferencedObject(ModelResources);
 }
 
 
@@ -838,9 +840,10 @@ void FCustomizableObjectCompiler::CompileInternal(bool bAsync)
 		// The second part will be filled when the core compilation finishes so that the CompilationContext is complete.
 		const FString ModelResourcesName = GetModelResourcesNameForPlatform(*CurrentObject, *CurrentOptions.TargetPlatform);
 		UObject* Outer = CurrentOptions.bIsCooking ? CurrentObject : StaticCast<UObject*>(GetTransientPackage());
-		ModelResources = NewObject<UModelResources>(Outer, FName(*ModelResourcesName), RF_Public);
-
-		ModelStreamableBulkData = MakeShared<FModelStreamableBulkData>();
+		TObjectPtr<UModelResources> ModelResources = NewObject<UModelResources>(Outer, FName(*ModelResourcesName), RF_Public);
+		PlatformData->ModelResources = TStrongObjectPtr<UModelResources>(ModelResources);
+		
+		PlatformData->ModelStreamableBulkData = MakeShared<FModelStreamableBulkData>();
 
 		ModelResources->CodeVersion = UCustomizableObjectPrivate::CurrentSupportedVersion;
 
@@ -1028,18 +1031,18 @@ void FCustomizableObjectCompiler::CompleteRequest(ECompilationStatePrivate State
 			System->UnlockObject(CurrentObject);
 		}
 
-		if (Model)
+		if (PlatformData->Model)
 		{
-			Model->GetPrivate()->UnloadRoms();
+			PlatformData->Model->GetPrivate()->UnloadRoms();
 		}
 
 		if (!CurrentOptions.bIsCooking)
 		{
 			if (Result == ECompilationResultPrivate::Success || Result == ECompilationResultPrivate::Warnings)
 			{
-				CurrentObject->GetPrivate()->SetModel(Model, GenerateIdentifier(*CurrentObject));
-				CurrentObject->GetPrivate()->SetModelResources(ModelResources, CurrentOptions.bIsCooking);
-				CurrentObject->GetPrivate()->SetModelStreamableBulkData(ModelStreamableBulkData, CurrentOptions.bIsCooking);
+				CurrentObject->GetPrivate()->SetModel(PlatformData->Model, GenerateIdentifier(*CurrentObject));
+				CurrentObject->GetPrivate()->SetModelResources(PlatformData->ModelResources.Get(), CurrentOptions.bIsCooking);
+				CurrentObject->GetPrivate()->SetModelStreamableBulkData(PlatformData->ModelStreamableBulkData, CurrentOptions.bIsCooking);
 			}
 			else
 			{
@@ -1105,9 +1108,8 @@ void FCustomizableObjectCompiler::CompleteRequest(ECompilationStatePrivate State
 	// Request completed, reset pointers and state
 	CurrentObject = nullptr;
 	CurrentRequest.Reset();
-	Model.Reset();
-	ModelResources = nullptr;
-	ModelStreamableBulkData.Reset();
+	PlatformData.Reset();
+	
 
 	UE_LOG(LogMutable, Verbose, TEXT("PROFILE: [ %16.8f ] Completed compile request."), FPlatformTime::Seconds());
 	UE_LOG(LogMutable, Verbose, TEXT("PROFILE: -----------------------------------------------------------"));
@@ -1189,9 +1191,9 @@ bool FCustomizableObjectCompiler::TryLoadCompiledDataFromDDC(UCustomizableObject
 	FCacheKey CacheKey = CurrentRequest->GetDerivedDataCacheKey();
 	CustomizableObject.GetPrivate()->LoadCompiledDataFromDDC(CurrentOptions, DefaultPolicy, &CacheKey);
 	
-	Model = CustomizableObject.GetPrivate()->GetModel();
-	ModelResources = CustomizableObject.GetPrivate()->GetModelResources(CurrentOptions.bIsCooking);
-	ModelStreamableBulkData = CustomizableObject.GetPrivate()->GetModelStreamableBulkData(CurrentOptions.bIsCooking);
+	PlatformData->Model = CustomizableObject.GetPrivate()->GetModel();
+	PlatformData->ModelResources = TStrongObjectPtr<UModelResources>(CustomizableObject.GetPrivate()->GetModelResources(CurrentOptions.bIsCooking));
+	PlatformData->ModelStreamableBulkData = CustomizableObject.GetPrivate()->GetModelStreamableBulkData(CurrentOptions.bIsCooking);
 	return CustomizableObject.IsCompiled();
 }
 
@@ -1275,7 +1277,8 @@ void FCustomizableObjectCompiler::FinishCompilationTask()
 	check(CompileTask.IsValid());
 
 	UpdateCompilerLogData();
-	Model = CompileTask->Model;
+	TSharedPtr<mu::FModel, ESPMode::ThreadSafe> Model = CompileTask->Model;
+	PlatformData->Model = Model;
 
 	// Generate a map that using the resource id tells the offset and size of the resource inside the bulk data
 	// At this point it is assumed that all data goes into a single file.
@@ -1283,7 +1286,7 @@ void FCustomizableObjectCompiler::FinishCompilationTask()
 	{
 		const int32 NumStreamingFiles = Model->GetRomCount();
 
-		TMap<uint32, FMutableStreamableBlock>& ModelStreamables = ModelStreamableBulkData->ModelStreamables;
+		TMap<uint32, FMutableStreamableBlock>& ModelStreamables = PlatformData->ModelStreamableBulkData->ModelStreamables;
 		ModelStreamables.Empty(NumStreamingFiles);
 
 		// TODO: Temp. Remove after unifying generated output files code between editor an package. UE-222777
@@ -1305,6 +1308,8 @@ void FCustomizableObjectCompiler::FinishCompilationTask()
 
 	// Second part of model resources can be filled now.
 	{
+		TObjectPtr<UModelResources> ModelResources = PlatformData->ModelResources.Get();
+
 		ModelResources->Skeletons.Reserve(CompilationContext->ReferencedSkeletons.Num());
 		for (const USkeleton* Skeleton : CompilationContext->ReferencedSkeletons)
 		{
@@ -1329,18 +1334,11 @@ void FCustomizableObjectCompiler::FinishCompilationTask()
 
 		ModelResources->SkinWeightProfilesInfo = MoveTemp(CompilationContext->SkinWeightProfilesInfo);
 
-		// Create the RealTimeMorphsTargets Blocks from the per mesh Morph data.
-		uint64 RealTimeMorphDataSize = 0;
-		for (const TPair<uint32, FRealTimeMorphMeshData>& MeshData : CompilationContext->RealTimeMorphTargetPerMeshData)
-		{
-			RealTimeMorphDataSize += MeshData.Value.Data.Num();
-		}
-
+		TSharedPtr<FModelStreamableBulkData> ModelStreamableBulkData = PlatformData->ModelStreamableBulkData;
 		ModelStreamableBulkData->RealTimeMorphStreamables.Empty(32);
-		ModelResources->EditorOnlyMorphTargetReconstructionData.Empty(RealTimeMorphDataSize);
 
 		uint64 RealTimeMorphDataOffsetInBytes = 0;
-		for (const TPair<uint32, FRealTimeMorphMeshData>& MeshData : CompilationContext->RealTimeMorphTargetPerMeshData)
+		for (TPair<uint32, FRealTimeMorphMeshData>& MeshData : CompilationContext->RealTimeMorphTargetPerMeshData)
 		{
 			const uint32 DataSizeInBytes = (uint32)MeshData.Value.Data.Num() * sizeof(FMorphTargetVertexData);
 			FRealTimeMorphStreamable& ResourceMeshData = ModelStreamableBulkData->RealTimeMorphStreamables.FindOrAdd(MeshData.Key);
@@ -1355,21 +1353,19 @@ void FCustomizableObjectCompiler::FinishCompilationTask()
 			ResourceMeshData.SourceId = MeshData.Value.SourceId;
 
 			RealTimeMorphDataOffsetInBytes += DataSizeInBytes;
-			ModelResources->EditorOnlyMorphTargetReconstructionData.Append(MeshData.Value.Data);
-		}
 
-		// Create the Clothing Blocks from the per mesh Morph data.
-		uint64 ClothingDataNum = 0;
-		for (const TPair<uint32, FClothingMeshData>& MeshData : CompilationContext->ClothingPerMeshData)
-		{
-			ClothingDataNum += MeshData.Value.Data.Num();
+			PlatformData->MorphStreamableData.Set(
+				MeshData.Key,
+				reinterpret_cast<uint8*>(MeshData.Value.Data.GetData()),
+				DataSizeInBytes);
+
+			MeshData.Value.Data.Empty();
 		}
 
 		ModelStreamableBulkData->ClothingStreamables.Empty(32);
-		ModelResources->EditorOnlyClothingMeshToMeshVertData.Empty(ClothingDataNum);
 
 		uint64 ClothingDataOffsetInBytes = 0;
-		for (const TPair<uint32, FClothingMeshData>& MeshData : CompilationContext->ClothingPerMeshData)
+		for (TPair<uint32, FClothingMeshData>& MeshData : CompilationContext->ClothingPerMeshData)
 		{
 			const uint32 DataSizeInBytes = (uint32)MeshData.Value.Data.Num() * sizeof(FCustomizableObjectMeshToMeshVertData);
 			FClothingStreamable& ResourceMeshData = ModelStreamableBulkData->ClothingStreamables.FindOrAdd(MeshData.Key);
@@ -1387,7 +1383,13 @@ void FCustomizableObjectCompiler::FinishCompilationTask()
 			ResourceMeshData.SourceId = MeshData.Value.SourceId;
 
 			ClothingDataOffsetInBytes += DataSizeInBytes;
-			ModelResources->EditorOnlyClothingMeshToMeshVertData.Append(MeshData.Value.Data);
+
+			PlatformData->ClothingStreamableData.Set(
+				MeshData.Key,
+				reinterpret_cast<uint8*>(MeshData.Value.Data.GetData()),
+				DataSizeInBytes);
+
+			MeshData.Value.Data.Empty();
 		}
 
 		ModelResources->ClothingAssetsData = MoveTemp(CompilationContext->ClothingAssetsData);
@@ -1442,7 +1444,8 @@ void FCustomizableObjectCompiler::FinishCompilationTask()
 
 	// Create SaveDD task
 	TRACE_BEGIN_REGION(UE_MUTABLE_SAVEDD_REGION);
-	SaveDDTask = MakeShareable(new FCustomizableObjectSaveDDRunnable(CurrentRequest, Model, *ModelResources, ModelStreamableBulkData));
+
+	SaveDDTask = MakeShareable(new FCustomizableObjectSaveDDRunnable(CurrentRequest, PlatformData));
 }
 
 
@@ -1462,7 +1465,7 @@ void FCustomizableObjectCompiler::FinishSavingDerivedDataTask()
 		check(!CurrentObject->GetPrivate()->CachedPlatformsData.Find(PlatformName));
 
 		MutablePrivate::FMutableCachedPlatformData& Data = CurrentObject->GetPrivate()->CachedPlatformsData.Add(PlatformName);
-		Data = MoveTemp(SaveDDTask->PlatformData);
+		Data = MoveTemp(*PlatformData.Get());
 
 		Data.ModelResources->InitCookData();
 	}
diff --git a/Engine/Plugins/Mutable/Source/CustomizableObjectEditor/Private/MuCOE/CustomizableObjectCompiler.h b/Engine/Plugins/Mutable/Source/CustomizableObjectEditor/Private/MuCOE/CustomizableObjectCompiler.h
--- a/Engine/Plugins/Mutable/Source/CustomizableObjectEditor/Private/MuCOE/CustomizableObjectCompiler.h
+++ b/Engine/Plugins/Mutable/Source/CustomizableObjectEditor/Private/MuCOE/CustomizableObjectCompiler.h
@@ -140,11 +140,8 @@ private:
 	/** Array used to protect from garbage collection those COs loaded asynchronously */
 	TArray<TObjectPtr<UObject>> ArrayGCProtect;
 
-	TSharedPtr<mu::FModel, ESPMode::ThreadSafe> Model;
-	TSharedPtr<FModelStreamableBulkData> ModelStreamableBulkData;
-
-	// Protected from GC with FCustomizableObjectCompiler::AddReferencedObjects
-	TObjectPtr<UModelResources> ModelResources;
+	/** Cached Platform Data. */
+	TSharedPtr<MutablePrivate::FMutableCachedPlatformData> PlatformData;
 
 	// Protected from GC with FCustomizableObjectCompiler::AddReferencedObjects
 	TObjectPtr<UCustomizableObject> CurrentObject = nullptr;
