[mutable] Core : Fixed Levelcount mismatch for "ImageNormalComposite_EmergencyFix"
- When trying to fill the mips of a texture we first need to ensure the amount of mips we want to fill with data are first present in the target texture.

#rnx

diff --git a/Engine/Plugins/Mutable/Source/MutableRuntime/Private/MuR/CodeRunner.cpp b/Engine/Plugins/Mutable/Source/MutableRuntime/Private/MuR/CodeRunner.cpp
--- a/Engine/Plugins/Mutable/Source/MutableRuntime/Private/MuR/CodeRunner.cpp
+++ b/Engine/Plugins/Mutable/Source/MutableRuntime/Private/MuR/CodeRunner.cpp
@@ -8,7 +8,6 @@
 #include "Logging/LogCategory.h"
 #include "Logging/LogMacros.h"
 #include "Math/IntPoint.h"
-#include "Math/UnrealMathSSE.h"
 #include "Math/UnrealMathUtility.h"
 #include "Math/Vector2D.h"
 #include "Misc/AssertionMacros.h"
@@ -27,14 +26,11 @@
 #include "MuR/OpImageBlend.h"
 #include "MuR/OpImageColourMap.h"
 #include "MuR/OpImageDisplace.h"
-#include "MuR/OpImageGradient.h"
 #include "MuR/OpImageInterpolate.h"
-#include "MuR/OpImageInvert.h"
 #include "MuR/OpImageLuminance.h"
 #include "MuR/OpImageNormalCombine.h"
 #include "MuR/OpImageProject.h"
 #include "MuR/OpImageRasterMesh.h"
-#include "MuR/OpImageSaturate.h"
 #include "MuR/OpImageTransform.h"
 #include "MuR/OpLayoutPack.h"
 #include "MuR/OpLayoutRemoveBlocks.h"
@@ -58,11 +54,9 @@
 #include "MuR/Parameters.h"
 #include "MuR/ParametersPrivate.h"
 #include "MuR/PhysicsBody.h"
-#include "MuR/Platform.h"
 #include "MuR/Skeleton.h"
 #include "MuR/SystemPrivate.h"
 #include "Templates/Tuple.h"
-#include "Trace/Detail/Channel.h"
 
 namespace
 {
@@ -3302,6 +3296,7 @@ namespace mu
 					int32 LevelCount = Base->GetLODCount();
 					
 					TSharedPtr<FImage> NormalFix = CloneOrTakeOver(Normal);
+					NormalFix->DataStorage.SetNumLODs(LevelCount);
 
 					FMipmapGenerationSettings MipSettings{};
 					ImOp.ImageMipmap(Settings.ImageCompressionQuality, NormalFix.Get(), NormalFix.Get(), StartLevel, LevelCount, MipSettings);
