[mutable] Fixed group parameters being shown in the editor preview parameter list when their directly-connected parent group parameter is set to None.
#tests Editor
#rnx

diff --git a/Engine/Plugins/Mutable/Source/CustomizableObjectEditor/Private/MuCOE/GenerateMutableSource/GenerateMutableSource.cpp b/Engine/Plugins/Mutable/Source/CustomizableObjectEditor/Private/MuCOE/GenerateMutableSource/GenerateMutableSource.cpp
--- a/Engine/Plugins/Mutable/Source/CustomizableObjectEditor/Private/MuCOE/GenerateMutableSource/GenerateMutableSource.cpp
+++ b/Engine/Plugins/Mutable/Source/CustomizableObjectEditor/Private/MuCOE/GenerateMutableSource/GenerateMutableSource.cpp
@@ -1112,11 +1112,13 @@ mu::Ptr<mu::NodeObject> GenerateMutableSource(const UEdGraphPin * Pin, FMutableG
 		{
 			if (const UEdGraphPin* ConnectedPin = FollowOutputPin(*Pin))
 			{
-				UCustomizableObjectNodeObject* NodeObject = Cast<UCustomizableObjectNodeObject>(ConnectedPin->GetOwningNode());
+				const FGuid* ParentId = nullptr;
+				FString ValueName;
 
-				if (NodeObject)
+				if (UCustomizableObjectNodeObject* NodeObject = Cast<UCustomizableObjectNodeObject>(ConnectedPin->GetOwningNode()))
 				{
-					const FGuid* ParentId = GenerationContext.GroupIdToExternalNodeMap.FindKey(NodeObject);
+					ParentId = GenerationContext.GroupIdToExternalNodeMap.FindKey(NodeObject);
+					ValueName = NodeObject->ObjectName;
 
 					// Group objects in the same graph aren't in the GroupIdToExternalNodeMap, so follow the pins instead
 					if (!ParentId && NodeObject->OutputPin())
@@ -1129,28 +1131,33 @@ mu::Ptr<mu::NodeObject> GenerateMutableSource(const UEdGraphPin * Pin, FMutableG
 							}
 						}
 					}
+				}
+				else if (UCustomizableObjectNodeObjectGroup* NodeObjectGroup = Cast<UCustomizableObjectNodeObjectGroup>(ConnectedPin->GetOwningNode()))
+				{
+					ParentId = &NodeObjectGroup->NodeGuid;
+					ValueName = TypedNodeGroup->GroupName;
+				}
 
-					if (ParentId)
+				if (ParentId)
+				{
+					FString* ParentParamName = GenerationContext.GuidToParamNameMap.Find(*ParentId);
+
+					if (ParentParamName)
 					{
-						FString* ParentParamName = GenerationContext.GuidToParamNameMap.Find(*ParentId);
+						ParameterUIData.ParamUIMetadata.ExtraInformation.Add(FString("__ParentParamName"), *ParentParamName);
 
-						if (ParentParamName)
+						if (bHideWhenNotSelected)
 						{
-							ParameterUIData.ParamUIMetadata.ExtraInformation.Add(FString("__ParentParamName"), *ParentParamName);
-
-							if (bHideWhenNotSelected)
-							{
-								ParameterUIData.ParamUIMetadata.ExtraInformation.Add(FString("__DisplayWhenParentValueEquals"), NodeObject->ObjectName);
-							}
+							ParameterUIData.ParamUIMetadata.ExtraInformation.Add(FString("__DisplayWhenParentValueEquals"), ValueName);
+						}
 
-							if (bCollapseUnderParent)
-							{
-								ParameterUIData.ParamUIMetadata.ExtraInformation.Add(FString("CollapseUnderParent"));
+						if (bCollapseUnderParent)
+						{
+							ParameterUIData.ParamUIMetadata.ExtraInformation.Add(FString("CollapseUnderParent"));
 
-								FMutableParameterData ParentParameterUIData;
-								ParentParameterUIData.ParamUIMetadata.ExtraInformation.Add(FString("__HasCollapsibleChildren"));
-								GenerationContext.ParameterUIDataMap.Add(*ParentParamName, ParentParameterUIData);
-							}
+							FMutableParameterData ParentParameterUIData;
+							ParentParameterUIData.ParamUIMetadata.ExtraInformation.Add(FString("__HasCollapsibleChildren"));
+							GenerationContext.ParameterUIDataMap.Add(*ParentParamName, ParentParameterUIData);
 						}
 					}
 				}
