[mutable] Fixed the default value color picker widget in a color parameter node showing linear color while the same parameter in the preview shows as sRGB.
#tests Editor
#rnx

diff --git a/Engine/Plugins/Mutable/Source/CustomizableObjectEditor/Private/MuCOE/CustomizableInstanceDetails.cpp b/Engine/Plugins/Mutable/Source/CustomizableObjectEditor/Private/MuCOE/CustomizableInstanceDetails.cpp
--- a/Engine/Plugins/Mutable/Source/CustomizableObjectEditor/Private/MuCOE/CustomizableInstanceDetails.cpp
+++ b/Engine/Plugins/Mutable/Source/CustomizableObjectEditor/Private/MuCOE/CustomizableInstanceDetails.cpp
@@ -1391,7 +1391,7 @@ TSharedRef<SWidget> FCustomizableInstanceDetails::GenerateColorWidget(const FStr
 		.Color(this, &FCustomizableInstanceDetails::GetColorParameterValue, ParamName)
 		.ShowBackgroundForAlpha(false)
 		.AlphaDisplayMode(EColorBlockAlphaDisplayMode::Ignore)
-		.UseSRGB(false)
+		.UseSRGB(true)
 		.OnMouseButtonDown(this, &FCustomizableInstanceDetails::OnColorBlockMouseButtonDown, ParamName)
 		.CornerRadius(FVector4(4.0f, 4.0f, 4.0f, 4.0f));
 }
