[mutable] Cleanup to remove dead and duplicated code.
#rnx

diff --git a/Engine/Plugins/Mutable/Source/MutableTools/Private/MuT/ASTOpImageCrop.cpp b/Engine/Plugins/Mutable/Source/MutableTools/Private/MuT/ASTOpImageCrop.cpp
--- a/Engine/Plugins/Mutable/Source/MutableTools/Private/MuT/ASTOpImageCrop.cpp
+++ b/Engine/Plugins/Mutable/Source/MutableTools/Private/MuT/ASTOpImageCrop.cpp
@@ -88,7 +88,6 @@ namespace mu
 				Args.source = Source->linkedAddress;
 			}
 
-			if (Source) Args.source = Source->linkedAddress;
 			Args.minX = Min[0];
 			Args.minY = Min[1];
 			Args.sizeX = Size[0];
diff --git a/Engine/Plugins/Mutable/Source/MutableTools/Private/MuT/ASTOpImageLayer.cpp b/Engine/Plugins/Mutable/Source/MutableTools/Private/MuT/ASTOpImageLayer.cpp
--- a/Engine/Plugins/Mutable/Source/MutableTools/Private/MuT/ASTOpImageLayer.cpp
+++ b/Engine/Plugins/Mutable/Source/MutableTools/Private/MuT/ASTOpImageLayer.cpp
@@ -22,7 +22,6 @@
 namespace mu
 {
 
-
 	ASTOpImageLayer::ASTOpImageLayer()
 		: base(this)
 		, blend(this)
@@ -31,7 +30,6 @@ namespace mu
 	}
 
 
-	//-------------------------------------------------------------------------------------------------
 	ASTOpImageLayer::~ASTOpImageLayer()
 	{
 		// Explicit call needed to avoid recursive destruction
@@ -39,7 +37,6 @@ namespace mu
 	}
 
 
-	//-------------------------------------------------------------------------------------------------
 	bool ASTOpImageLayer::IsEqual(const ASTOp& InOtherUntyped) const
 	{
 		if (InOtherUntyped.GetOpType() == GetOpType())
@@ -57,7 +54,6 @@ namespace mu
 	}
 
 
-	//-------------------------------------------------------------------------------------------------
 	uint64 ASTOpImageLayer::Hash() const
 	{
 		uint64 res = std::hash<EOpType>()(GetOpType());
@@ -68,7 +64,6 @@ namespace mu
 	}
 
 
-	//-------------------------------------------------------------------------------------------------
 	mu::Ptr<ASTOp> ASTOpImageLayer::Clone(MapChildFuncRef mapChild) const
 	{
 		Ptr<ASTOpImageLayer> n = new ASTOpImageLayer();
@@ -83,7 +78,6 @@ namespace mu
 	}
 
 
-	//-------------------------------------------------------------------------------------------------
 	void ASTOpImageLayer::ForEachChild(const TFunctionRef<void(ASTChild&)> f)
 	{
 		f(base);
@@ -92,7 +86,6 @@ namespace mu
 	}
 
 
-	//-------------------------------------------------------------------------------------------------
 	void ASTOpImageLayer::Link(FProgram& program, FLinkerOptions*)
 	{
 		// Already linked?
@@ -118,7 +111,6 @@ namespace mu
 	}
 
 
-	//-------------------------------------------------------------------------------------------------
 	FImageDesc ASTOpImageLayer::GetImageDesc(bool returnBestOption, FGetImageDescContext* context) const 
 	{
 		FImageDesc res;
@@ -156,7 +148,6 @@ namespace mu
 	}
 
 
-	//-------------------------------------------------------------------------------------------------
 	void ASTOpImageLayer::GetLayoutBlockSize(int* pBlockX, int* pBlockY)
 	{
 		if (base)
@@ -166,7 +157,6 @@ namespace mu
 	}
 
 
-	//-------------------------------------------------------------------------------------------------
 	mu::Ptr<ImageSizeExpression> ASTOpImageLayer::GetImageSizeExpression() const
 	{
 		if (base)
@@ -178,8 +168,6 @@ namespace mu
 	}
 
 
-
-	//---------------------------------------------------------------------------------------------
 	Ptr<ASTOp> ASTOpImageLayer::OptimiseSemantic(const FModelOptimizationOptions& options, int32 Pass) const
 	{
 		Ptr<ASTOp> at;
@@ -548,116 +536,6 @@ namespace mu
 			return at;
 		}
 
-		// Promote conditions from the base
-		EOpType baseType = baseAt->GetOpType();
-		switch (baseType)
-		{
-			// Seems to cause operation explosion in optimizer in bandit model.
-			// moved to generic sink in the default.
-//            case EOpType::IM_CONDITIONAL:
-//            {
-//                bModified = true;
-
-//                OP op = program.m_code[baseAt];
-
-//                OP aOp = program.m_code[at];
-//                aOp.Args.ImageLayer.base = program.m_code[baseAt].Args.Conditional.yes;
-//                op.args.Conditional.yes = program.AddOp( aOp );
-
-//                OP bOp = program.m_code[at];
-//                bOp.Args.ImageLayer.base = program.m_code[baseAt].Args.Conditional.no;
-//                op.args.Conditional.no = program.AddOp( bOp );
-
-//                at = program.AddOp( op );
-//                break;
-//            }
-
-		case EOpType::IM_SWITCH:
-		{
-			// Warning:
-			// It seems to cause data explosion in optimizer in some models. Because
-			// all switch branches become unique constants
-
-//                // See if the blended has an identical switch, to optimise it too
-//                auto baseSwitch = dynamic_cast<const ASTOpSwitch*>( baseAt.get() );
-//                auto blendedSwitch = dynamic_cast<const ASTOpSwitch*>( blendAt.get() );
-//                auto maskSwitch = dynamic_cast<const ASTOpSwitch*>( maskAt.get() );
-
-//                bool blendedToo = baseSwitch->IsCompatibleWith( blendedSwitch );
-//                bool maskToo = baseSwitch->IsCompatibleWith( maskSwitch );
-
-//                // Move the layer operation down all the paths
-//                auto nop = mu::Clone<ASTOpSwitch>(baseSwitch);
-
-//                if (nop->def)
-//                {
-//                    auto defOp = mu::Clone<ASTOpFixed>(this);
-//                    defOp->SetChild( defOp->op.args.ImageLayer.base, nop->def );
-//                    if (blendedToo)
-//                    {
-//                        defOp->SetChild( defOp->op.args.ImageLayer.blended, blendedSwitch->def );
-//                    }
-//                    if (maskToo)
-//                    {
-//                        defOp->SetChild( defOp->op.args.ImageLayer.mask, maskSwitch->def );
-//                    }
-//                    nop->def = defOp;
-//                }
-
-//                for ( size_t v=0; v<nop->cases.Num(); ++v )
-//                {
-//                    if ( nop->cases[v].branch )
-//                    {
-//                        auto bOp = mu::Clone<ASTOpFixed>(this);
-//                        bOp->SetChild( bOp->op.args.ImageLayer.base, nop->cases[v].branch );
-
-//                        if (blendedToo)
-//                        {
-//                            bOp->SetChild( bOp->op.args.ImageLayer.blended, blendedSwitch->FindBranch( nop->cases[v].condition ) );
-//                        }
-
-//                        if (maskToo)
-//                        {
-//                            bOp->SetChild( bOp->op.args.ImageLayer.mask, maskSwitch->FindBranch( nop->cases[v].condition ) );
-//                        }
-
-//                        nop->cases[v].branch = bOp;
-//                    }
-//                }
-
-//                at = nop;
-
-			break;
-		}
-
-		default:
-		{
-			// Try generic base sink
-//                OP templateOp = program.m_code[at];
-
-//                CloneNeutralPartialTreeVisitor neutral;
-//                EOpType supportedOps[] = {
-//                    EOpType::IM_CONDITIONAL,
-//                    EOpType::NONE
-//                };
-//                auto newRoot = neutral.Apply
-//                         (
-//                            program.m_code[at].Args.ImageLayer.base, program,
-//                            Image, supportedOps,
-//                            &templateOp, &templateOp.Args.ImageLayer.base);
-
-//                // If there is a change
-//                if (newRoot!=program.m_code[at].Args.ImageLayer.base)
-//                {
-//                    bModified = true;
-//                    at = newRoot;
-//                }
-
-			break;
-		}
-
-		}
-
 		// If we failed to optimize so far, see if it is worth optimizing the blended branch only.
 		if (!at && blendAt)
 		{
@@ -665,6 +543,7 @@ namespace mu
 			switch (BlendType)
 			{
 
+
 			case EOpType::IM_SWITCH:
 			{
 				const ASTOpSwitch* BlendSwitch = static_cast<const ASTOpSwitch*>(blendAt.get());
