[mutable] Replaced usage of LODSettings data from CO data cached in ModelResources for COI updates
- Swaped usage of CO LODSettings to ModelResources hosted LODSettings for COI updates
- Made CO LODSettigs Editor Only (as will only be used during compilation, so, only in editor)

#rnx

diff --git a/Engine/Plugins/Mutable/Source/CustomizableObject/Internal/MuCO/CustomizableObjectPrivate.h b/Engine/Plugins/Mutable/Source/CustomizableObject/Internal/MuCO/CustomizableObjectPrivate.h
--- a/Engine/Plugins/Mutable/Source/CustomizableObject/Internal/MuCO/CustomizableObjectPrivate.h
+++ b/Engine/Plugins/Mutable/Source/CustomizableObject/Internal/MuCO/CustomizableObjectPrivate.h
@@ -1314,6 +1314,14 @@ public:
 	UPROPERTY()
 	TArray<FName> ComponentNamesPerObjectComponent;
 
+	/** Minimum LOD to render per Platform. */
+	UPROPERTY()
+	FPerPlatformInt MinLOD;
+
+	/** Minimum LOD to render per Quality level. */
+	UPROPERTY()
+	FPerQualityLevelInt MinQualityLevelLOD;
+	
 	UPROPERTY()
 	FString ReleaseVersion;
 
diff --git a/Engine/Plugins/Mutable/Source/CustomizableObject/Private/MuCO/CustomizableObject.cpp b/Engine/Plugins/Mutable/Source/CustomizableObject/Private/MuCO/CustomizableObject.cpp
--- a/Engine/Plugins/Mutable/Source/CustomizableObject/Private/MuCO/CustomizableObject.cpp
+++ b/Engine/Plugins/Mutable/Source/CustomizableObject/Private/MuCO/CustomizableObject.cpp
@@ -3,12 +3,10 @@
 #include "MuCO/CustomizableObject.h"
 
 #include "Algo/Copy.h"
-#include "AssetRegistry/AssetRegistryModule.h"
 #include "Async/AsyncFileHandle.h"
 #include "EdGraph/EdGraph.h"
 #include "Engine/Engine.h"
 #include "Engine/SkeletalMesh.h"
-#include "Engine/SkeletalMeshLODSettings.h"
 #include "Animation/AnimInstance.h"
 #include "Animation/Skeleton.h"
 #include "Engine/AssetUserData.h"
@@ -17,7 +15,6 @@
 #include "Input/Reply.h"
 #include "Interfaces/ITargetPlatform.h"
 #include "Materials/MaterialInterface.h"
-#include "Misc/DataValidation.h"
 #include "Misc/PackageName.h"
 #include "Misc/Paths.h"
 #include "MuCO/CustomizableObjectInstance.h"
@@ -29,7 +26,6 @@
 #include "MuCO/ICustomizableObjectModule.h"
 #include "MuCO/MutableProjectorTypeUtils.h"
 #include "MuCO/UnrealMutableModelDiskStreamer.h"
-#include "MuCO/UnrealPortabilityHelpers.h"
 #include "MuR/Model.h"
 #include "MuR/ModelPrivate.h"
 #include "MuR/Operations.h"
@@ -40,7 +36,6 @@
 #include "UObject/AssetRegistryTagsContext.h"
 #include "UObject/ObjectSaveContext.h"
 #include "MuCO/ICustomizableObjectEditorModule.h"
-#include "MuCO/CustomizableObjectSystemPrivate.h"
 #include "Misc/MessageDialog.h"
 
 #if WITH_EDITOR
@@ -48,7 +43,6 @@
 #include "Engine/TextureLODSettings.h"
 
 #include "DerivedDataCache.h"
-#include "DerivedDataCacheInterface.h"
 #include "DerivedDataCacheKey.h"
 #include "DerivedDataRequestOwner.h"
 #endif
@@ -2030,28 +2024,29 @@ uint8 UCustomizableObjectPrivate::GetMinLODIndex(const FName& ComponentName) con
 {
 	uint8 MinLODIdx = 0;
 
-	// Use the Scalability quality settings to determine what the MinLOD will be used
-	if (GEngine && GEngine->UseSkeletalMeshMinLODPerQualityLevels)
+	const UModelResources* LocalModelResources = GetModelResources();
+	if (LocalModelResources)
 	{
-		if (UCustomizableObjectSystem::GetInstance() != nullptr)
+		// Use the Scalability quality settings to determine what the MinLOD will be used
+		if (GEngine && GEngine->UseSkeletalMeshMinLODPerQualityLevels)
 		{
-			// Get the quality level for the skeletal meshes for the current scalability setting. The bigger the better.
-			// Value extracted from the Scalability.ini of the current platform (if found).
-			const int32 QualityLevel = UCustomizableObjectSystem::GetInstance()->GetSkeletalMeshMinLODQualityLevel();
-			
-			// Get the MinLOD to use for the target quality level. The lower the better.
-			MinLODIdx = GetPublic()->LODSettings.MinQualityLevelLOD.GetValue(QualityLevel);
+			if (UCustomizableObjectSystem::GetInstance() != nullptr)
+			{
+				// Get the quality level for the skeletal meshes for the current scalability setting. The bigger the better.
+				// Value extracted from the Scalability.ini of the current platform (if found).
+				const int32 QualityLevel = UCustomizableObjectSystem::GetInstance()->GetSkeletalMeshMinLODQualityLevel();
+				MinLODIdx = LocalModelResources->MinQualityLevelLOD.GetValue(QualityLevel);
+			}
+		}
+		else
+		{
+			MinLODIdx = LocalModelResources->MinLOD.GetValue();
 		}
 	}
-	else
-	{
-		MinLODIdx = GetPublic()->LODSettings.MinLOD.GetValue();
-	}
-
+	
 	// Get the first lod the current platform can get generated. It represents the Min LOD mutable can generate for the current platform.
 	uint8 FirstLODAvailable = 0;
 
-	const UModelResources* LocalModelResources = GetModelResources();
 	if (LocalModelResources)
 	{
 		if (const uint8* Result = LocalModelResources->ComponentFirstLODAvailable.Find(ComponentName))
diff --git a/Engine/Plugins/Mutable/Source/CustomizableObject/Private/MuCO/CustomizableObjectInstance.cpp b/Engine/Plugins/Mutable/Source/CustomizableObject/Private/MuCO/CustomizableObjectInstance.cpp
--- a/Engine/Plugins/Mutable/Source/CustomizableObject/Private/MuCO/CustomizableObjectInstance.cpp
+++ b/Engine/Plugins/Mutable/Source/CustomizableObject/Private/MuCO/CustomizableObjectInstance.cpp
@@ -15,7 +15,6 @@
 #include "Engine/SkeletalMesh.h"
 #include "Engine/SkeletalMeshSocket.h"
 #include "Engine/SkeletalMeshLODSettings.h"
-#include "Engine/Texture2DArray.h"
 #include "MaterialDomain.h"
 #include "Materials/Material.h"
 #include "Materials/MaterialInstanceDynamic.h"
@@ -25,27 +24,22 @@
 
 #include "MuCO/CustomizableObjectSystemPrivate.h"
 #include "MuCO/CustomizableObjectSkeletalMesh.h"
-#include "MuCO/CustomizableInstanceLODManagement.h"
 #include "MuCO/CustomizableObjectInstancePrivate.h"
 #include "MuCO/CustomizableObjectExtension.h"
 #include "MuCO/CustomizableObjectMipDataProvider.h"
 #include "MuCO/CustomizableObjectPrivate.h"
 #include "MuCO/CustomizableObjectInstanceUsage.h"
 #include "MuCO/CustomizableObjectInstanceAssetUserData.h"
-#include "MuCO/DefaultImageProvider.h"
 #include "MuCO/ICustomizableObjectModule.h"
 #include "MuCO/UnrealConversionUtils.h"
 #include "MuCO/UnrealPortabilityHelpers.h"
 #include "MuCO/LogBenchmarkUtil.h"
-#include "MuR/Model.h"
 
 #include "PhysicsEngine/PhysicsAsset.h"
 #include "Rendering/Texture2DResource.h"
 #include "RenderingThread.h"
 #include "SkeletalMergingLibrary.h"
-#include "UnrealMutableImageProvider.h"
 #include "MuCO/ICustomizableObjectEditorModule.h"
-#include "UObject/ObjectSaveContext.h"
 #include "UObject/UObjectIterator.h"
 #include "Widgets/Notifications/SNotificationList.h"
 #include "PhysicsEngine/AggregateGeom.h"
@@ -67,7 +61,6 @@
 #include "MessageLogModule.h"
 #include "UnrealEdMisc.h"
 #include "Subsystems/AssetEditorSubsystem.h"
-#include "Misc/TransactionObjectEvent.h"
 #include "Application/ThrottleManager.h"
 #endif
 
@@ -3763,8 +3756,9 @@ void UCustomizableInstancePrivate::InitSkeletalMeshData(const TSharedRef<FUpdate
 	}
 
 	// Set Min LOD (Override the Reference Skeletal Mesh LOD Settings)
-	SkeletalMesh->SetMinLod(FMath::Max(CustomizableObject.LODSettings.MinLOD.GetDefault(), static_cast<int32>(OperationData->FirstLODAvailable[ComponentName])));
-	SkeletalMesh->SetQualityLevelMinLod(CustomizableObject.LODSettings.MinQualityLevelLOD);
+	const UModelResources& ModelResources = CustomizableObject.GetPrivate()->GetModelResourcesChecked();
+	SkeletalMesh->SetMinLod(FMath::Max(ModelResources.MinLOD.GetDefault(), static_cast<int32>(OperationData->FirstLODAvailable[ComponentName])));
+	SkeletalMesh->SetQualityLevelMinLod(ModelResources.MinQualityLevelLOD);
 
 	// Set up unreal's default material, will be replaced when building materials
 	{
diff --git a/Engine/Plugins/Mutable/Source/CustomizableObject/Public/MuCO/CustomizableObject.h b/Engine/Plugins/Mutable/Source/CustomizableObject/Public/MuCO/CustomizableObject.h
--- a/Engine/Plugins/Mutable/Source/CustomizableObject/Public/MuCO/CustomizableObject.h
+++ b/Engine/Plugins/Mutable/Source/CustomizableObject/Public/MuCO/CustomizableObject.h
@@ -3,7 +3,6 @@
 #pragma once
 
 #include "RHIDefinitions.h"
-#include "SkeletalMeshTypes.h"
 #include "MuCO/CustomizableObjectParameterTypeDefinitions.h"
 #include "MuCO/CustomizableObjectUIData.h"
 #include "MuCO/CustomizableObject_Deprecated.h"
@@ -236,7 +235,7 @@ public:
 	FInstancedStruct VersionStruct;
 #endif
 
-	// DEPRECATED, Will be per component
+	// DEPRECATED, Will be per component UE-243403
 	UPROPERTY(EditAnywhere, Category = CustomizableObject, meta = (DisplayName = "LOD Settings"))
 	FMutableLODSettings LODSettings;
 	
diff --git a/Engine/Plugins/Mutable/Source/CustomizableObjectEditor/Private/MuCOE/CustomizableObjectCompiler.cpp b/Engine/Plugins/Mutable/Source/CustomizableObjectEditor/Private/MuCOE/CustomizableObjectCompiler.cpp
--- a/Engine/Plugins/Mutable/Source/CustomizableObjectEditor/Private/MuCOE/CustomizableObjectCompiler.cpp
+++ b/Engine/Plugins/Mutable/Source/CustomizableObjectEditor/Private/MuCOE/CustomizableObjectCompiler.cpp
@@ -14,26 +14,21 @@
 #include "MuCO/CustomizableObjectPrivate.h"
 #include "MuCO/CustomizableObjectInstance.h"
 #include "MuCO/CustomizableObjectSystem.h"
-#include "MuCO/ICustomizableObjectModule.h"
 #include "MuCOE/CustomizableObjectCompileRunnable.h"
 #include "MuCOE/CustomizableObjectEditorLogger.h"
-#include "MuCOE/GenerateMutableSource/GenerateMutableSourceImage.h"
 #include "MuCOE/GraphTraversal.h"
 #include "MuCOE/ICustomizableObjectPopulationModule.h"
 #include "MuCOE/Nodes/CustomizableObjectNodeObjectGroup.h"
 #include "MuCOE/Nodes/CustomizableObjectNodeTable.h"
 #include "MuCOE/CustomizableObjectEditorModule.h"
-#include "UObject/ICookInfo.h"
 #include "UObject/UObjectIterator.h"
 #include "Widgets/Notifications/SNotificationList.h"
 #include "Interfaces/ITargetPlatform.h"
 #include "Misc/App.h"
-#include "MuCO/CustomizableObjectPrivate.h"
 #include "MuCO/CustomizableObjectSystemPrivate.h"
 #include "MuR/Model.h"
 #include "MuR/ModelPrivate.h"
 #include "MuCOE/CustomizableObjectVersionBridge.h"
-#include "Serialization/ObjectAndNameAsStringProxyArchive.h"
 
 class UTexture2D;
 
@@ -906,6 +901,10 @@ void FCustomizableObjectCompiler::CompileInternal(bool bAsync)
 										ImageProp.AddressX, ImageProp.AddressY });
 		}
 
+		// Copy the LODSettings data found in the CO into the ModelResources
+		ModelResources->MinLOD = CurrentObject->LODSettings.MinLOD;
+		ModelResources->MinQualityLevelLOD = CurrentObject->LODSettings.MinQualityLevelLOD;
+		
 		ModelResources->ParameterUIDataMap = MoveTemp(GenerationContext.ParameterUIDataMap);
 		ModelResources->StateUIDataMap = MoveTemp(GenerationContext.StateUIDataMap);
 		ModelResources->IntParameterOptionDataTable = MoveTemp(GenerationContext.IntParameterOptionDataTable);
