[mutable] Fix possible deadlock when shutting down the editor in the middle of a CO compilation.
#rnx

diff --git a/Engine/Plugins/Mutable/Source/CustomizableObject/Internal/MuCO/CustomizableObjectSystemPrivate.h b/Engine/Plugins/Mutable/Source/CustomizableObject/Internal/MuCO/CustomizableObjectSystemPrivate.h
--- a/Engine/Plugins/Mutable/Source/CustomizableObject/Internal/MuCO/CustomizableObjectSystemPrivate.h
+++ b/Engine/Plugins/Mutable/Source/CustomizableObject/Internal/MuCO/CustomizableObjectSystemPrivate.h
@@ -824,8 +824,6 @@ public:
 	 */
 	void AddPendingLoad(UCustomizableObject*);
 
-	CUSTOMIZABLEOBJECT_API void AddGameThreadCompileTask(TFunction<void()>&& Task);
-
 #endif
 	
 	/** Mutable TaskGraph system (Mutable Thread). */
@@ -840,9 +838,6 @@ public:
 	UPROPERTY(Transient)
 	TArray<TObjectPtr<UCustomizableObject>> ObjectsPendingLoad;
 
-	/** */
-	TQueue< TFunction<void()>, EQueueMode::Mpsc > PendingGameThreadCompileTasks;
-
 #endif
 
 	FLogBenchmarkUtil LogBenchmarkUtil;
diff --git a/Engine/Plugins/Mutable/Source/CustomizableObject/Private/MuCO/CustomizableObjectSystem.cpp b/Engine/Plugins/Mutable/Source/CustomizableObject/Private/MuCO/CustomizableObjectSystem.cpp
--- a/Engine/Plugins/Mutable/Source/CustomizableObject/Private/MuCO/CustomizableObjectSystem.cpp
+++ b/Engine/Plugins/Mutable/Source/CustomizableObject/Private/MuCO/CustomizableObjectSystem.cpp
@@ -1693,12 +1693,6 @@ void UCustomizableObjectSystemPrivate::AddPendingLoad(UCustomizableObject* CO)
 }
 
 
-void UCustomizableObjectSystemPrivate::AddGameThreadCompileTask(TFunction<void()>&& Task)
-{
-	PendingGameThreadCompileTasks.Enqueue(Task);
-}
-
-
 bool UCustomizableObjectSystem::LockObject(UCustomizableObject* InObject)
 {
 	check(InObject != nullptr);
@@ -3949,23 +3943,6 @@ int32 UCustomizableObjectSystem::TickInternal(const bool bBlocking)
 		return 0;
 	}
 
-	// See if there are compile-time tasks to run
-	constexpr double MaxSecondsPerFrame = 0.4;
-	double MaxTime = FPlatformTime::Seconds() + MaxSecondsPerFrame;
-
-	TFunction<void()> Task;
-	while (GetPrivate()->PendingGameThreadCompileTasks.Dequeue(Task))
-	{
-		Task();
-
-		// Simple time limit enforcement to avoid blocking the game thread if there are many requests.
-		double CurrentTime = FPlatformTime::Seconds();
-		if (CurrentTime >= MaxTime)
-		{
-			break;
-		}
-	}
-
 	// See if any COs pending to load can be completed.
 	for ( TArray<TObjectPtr<UCustomizableObject>>::TIterator It = GetPrivate()->ObjectsPendingLoad.CreateIterator(); It; ++It )
 	{
diff --git a/Engine/Plugins/Mutable/Source/CustomizableObjectEditor/Private/MuCOE/CustomizableObjectCompileRunnable.cpp b/Engine/Plugins/Mutable/Source/CustomizableObjectEditor/Private/MuCOE/CustomizableObjectCompileRunnable.cpp
--- a/Engine/Plugins/Mutable/Source/CustomizableObjectEditor/Private/MuCOE/CustomizableObjectCompileRunnable.cpp
+++ b/Engine/Plugins/Mutable/Source/CustomizableObjectEditor/Private/MuCOE/CustomizableObjectCompileRunnable.cpp
@@ -205,6 +205,14 @@ uint32 FCustomizableObjectCompileRunnable::Run()
 					MUTABLE_CPUPROFILER_SCOPE(LoadMeshFunc);
 					check(IsInGameThread());
 
+					// If we are shutting down, we are not allowed to try to load anything.
+					if (IsEngineExitRequested())
+					{
+						*ResolvedMesh = MakeShared<mu::FMesh>();
+						CompletionEvent.Trigger();
+						return;
+					}
+
 					// Ensure we don't pull unwanted data into the package when cooking.
 					FCookLoadScope CookLoadScope(ECookLoadType::EditorOnly);
 
@@ -234,7 +242,15 @@ uint32 FCustomizableObjectCompileRunnable::Run()
 			else
 			{
 				check(!bRunImmediatlyIfPossible);
-				UCustomizableObjectSystem::GetInstance()->GetPrivate()->AddGameThreadCompileTask(MoveTemp(LoadMeshFunc));
+				
+				// Could store event here
+				//LastTask = 
+				FFunctionGraphTask::CreateAndDispatchWhenReady(LoadMeshFunc,
+					TStatId{}, 
+					nullptr, // Prerequisites
+					ENamedThreads::GameThread
+					);
+
 			}
 
 			return CompletionEvent;
diff --git a/Engine/Plugins/Mutable/Source/MutableTools/Private/MuT/ASTOpMeshMaskDiff.cpp b/Engine/Plugins/Mutable/Source/MutableTools/Private/MuT/ASTOpMeshMaskDiff.cpp
--- a/Engine/Plugins/Mutable/Source/MutableTools/Private/MuT/ASTOpMeshMaskDiff.cpp
+++ b/Engine/Plugins/Mutable/Source/MutableTools/Private/MuT/ASTOpMeshMaskDiff.cpp
@@ -83,6 +83,11 @@ namespace mu
 
 	mu::Ptr<ASTOp> ASTOpMeshMaskDiff::OptimiseSemantic(const FModelOptimizationOptions&, int32 Pass) const
 	{
+		if (!Fragment)
+		{
+			return nullptr;
+		}
+
 		mu::Ptr<ASTOp> ResultOp;
 
 		EOpType FragmentType = Fragment->GetOpType();
diff --git a/Engine/Plugins/Mutable/Source/MutableTools/Private/MuT/CodeGenerator.cpp b/Engine/Plugins/Mutable/Source/MutableTools/Private/MuT/CodeGenerator.cpp
--- a/Engine/Plugins/Mutable/Source/MutableTools/Private/MuT/CodeGenerator.cpp
+++ b/Engine/Plugins/Mutable/Source/MutableTools/Private/MuT/CodeGenerator.cpp
@@ -2475,6 +2475,10 @@ namespace mu
 					[Layout, MaskImage, ResolveOriginalMesh, LastMeshTask, PreModifiersTask, Modifier, TypedClipNode, CompilerOptions = CompilerOptions, ErrorLog = ErrorLog, ErrorContext, OriginalMeshCopy]() mutable
 					{
 						TSharedPtr<FMesh> TargetMesh = *ResolveOriginalMesh;
+						if (!TargetMesh)
+						{
+							return FMeshGenerationResult();
+						}
 						FMeshGenerationResult LastMeshResults = LastMeshTask.GetResult();
 						FMeshGenerationResult BaseResults = PreModifiersTask.GetResult();
 
diff --git a/Engine/Plugins/Mutable/Source/MutableTools/Private/MuT/CodeGenerator_Image.cpp b/Engine/Plugins/Mutable/Source/MutableTools/Private/MuT/CodeGenerator_Image.cpp
--- a/Engine/Plugins/Mutable/Source/MutableTools/Private/MuT/CodeGenerator_Image.cpp
+++ b/Engine/Plugins/Mutable/Source/MutableTools/Private/MuT/CodeGenerator_Image.cpp
@@ -1383,7 +1383,12 @@ namespace mu
 			// that whould be fixed with the general ordering design. 
 			// Actually use the MeshOp, otherwise the projector will only project to the first option if the mesh operation is a switch. 
 			Ptr<ASTOp> CurrentMeshToProjectOp = MeshResult.MeshOp;
-			//Ptr<ASTOp> CurrentMeshToProjectOp = MeshResult.BaseMeshOp;
+			if (!MeshResult.MeshOp)
+			{
+				Result.op = nullptr;
+				return;
+			}
+
 
             if (projectorResult.type == EProjectorType::Wrapping)
             {
@@ -2044,7 +2049,7 @@ namespace mu
     {
         Ptr<ASTOp> Result = at;
 
-		if (size[0] > 0 && size[1] > 0)
+		if (at && size[0] > 0 && size[1] > 0)
 		{
 			if (UE::Math::TIntVector2<int32>(at->GetImageDesc().m_size) != size)
 			{
