[mutable] New method to get the list of assets relevant for a CO compilation. It is less precise but up to 100x faster.
#rnx

diff --git a/Engine/Plugins/Mutable/Source/CustomizableObject/Internal/MuCO/ICustomizableObjectEditorModule.h b/Engine/Plugins/Mutable/Source/CustomizableObject/Internal/MuCO/ICustomizableObjectEditorModule.h
--- a/Engine/Plugins/Mutable/Source/CustomizableObject/Internal/MuCO/ICustomizableObjectEditorModule.h
+++ b/Engine/Plugins/Mutable/Source/CustomizableObject/Internal/MuCO/ICustomizableObjectEditorModule.h
@@ -65,7 +65,7 @@ public:
 	  * @param RemovedPackages list of removed packages since the last compilation.
 	  * @param bReleaseVersionDiff true if the Release Version has changed since the last compilation.
 	  * @return true if the compilation is out of date. */
-	virtual bool IsCompilationOutOfDate(const UCustomizableObject& Object, bool bSkipIndirectReferences, TArray<FName>& OutOfDatePackages, TArray<FName>& AddedPackages, TArray<FName>& RemovedPackages, bool& bReleaseVersionDiff) const = 0;
+	virtual bool IsCompilationOutOfDate(UCustomizableObject& Object, bool bSkipIndirectReferences, TArray<FName>& OutOfDatePackages, TArray<FName>& AddedPackages, TArray<FName>& RemovedPackages, bool& bReleaseVersionDiff) const = 0;
 
 	/** See GraphTraversal::IsRootObject(...) */
 	virtual bool IsRootObject(const UCustomizableObject& Object) const = 0;
@@ -108,9 +108,8 @@ public:
 
 	virtual USkeletalMesh* GetReferenceSkeletalMesh(const UCustomizableObject& Object, const FName& Component) const = 0;
 	
-	/** Perform a fast compilation pass to get all participating objects.
-	 *  @param bLoadObjects Load any object. If false, no objects will load. If true, only objects strictly required to get the full list of participating objects will load. */
-	virtual TMap<FName, FGuid> GetParticipatingObjects(const UCustomizableObject* Object, bool bLoadObjects, const FCompilationOptions* Options = nullptr) const = 0;
+	/** Perform a fast compilation pass to get all participating objects. */
+	virtual TMap<FName, FGuid> GetParticipatingObjects(UCustomizableObject* Object, const FCompilationOptions* Options = nullptr) const = 0;
 
 	virtual void BackwardsCompatibleFixup(UEdGraph& Graph, int32 CustomizableObjectCustomVersion) = 0;
 	
diff --git a/Engine/Plugins/Mutable/Source/CustomizableObject/Private/MuCO/CustomizableObject.cpp b/Engine/Plugins/Mutable/Source/CustomizableObject/Private/MuCO/CustomizableObject.cpp
--- a/Engine/Plugins/Mutable/Source/CustomizableObject/Private/MuCO/CustomizableObject.cpp
+++ b/Engine/Plugins/Mutable/Source/CustomizableObject/Private/MuCO/CustomizableObject.cpp
@@ -3061,7 +3061,7 @@ TArray<uint8> UCustomizableObjectPrivate::BuildDerivedDataKey(FCompilationOption
 	// Participating objects hash
 	if (ICustomizableObjectEditorModule* Module = ICustomizableObjectEditorModule::Get())
 	{
-		TArray<TTuple<FName, FGuid>> ParticipatingObjects = Module->GetParticipatingObjects(GetPublic(), true, &Options).Array();
+		TArray<TTuple<FName, FGuid>> ParticipatingObjects = Module->GetParticipatingObjects(GetPublic(), &Options).Array();
 		ParticipatingObjects.Sort([](const TTuple<FName, FGuid>& A, const TTuple<FName, FGuid>& B)
 		{
 			return A.Get<0>().LexicalLess(B.Get<0>()) && A.Get<1>() < B.Get<1>();
diff --git a/Engine/Plugins/Mutable/Source/CustomizableObjectEditor/Private/MuCOE/CustomizableObjectCompiler.cpp b/Engine/Plugins/Mutable/Source/CustomizableObjectEditor/Private/MuCOE/CustomizableObjectCompiler.cpp
--- a/Engine/Plugins/Mutable/Source/CustomizableObjectEditor/Private/MuCOE/CustomizableObjectCompiler.cpp
+++ b/Engine/Plugins/Mutable/Source/CustomizableObjectEditor/Private/MuCOE/CustomizableObjectCompiler.cpp
@@ -506,7 +506,6 @@ void ProcessChildObjectsRecursively(const UCustomizableObject* ParentObject, FMu
 			}
 
 			GenerationContext.GroupIdToExternalNodeMap.Add(Root->ParentObjectGroupId, Root);
-			GenerationContext.AddParticipatingObject(*ParentObject);
 
 			TArray<UCustomizableObjectNodeObjectGroup*> GroupNodes;
 			ChildObject->GetPrivate()->GetSource()->GetNodesOfClass<UCustomizableObjectNodeObjectGroup>(GroupNodes);
@@ -627,7 +626,7 @@ mu::Ptr<mu::NodeObject> GenerateMutableRoot(
 			const int32 MaxIndex = Object->GetPrivate()->GetWorkingSet().Num();
 			for (int32 i = 0; i < MaxIndex; ++i)
 			{
-				if (UCustomizableObject* WorkingSetObject = GenerationContext.LoadObject(Object->GetPrivate()->GetWorkingSet()[i], true))
+				if (UCustomizableObject* WorkingSetObject = GenerationContext.LoadObject(Object->GetPrivate()->GetWorkingSet()[i]))
 				{
 					ArrayCustomizableObject.Reset();
 
@@ -798,7 +797,7 @@ void FCustomizableObjectCompiler::CompileInternal(bool bAsync)
 	FMutableGraphGenerationContext GenerationContext(*CompilationContext);
 
 	// Perform a first participating objects pass
-	GenerationContext.ParticipatingObjects = ICustomizableObjectEditorModule::GetChecked().GetParticipatingObjects(CurrentObject, true, &CurrentOptions);
+	TMap<FName, FGuid> ParticipatingObjects = ICustomizableObjectEditorModule::GetChecked().GetParticipatingObjects(CurrentObject, &CurrentOptions);
 
 	// Clear Messages from previous Compilations
 	CompilationLogsContainer.ClearMessageCounters();
@@ -953,7 +952,7 @@ void FCustomizableObjectCompiler::CompileInternal(bool bAsync)
 
 		ModelResources->ComponentFirstLODAvailable = GenerationContext.FirstLODAvailable;
 		
-		ModelResources->ParticipatingObjects = MoveTemp(GenerationContext.ParticipatingObjects);
+		ModelResources->ParticipatingObjects = MoveTemp(ParticipatingObjects);
 		
 		if (CurrentOptions.bGatherReferences)
 		{
@@ -1222,8 +1221,6 @@ mu::Ptr<mu::Node> FCustomizableObjectCompiler::Export(UCustomizableObject* Objec
 	
 	CompilationContext = MakeShared<FMutableCompilationContext>(Object, this, CompilerOptions);
 	FMutableGraphGenerationContext GenerationContext(*CompilationContext);
-
-	GenerationContext.bSkipParticipatingObjectsPass = true;
 	
 	// Generate the mutable node expression
 	mu::NodeObjectPtr MutableRoot = GenerateMutableRoot(Object, GenerationContext);
diff --git a/Engine/Plugins/Mutable/Source/CustomizableObjectEditor/Private/MuCOE/CustomizableObjectEditorModule.cpp b/Engine/Plugins/Mutable/Source/CustomizableObjectEditor/Private/MuCOE/CustomizableObjectEditorModule.cpp
--- a/Engine/Plugins/Mutable/Source/CustomizableObjectEditor/Private/MuCOE/CustomizableObjectEditorModule.cpp
+++ b/Engine/Plugins/Mutable/Source/CustomizableObjectEditor/Private/MuCOE/CustomizableObjectEditorModule.cpp
@@ -68,7 +68,12 @@
 #include "MuCOE/Nodes/CustomizableObjectNodeTableDetails.h"
 #include "MuCOE/Nodes/CustomizableObjectNodeComponentMeshDetails.h"
 #include "MuCOE/Nodes/CustomizableObjectNodeVariation.h"
+#include "MuCOE/Nodes/CustomizableObjectNodeTexture.h"
 #include "MuCOE/Nodes/CustomizableObjectNodeTextureVariation.h"
+#include "MuCOE/Nodes/CustomizableObjectNodeTextureProject.h"
+#include "MuCOE/Nodes/CustomizableObjectNodeTextureTransform.h"
+#include "MuCOE/Nodes/CustomizableObjectNodeAnimationPose.h"
+#include "MuCOE/Nodes/CustomizableObjectNodeCurve.h"
 #include "MuCOE/Widgets/CustomizableObjectVariationCustomization.h"
 #include "MuCOE/Widgets/CustomizableObjectLODReductionSettings.h"
 #include "PropertyEditorModule.h"
@@ -86,6 +91,8 @@
 #include "MuCOE/CustomizableObjectGraph.h"
 #include "Nodes/CustomizableObjectNodeComponentMeshDetails.h"
 #include "Nodes/CustomizableObjectNodeModifierTransformInMeshDetails.h"
+#include "MuCOE/GenerateMutableSource/GenerateMutableSourceTable.h"
+
 
 class AActor;
 class FString;
@@ -510,7 +517,7 @@ void GetReferencingPackages(const UCustomizableObject& Object, TArray<FAssetData
 }
 
 
-bool FCustomizableObjectEditorModule::IsCompilationOutOfDate(const UCustomizableObject& Object, bool bSkipIndirectReferences, TArray<FName>& OutOfDatePackages, TArray<FName>& AddedPackages, TArray<FName>& RemovedPackages, bool& bReleaseVersion) const
+bool FCustomizableObjectEditorModule::IsCompilationOutOfDate(UCustomizableObject& Object, bool bSkipIndirectReferences, TArray<FName>& OutOfDatePackages, TArray<FName>& AddedPackages, TArray<FName>& RemovedPackages, bool& bReleaseVersion) const
 {
 	MUTABLE_CPUPROFILER_SCOPE(FCustomizableObjectEditorModule::IsCompilationOutOfDate)
 	
@@ -533,6 +540,7 @@ bool FCustomizableObjectEditorModule::IsCompilationOutOfDate(const UCustomizable
 		TSoftObjectPtr SoftObjectPtr(FSoftObjectPath(ParticipatingObject.Key.ToString()));
 		if (SoftObjectPtr) // If loaded
 		{
+			// TODO UE-226453 
 			PRAGMA_DISABLE_DEPRECATION_WARNINGS
 			const FGuid PackageGuid = SoftObjectPtr->GetPackage()->GetGuid();
 			PRAGMA_ENABLE_DEPRECATION_WARNINGS
@@ -568,7 +576,7 @@ bool FCustomizableObjectEditorModule::IsCompilationOutOfDate(const UCustomizable
 	{
 		// Due to performance issues, we will skip loading all objects. We can do that since loading/not loading objects do not affect the number of indirect objects discovered
 		// (e.g., we will traverse the same number of COs/Tables regardless if we do not load meshes/textures...). 
-		TMap<FName, FGuid> ParticipatingObjects = GetParticipatingObjects(&Object, false);
+		TMap<FName, FGuid> ParticipatingObjects = GetParticipatingObjects(&Object);
 		
 		for (const TTuple<FName, FGuid>& ParticipatingObject : ParticipatingObjects)
 		{
@@ -688,20 +696,211 @@ USkeletalMesh* FCustomizableObjectEditorModule::GetReferenceSkeletalMesh(const U
 }
 
 
-TMap<FName, FGuid> FCustomizableObjectEditorModule::GetParticipatingObjects(const UCustomizableObject* Object, bool bLoadObjects, const FCompilationOptions* InOptions) const
+TMap<FName, FGuid> FCustomizableObjectEditorModule::GetParticipatingObjects(UCustomizableObject* Object, const FCompilationOptions* InOptions) const
 {
-	MUTABLE_CPUPROFILER_SCOPE(FCustomizableObjectEditorModule::GetParticipatingObjects)
+	MUTABLE_CPUPROFILER_SCOPE(FCustomizableObjectEditorModule::GetParticipatingObjectsNew);
 
 	FCompilationOptions Options = InOptions ? *InOptions : Object->GetPrivate()->GetCompileOptions();
-		
 	FMutableCompilationContext CompilationContext(Object, nullptr, Options);
 	FMutableGraphGenerationContext Context(CompilationContext);
-	Context.bParticipatingObjectsPass = true;
-	Context.bLoadObjects = bLoadObjects;
 
-	GenerateMutableRoot(Object, Context);
 
-	return Context.ParticipatingObjects;
+	FAssetRegistryModule& AssetRegistryModule = FModuleManager::LoadModuleChecked<FAssetRegistryModule>("AssetRegistry");
+
+	// Store the list of participating assets here
+	TMap<FName, FGuid> Result;
+
+	// Use these 2 helpers to add assets to the list.
+	auto AddObject = [&Result]
+	(UObject* Candidate)
+		{
+			if (Candidate)
+			{
+				// TODO UE-226453
+				PRAGMA_DISABLE_DEPRECATION_WARNINGS
+				const FGuid PackageGuid = Candidate->GetPackage()->GetGuid();
+				PRAGMA_ENABLE_DEPRECATION_WARNINGS
+
+				FName PackageName = Candidate->GetPackage()->GetFName();
+				Result.Emplace(PackageName, PackageGuid);
+			}
+		};
+
+	auto AddObjectSoft = [&Result, &AssetRegistryModule]
+	(const FSoftObjectPath& SoftPath)
+		{
+			FAssetPackageData AssetPackageData;
+			const UE::AssetRegistry::EExists Query = AssetRegistryModule.Get().TryGetAssetPackageData(SoftPath.GetLongPackageFName(), AssetPackageData);
+			if (Query != UE::AssetRegistry::EExists::Exists)
+			{
+				return;
+			}
+
+			PRAGMA_DISABLE_DEPRECATION_WARNINGS
+			const FGuid PackageGuid = AssetPackageData.PackageGuid;
+			PRAGMA_ENABLE_DEPRECATION_WARNINGS
+
+			Result.Emplace(SoftPath.GetLongPackageFName(), PackageGuid);
+		};
+
+	// Gather all related COs
+	TSet<UCustomizableObject*> AllObjects;
+	GetAllObjectsInGraph(Object, AllObjects);
+
+	// Scan all nodes in each CO
+	for (UCustomizableObject* CO : AllObjects)
+	{
+		AddObject(CO);
+
+		UEdGraph* Graph = CO->GetPrivate()->GetSource();
+
+		// For each node, based on its type.
+		// This could also be implemented as a virtual method on the nodes, but only a few of them are relevant.
+		for ( UEdGraphNode* GenericNode : Graph->Nodes )
+		{
+			if (UCustomizableObjectNodeSkeletalMesh* SkeletalNode = Cast<UCustomizableObjectNodeSkeletalMesh>(GenericNode))
+			{
+				AddObjectSoft(SkeletalNode->SkeletalMesh.ToSoftObjectPath());
+			}
+
+			else if (UCustomizableObjectNodeStaticMesh* StaticMeshNode = Cast<UCustomizableObjectNodeStaticMesh>(GenericNode))
+			{
+				AddObjectSoft(StaticMeshNode->StaticMesh.ToSoftObjectPath());
+			}
+
+			else if (UCustomizableObjectNodeAnimationPose* PoseNode = Cast<UCustomizableObjectNodeAnimationPose>(GenericNode))
+			{
+				AddObjectSoft(PoseNode->PoseAsset);
+			}
+
+			else if (UCustomizableObjectNodeCurve* CurveNode = Cast<UCustomizableObjectNodeCurve>(GenericNode))
+			{
+				AddObjectSoft(CurveNode->CurveAsset);
+			}
+
+			else if (UCustomizableObjectNodeComponentMesh* CompNode = Cast<UCustomizableObjectNodeComponentMesh>(GenericNode))
+			{
+				AddObjectSoft(CompNode->ReferenceSkeletalMesh);
+			}
+
+			else if (UCustomizableObjectNodeTexture* TexNode = Cast<UCustomizableObjectNodeTexture>(GenericNode))
+			{
+				AddObjectSoft(TexNode->Texture);
+			}
+
+			else if (UCustomizableObjectNodeTextureProject* ProjectNode = Cast<UCustomizableObjectNodeTextureProject>(GenericNode))
+			{
+				AddObjectSoft(ProjectNode->ReferenceTexture);
+			}
+
+			else if (UCustomizableObjectNodeTextureTransform* TransformNode = Cast<UCustomizableObjectNodeTextureTransform>(GenericNode))
+			{
+				AddObjectSoft(TransformNode->ReferenceTexture);
+			}
+
+			else if (UCustomizableObjectNodeMaterial* MatNode = Cast<UCustomizableObjectNodeMaterial>(GenericNode))
+			{
+				AddObjectSoft(MatNode->GetMaterial());
+				int32 NumImages = MatNode->GetNumParameters(EMaterialParameterType::Texture);
+				for (int32 ImageIndex = 0; ImageIndex < NumImages; ++ImageIndex)
+				{
+					AddObject(MatNode->GetImageReferenceTexture(ImageIndex));
+				}
+			}
+
+			else if (UCustomizableObjectNodeTable* TableNode = Cast<UCustomizableObjectNodeTable>(GenericNode))
+			{
+				UDataTable* DataTable = nullptr;
+					
+				if (TableNode->TableDataGatheringMode == ETableDataGatheringSource::ETDGM_AssetRegistry)
+				{
+					DataTable = GenerateDataTableFromStruct(TableNode, Context);
+				}
+				else
+				{
+					DataTable = TableNode->Table;
+				}
+
+				if (DataTable)
+				{
+					AddObject(TableNode->Table);
+
+					TArray<FName> RowNames = TableNode->GetEnabledRows(*TableNode->Table);
+
+					TArray<FName> ExpectedPropNames = DataTableUtils::GetStructPropertyNames(DataTable->RowStruct);
+					for (const FName& ColumnName : ExpectedPropNames)
+					{
+						FProperty* ColumnProperty = DataTable->FindTableProperty(ColumnName);
+						if (!ColumnProperty)
+						{
+							continue;
+						}
+
+						const FSoftObjectProperty* SoftObjectProperty = CastField<FSoftObjectProperty>(ColumnProperty);
+						if (!SoftObjectProperty)
+						{
+							continue;
+						}
+
+						for (const FName& RowName : RowNames )
+						{
+							uint8* RowData = DataTable->FindRowUnchecked(RowName);
+							if (RowData)
+							{
+								uint8* CellData = ColumnProperty->ContainerPtrToValuePtr<uint8>(RowData, 0);
+								if (CellData)
+								{
+									const FSoftObjectPtr& Path = SoftObjectProperty->GetPropertyValue(CellData);
+									AddObjectSoft(Path.ToSoftObjectPath());
+								}
+							}
+						}
+					}
+				}
+			}
+		}
+	}
+
+	// Get one level of references for each participating object
+	// This will add some unnecessary assets. It could be further refined by filtering asset types or ignoring assets in Engine.
+	// If more fine-tunning is needed, assets could specificy a recursion level when added to the list from nodes.
+
+	// Copy the original asset list since we will be modifying it on every iteration.
+	TMap<FName, FGuid> ResultCopy = Result;
+	{
+		TArray<FAssetData> Assets;
+		TArray<FSoftObjectPath> References;
+		for (const TPair<FName, FGuid>& Entry : ResultCopy)
+		{
+			Assets.Reset();
+			References.Reset();
+			TArray<FName> PackageDependencies;
+			AssetRegistryModule.Get().GetDependencies(Entry.Key, PackageDependencies);
+
+			FARFilter Filter;
+			for (const FName& PackageName : PackageDependencies)
+			{
+				if (!PackageName.ToString().StartsWith(TEXT("/TempAutosave")))
+				{
+					Filter.PackageNames.Add(PackageName);
+				}
+			}
+			AssetRegistryModule.Get().GetAssets(Filter, Assets);
+
+			for (const FAssetData& Dep : Assets)
+			{
+				// Assets could be filtered like this:
+				//const UClass* AssetClass = Dep.GetClass(EResolveClass::Yes);
+				//if (AssetClass == USkeleton::StaticClass())
+
+				AddObjectSoft(Dep.GetSoftObjectPath());
+			}
+		}
+
+	}
+
+	// Done
+	return MoveTemp(Result);
 }
 
 
diff --git a/Engine/Plugins/Mutable/Source/CustomizableObjectEditor/Private/MuCOE/CustomizableObjectEditorModule.h b/Engine/Plugins/Mutable/Source/CustomizableObjectEditor/Private/MuCOE/CustomizableObjectEditorModule.h
--- a/Engine/Plugins/Mutable/Source/CustomizableObjectEditor/Private/MuCOE/CustomizableObjectEditorModule.h
+++ b/Engine/Plugins/Mutable/Source/CustomizableObjectEditor/Private/MuCOE/CustomizableObjectEditorModule.h
@@ -42,7 +42,7 @@ public:
 
 	// ICustomizableObjectEditorModule interface
 	virtual FCustomizableObjectEditorLogger& GetLogger() override;
-	virtual bool IsCompilationOutOfDate(const UCustomizableObject& Object, bool bSkipIndirectReferences, TArray<FName>& OutOfDatePackages, TArray<FName>& AddedPackages, TArray<FName>& RemovedPackages, bool& bVersionDiff) const override;
+	virtual bool IsCompilationOutOfDate(UCustomizableObject& Object, bool bSkipIndirectReferences, TArray<FName>& OutOfDatePackages, TArray<FName>& AddedPackages, TArray<FName>& RemovedPackages, bool& bVersionDiff) const override;
 	virtual bool IsRootObject(const UCustomizableObject& Object) const override;
 	virtual FString GetCurrentReleaseVersionForObject(const UCustomizableObject& Object) const override;
 	virtual UCustomizableObject* GetRootObject(UCustomizableObject* ChildObject) const override;
@@ -50,7 +50,7 @@ public:
 	virtual void GetRelatedObjects(UCustomizableObject*, TSet<UCustomizableObject*>& OutRelated) const override;
 	virtual void BakeCustomizableObjectInstance(UCustomizableObjectInstance* InTargetInstance, const FBakingConfiguration& InBakingConfig) override;
 	virtual USkeletalMesh* GetReferenceSkeletalMesh(const UCustomizableObject& Object, const FName& ComponentName) const override;
-	virtual TMap<FName, FGuid> GetParticipatingObjects(const UCustomizableObject* Object, bool bLoadObjects, const FCompilationOptions* Options = nullptr) const override;
+	virtual TMap<FName, FGuid> GetParticipatingObjects(UCustomizableObject* Object, const FCompilationOptions* Options = nullptr) const override;
 	virtual void BackwardsCompatibleFixup(UEdGraph& Graph, int32 CustomizableObjectCustomVersion) override;
 	virtual void PostBackwardsCompatibleFixup(UEdGraph& Graph) override;
 	virtual void CancelCompileRequests() override;
diff --git a/Engine/Plugins/Mutable/Source/CustomizableObjectEditor/Private/MuCOE/CustomizableObjectLayout.cpp b/Engine/Plugins/Mutable/Source/CustomizableObjectEditor/Private/MuCOE/CustomizableObjectLayout.cpp
--- a/Engine/Plugins/Mutable/Source/CustomizableObjectEditor/Private/MuCOE/CustomizableObjectLayout.cpp
+++ b/Engine/Plugins/Mutable/Source/CustomizableObjectEditor/Private/MuCOE/CustomizableObjectLayout.cpp
@@ -103,7 +103,6 @@ void UCustomizableObjectLayout::GenerateAutomaticBlocksFromUVs()
 	
 	FMutableCompilationContext CompilationContext(Object, &Compiler, Options);
 	FMutableGraphGenerationContext GenerationContext(CompilationContext);
-	GenerationContext.bSkipParticipatingObjectsPass = true;
 
 	bool bOutWasEmpty = false;
 	mu::Ptr<mu::NodeLayout> LayoutNode = CreateMutableLayoutNode(GenerationContext, this, false, bOutWasEmpty );
diff --git a/Engine/Plugins/Mutable/Source/CustomizableObjectEditor/Private/MuCOE/ExtensionDataCompilerInterface.cpp b/Engine/Plugins/Mutable/Source/CustomizableObjectEditor/Private/MuCOE/ExtensionDataCompilerInterface.cpp
--- a/Engine/Plugins/Mutable/Source/CustomizableObjectEditor/Private/MuCOE/ExtensionDataCompilerInterface.cpp
+++ b/Engine/Plugins/Mutable/Source/CustomizableObjectEditor/Private/MuCOE/ExtensionDataCompilerInterface.cpp
@@ -59,8 +59,3 @@ void FExtensionDataCompilerInterface::CompilerLog(const FText& InLogText, const
 	GenerationContext.Log(InLogText, InNode);
 }
 
-void FExtensionDataCompilerInterface::AddParticipatingObject(const UObject& Object)
-{
-	GenerationContext.AddParticipatingObject(Object);
-}
-
diff --git a/Engine/Plugins/Mutable/Source/CustomizableObjectEditor/Private/MuCOE/GenerateMutableSource/GenerateMutableSource.cpp b/Engine/Plugins/Mutable/Source/CustomizableObjectEditor/Private/MuCOE/GenerateMutableSource/GenerateMutableSource.cpp
--- a/Engine/Plugins/Mutable/Source/CustomizableObjectEditor/Private/MuCOE/GenerateMutableSource/GenerateMutableSource.cpp
+++ b/Engine/Plugins/Mutable/Source/CustomizableObjectEditor/Private/MuCOE/GenerateMutableSource/GenerateMutableSource.cpp
@@ -280,62 +280,7 @@ void FMutableGraphGenerationContext::Log(const FText& Message, const TArray<cons
 
 void FMutableGraphGenerationContext::Log(const FText& Message, const UObject* Context, const EMessageSeverity::Type MessageSeverity, const bool bAddBaseObjectInfo, const ELoggerSpamBin SpamBin) const
 {
-	if (!bParticipatingObjectsPass)
-	{
-		CompilationContext.Log(Message, Context, MessageSeverity, bAddBaseObjectInfo, SpamBin);
-	}
-}
-
-
-void FMutableGraphGenerationContext::AddParticipatingObject(const FSoftObjectPath& SoftPath)
-{
-	const FAssetRegistryModule& AssetRegistryModule = FModuleManager::LoadModuleChecked<FAssetRegistryModule>("AssetRegistry");
-
-	FAssetPackageData AssetPackageData;
-	const UE::AssetRegistry::EExists Result = AssetRegistryModule.Get().TryGetAssetPackageData(SoftPath.GetLongPackageFName(), AssetPackageData);
-	if (Result != UE::AssetRegistry::EExists::Exists)
-	{
-		return;
-	}
-
-	PRAGMA_DISABLE_DEPRECATION_WARNINGS
-	const FGuid PackageGuid = AssetPackageData.PackageGuid;
-	PRAGMA_ENABLE_DEPRECATION_WARNINGS
-
-	AddParticipatingObjectChecked(SoftPath.GetLongPackageFName(), PackageGuid);
-}
-
-
-void FMutableGraphGenerationContext::AddParticipatingObject(const FSoftObjectPtr& SoftObject)
-{
-	AddParticipatingObject(SoftObject.ToSoftObjectPath());
-}
-
-
-void FMutableGraphGenerationContext::AddParticipatingObject(const UObject& InObject)
-{
-	PRAGMA_DISABLE_DEPRECATION_WARNINGS
-	const FGuid PackageGuid = InObject.GetPackage()->GetGuid();
-	PRAGMA_ENABLE_DEPRECATION_WARNINGS
-
-	AddParticipatingObjectChecked(InObject.GetPackage()->GetFName(), PackageGuid);
-}
-
-
-void FMutableGraphGenerationContext::AddParticipatingObjectChecked(const FName& PackageName, const FGuid& PackageGuid)
-{
-	if (bParticipatingObjectsPass)
-	{
-		ParticipatingObjects.Emplace(PackageName, PackageGuid);
-	}
-	else
-	{
-		if (!bSkipParticipatingObjectsPass)
-		{
-			FGuid* Result = ParticipatingObjects.Find(PackageName);
-			ensureAlwaysMsgf(Result, TEXT("Participating Object and Generate Source pass missatch. Asset mismatch: %s"), *PackageName.ToString());
-		}
-	}
+	CompilationContext.Log(Message, Context, MessageSeverity, bAddBaseObjectInfo, SpamBin);
 }
 
 
@@ -355,11 +300,6 @@ TSharedPtr<mu::FMesh> FMutableGraphGenerationContext::FindGeneratedMesh( const F
 
 int32 FMutableGraphGenerationContext::AddStreamedResource(UAssetUserData& AssetUserData)
 {
-	if (bParticipatingObjectsPass)
-	{
-		return INDEX_NONE;
-	}
-
 	const uint32 AssetIdentifier = AssetUserData.GetUniqueID();
 
 	// Return resource index if found.
@@ -453,11 +393,9 @@ FMutableComponentInfo* FMutableGraphGenerationContext::GetCurrentComponentInfo()
 }
 
 
-UObject* FMutableGraphGenerationContext::LoadObject(const FSoftObjectPtr& SoftObject, bool bParticipatingObjectsPassLoad)
+UObject* FMutableGraphGenerationContext::LoadObject(const FSoftObjectPtr& SoftObject)
 {
-	return bLoadObjects && (!bParticipatingObjectsPass || bParticipatingObjectsPassLoad) ?
-		MutablePrivate::LoadObject(SoftObject) :
-		nullptr;
+	return  MutablePrivate::LoadObject(SoftObject);
 }
 
 
@@ -819,8 +757,6 @@ mu::Ptr<mu::NodeObject> GenerateMutableSource(const UEdGraphPin * Pin, FMutableG
 	CheckNumOutputs(*Pin, GenerationContext);
 	
 	UCustomizableObjectNode* Node = CastChecked<UCustomizableObjectNode>(Pin->GetOwningNode());
-
-	GenerationContext.AddParticipatingObject(*GetRootObject(*Node));
 	
 	const FGeneratedKey Key(reinterpret_cast<void*>(&GenerateMutableSource), *Pin, *Node, GenerationContext, true);
 	if (const FGeneratedData* Generated = GenerationContext.Generated.Find(Key))
@@ -1402,21 +1338,18 @@ void PopulateReferenceSkeletalMeshesData(FMutableGraphGenerationContext& Generat
 		// Physics Asset
 		if (UPhysicsAsset* PhysicsAsset = RefSkeletalMesh->GetPhysicsAsset())
 		{
-			GenerationContext.AddParticipatingObject(*PhysicsAsset);
 			Data.PhysicsAsset = PhysicsAsset;
 		}
 
 		// Post ProcessAnimInstance
 		if(const TSubclassOf<UAnimInstance> PostProcessAnimInstance = RefSkeletalMesh->GetPostProcessAnimBlueprint())
 		{
-			GenerationContext.AddParticipatingObject(*PostProcessAnimInstance.Get());
 			Data.PostProcessAnimInst = PostProcessAnimInstance;
 		}
 		
 		// Shadow Physics Asset
 		if (UPhysicsAsset* PhysicsAsset = RefSkeletalMesh->GetShadowPhysicsAsset())
 		{
-			GenerationContext.AddParticipatingObject(*PhysicsAsset);
 			Data.ShadowPhysicsAsset = PhysicsAsset;
 		}
 
diff --git a/Engine/Plugins/Mutable/Source/CustomizableObjectEditor/Private/MuCOE/GenerateMutableSource/GenerateMutableSource.h b/Engine/Plugins/Mutable/Source/CustomizableObjectEditor/Private/MuCOE/GenerateMutableSource/GenerateMutableSource.h
--- a/Engine/Plugins/Mutable/Source/CustomizableObjectEditor/Private/MuCOE/GenerateMutableSource/GenerateMutableSource.h
+++ b/Engine/Plugins/Mutable/Source/CustomizableObjectEditor/Private/MuCOE/GenerateMutableSource/GenerateMutableSource.h
@@ -719,24 +719,6 @@ struct FMutableGraphGenerationContext
 	 * When a Mesh Section node is found, it will compile all Group Projector nodes in the stack. */
 	TArray<TArray<UCustomizableObjectNodeGroupProjectorParameter*>> CurrentGroupProjectors;
 
-	/** See FCustomizableObjectPrivateData::ParticipatingObjects. */
-	TMap<FName, FGuid> ParticipatingObjects;
-
-	/** Participating objects.
-	* See FCustomizableObjectPrivateData::ParticipatingObjects.
-	*/
-	void AddParticipatingObject(const UObject& Object);
-
-	void AddParticipatingObject(const FSoftObjectPath& SoftPath);
-
-	void AddParticipatingObject(const FSoftObjectPtr& SoftObject);
-
-	template<typename T>
-	void AddParticipatingObject(const TSoftClassPtr<T>& SoftClass)
-	{
-		AddParticipatingObject(SoftClass.ToSoftObjectPath());
-	}
-
 	/** Message logging */
 	FString GetObjectName();
 	void Log(const FText& Message, const TArray<const UObject*>& Context, const EMessageSeverity::Type MessageSeverity = EMessageSeverity::Warning, const bool bAddBaseObjectInfo = true, const ELoggerSpamBin SpamBin = ELoggerSpamBin::ShowAll) const;
@@ -759,22 +741,18 @@ struct FMutableGraphGenerationContext
 	/** Get the reference skeletal mesh associated to the current mesh component being generated */
 	FMutableComponentInfo* GetCurrentComponentInfo();
 
-	UObject* LoadObject(const FSoftObjectPtr& SoftObject, bool bParticipatingObjectsPassLoad = false);
+	UObject* LoadObject(const FSoftObjectPtr& SoftObject);
 	
 	template<typename T>
-	T* LoadObject(const TSoftObjectPtr<T>& SoftObject, bool bParticipatingObjectsPassLoad = false)
+	T* LoadObject(const TSoftObjectPtr<T>& SoftObject)
 	{
-		return bLoadObjects && (!bParticipatingObjectsPass || bParticipatingObjectsPassLoad) ?
-			MutablePrivate::LoadObject<T>(SoftObject) :
-			nullptr;
+		return MutablePrivate::LoadObject<T>(SoftObject);
 	}
 
 	template<typename T>
-	UClass* LoadClass(const TSoftClassPtr<T>& SoftClass, bool bParticipatingObjectsPassLoad = false)
+	UClass* LoadClass(const TSoftClassPtr<T>& SoftClass)
 	{
-		return bLoadObjects && (!bParticipatingObjectsPass || bParticipatingObjectsPassLoad) ?
-			MutablePrivate::LoadClass<T>(SoftClass) :
-			nullptr;
+		return MutablePrivate::LoadClass<T>(SoftClass);
 	}
 
 	/** Only compiled components. All components types. Index is the ObjectComponentIndex. */
@@ -818,15 +796,6 @@ struct FMutableGraphGenerationContext
 	bool bEnableLODStreaming = true;
 
 	bool bPartialCompilation = false;
-
-	/** true if performing the Participating Objects pass. */
-	bool bParticipatingObjectsPass = false;
-
-	/** true if the Participating Objects have been skipped. */
-	bool bSkipParticipatingObjectsPass = false;
-
-	/** Load any Soft Object/Class Pointers. */
-	bool bLoadObjects = true;
 	
 	// Based on the last object visited.
 	ECustomizableObjectAutomaticLODStrategy CurrentAutoLODStrategy = ECustomizableObjectAutomaticLODStrategy::Manual;
@@ -902,13 +871,6 @@ struct FMutableGraphGenerationContext
 
 	/** Index of the Referenced Material being generated */
 	int32 CurrentReferencedMaterialIndex = -1;
-
-private:
-
-	/** Two operation modes:
-	 * 1. bParticipatingObjectsPass = true, add the participating object.
-	 * 2. bParticipatingObjectsPass = false, check that the participating object was discovered in the participating object pass. */
-	void AddParticipatingObjectChecked(const FName& PackageName, const FGuid& PackageGuid);
 };
 
 
diff --git a/Engine/Plugins/Mutable/Source/CustomizableObjectEditor/Private/MuCOE/GenerateMutableSource/GenerateMutableSourceComponent.cpp b/Engine/Plugins/Mutable/Source/CustomizableObjectEditor/Private/MuCOE/GenerateMutableSource/GenerateMutableSourceComponent.cpp
--- a/Engine/Plugins/Mutable/Source/CustomizableObjectEditor/Private/MuCOE/GenerateMutableSource/GenerateMutableSourceComponent.cpp
+++ b/Engine/Plugins/Mutable/Source/CustomizableObjectEditor/Private/MuCOE/GenerateMutableSource/GenerateMutableSourceComponent.cpp
@@ -541,9 +541,6 @@ void FirstPass(UCustomizableObjectNodeComponentMesh& Node, FMutableGraphGenerati
 	GenerationContext.CompilationContext.ReferencedSkeletons.Add(RefSkeleton);
 
 	GenerationContext.ComponentNames.Add(Node.ComponentName);
-
-	// Add reference meshes to the participating objects
-	GenerationContext.AddParticipatingObject(*RefSkeletalMesh);
 }
 
 
diff --git a/Engine/Plugins/Mutable/Source/CustomizableObjectEditor/Private/MuCOE/GenerateMutableSource/GenerateMutableSourceGroupProjector.cpp b/Engine/Plugins/Mutable/Source/CustomizableObjectEditor/Private/MuCOE/GenerateMutableSource/GenerateMutableSourceGroupProjector.cpp
--- a/Engine/Plugins/Mutable/Source/CustomizableObjectEditor/Private/MuCOE/GenerateMutableSource/GenerateMutableSourceGroupProjector.cpp
+++ b/Engine/Plugins/Mutable/Source/CustomizableObjectEditor/Private/MuCOE/GenerateMutableSource/GenerateMutableSourceGroupProjector.cpp
@@ -74,11 +74,6 @@ mu::NodeImagePtr GenerateMutableSourceGroupProjector(const int32 NodeLOD, const
 			if (!GroupProjectionReferenceTexture)
 			{
 				GroupProjectionReferenceTexture = ProjectorTempData.CustomizableObjectNodeGroupProjectorParameter->ReferenceTexture;
-				
-				if (GroupProjectionReferenceTexture)
-				{
-					GenerationContext.AddParticipatingObject(*GroupProjectionReferenceTexture);
-				}
 			}
 
 			const bool bProjectToImage = [&]
diff --git a/Engine/Plugins/Mutable/Source/CustomizableObjectEditor/Private/MuCOE/GenerateMutableSource/GenerateMutableSourceImage.cpp b/Engine/Plugins/Mutable/Source/CustomizableObjectEditor/Private/MuCOE/GenerateMutableSource/GenerateMutableSourceImage.cpp
--- a/Engine/Plugins/Mutable/Source/CustomizableObjectEditor/Private/MuCOE/GenerateMutableSource/GenerateMutableSourceImage.cpp
+++ b/Engine/Plugins/Mutable/Source/CustomizableObjectEditor/Private/MuCOE/GenerateMutableSource/GenerateMutableSourceImage.cpp
@@ -108,8 +108,6 @@ mu::NodeImagePtr GenerateMutableSourceImage(const UEdGraphPin* Pin, FMutableGrap
 		UTexture2D* BaseTexture = TypedNodeTex->Texture;
 		if (BaseTexture)
 		{
-			GenerationContext.AddParticipatingObject(*BaseTexture);
-			
 			// Check the specific image cache
 			FGeneratedImageKey ImageKey = FGeneratedImageKey(Pin);
 			mu::NodeImagePtr ImageNode;
@@ -189,8 +187,6 @@ mu::NodeImagePtr GenerateMutableSourceImage(const UEdGraphPin* Pin, FMutableGrap
 		const UTexture2D* ReferenceTexture = TypedNodeParam->ReferenceValue;
 		if (ReferenceTexture)
 		{
-			GenerationContext.AddParticipatingObject(*ReferenceTexture);
-
 			const uint32 LODBias = ComputeLODBiasForTexture(GenerationContext, *TypedNodeParam->ReferenceValue, ReferenceTexture, ReferenceTextureSize);
 			TextureSize.X = FMath::Max(ReferenceTexture->Source.GetSizeX() >> LODBias, 1);
 			TextureSize.Y = FMath::Max(ReferenceTexture->Source.GetSizeY() >> LODBias, 1);
@@ -624,8 +620,6 @@ mu::NodeImagePtr GenerateMutableSourceImage(const UEdGraphPin* Pin, FMutableGrap
 		int32 MaxReferenceTextureSizeInGame = ReferenceTextureSize;
 		if (TypedNodeProject->ReferenceTexture)
 		{
-			GenerationContext.AddParticipatingObject(*TypedNodeProject->ReferenceTexture);
-
 			const UTextureLODSettings& TextureLODSettings = CompilationOptions.TargetPlatform->GetTextureLODSettings();
 			MaxReferenceTextureSizeInGame = GetTextureSizeInGame(*TypedNodeProject->ReferenceTexture, TextureLODSettings);
 		}
@@ -857,8 +851,6 @@ mu::NodeImagePtr GenerateMutableSourceImage(const UEdGraphPin* Pin, FMutableGrap
 		int32 MaxReferenceTextureSizeInGame = ReferenceTextureSize;
 		if (TypedNodeTransform->ReferenceTexture)
 		{
-			GenerationContext.AddParticipatingObject(*TypedNodeTransform->ReferenceTexture);
-
 			const UTextureLODSettings& TextureLODSettings = CompilationOptions.TargetPlatform->GetTextureLODSettings();
 			MaxReferenceTextureSizeInGame = GetTextureSizeInGame(*TypedNodeTransform->ReferenceTexture, TextureLODSettings);
 		}
diff --git a/Engine/Plugins/Mutable/Source/CustomizableObjectEditor/Private/MuCOE/GenerateMutableSource/GenerateMutableSourceLayout.cpp b/Engine/Plugins/Mutable/Source/CustomizableObjectEditor/Private/MuCOE/GenerateMutableSource/GenerateMutableSourceLayout.cpp
--- a/Engine/Plugins/Mutable/Source/CustomizableObjectEditor/Private/MuCOE/GenerateMutableSource/GenerateMutableSourceLayout.cpp
+++ b/Engine/Plugins/Mutable/Source/CustomizableObjectEditor/Private/MuCOE/GenerateMutableSource/GenerateMutableSourceLayout.cpp
@@ -144,8 +144,6 @@ mu::FSourceLayoutBlock ToMutable(FMutableGraphGenerationContext& GenerationConte
 
 	if (UnrealBlock.Mask)
 	{
-		GenerationContext.AddParticipatingObject(*UnrealBlock.Mask);
-
 		// In the editor the src data can be directly accessed
 		TSharedPtr<mu::FImage> MaskImage = MakeShared<mu::FImage>();
 
diff --git a/Engine/Plugins/Mutable/Source/CustomizableObjectEditor/Private/MuCOE/GenerateMutableSource/GenerateMutableSourceMesh.cpp b/Engine/Plugins/Mutable/Source/CustomizableObjectEditor/Private/MuCOE/GenerateMutableSource/GenerateMutableSourceMesh.cpp
--- a/Engine/Plugins/Mutable/Source/CustomizableObjectEditor/Private/MuCOE/GenerateMutableSource/GenerateMutableSourceMesh.cpp
+++ b/Engine/Plugins/Mutable/Source/CustomizableObjectEditor/Private/MuCOE/GenerateMutableSource/GenerateMutableSourceMesh.cpp
@@ -2284,8 +2284,6 @@ TSharedPtr<mu::FMesh> ConvertSkeletalMeshToMutable(
 		return nullptr;
 	}
 
-	GenerationContext.AddParticipatingObject(InSkeletalMesh.ToSoftObjectPath());
-
 	// Prepare the data that is needed for the core mesh conversion.
 	FMutableSourceMeshData Source;
 	Source.Mesh = InSkeletalMesh;
@@ -2339,7 +2337,6 @@ TSharedPtr<mu::FMesh> ConvertSkeletalMeshToMutable(
 				const UClass* AssetClass = Dep.GetClass(EResolveClass::Yes);
 				if (AssetClass == USkeleton::StaticClass())
 				{
-					GenerationContext.AddParticipatingObject(Dep.GetSoftObjectPath());
 					bSkeletonFound = true;
 				}
 			}
@@ -2360,8 +2357,6 @@ TSharedPtr<mu::FMesh> ConvertSkeletalMeshToMutable(
 				const UClass* AssetClass = Dep.GetClass(EResolveClass::Yes);
 				if (AssetClass->ImplementsInterface(UClothingAssetCommon::StaticClass()))
 				{
-					GenerationContext.AddParticipatingObject(Dep.GetSoftObjectPath());
-
 					TArray<FAssetData> ClothAssets;
 					{
 						TArray<FName> ClothPackageDependencies;
@@ -2377,15 +2372,6 @@ TSharedPtr<mu::FMesh> ConvertSkeletalMeshToMutable(
 						}
 						AssetRegistryModule.Get().GetAssets(ClothFilter, ClothAssets);
 					}
-
-					for (const FAssetData& ClothDep : ClothAssets)
-					{
-						const UClass* ClothAssetClass = ClothDep.GetClass(EResolveClass::Yes);
-						if (ClothAssetClass->ImplementsInterface(UPhysicsAsset::StaticClass()))
-						{
-							GenerationContext.AddParticipatingObject(ClothDep.GetSoftObjectPath());
-						}
-					}
 				}
 			}
 		}
@@ -2417,8 +2403,6 @@ TSharedPtr<mu::FMesh> ConvertStaticMeshToMutable(const UStaticMesh* StaticMesh,
 		return nullptr;
 	}
 
-	GenerationContext.AddParticipatingObject(*StaticMesh);
-
 	TSharedPtr<mu::FMesh> MutableMesh = MakeShared<mu::FMesh>();
 
 	// Vertices
@@ -3499,7 +3483,7 @@ mu::Ptr<mu::NodeMesh> GenerateMutableSourceMesh(const UEdGraphPin* Pin,
 		mu::NodeMeshConstantPtr MeshNode = new mu::NodeMeshConstant();
 		Result = MeshNode;
 
-		GenerationContext.LoadObject(TypedNodeSkel->SkeletalMesh, true);
+		GenerationContext.LoadObject(TypedNodeSkel->SkeletalMesh);
 		if (TypedNodeSkel->SkeletalMesh)
 		{
 			int32 LODIndexConnected = -1; // LOD which the pin is connected to
@@ -3522,8 +3506,6 @@ mu::Ptr<mu::NodeMesh> GenerateMutableSourceMesh(const UEdGraphPin* Pin,
 
 			if (!TypedNodeSkel->AnimInstance.IsNull())
 			{
-				GenerationContext.AddParticipatingObject(TypedNodeSkel->AnimInstance);
-
 				FName SlotIndex = TypedNodeSkel->AnimBlueprintSlotName;
 				const int32 AnimInstanceIndex = GenerationContext.AnimBPAssets.AddUnique(TypedNodeSkel->AnimInstance);
 
@@ -3579,14 +3561,6 @@ mu::Ptr<mu::NodeMesh> GenerateMutableSourceMesh(const UEdGraphPin* Pin,
 			{
 				MeshNode->SetValue(MutableMesh);
 
-				if (TypedNodeSkel->SkeletalMesh->GetPhysicsAsset() && 
-					MutableMesh->GetPhysicsBody() && 
-					MutableMesh->GetPhysicsBody()->GetBodyCount())
-				{
-					TSoftObjectPtr<UPhysicsAsset> PhysicsAsset = TypedNodeSkel->SkeletalMesh->GetPhysicsAsset();
-					GenerationContext.AddParticipatingObject(*PhysicsAsset);
-				}
-
 				if (GenerationContext.CompilationContext.Options.bSkinWeightProfilesEnabled && LODIndex>=0)
 				{
 					if (ImportedModel && ImportedModel->LODModels.IsValidIndex(LODIndex))
@@ -3709,7 +3683,7 @@ mu::Ptr<mu::NodeMesh> GenerateMutableSourceMesh(const UEdGraphPin* Pin,
 
 	else if (const UCustomizableObjectNodeStaticMesh* TypedNodeStatic = Cast<UCustomizableObjectNodeStaticMesh>(Node))
 	{
-		GenerationContext.LoadObject(TypedNodeStatic->StaticMesh, true);
+		GenerationContext.LoadObject(TypedNodeStatic->StaticMesh);
 		if (TypedNodeStatic->StaticMesh == nullptr)
 		{
 			FString Msg = FString::Printf(TEXT("The UCustomizableObjectNodeStaticMesh node %s has no static mesh assigned"), *Node->GetName());
diff --git a/Engine/Plugins/Mutable/Source/CustomizableObjectEditor/Private/MuCOE/GenerateMutableSource/GenerateMutableSourceSurface.cpp b/Engine/Plugins/Mutable/Source/CustomizableObjectEditor/Private/MuCOE/GenerateMutableSource/GenerateMutableSourceSurface.cpp
--- a/Engine/Plugins/Mutable/Source/CustomizableObjectEditor/Private/MuCOE/GenerateMutableSource/GenerateMutableSourceSurface.cpp
+++ b/Engine/Plugins/Mutable/Source/CustomizableObjectEditor/Private/MuCOE/GenerateMutableSource/GenerateMutableSourceSurface.cpp
@@ -270,7 +270,6 @@ mu::Ptr<mu::NodeSurface> GenerateMutableSourceSurface(const UEdGraphPin * Pin, F
 		uint32 SurfaceMetadataUniqueHash = 0; // Value 0 is used as invalid hash.
 		if (TypedNodeMat->GetMaterial())
 		{
-			GenerationContext.AddParticipatingObject(*TypedNodeMat->GetMaterial());
 			GenerationContext.CurrentReferencedMaterialIndex = GenerationContext.ReferencedMaterials.AddUnique(TypedNodeMat->GetMaterial());
 
 			if (ConnectedMaterialPin)
@@ -504,8 +503,6 @@ mu::Ptr<mu::NodeSurface> GenerateMutableSourceSurface(const UEdGraphPin * Pin, F
 						{
 							if (ReferenceTexture)
 							{
-								GenerationContext.AddParticipatingObject(*ReferenceTexture);
-
 								// Store properties for the generated images
 								Props.TextureParameterName = ImageName;
 								Props.ImagePropertiesIndex = GenerationContext.ImageProperties.Num() - 1;
diff --git a/Engine/Plugins/Mutable/Source/CustomizableObjectEditor/Private/MuCOE/GenerateMutableSource/GenerateMutableSourceTable.cpp b/Engine/Plugins/Mutable/Source/CustomizableObjectEditor/Private/MuCOE/GenerateMutableSource/GenerateMutableSourceTable.cpp
--- a/Engine/Plugins/Mutable/Source/CustomizableObjectEditor/Private/MuCOE/GenerateMutableSource/GenerateMutableSourceTable.cpp
+++ b/Engine/Plugins/Mutable/Source/CustomizableObjectEditor/Private/MuCOE/GenerateMutableSource/GenerateMutableSourceTable.cpp
@@ -46,10 +46,8 @@ bool FillTableColumn(const UCustomizableObjectNodeTable* TableNode, mu::Ptr<mu::
 		
 		if (SoftObjectProperty->PropertyClass->IsChildOf(USkeletalMesh::StaticClass()))
 		{
-			GenerationContext.AddParticipatingObject(SoftObject);
-
 			// TODO: This is still forcing the load of the skeletal mesh asset. 
-			UObject* Object = GenerationContext.LoadObject(SoftObject, true);
+			UObject* Object = GenerationContext.LoadObject(SoftObject);
 			USkeletalMesh* SkeletalMesh = Cast<USkeletalMesh>(Object);
 
 			if(!SkeletalMesh)
@@ -125,8 +123,6 @@ bool FillTableColumn(const UCustomizableObjectNodeTable* TableNode, mu::Ptr<mu::
 
 											if (!AnimInstance.IsNull())
 											{
-												GenerationContext.AddParticipatingObject(AnimInstance);
-
 												const int32 AnimInstanceIndex = GenerationContext.AnimBPAssets.AddUnique(AnimInstance);
 
 												AnimBPAssetTag = GenerateAnimationInstanceTag(AnimInstanceIndex, SlotIndex);
@@ -286,12 +282,6 @@ bool FillTableColumn(const UCustomizableObjectNodeTable* TableNode, mu::Ptr<mu::
 
 			if (MutableMesh)
 			{
-				if (SkeletalMesh->GetPhysicsAsset())
-				{
-					TSoftObjectPtr<UPhysicsAsset> PhysicsAsset = SkeletalMesh->GetPhysicsAsset();
-					GenerationContext.AddParticipatingObject(*PhysicsAsset);
-				}
-
 				if (!AnimBPAssetTag.IsEmpty())
 				{
 					AddTagToMutableMeshUnique(*MutableMesh, AnimBPAssetTag);
@@ -329,8 +319,6 @@ bool FillTableColumn(const UCustomizableObjectNodeTable* TableNode, mu::Ptr<mu::
 
 		else if (SoftObjectProperty->PropertyClass->IsChildOf(UStaticMesh::StaticClass()))
 		{
-			GenerationContext.AddParticipatingObject(SoftObject);
-			
 			UObject* Object = GenerationContext.LoadObject(SoftObject);
 
 			UStaticMesh* StaticMesh = Cast<UStaticMesh>(Object);
@@ -403,8 +391,6 @@ bool FillTableColumn(const UCustomizableObjectNodeTable* TableNode, mu::Ptr<mu::
 
 		else if (SoftObjectProperty->PropertyClass->IsChildOf(UTexture::StaticClass()))
 		{
-			GenerationContext.AddParticipatingObject(SoftObject);
-
 			UObject* Object = GenerationContext.LoadObject(SoftObject);
 			UTexture* Texture = Cast<UTexture>(Object);
 
@@ -435,8 +421,6 @@ bool FillTableColumn(const UCustomizableObjectNodeTable* TableNode, mu::Ptr<mu::
 
 		else if (SoftObjectProperty->PropertyClass->IsChildOf(UMaterialInterface::StaticClass()))
 		{
-			GenerationContext.AddParticipatingObject(SoftObject);
-			
 			UObject* Object = GenerationContext.LoadObject(SoftObject);
 
 			// Get display name of the column of the data table (name showed in the table and struct editors)
@@ -454,8 +438,6 @@ bool FillTableColumn(const UCustomizableObjectNodeTable* TableNode, mu::Ptr<mu::
 				return false;
 			}
 
-			GenerationContext.AddParticipatingObject(*ReferenceMaterial);
-
 			const bool bTableMaterialCheckDisabled = GenerationContext.CompilationContext.Object->GetPrivate()->IsTableMaterialsParentCheckDisabled();
 			const bool bMaterialParentMismatch = !bTableMaterialCheckDisabled && MaterialInstance
 												 && ReferenceMaterial->GetMaterial() != MaterialInstance->GetMaterial();
@@ -571,8 +553,6 @@ bool FillTableColumn(const UCustomizableObjectNodeTable* TableNode, mu::Ptr<mu::
 
 		else if (SoftObjectProperty->PropertyClass->IsChildOf(UPoseAsset::StaticClass()))
 		{
-			GenerationContext.AddParticipatingObject(SoftObject);
-			
 			UObject* Object = GenerationContext.LoadObject(SoftObject);
 			
 			if (UPoseAsset* PoseAsset = Cast<UPoseAsset>(Object))
@@ -680,20 +660,6 @@ bool FillTableColumn(const UCustomizableObjectNodeTable* TableNode, mu::Ptr<mu::
 }
 
 
-uint8* GetCellData(const FName& RowName, const UDataTable& DataTable, const FProperty& ColumnProperty)
-{
-	// Get Row Data
-	uint8* RowData = DataTable.FindRowUnchecked(RowName);
-
-	if (RowData)
-	{
-		return ColumnProperty.ContainerPtrToValuePtr<uint8>(RowData, 0);
-	}
-
-	return nullptr;
-}
-
-
 FName GetAnotherOption(FName SelectedOptionName, const TArray<FName>& RowNames)
 {
 	for (const FName& CandidateOption : RowNames)
@@ -708,55 +674,6 @@ FName GetAnotherOption(FName SelectedOptionName, const TArray<FName>& RowNames)
 }
 
 
-TArray<FName> GetEnabledRows(const UDataTable& DataTable, const UCustomizableObjectNodeTable& TableNode)
-{
-	TArray<FName> RowNames;
-	const UScriptStruct* TableStruct = DataTable.GetRowStruct();
-
-	if (!TableStruct)
-	{
-		return RowNames;
-	}
-
-	TArray<FName> TableRowNames = DataTable.GetRowNames();
-	FBoolProperty* BoolProperty = nullptr;
-
-	// Sort them to avoid cooked data indeterminism problems. Rows may come from different tables and their loading order
-	// is not defined.
-	TableRowNames.Sort([](const FName& A, const FName& B) { return A.ToString() < B.ToString(); });
-
-	for (TFieldIterator<FProperty> PropertyIt(TableStruct); PropertyIt && TableNode.bDisableCheckedRows; ++PropertyIt)
-	{
-		BoolProperty = CastField<FBoolProperty>(*PropertyIt);
-
-		if (BoolProperty)
-		{
-			for (const FName& RowName : TableRowNames)
-			{
-				if (uint8* CellData = GetCellData(RowName, DataTable, *BoolProperty))
-				{
-					if (!BoolProperty->GetPropertyValue(CellData))
-					{
-						RowNames.Add(RowName);
-					}
-				}
-			}
-
-			// There should be only one Bool column
-			break;
-		}
-	}
-
-	// There is no Bool column or we don't want to disable rows
-	if (!BoolProperty)
-	{
-		return TableRowNames;
-	}
-
-	return RowNames;
-}
-
-
 void RestrictRowNamesToSelectedOption(TArray<FName>& InOutRowNames, const UCustomizableObjectNodeTable& TableNode, FMutableGraphGenerationContext& GenerationContext)
 {
 	if (!GenerationContext.CompilationContext.Options.ParamNamesToSelectedOptions.IsEmpty())
@@ -835,7 +752,7 @@ void RestrictRowContentByVersion( TArray<FName>& InOutRowNames, const UDataTable
 
 	for (int32 RowIndex = 0; RowIndex < InOutRowNames.Num(); ++RowIndex)
 	{
-		if (uint8* CellData = GetCellData(InOutRowNames[RowIndex], DataTable, *ColumnProperty))
+		if (uint8* CellData = UCustomizableObjectNodeTable::GetCellData(InOutRowNames[RowIndex], DataTable, *ColumnProperty))
 		{
 			if (!CustomizableObjectVersionBridgeInterface->IsVersionPropertyIncludedInCurrentRelease(*ColumnProperty, CellData))
 			{
@@ -892,7 +809,7 @@ TArray<FName> GetRowsToCompile(const UDataTable& DataTable, const UCustomizableO
 	}
 	else
 	{
-		TArray<FName> RowNames = GetEnabledRows(DataTable, TableNode);
+		TArray<FName> RowNames = TableNode.GetEnabledRows(DataTable);
 
 		if (!RowNames.IsEmpty())
 		{
@@ -923,60 +840,56 @@ bool GenerateTableColumn(const UCustomizableObjectNodeTable* TableNode, const UE
 	{
 		return false;
 	}
-	
-	GenerationContext.AddParticipatingObject(*DataTable);
 
 	// Getting names of the rows to access the information
 	TArray<uint32> RowIds;
 	TArray<FName> RowNames = GetRowsToCompile(*DataTable, *TableNode, GenerationContext, RowIds);
 
 	// Pre-pass to request async loading of all data. This seems to be slightly faster because it avoids the sync after every separate load.
-	bool bPreload = !GenerationContext.bParticipatingObjectsPass;
-	//if (bPreload)
+	// Avoid pre-loading meshes, which should be loaded on-demand during the core compilation
+	//bool bPreload = true;
+	//if (const FSoftObjectProperty* SoftObjectProperty = CastField<FSoftObjectProperty>(ColumnProperty))
 	//{
-	//	if (const FSoftObjectProperty* SoftObjectProperty = CastField<FSoftObjectProperty>(ColumnProperty))
+	//	if (SoftObjectProperty->PropertyClass->IsChildOf(USkeletalMesh::StaticClass()))
 	//	{
-	//		if (SoftObjectProperty->PropertyClass->IsChildOf(USkeletalMesh::StaticClass()))
-	//		{
-	//			bPreload = false;
-	//		}
-	//		else if (SoftObjectProperty->PropertyClass->IsChildOf(UStaticMesh::StaticClass()))
-	//		{
-	//			bPreload = false;
-	//		}
+	//		bPreload = false;
+	//	}
+	//	else if (SoftObjectProperty->PropertyClass->IsChildOf(UStaticMesh::StaticClass()))
+	//	{
+	//		bPreload = false;
 	//	}
 	//}
 
-	if (bPreload)
-	{
-		FCookLoadScope CookLoadScope(ECookLoadType::EditorOnly);
+	//if (bPreload)
+	//{
+	//	FCookLoadScope CookLoadScope(ECookLoadType::EditorOnly);
 
-		TArray<int32> LoadRequests;
-		LoadRequests.Reserve(RowNames.Num());
+	//	TArray<int32> LoadRequests;
+	//	LoadRequests.Reserve(RowNames.Num());
 
-		MUTABLE_CPUPROFILER_SCOPE(Preload);
-		for (int32 RowIndex = 0; RowIndex < RowNames.Num(); ++RowIndex)
-		{
-			if (uint8* CellData = GetCellData(RowNames[RowIndex], *DataTable, *ColumnProperty))
-			{
-				// Getting property type
-				if (const FSoftObjectProperty* SoftObjectProperty = CastField<FSoftObjectProperty>(ColumnProperty))
-				{
-					const FSoftObjectPtr& Path = SoftObjectProperty->GetPropertyValue(CellData);
-					const FString PackageName = Path.GetLongPackageName();
-					if (!PackageName.IsEmpty())
-					{
-						LoadRequests.Add( LoadPackageAsync(PackageName) );
-					}
-				}
-			}
-		}
+	//	MUTABLE_CPUPROFILER_SCOPE(Preload);
+	//	for (int32 RowIndex = 0; RowIndex < RowNames.Num(); ++RowIndex)
+	//	{
+	//		if (uint8* CellData = UCustomizableObjectNodeTable::GetCellData(RowNames[RowIndex], *DataTable, *ColumnProperty))
+	//		{
+	//			// Getting property type
+	//			if (const FSoftObjectProperty* SoftObjectProperty = CastField<FSoftObjectProperty>(ColumnProperty))
+	//			{
+	//				const FSoftObjectPtr& Path = SoftObjectProperty->GetPropertyValue(CellData);
+	//				const FString PackageName = Path.GetLongPackageName();
+	//				if (!PackageName.IsEmpty())
+	//				{
+	//					LoadRequests.Add( LoadPackageAsync(PackageName) );
+	//				}
+	//			}
+	//		}
+	//	}
 
-		{
-			MUTABLE_CPUPROFILER_SCOPE(Flush);
-			FlushAsyncLoading(LoadRequests);
-		}
-	}
+	//	{
+	//		MUTABLE_CPUPROFILER_SCOPE(Flush);
+	//		FlushAsyncLoading(LoadRequests);
+	//	}
+	//}
 
 	// Variable to check if something failed during a Cell generation
 	bool bCellGenerated = true;
@@ -995,7 +908,7 @@ bool GenerateTableColumn(const UCustomizableObjectNodeTable* TableNode, const UE
 	// Set the value of each row
 	for (int32 RowIndex = 0; RowIndex < RowNames.Num() && bCellGenerated; ++RowIndex)
 	{
-		if (uint8* CellData = GetCellData(RowNames[RowIndex], *DataTable, *ColumnProperty))
+		if (uint8* CellData = UCustomizableObjectNodeTable::GetCellData(RowNames[RowIndex], *DataTable, *ColumnProperty))
 		{
 			bCellGenerated = FillTableColumn(TableNode, MutableTable, DataTableColumnName, RowNames[RowIndex].ToString(), RowIds[RowIndex], CellData, ColumnProperty,
 				LODIndexConnected, SectionIndexConnected, LODIndex, SectionIndex, SectionMetadataId, bOnlyConnectedLOD, GenerationContext);
@@ -1063,7 +976,7 @@ void GenerateTableParameterUIData(const UDataTable* DataTable, const UCustomizab
 		{
 			FName RowName = RowNames[NameIndex];
 
-			if (uint8* MetadataCellData = GetCellData(RowName, *DataTable, *MetadataColumnProperty))
+			if (uint8* MetadataCellData = UCustomizableObjectNodeTable::GetCellData(RowName, *DataTable, *MetadataColumnProperty))
 			{
 				FMutableParamUIMetadata MetadataValue = *reinterpret_cast<FMutableParamUIMetadata*>(MetadataCellData);
 
@@ -1072,7 +985,7 @@ void GenerateTableParameterUIData(const UDataTable* DataTable, const UCustomizab
 				// Add thumbnail
 				if (bIsValidThumbnailColumn && MetadataValue.EditorUIThumbnailObject.IsNull())
 				{
-					if (uint8* ThumbnailCellData = GetCellData(RowName, *DataTable, *ThumbnailColumnProperty))
+					if (uint8* ThumbnailCellData = UCustomizableObjectNodeTable::GetCellData(RowName, *DataTable, *ThumbnailColumnProperty))
 					{
 						FSoftObjectPtr* ObjectPtr = reinterpret_cast<FSoftObjectPtr*>(ThumbnailCellData);
 						IntegerMetadata.ParamUIMetadata.EditorUIThumbnailObject = ObjectPtr->ToSoftObjectPath();
@@ -1097,7 +1010,7 @@ void GenerateTableParameterUIData(const UDataTable* DataTable, const UCustomizab
 							{
 								if (StructProperty->Struct == TBaseStructure<FGameplayTagContainer>::Get())
 								{
-									if (uint8* TagCellData = GetCellData(RowName, *DataTable, *ColumnProperty))
+									if (uint8* TagCellData = UCustomizableObjectNodeTable::GetCellData(RowName, *DataTable, *ColumnProperty))
 									{
 										FGameplayTagContainer* TagContainer = reinterpret_cast<FGameplayTagContainer*>(TagCellData);
 										IntegerMetadata.ParamUIMetadata.GameplayTags.AppendTags(*TagContainer);
@@ -1205,20 +1118,6 @@ mu::Ptr<mu::FTable> GenerateMutableSourceTable(const UDataTable* DataTable, cons
 }
 
 
-void AddCompositeTablesToParticipatingObjetcts(const UDataTable* Table, FMutableGraphGenerationContext& GenerationContext)
-{
-	if (const UCompositeDataTable* CompositeTable = Cast<UCompositeDataTable>(Table))
-	{
-		GenerationContext.AddParticipatingObject(*CompositeTable);
-
-		/*for (const TArray<TObjectPtr<UDataTable>>& ParentTable : CompositeTable.ParentTables) // TODO 
-		{
-			AddCompositeTablesToParticipatingObjetcts(ParentTable, GenerationContext);
-		}*/
-	}
-}
-
-
 UDataTable* GetDataTable(const UCustomizableObjectNodeTable* TableNode, FMutableGraphGenerationContext& GenerationContext)
 {
 	UDataTable* OutDataTable = nullptr;
@@ -1232,8 +1131,6 @@ UDataTable* GetDataTable(const UCustomizableObjectNodeTable* TableNode, FMutable
 		OutDataTable = TableNode->Table;
 	}
 
-	AddCompositeTablesToParticipatingObjetcts(OutDataTable, GenerationContext);
-
 	return OutDataTable;
 }
 
diff --git a/Engine/Plugins/Mutable/Source/CustomizableObjectEditor/Private/MuCOE/Nodes/CustomizableObjectNodeTable.cpp b/Engine/Plugins/Mutable/Source/CustomizableObjectEditor/Private/MuCOE/Nodes/CustomizableObjectNodeTable.cpp
--- a/Engine/Plugins/Mutable/Source/CustomizableObjectEditor/Private/MuCOE/Nodes/CustomizableObjectNodeTable.cpp
+++ b/Engine/Plugins/Mutable/Source/CustomizableObjectEditor/Private/MuCOE/Nodes/CustomizableObjectNodeTable.cpp
@@ -1518,4 +1518,67 @@ int32 UCustomizableObjectNodeTable::GetDefaultSkeletalMaterialIndexFor(const UEd
 }
 
 
+uint8* UCustomizableObjectNodeTable::GetCellData(const FName& RowName, const UDataTable& DataTable, const FProperty& ColumnProperty)
+{
+	// Get Row Data
+	uint8* RowData = DataTable.FindRowUnchecked(RowName);
+
+	if (RowData)
+	{
+		return ColumnProperty.ContainerPtrToValuePtr<uint8>(RowData, 0);
+	}
+
+	return nullptr;
+}
+
+
+TArray<FName> UCustomizableObjectNodeTable::GetEnabledRows(const UDataTable& DataTable) const
+{
+	TArray<FName> RowNames;
+	const UScriptStruct* TableStruct = DataTable.GetRowStruct();
+
+	if (!TableStruct)
+	{
+		return RowNames;
+	}
+
+	TArray<FName> TableRowNames = DataTable.GetRowNames();
+	FBoolProperty* BoolProperty = nullptr;
+
+	// Sort them to avoid cooked data indeterminism problems. Rows may come from different tables and their loading order
+	// is not defined.
+	TableRowNames.Sort([](const FName& A, const FName& B) { return A.ToString() < B.ToString(); });
+
+	for (TFieldIterator<FProperty> PropertyIt(TableStruct); PropertyIt && bDisableCheckedRows; ++PropertyIt)
+	{
+		BoolProperty = CastField<FBoolProperty>(*PropertyIt);
+
+		if (BoolProperty)
+		{
+			for (const FName& RowName : TableRowNames)
+			{
+				if (uint8* CellData = GetCellData(RowName, DataTable, *BoolProperty))
+				{
+					if (!BoolProperty->GetPropertyValue(CellData))
+					{
+						RowNames.Add(RowName);
+					}
+				}
+			}
+
+			// There should be only one Bool column
+			break;
+		}
+	}
+
+	// There is no Bool column or we don't want to disable rows
+	if (!BoolProperty)
+	{
+		return TableRowNames;
+	}
+
+	return RowNames;
+}
+
+
 #undef LOCTEXT_NAMESPACE
diff --git a/Engine/Plugins/Mutable/Source/CustomizableObjectEditor/Private/MuCOE/Nodes/CustomizableObjectNodeTable.h b/Engine/Plugins/Mutable/Source/CustomizableObjectEditor/Private/MuCOE/Nodes/CustomizableObjectNodeTable.h
--- a/Engine/Plugins/Mutable/Source/CustomizableObjectEditor/Private/MuCOE/Nodes/CustomizableObjectNodeTable.h
+++ b/Engine/Plugins/Mutable/Source/CustomizableObjectEditor/Private/MuCOE/Nodes/CustomizableObjectNodeTable.h
@@ -400,6 +400,12 @@ public:
 	/** Returns the SkeletalMesh section associated to the given skeletal mesh output pin. */
 	const FSkelMeshSection* GetDefaultSkeletalMeshSectionFor(const UEdGraphPin& MeshPin) const;
 
+	/** */
+	TArray<FName> GetEnabledRows(const UDataTable& DataTable) const;
+
+	/** */
+	static uint8* GetCellData(const FName& RowName, const UDataTable& DataTable, const FProperty& ColumnProperty);
+
 private:
 	/** Number of properties to know when the node needs an update */
 	UPROPERTY()
diff --git a/Engine/Plugins/Mutable/Source/CustomizableObjectEditor/Public/MuCOE/ExtensionDataCompilerInterface.h b/Engine/Plugins/Mutable/Source/CustomizableObjectEditor/Public/MuCOE/ExtensionDataCompilerInterface.h
--- a/Engine/Plugins/Mutable/Source/CustomizableObjectEditor/Public/MuCOE/ExtensionDataCompilerInterface.h
+++ b/Engine/Plugins/Mutable/Source/CustomizableObjectEditor/Public/MuCOE/ExtensionDataCompilerInterface.h
@@ -60,7 +60,5 @@ public:
 	/** Adds a compiler log message to be displayed at the end of the compilation process */
 	void CompilerLog(const FText& InLogText, const class UCustomizableObjectNode* InNode);
 
-	void AddParticipatingObject(const UObject& Object);
-
 	FMutableGraphGenerationContext& GenerationContext;
 };
\ No newline at end of file
