[mutable] Changed the color of a + button to add tags in the MeshSection and EditMeshSection nodes, as they looked as if they were in a disabled state in the old color.
#tests Editor
#rnx

diff --git a/Engine/Plugins/Mutable/Source/CustomizableObjectEditor/Private/MuCOE/SMutableTagListWidget.cpp b/Engine/Plugins/Mutable/Source/CustomizableObjectEditor/Private/MuCOE/SMutableTagListWidget.cpp
--- a/Engine/Plugins/Mutable/Source/CustomizableObjectEditor/Private/MuCOE/SMutableTagListWidget.cpp
+++ b/Engine/Plugins/Mutable/Source/CustomizableObjectEditor/Private/MuCOE/SMutableTagListWidget.cpp
@@ -44,7 +44,7 @@ void SMutableTagComboBox::Construct(const FArguments& InArgs)
 		.ButtonStyle(FAppStyle::Get(), "NoBorder")
 		.MenuButtonBrush(InArgs._MenuButtonBrush)
 		.OnSelectionChanged(InArgs._OnSelectionChanged)
-		.ForegroundColor(FSlateColor::UseSubduedForeground())
+		.ForegroundColor(FSlateColor(FLinearColor(.35f, .35f, .35f)))
 		.Content()
 		[
 			InArgs._Content.Widget
@@ -478,7 +478,7 @@ TSharedRef<ITableRow> SMutableTagListWidget::GenerateTagListItemRow(TSharedPtr<F
 						[
 							SNew(SImage)
 								.Image(FAppStyle::GetBrush(TEXT("Icons.MinusCircle")))
-								.ColorAndOpacity(FSlateColor::UseSubduedForeground())
+								.ColorAndOpacity(FSlateColor(FLinearColor(.35f, .35f, .35f)))
 						]
 				]
 		];
