[mutable] Fixed direct removal of include.
- A deprecation must have been used instead of the removal of the include as it was set in a public header.

#rnx

diff --git a/Engine/Plugins/Mutable/Source/CustomizableObject/Public/MuCO/CustomizableObject.h b/Engine/Plugins/Mutable/Source/CustomizableObject/Public/MuCO/CustomizableObject.h
--- a/Engine/Plugins/Mutable/Source/CustomizableObject/Public/MuCO/CustomizableObject.h
+++ b/Engine/Plugins/Mutable/Source/CustomizableObject/Public/MuCO/CustomizableObject.h
@@ -8,6 +8,10 @@
 #include "MuCO/CustomizableObject_Deprecated.h"
 #include "Templates/TypeCompatibleBytes.h"
 
+#if UE_ENABLE_INCLUDE_ORDER_DEPRECATED_IN_5_6
+#include "SkeletalMeshTypes.h"
+#endif //UE_ENABLE_INCLUDE_ORDER_DEPRECATED_IN_5_6
+
 #include "CustomizableObject.generated.h"
 
 class FReply;
