[mutable] Fix crash with composite tables and participating object detection.
#rnx

diff --git a/Engine/Plugins/Mutable/Source/CustomizableObjectEditor/Private/MuCOE/CustomizableObjectEditorModule.cpp b/Engine/Plugins/Mutable/Source/CustomizableObjectEditor/Private/MuCOE/CustomizableObjectEditorModule.cpp
--- a/Engine/Plugins/Mutable/Source/CustomizableObjectEditor/Private/MuCOE/CustomizableObjectEditorModule.cpp
+++ b/Engine/Plugins/Mutable/Source/CustomizableObjectEditor/Private/MuCOE/CustomizableObjectEditorModule.cpp
@@ -823,9 +823,9 @@ TMap<FName, FGuid> FCustomizableObjectEditorModule::GetParticipatingObjects(UCus
 
 				if (DataTable)
 				{
-					AddObject(TableNode->Table);
+					AddObject(DataTable);
 
-					TArray<FName> RowNames = TableNode->GetEnabledRows(*TableNode->Table);
+					TArray<FName> RowNames = TableNode->GetEnabledRows(*DataTable);
 
 					TArray<FName> ExpectedPropNames = DataTableUtils::GetStructPropertyNames(DataTable->RowStruct);
 					for (const FName& ColumnName : ExpectedPropNames)
