Nonunity/Stresstest compile errors

diff --git a/Engine/Plugins/Mutable/Source/CustomizableObject/Private/MuCO/CustomizableObjectModule.cpp b/Engine/Plugins/Mutable/Source/CustomizableObject/Private/MuCO/CustomizableObjectModule.cpp
--- a/Engine/Plugins/Mutable/Source/CustomizableObject/Private/MuCO/CustomizableObjectModule.cpp
+++ b/Engine/Plugins/Mutable/Source/CustomizableObject/Private/MuCO/CustomizableObjectModule.cpp
@@ -15,12 +15,14 @@
 #include "MuCO/CustomizableObjectSystem.h"
 #include "MuCO/ICustomizableObjectModule.h"
 #include "UObject/StrongObjectPtr.h"
+#include "UObject/UObjectIterator.h"
 #include "GPUSkinPublicDefs.h"
 #include "Components/SkeletalMeshComponent.h"
 
 /**
  * Customizable Object module implementation (private)
  */
+
 class FCustomizableObjectModule : public ICustomizableObjectModule
 {
 public:
diff --git a/Engine/Source/Runtime/Engine/Private/ImageUtils.cpp b/Engine/Source/Runtime/Engine/Private/ImageUtils.cpp
--- a/Engine/Source/Runtime/Engine/Private/ImageUtils.cpp
+++ b/Engine/Source/Runtime/Engine/Private/ImageUtils.cpp
@@ -24,6 +24,10 @@ ImageUtils.cpp: Image utility functions.
 #include "DDSFile.h"
 #include "TextureResource.h"
 
+#ifdef LoadImage // Seems like some windows include is leaking through (LoadImage define can be set to LoadImageW)
+#undef LoadImage
+#endif
+
 DEFINE_LOG_CATEGORY_STATIC(LogImageUtils, Log, All);
 
 #define LOCTEXT_NAMESPACE "ImageUtils"
diff --git a/Engine/Source/Runtime/Experimental/Chaos/Private/Chaos/Serialization/SolverSerializer.cpp b/Engine/Source/Runtime/Experimental/Chaos/Private/Chaos/Serialization/SolverSerializer.cpp
--- a/Engine/Source/Runtime/Experimental/Chaos/Private/Chaos/Serialization/SolverSerializer.cpp
+++ b/Engine/Source/Runtime/Experimental/Chaos/Private/Chaos/Serialization/SolverSerializer.cpp
@@ -10,6 +10,7 @@
 #include "ChaosVisualDebugger/ChaosVDDataWrapperUtils.h"
 #include "DataWrappers/ChaosVDCharacterGroundConstraintDataWrappers.h"
 #include "DataWrappers/ChaosVDJointDataWrappers.h"
+#include "Serialization/MemoryReader.h"
 
 namespace Chaos
 {
