[mutable] Updated Paramters and Constant pin names
- Replaced generic name for the name of the actual resource (Value->Texure,Enum,Projector...)
- Added "virtual" keywords
- Removed unnecesary functions
- Fixed some file format issues.

diff --git a/Engine/Plugins/Mutable/Source/CustomizableObject/Internal/MuCO/CustomizableObjectCustomVersion.h b/Engine/Plugins/Mutable/Source/CustomizableObject/Internal/MuCO/CustomizableObjectCustomVersion.h
--- a/Engine/Plugins/Mutable/Source/CustomizableObject/Internal/MuCO/CustomizableObjectCustomVersion.h
+++ b/Engine/Plugins/Mutable/Source/CustomizableObject/Internal/MuCO/CustomizableObjectCustomVersion.h
@@ -170,7 +170,9 @@ struct CUSTOMIZABLEOBJECT_API FCustomizableObjectCustomVersion
 
 		TableNoneOptionsMovedToUnrealCode,
 
-		UpdatedNodesPinName,	
+		UpdatedNodesPinName,
+
+		UpdatedNodesPinName2,
 
 		// -----<new versions can be added above this line>-------------------------------------------------
 		VersionPlusOne,
diff --git a/Engine/Plugins/Mutable/Source/CustomizableObjectEditor/Private/MuCOE/Nodes/CustomizableObjectNodeColorConstant.cpp b/Engine/Plugins/Mutable/Source/CustomizableObjectEditor/Private/MuCOE/Nodes/CustomizableObjectNodeColorConstant.cpp
--- a/Engine/Plugins/Mutable/Source/CustomizableObjectEditor/Private/MuCOE/Nodes/CustomizableObjectNodeColorConstant.cpp
+++ b/Engine/Plugins/Mutable/Source/CustomizableObjectEditor/Private/MuCOE/Nodes/CustomizableObjectNodeColorConstant.cpp
@@ -2,6 +2,7 @@
 
 #include "MuCOE/Nodes/CustomizableObjectNodeColorConstant.h"
 
+#include "MuCO/CustomizableObjectCustomVersion.h"
 #include "MuCOE/CustomizableObjectEditorStyle.h"
 #include "MuCOE/CustomizableObjectGraph.h"
 #include "MuCOE/EdGraphSchema_CustomizableObject.h"
@@ -26,22 +27,14 @@ UCustomizableObjectNodeColorConstant::UCustomizableObjectNodeColorConstant()
 }
 
 
-void UCustomizableObjectNodeColorConstant::PostEditChangeProperty(FPropertyChangedEvent& PropertyChangedEvent)
-{
-	//UProperty* PropertyThatChanged = PropertyChangedEvent.Property;
-	//if ( PropertyThatChanged && PropertyThatChanged->GetName() == TEXT("NumLODs") )
-	//{
-	//}
-
-	Super::PostEditChangeProperty(PropertyChangedEvent);
-}
-
-
 void UCustomizableObjectNodeColorConstant::AllocateDefaultPins(UCustomizableObjectNodeRemapPins* RemapPins)
 {
-	const UEdGraphSchema_CustomizableObject* Schema = GetDefault<UEdGraphSchema_CustomizableObject>();
-
-	UEdGraphPin* ValuePin = CustomCreatePin(EGPD_Output, Schema->PC_Color, FName("Value"));
+	const FName Type = UEdGraphSchema_CustomizableObject::PC_Color;
+	const FName PinName = UEdGraphSchema_CustomizableObject::GetPinCategoryName(Type);
+	const FText PinFriendlyName = UEdGraphSchema_CustomizableObject::GetPinCategoryFriendlyName(Type);
+	
+	UEdGraphPin* ValuePin = CustomCreatePin(EGPD_Output, Type, PinName);
+	ValuePin->PinFriendlyName = PinFriendlyName;
 	ValuePin->bDefaultValueIsIgnored = true;
 }
 
@@ -65,6 +58,21 @@ FText UCustomizableObjectNodeColorConstant::GetTooltipText() const
 }
 
 
+void UCustomizableObjectNodeColorConstant::BackwardsCompatibleFixup(int32 CustomizableObjectCustomVersion)
+{
+	Super::BackwardsCompatibleFixup(CustomizableObjectCustomVersion);
+
+	if (CustomizableObjectCustomVersion == FCustomizableObjectCustomVersion::UpdatedNodesPinName2)
+	{
+		if (UEdGraphPin* Pin = FindPin(TEXT("Value")))
+		{
+			Pin->PinName = TEXT("Color");
+			Pin->PinFriendlyName = LOCTEXT("Color_Pin_Category", "Color");;
+		}
+	}
+}
+
+
 TSharedPtr<SGraphNode> UCustomizableObjectNodeColorConstant::CreateVisualWidget()
 {
 	return SNew(SGraphNodeColorConstant, this);
@@ -82,10 +90,6 @@ void SGraphNodeColorConstant::Construct(const FArguments& InArgs, UEdGraphNode*
 	SCustomizableObjectNode::Construct({}, InGraphNode);
 }
 
-void SGraphNodeColorConstant::UpdateGraphNode()
-{
-	SGraphNode::UpdateGraphNode();
-}
 
 void SGraphNodeColorConstant::SetDefaultTitleAreaWidget(TSharedRef<SOverlay> DefaultTitleAreaWidget)
 {
diff --git a/Engine/Plugins/Mutable/Source/CustomizableObjectEditor/Private/MuCOE/Nodes/CustomizableObjectNodeColorConstant.h b/Engine/Plugins/Mutable/Source/CustomizableObjectEditor/Private/MuCOE/Nodes/CustomizableObjectNodeColorConstant.h
--- a/Engine/Plugins/Mutable/Source/CustomizableObjectEditor/Private/MuCOE/Nodes/CustomizableObjectNodeColorConstant.h
+++ b/Engine/Plugins/Mutable/Source/CustomizableObjectEditor/Private/MuCOE/Nodes/CustomizableObjectNodeColorConstant.h
@@ -41,9 +41,6 @@ public:
 	// Builds the SGraphNodeFloatConstant when needed
 	void Construct(const FArguments& InArgs, UEdGraphNode* InGraphNode);
 
-	// Calls the needed functions to build the SGraphNode widgets
-	void UpdateGraphNode();
-
 	// Overriden functions to build the SGraphNode widgets
 	virtual void SetDefaultTitleAreaWidget(TSharedRef<SOverlay> DefaultTitleAreaWidget) override;
 	virtual void CreateBelowPinControls(TSharedPtr<SVerticalBox> MainBox) override;
@@ -88,22 +85,20 @@ public:
 	/**  */
 	UPROPERTY(EditAnywhere, Category=CustomizableObject, meta = (DontUpdateWhileEditing))
 	FLinearColor Value;
-
-	// UObject interface.
-	void PostEditChangeProperty(FPropertyChangedEvent& PropertyChangedEvent) override;
-
+	
 	// Begin EdGraphNode interface
-	FText GetNodeTitle(ENodeTitleType::Type TitleType) const override;
-	FLinearColor GetNodeTitleColor() const override;
-	FText GetTooltipText() const override;
+	virtual FText GetNodeTitle(ENodeTitleType::Type TitleType) const override;
+	virtual FLinearColor GetNodeTitleColor() const override;
+	virtual FText GetTooltipText() const override;
 	// End EdGraphNode interface
 
 	// UCustomizableObjectNode interface
-	void AllocateDefaultPins(UCustomizableObjectNodeRemapPins* RemapPins) override;
-	bool IsAffectedByLOD() const override { return false; }
+	virtual void AllocateDefaultPins(UCustomizableObjectNodeRemapPins* RemapPins) override;
+	virtual bool IsAffectedByLOD() const override { return false; }
+	virtual void BackwardsCompatibleFixup(int32 CustomizableObjectCustomVersion) override;
 
 	// Creates the SGraph Node widget for the Color Editor
-	TSharedPtr<SGraphNode> CreateVisualWidget() override;
+	virtual TSharedPtr<SGraphNode> CreateVisualWidget() override;
 
 	// Determines if the Node is collapsed or not
 	bool bCollapsed;
diff --git a/Engine/Plugins/Mutable/Source/CustomizableObjectEditor/Private/MuCOE/Nodes/CustomizableObjectNodeColorParameter.cpp b/Engine/Plugins/Mutable/Source/CustomizableObjectEditor/Private/MuCOE/Nodes/CustomizableObjectNodeColorParameter.cpp
--- a/Engine/Plugins/Mutable/Source/CustomizableObjectEditor/Private/MuCOE/Nodes/CustomizableObjectNodeColorParameter.cpp
+++ b/Engine/Plugins/Mutable/Source/CustomizableObjectEditor/Private/MuCOE/Nodes/CustomizableObjectNodeColorParameter.cpp
@@ -10,11 +10,29 @@ class UCustomizableObjectNodeRemapPins;
 #define LOCTEXT_NAMESPACE "CustomizableObjectEditor"
 
 
-void UCustomizableObjectNodeColorParameter::AllocateDefaultPins(UCustomizableObjectNodeRemapPins* RemapPins)
+void UCustomizableObjectNodeColorParameter::BackwardsCompatibleFixup(int32 CustomizableObjectCustomVersion)
 {
-	const UEdGraphSchema_CustomizableObject* Schema = GetDefault<UEdGraphSchema_CustomizableObject>();
+	Super::BackwardsCompatibleFixup(CustomizableObjectCustomVersion);
+	
+	if (CustomizableObjectCustomVersion == FCustomizableObjectCustomVersion::UpdatedNodesPinName2)
+	{
+		if (UEdGraphPin* Pin = FindPin(TEXT("Value")))
+		{
+			Pin->PinName = TEXT("Color");
+			Pin->PinFriendlyName = LOCTEXT("Color_Pin_Category", "Color");;
+		}
+	}
+}
+
 
-	UEdGraphPin* ValuePin = CustomCreatePin(EGPD_Output, Schema->PC_Color, FName("Value"));
+void UCustomizableObjectNodeColorParameter::AllocateDefaultPins(UCustomizableObjectNodeRemapPins* RemapPins)
+{
+	const FName Type = UEdGraphSchema_CustomizableObject::PC_Color;
+	const FName PinName = UEdGraphSchema_CustomizableObject::GetPinCategoryName(Type);
+	const FText PinFriendlyName = UEdGraphSchema_CustomizableObject::GetPinCategoryFriendlyName(Type);
+	
+	UEdGraphPin* ValuePin = CustomCreatePin(EGPD_Output, Type, PinName);
+	ValuePin->PinFriendlyName = PinFriendlyName;
 	ValuePin->bDefaultValueIsIgnored = true;
 }
 
diff --git a/Engine/Plugins/Mutable/Source/CustomizableObjectEditor/Private/MuCOE/Nodes/CustomizableObjectNodeColorParameter.h b/Engine/Plugins/Mutable/Source/CustomizableObjectEditor/Private/MuCOE/Nodes/CustomizableObjectNodeColorParameter.h
--- a/Engine/Plugins/Mutable/Source/CustomizableObjectEditor/Private/MuCOE/Nodes/CustomizableObjectNodeColorParameter.h
+++ b/Engine/Plugins/Mutable/Source/CustomizableObjectEditor/Private/MuCOE/Nodes/CustomizableObjectNodeColorParameter.h
@@ -28,14 +28,15 @@ public:
 	FMutableParamUIMetadata ParamUIMetadata;
 
 	// Begin EdGraphNode interface
-	FText GetNodeTitle(ENodeTitleType::Type TitleType) const override;
-	FLinearColor GetNodeTitleColor() const override;
-	FText GetTooltipText() const override;
+	virtual FText GetNodeTitle(ENodeTitleType::Type TitleType) const override;
+	virtual FLinearColor GetNodeTitleColor() const override;
+	virtual FText GetTooltipText() const override;
 	virtual void OnRenameNode(const FString& NewName) override;
 	virtual bool GetCanRenameNode() const override { return true; }
+	virtual void BackwardsCompatibleFixup(int32 CustomizableObjectCustomVersion) override;
 
 	// UCustomizableObjectNode interface
-	void AllocateDefaultPins(UCustomizableObjectNodeRemapPins* RemapPins) override;
-	bool IsAffectedByLOD() const override { return false; }
+	virtual void AllocateDefaultPins(UCustomizableObjectNodeRemapPins* RemapPins) override;
+	virtual bool IsAffectedByLOD() const override { return false; }
 };
 
diff --git a/Engine/Plugins/Mutable/Source/CustomizableObjectEditor/Private/MuCOE/Nodes/CustomizableObjectNodeEnumParameter.cpp b/Engine/Plugins/Mutable/Source/CustomizableObjectEditor/Private/MuCOE/Nodes/CustomizableObjectNodeEnumParameter.cpp
--- a/Engine/Plugins/Mutable/Source/CustomizableObjectEditor/Private/MuCOE/Nodes/CustomizableObjectNodeEnumParameter.cpp
+++ b/Engine/Plugins/Mutable/Source/CustomizableObjectEditor/Private/MuCOE/Nodes/CustomizableObjectNodeEnumParameter.cpp
@@ -22,12 +22,30 @@ void UCustomizableObjectNodeEnumParameter::PostEditChangeProperty(FPropertyChang
 }
 
 
-void UCustomizableObjectNodeEnumParameter::AllocateDefaultPins(UCustomizableObjectNodeRemapPins* RemapPins)
+void UCustomizableObjectNodeEnumParameter::BackwardsCompatibleFixup(int32 CustomizableObjectCustomVersion)
 {
-	const UEdGraphSchema_CustomizableObject* Schema = GetDefault<UEdGraphSchema_CustomizableObject>();
+	Super::BackwardsCompatibleFixup(CustomizableObjectCustomVersion);
+	
+	if (CustomizableObjectCustomVersion == FCustomizableObjectCustomVersion::UpdatedNodesPinName2)
+	{
+		if (UEdGraphPin* Pin = FindPin(TEXT("Value")))
+		{
+			Pin->PinName = TEXT("Enum");
+			Pin->PinFriendlyName = LOCTEXT("Enum_Pin_Category", "Enum");
+		}
+	}
+}
+
 
-	UEdGraphPin* ValuePin = CustomCreatePin(EGPD_Output, Schema->PC_Enum, FName("Value"));
+void UCustomizableObjectNodeEnumParameter::AllocateDefaultPins(UCustomizableObjectNodeRemapPins* RemapPins)
+{
+	const FName Type = UEdGraphSchema_CustomizableObject::PC_Enum;
+	const FName PinName = UEdGraphSchema_CustomizableObject::GetPinCategoryName(Type);
+	const FText PinFriendlyName = UEdGraphSchema_CustomizableObject::GetPinCategoryFriendlyName(Type);
+	
+	UEdGraphPin* ValuePin = CustomCreatePin(EGPD_Output, Type, PinName);
 	ValuePin->bDefaultValueIsIgnored = true;
+	ValuePin->PinFriendlyName = PinFriendlyName;
 }
 
 
diff --git a/Engine/Plugins/Mutable/Source/CustomizableObjectEditor/Private/MuCOE/Nodes/CustomizableObjectNodeEnumParameter.h b/Engine/Plugins/Mutable/Source/CustomizableObjectEditor/Private/MuCOE/Nodes/CustomizableObjectNodeEnumParameter.h
--- a/Engine/Plugins/Mutable/Source/CustomizableObjectEditor/Private/MuCOE/Nodes/CustomizableObjectNodeEnumParameter.h
+++ b/Engine/Plugins/Mutable/Source/CustomizableObjectEditor/Private/MuCOE/Nodes/CustomizableObjectNodeEnumParameter.h
@@ -45,17 +45,18 @@ public:
 	TArray<FCustomizableObjectNodeEnumValue> Values;
 
 	// UObject interface.
-	void PostEditChangeProperty(FPropertyChangedEvent& PropertyChangedEvent) override;
+	virtual void PostEditChangeProperty(FPropertyChangedEvent& PropertyChangedEvent) override;
 
 	// Begin EdGraphNode interface
-	FText GetNodeTitle(ENodeTitleType::Type TitleType) const override;
-	FLinearColor GetNodeTitleColor() const override;
-	FText GetTooltipText() const override;
+	virtual FText GetNodeTitle(ENodeTitleType::Type TitleType) const override;
+	virtual FLinearColor GetNodeTitleColor() const override;
+	virtual FText GetTooltipText() const override;
 	virtual void OnRenameNode(const FString& NewName) override;
 	virtual bool GetCanRenameNode() const override { return true; }
+	virtual void BackwardsCompatibleFixup(int32 CustomizableObjectCustomVersion) override;
 
 	// UCustomizableObjectNode interface
-	void AllocateDefaultPins(UCustomizableObjectNodeRemapPins* RemapPins) override;
+	virtual void AllocateDefaultPins(UCustomizableObjectNodeRemapPins* RemapPins) override;
 	virtual bool IsAffectedByLOD() const override { return false; }
 
 };
diff --git a/Engine/Plugins/Mutable/Source/CustomizableObjectEditor/Private/MuCOE/Nodes/CustomizableObjectNodeFloatConstant.cpp b/Engine/Plugins/Mutable/Source/CustomizableObjectEditor/Private/MuCOE/Nodes/CustomizableObjectNodeFloatConstant.cpp
--- a/Engine/Plugins/Mutable/Source/CustomizableObjectEditor/Private/MuCOE/Nodes/CustomizableObjectNodeFloatConstant.cpp
+++ b/Engine/Plugins/Mutable/Source/CustomizableObjectEditor/Private/MuCOE/Nodes/CustomizableObjectNodeFloatConstant.cpp
@@ -2,6 +2,7 @@
 
 #include "MuCOE/Nodes/CustomizableObjectNodeFloatConstant.h"
 
+#include "MuCO/CustomizableObjectCustomVersion.h"
 #include "MuCOE/CustomizableObjectEditorStyle.h"
 #include "MuCOE/CustomizableObjectGraph.h"
 #include "MuCOE/EdGraphSchema_CustomizableObject.h"
@@ -24,22 +25,14 @@ UCustomizableObjectNodeFloatConstant::UCustomizableObjectNodeFloatConstant()
 }
 
 
-void UCustomizableObjectNodeFloatConstant::PostEditChangeProperty(FPropertyChangedEvent& PropertyChangedEvent)
-{
-	//UProperty* PropertyThatChanged = PropertyChangedEvent.Property;
-	//if ( PropertyThatChanged && PropertyThatChanged->GetName() == TEXT("NumLODs") )
-	{
-	}
-
-	Super::PostEditChangeProperty(PropertyChangedEvent);
-}
-
-
 void UCustomizableObjectNodeFloatConstant::AllocateDefaultPins(UCustomizableObjectNodeRemapPins* RemapPins)
 {
-	const UEdGraphSchema_CustomizableObject* Schema = GetDefault<UEdGraphSchema_CustomizableObject>();
-
-	UEdGraphPin* ValuePin = CustomCreatePin(EGPD_Output, Schema->PC_Float, FName("Value"));
+	const FName Type = UEdGraphSchema_CustomizableObject::PC_Float;
+	const FName PinName = UEdGraphSchema_CustomizableObject::GetPinCategoryName(Type);
+	const FText PinFriendlyName = UEdGraphSchema_CustomizableObject::GetPinCategoryFriendlyName(Type);
+	
+	UEdGraphPin* ValuePin = CustomCreatePin(EGPD_Output, Type, PinName);
+	ValuePin->PinFriendlyName = PinFriendlyName;
 	ValuePin->bDefaultValueIsIgnored = true;
 }
 
@@ -63,6 +56,20 @@ FText UCustomizableObjectNodeFloatConstant::GetTooltipText() const
 }
 
 
+void UCustomizableObjectNodeFloatConstant::BackwardsCompatibleFixup(int32 CustomizableObjectCustomVersion)
+{
+	Super::BackwardsCompatibleFixup(CustomizableObjectCustomVersion);
+
+	if (CustomizableObjectCustomVersion == FCustomizableObjectCustomVersion::UpdatedNodesPinName2)
+	{
+		if (UEdGraphPin* Pin = FindPin(TEXT("Value")))
+		{
+			Pin->PinName = TEXT("Float");
+			Pin->PinFriendlyName = LOCTEXT("Float_Pin_Category", "Float");
+		}
+	}
+}
+
 TSharedPtr<SGraphNode> UCustomizableObjectNodeFloatConstant::CreateVisualWidget()
 {
 	return SNew(SGraphNodeFloatConstant, this);
diff --git a/Engine/Plugins/Mutable/Source/CustomizableObjectEditor/Private/MuCOE/Nodes/CustomizableObjectNodeFloatConstant.h b/Engine/Plugins/Mutable/Source/CustomizableObjectEditor/Private/MuCOE/Nodes/CustomizableObjectNodeFloatConstant.h
--- a/Engine/Plugins/Mutable/Source/CustomizableObjectEditor/Private/MuCOE/Nodes/CustomizableObjectNodeFloatConstant.h
+++ b/Engine/Plugins/Mutable/Source/CustomizableObjectEditor/Private/MuCOE/Nodes/CustomizableObjectNodeFloatConstant.h
@@ -52,8 +52,6 @@ private:
 	// Callback for the OnValueCommited of the SpinBox
 	void OnSpinBoxValueCommitted(float Value, ETextCommit::Type);
 
-private:
-
 	// Pointer to the NodeFloatconstant that owns this SGraphNode
 	UCustomizableObjectNodeFloatConstant* NodeFloatConstant;
 
@@ -75,20 +73,18 @@ public:
 	UPROPERTY(EditAnywhere, Category = CustomizableObject, meta = (DontUpdateWhileEditing))
 	float Value;
 
-	// UObject interface.
-	void PostEditChangeProperty(FPropertyChangedEvent& PropertyChangedEvent) override;
-
 	// Begin EdGraphNode interface
-	FText GetNodeTitle(ENodeTitleType::Type TitleType) const override;
-	FLinearColor GetNodeTitleColor() const override;
-	FText GetTooltipText() const override;
+	virtual FText GetNodeTitle(ENodeTitleType::Type TitleType) const override;
+	virtual FLinearColor GetNodeTitleColor() const override;
+	virtual FText GetTooltipText() const override;
 
 	// UCustomizableObjectNode interface
-	void AllocateDefaultPins(UCustomizableObjectNodeRemapPins* RemapPins) override;
-	bool IsAffectedByLOD() const override { return false; }
+	virtual void AllocateDefaultPins(UCustomizableObjectNodeRemapPins* RemapPins) override;
+	virtual bool IsAffectedByLOD() const override { return false; }
+	virtual void BackwardsCompatibleFixup(int32 CustomizableObjectCustomVersion) override;
 
 	// Creates the SGraph Node widget for the thumbnail
-	TSharedPtr<SGraphNode> CreateVisualWidget() override;
+	virtual TSharedPtr<SGraphNode> CreateVisualWidget() override;
 
 	// Determines if the Node is collapsed or not
 	bool bCollapsed;
diff --git a/Engine/Plugins/Mutable/Source/CustomizableObjectEditor/Private/MuCOE/Nodes/CustomizableObjectNodeFloatParameter.cpp b/Engine/Plugins/Mutable/Source/CustomizableObjectEditor/Private/MuCOE/Nodes/CustomizableObjectNodeFloatParameter.cpp
--- a/Engine/Plugins/Mutable/Source/CustomizableObjectEditor/Private/MuCOE/Nodes/CustomizableObjectNodeFloatParameter.cpp
+++ b/Engine/Plugins/Mutable/Source/CustomizableObjectEditor/Private/MuCOE/Nodes/CustomizableObjectNodeFloatParameter.cpp
@@ -4,7 +4,6 @@
 
 #include "MuCO/CustomizableObjectCustomVersion.h"
 #include "MuCOE/EdGraphSchema_CustomizableObject.h"
-#include "MuCOE/UnrealEditorPortabilityHelpers.h"
 
 class UCustomizableObjectNodeRemapPins;
 
@@ -25,10 +24,13 @@ void UCustomizableObjectNodeFloatParameter::PostEditChangeProperty(FPropertyChan
 
 void UCustomizableObjectNodeFloatParameter::AllocateDefaultPins(UCustomizableObjectNodeRemapPins* RemapPins)
 {
-	const UEdGraphSchema_CustomizableObject* Schema = GetDefault<UEdGraphSchema_CustomizableObject>();
-
-	UEdGraphPin* ValuePin = CustomCreatePin(EGPD_Output, Schema->PC_Float, FName("Value"));
+	const FName Type = UEdGraphSchema_CustomizableObject::PC_Float;
+	const FName PinName = UEdGraphSchema_CustomizableObject::GetPinCategoryName(Type);
+	const FText PinFriendlyName = UEdGraphSchema_CustomizableObject::GetPinCategoryFriendlyName(Type);
+	
+	UEdGraphPin* ValuePin = CustomCreatePin(EGPD_Output, Type, PinName);
 	ValuePin->bDefaultValueIsIgnored = true;
+	ValuePin->PinFriendlyName = PinFriendlyName;
 }
 
 
@@ -89,6 +91,14 @@ void UCustomizableObjectNodeFloatParameter::BackwardsCompatibleFixup(int32 Custo
 	{
 		ReconstructNode();
 	}
+	else if (CustomizableObjectCustomVersion == FCustomizableObjectCustomVersion::UpdatedNodesPinName2)
+	{
+		if (UEdGraphPin* Pin = FindPin(TEXT("Value")))
+		{
+			Pin->PinName = TEXT("Float");
+			Pin->PinFriendlyName = LOCTEXT("Float_Pin_Category", "Float");
+		}
+	}
 }
 
 
diff --git a/Engine/Plugins/Mutable/Source/CustomizableObjectEditor/Private/MuCOE/Nodes/CustomizableObjectNodeFloatParameter.h b/Engine/Plugins/Mutable/Source/CustomizableObjectEditor/Private/MuCOE/Nodes/CustomizableObjectNodeFloatParameter.h
--- a/Engine/Plugins/Mutable/Source/CustomizableObjectEditor/Private/MuCOE/Nodes/CustomizableObjectNodeFloatParameter.h
+++ b/Engine/Plugins/Mutable/Source/CustomizableObjectEditor/Private/MuCOE/Nodes/CustomizableObjectNodeFloatParameter.h
@@ -42,18 +42,18 @@ public:
 	FMutableParamUIMetadata ParamUIMetadata;
 	
 	// UObject interface.
-	void PostEditChangeProperty(FPropertyChangedEvent& PropertyChangedEvent) override;
+	virtual void PostEditChangeProperty(FPropertyChangedEvent& PropertyChangedEvent) override;
 	
 	// Begin EdGraphNode interface
-	FText GetNodeTitle(ENodeTitleType::Type TitleType) const override;
-	FLinearColor GetNodeTitleColor() const override;
-	FText GetTooltipText() const override;
+	virtual FText GetNodeTitle(ENodeTitleType::Type TitleType) const override;
+	virtual FLinearColor GetNodeTitleColor() const override;
+	virtual FText GetTooltipText() const override;
 	virtual bool GetCanRenameNode() const override { return true; }
 
 	// UCustomizableObjectNode interface
-	void AllocateDefaultPins(UCustomizableObjectNodeRemapPins* RemapPins) override;
-	bool IsAffectedByLOD() const override;
-	void BackwardsCompatibleFixup(int32 CustomizableObjectCustomVersion) override;
+	virtual void AllocateDefaultPins(UCustomizableObjectNodeRemapPins* RemapPins) override;
+	virtual bool IsAffectedByLOD() const override;
+	virtual void BackwardsCompatibleFixup(int32 CustomizableObjectCustomVersion) override;
 	virtual void OnRenameNode(const FString& NewName) override;
 
 };
diff --git a/Engine/Plugins/Mutable/Source/CustomizableObjectEditor/Private/MuCOE/Nodes/CustomizableObjectNodeGroupProjectorParameter.cpp b/Engine/Plugins/Mutable/Source/CustomizableObjectEditor/Private/MuCOE/Nodes/CustomizableObjectNodeGroupProjectorParameter.cpp
--- a/Engine/Plugins/Mutable/Source/CustomizableObjectEditor/Private/MuCOE/Nodes/CustomizableObjectNodeGroupProjectorParameter.cpp
+++ b/Engine/Plugins/Mutable/Source/CustomizableObjectEditor/Private/MuCOE/Nodes/CustomizableObjectNodeGroupProjectorParameter.cpp
@@ -5,7 +5,6 @@
 #include "MuCO/CustomizableObjectCustomVersion.h"
 #include "MuCO/LoadUtils.h"
 #include "MuCOE/EdGraphSchema_CustomizableObject.h"
-#include "MuCOE/UnrealEditorPortabilityHelpers.h"
 
 
 class UCustomizableObjectNodeRemapPins;
@@ -122,30 +121,44 @@ void UCustomizableObjectNodeGroupProjectorParameter::BackwardsCompatibleFixup(in
 	
 	if (CustomizableObjectCustomVersion == FCustomizableObjectCustomVersion::GroupProjectorPinTypeAdded)
 	{
-		if (UEdGraphPin* Pin = ProjectorPin())
+		if (UEdGraphPin* Pin = FindPin(TEXT("Value")))
 		{
 			Pin->PinType.PinCategory = UEdGraphSchema_CustomizableObject::PC_GroupProjector;
 		}
 	}
-
-	if (CustomizableObjectCustomVersion == FCustomizableObjectCustomVersion::GroupProjectorImagePinRemoved)
+	else if (CustomizableObjectCustomVersion == FCustomizableObjectCustomVersion::GroupProjectorImagePinRemoved)
 	{
 		ReconstructNode();
 	}
+	else if (CustomizableObjectCustomVersion == FCustomizableObjectCustomVersion::UpdatedNodesPinName2)
+	{
+		// Looking for the projector pin as it was renamed in the same compatibility fix but of the parent
+		if (UEdGraphPin* Pin = FindPin(TEXT("Projector")))
+		{
+			Pin->PinName = TEXT("Group Projector");
+			Pin->PinFriendlyName = LOCTEXT("Group_Projector_Pin_Category", "Group Projector");
+		}
+	}
 }
 
 
 void UCustomizableObjectNodeGroupProjectorParameter::AllocateDefaultPins(UCustomizableObjectNodeRemapPins* RemapPins)
 {
-	const UEdGraphSchema_CustomizableObject* Schema = GetDefault<UEdGraphSchema_CustomizableObject>();
-
-	CustomCreatePin(EGPD_Output, Schema->PC_GroupProjector, TEXT("Value"));	
+	const FName Type = UEdGraphSchema_CustomizableObject::PC_GroupProjector;
+	const FName PinName = UEdGraphSchema_CustomizableObject::GetPinCategoryName(Type);
+	const FText PinFriendlyName = UEdGraphSchema_CustomizableObject::GetPinCategoryFriendlyName(Type);
+	
+	UEdGraphPin* ValuePin = CustomCreatePin(EGPD_Output, Type, PinName);
+	ValuePin->bDefaultValueIsIgnored = true;
+	ValuePin->PinFriendlyName = PinFriendlyName;
 }
 
 
 UEdGraphPin& UCustomizableObjectNodeGroupProjectorParameter::OutputPin() const
 {
-	return *FindPin(TEXT("Value"));
+	const FName Type = UEdGraphSchema_CustomizableObject::PC_GroupProjector;
+	const FName PinName = UEdGraphSchema_CustomizableObject::GetPinCategoryName(Type);
+	return *FindPin(PinName);
 }
 
 
diff --git a/Engine/Plugins/Mutable/Source/CustomizableObjectEditor/Private/MuCOE/Nodes/CustomizableObjectNodeGroupProjectorParameter.h b/Engine/Plugins/Mutable/Source/CustomizableObjectEditor/Private/MuCOE/Nodes/CustomizableObjectNodeGroupProjectorParameter.h
--- a/Engine/Plugins/Mutable/Source/CustomizableObjectEditor/Private/MuCOE/Nodes/CustomizableObjectNodeGroupProjectorParameter.h
+++ b/Engine/Plugins/Mutable/Source/CustomizableObjectEditor/Private/MuCOE/Nodes/CustomizableObjectNodeGroupProjectorParameter.h
@@ -103,7 +103,7 @@ public:
 
 	// UCustomizableObjectNode interface.
 	virtual void BackwardsCompatibleFixup(int32 CustomizableObjectCustomVersion) override;
-	virtual void AllocateDefaultPins(UCustomizableObjectNodeRemapPins* RemapPins);
+	virtual void AllocateDefaultPins(UCustomizableObjectNodeRemapPins* RemapPins) override;
 
 	// Own interface
 	UEdGraphPin& OutputPin() const;
diff --git a/Engine/Plugins/Mutable/Source/CustomizableObjectEditor/Private/MuCOE/Nodes/CustomizableObjectNodeProjectorConstant.cpp b/Engine/Plugins/Mutable/Source/CustomizableObjectEditor/Private/MuCOE/Nodes/CustomizableObjectNodeProjectorConstant.cpp
--- a/Engine/Plugins/Mutable/Source/CustomizableObjectEditor/Private/MuCOE/Nodes/CustomizableObjectNodeProjectorConstant.cpp
+++ b/Engine/Plugins/Mutable/Source/CustomizableObjectEditor/Private/MuCOE/Nodes/CustomizableObjectNodeProjectorConstant.cpp
@@ -4,7 +4,6 @@
 
 #include "MuCO/CustomizableObjectCustomVersion.h"
 #include "MuCOE/EdGraphSchema_CustomizableObject.h"
-#include "MuCOE/UnrealEditorPortabilityHelpers.h"
 
 class UCustomizableObjectNodeRemapPins;
 
@@ -43,10 +42,12 @@ void UCustomizableObjectNodeProjectorConstant::PostEditChangeProperty(FPropertyC
 
 void UCustomizableObjectNodeProjectorConstant::AllocateDefaultPins(UCustomizableObjectNodeRemapPins* RemapPins)
 {
-	const UEdGraphSchema_CustomizableObject* Schema = GetDefault<UEdGraphSchema_CustomizableObject>();
-
-	FString PinName = TEXT("Value");
-	UEdGraphPin* ValuePin = CustomCreatePin(EGPD_Output, Schema->PC_Projector, FName(*PinName));
+	const FName Type = UEdGraphSchema_CustomizableObject::PC_Projector;
+	const FName PinName = UEdGraphSchema_CustomizableObject::GetPinCategoryName(Type);
+	const FText PinFriendlyName = UEdGraphSchema_CustomizableObject::GetPinCategoryFriendlyName(Type);
+	
+	UEdGraphPin* ValuePin = CustomCreatePin(EGPD_Output, Type, PinName);
+	ValuePin->PinFriendlyName = PinFriendlyName;
 	ValuePin->bDefaultValueIsIgnored = true;
 }
 
@@ -59,11 +60,18 @@ void UCustomizableObjectNodeProjectorConstant::BackwardsCompatibleFixup(int32 Cu
 	{
 		Value.ProjectionType = ProjectionType_DEPRECATED;
 	}
-
-	if (CustomizableObjectCustomVersion == FCustomizableObjectCustomVersion::SnapToBoneComponentIndexToName)
+	else if (CustomizableObjectCustomVersion == FCustomizableObjectCustomVersion::SnapToBoneComponentIndexToName)
 	{
 		ReferenceSkeletonComponent = FName(FString::FromInt(ReferenceSkeletonIndex_DEPRECATED));
 	}
+	else if (CustomizableObjectCustomVersion == FCustomizableObjectCustomVersion::UpdatedNodesPinName2)
+	{
+		if (UEdGraphPin* Pin = FindPin(TEXT("Value")))
+		{
+			Pin->PinName = TEXT("Projector");
+			Pin->PinFriendlyName = LOCTEXT("Projector_Pin_Category", "Projector");
+		}
+	}
 }
 
 
@@ -121,7 +129,7 @@ void UCustomizableObjectNodeProjectorConstant::SetProjectorScale(const FVector&
 }
 
 
-float UCustomizableObjectNodeProjectorConstant::GetProjectorAngle()
+float UCustomizableObjectNodeProjectorConstant::GetProjectorAngle() const
 {
 	return ProjectionAngle;
 }
diff --git a/Engine/Plugins/Mutable/Source/CustomizableObjectEditor/Private/MuCOE/Nodes/CustomizableObjectNodeProjectorConstant.h b/Engine/Plugins/Mutable/Source/CustomizableObjectEditor/Private/MuCOE/Nodes/CustomizableObjectNodeProjectorConstant.h
--- a/Engine/Plugins/Mutable/Source/CustomizableObjectEditor/Private/MuCOE/Nodes/CustomizableObjectNodeProjectorConstant.h
+++ b/Engine/Plugins/Mutable/Source/CustomizableObjectEditor/Private/MuCOE/Nodes/CustomizableObjectNodeProjectorConstant.h
@@ -63,7 +63,6 @@ public:
 	virtual void BackwardsCompatibleFixup(int32 CustomizableObjectCustomVersion) override;
 
 	// Own interface
-	UEdGraphPin* ValuePin();
 
 	ECustomizableObjectProjectorType GetProjectorType() const;
 
@@ -83,7 +82,7 @@ public:
 
 	void SetProjectorScale(const FVector& Scale);
 
-	float GetProjectorAngle();
+	float GetProjectorAngle() const;
 
 	void SetProjectorAngle(float Angle);
 
diff --git a/Engine/Plugins/Mutable/Source/CustomizableObjectEditor/Private/MuCOE/Nodes/CustomizableObjectNodeProjectorParameter.cpp b/Engine/Plugins/Mutable/Source/CustomizableObjectEditor/Private/MuCOE/Nodes/CustomizableObjectNodeProjectorParameter.cpp
--- a/Engine/Plugins/Mutable/Source/CustomizableObjectEditor/Private/MuCOE/Nodes/CustomizableObjectNodeProjectorParameter.cpp
+++ b/Engine/Plugins/Mutable/Source/CustomizableObjectEditor/Private/MuCOE/Nodes/CustomizableObjectNodeProjectorParameter.cpp
@@ -4,7 +4,6 @@
 
 #include "MuCO/CustomizableObjectCustomVersion.h"
 #include "MuCOE/EdGraphSchema_CustomizableObject.h"
-#include "MuCOE/UnrealEditorPortabilityHelpers.h"
 
 class UCustomizableObjectNodeRemapPins;
 
@@ -32,11 +31,13 @@ void UCustomizableObjectNodeProjectorParameter::PostEditChangeProperty(FProperty
 
 void UCustomizableObjectNodeProjectorParameter::AllocateDefaultPins(UCustomizableObjectNodeRemapPins* RemapPins)
 {
-	const UEdGraphSchema_CustomizableObject* Schema = GetDefault<UEdGraphSchema_CustomizableObject>();
-
-	FString PinName = TEXT("Value");
-	UEdGraphPin* ValuePin = CustomCreatePin(EGPD_Output, Schema->PC_Projector, FName(*PinName));
+	const FName Type = UEdGraphSchema_CustomizableObject::PC_Projector;
+	const FName PinName = UEdGraphSchema_CustomizableObject::GetPinCategoryName(Type);
+	const FText PinFriendlyName = UEdGraphSchema_CustomizableObject::GetPinCategoryFriendlyName(Type);
+	
+	UEdGraphPin* ValuePin = CustomCreatePin(EGPD_Output, Type, PinName);
 	ValuePin->bDefaultValueIsIgnored = true;
+	ValuePin->PinFriendlyName = PinFriendlyName;
 }
 
 
@@ -48,11 +49,18 @@ void UCustomizableObjectNodeProjectorParameter::BackwardsCompatibleFixup(int32 C
 	{
 		DefaultValue.ProjectionType = ProjectionType_DEPRECATED;
 	}
-
-	if (CustomizableObjectCustomVersion == FCustomizableObjectCustomVersion::SnapToBoneComponentIndexToName)
+	else if (CustomizableObjectCustomVersion == FCustomizableObjectCustomVersion::SnapToBoneComponentIndexToName)
 	{
 		ReferenceSkeletonComponent = FName(FString::FromInt(ReferenceSkeletonIndex_DEPRECATED));
 	}
+	else if (CustomizableObjectCustomVersion == FCustomizableObjectCustomVersion::UpdatedNodesPinName2)
+	{
+		if (UEdGraphPin* Pin = FindPin(TEXT("Value")))
+		{
+			Pin->PinName = TEXT("Projector");
+			Pin->PinFriendlyName = LOCTEXT("Projector_Pin_Category", "Projector");
+		}
+	}
 }
 
 
diff --git a/Engine/Plugins/Mutable/Source/CustomizableObjectEditor/Private/MuCOE/Nodes/CustomizableObjectNodeProjectorParameter.h b/Engine/Plugins/Mutable/Source/CustomizableObjectEditor/Private/MuCOE/Nodes/CustomizableObjectNodeProjectorParameter.h
--- a/Engine/Plugins/Mutable/Source/CustomizableObjectEditor/Private/MuCOE/Nodes/CustomizableObjectNodeProjectorParameter.h
+++ b/Engine/Plugins/Mutable/Source/CustomizableObjectEditor/Private/MuCOE/Nodes/CustomizableObjectNodeProjectorParameter.h
@@ -66,10 +66,6 @@ public:
 	virtual void BackwardsCompatibleFixup(int32 CustomizableObjectCustomVersion) override;
 
 	// Own interface
-	UEdGraphPin* ProjectorPin() const
-	{
-		return FindPin(TEXT("Value"));
-	}
 
 	ECustomizableObjectProjectorType GetProjectorType() const;
 
diff --git a/Engine/Plugins/Mutable/Source/CustomizableObjectEditor/Private/MuCOE/Nodes/CustomizableObjectNodeTextureParameter.cpp b/Engine/Plugins/Mutable/Source/CustomizableObjectEditor/Private/MuCOE/Nodes/CustomizableObjectNodeTextureParameter.cpp
--- a/Engine/Plugins/Mutable/Source/CustomizableObjectEditor/Private/MuCOE/Nodes/CustomizableObjectNodeTextureParameter.cpp
+++ b/Engine/Plugins/Mutable/Source/CustomizableObjectEditor/Private/MuCOE/Nodes/CustomizableObjectNodeTextureParameter.cpp
@@ -12,11 +12,13 @@ class UCustomizableObjectNodeRemapPins;
 
 void UCustomizableObjectNodeTextureParameter::AllocateDefaultPins(UCustomizableObjectNodeRemapPins* RemapPins)
 {
-	const UEdGraphSchema_CustomizableObject* Schema = GetDefault<UEdGraphSchema_CustomizableObject>();
-
-	FString PinName = TEXT("Value");
-	UEdGraphPin* ValuePin = CustomCreatePin(EGPD_Output, Schema->PC_Image, FName(*PinName));
+	const FName Type = UEdGraphSchema_CustomizableObject::PC_Image;
+	const FName PinName = UEdGraphSchema_CustomizableObject::GetPinCategoryName(Type);
+	const FText PinFriendlyName = UEdGraphSchema_CustomizableObject::GetPinCategoryFriendlyName(Type);
+	
+	UEdGraphPin* ValuePin = CustomCreatePin(EGPD_Output, Type, PinName);
 	ValuePin->bDefaultValueIsIgnored = true;
+	ValuePin->PinFriendlyName = PinFriendlyName;
 }
 
 
@@ -74,6 +76,14 @@ void UCustomizableObjectNodeTextureParameter::BackwardsCompatibleFixup(int32 Cus
 		ReferenceValue = DefaultValue;
 		DefaultValue = {};
 	}
+	else if (CustomizableObjectCustomVersion == FCustomizableObjectCustomVersion::UpdatedNodesPinName2)
+	{
+		if (UEdGraphPin* Pin = FindPin(TEXT("Value")))
+		{
+			Pin->PinName = TEXT("Texture");
+			Pin->PinFriendlyName = LOCTEXT("Image_Pin_Category", "Texture");
+		}
+	}
 }
 
 #undef LOCTEXT_NAMESPACE
diff --git a/Engine/Plugins/Mutable/Source/CustomizableObjectEditor/Private/MuCOE/Nodes/CustomizableObjectNodeTransformConstant.cpp b/Engine/Plugins/Mutable/Source/CustomizableObjectEditor/Private/MuCOE/Nodes/CustomizableObjectNodeTransformConstant.cpp
--- a/Engine/Plugins/Mutable/Source/CustomizableObjectEditor/Private/MuCOE/Nodes/CustomizableObjectNodeTransformConstant.cpp
+++ b/Engine/Plugins/Mutable/Source/CustomizableObjectEditor/Private/MuCOE/Nodes/CustomizableObjectNodeTransformConstant.cpp
@@ -2,9 +2,8 @@
 
 #include "MuCOE/Nodes/CustomizableObjectNodeTransformConstant.h"
 
-#include "Components/HorizontalBox.h"
+#include "MuCO/CustomizableObjectCustomVersion.h"
 #include "MuCOE/CustomizableObjectEditorStyle.h"
-#include "MuCOE/CustomizableObjectGraph.h"
 #include "MuCOE/EdGraphSchema_CustomizableObject.h"
 #include "Widgets/Images/SImage.h"
 #include "Widgets/Input/SCheckBox.h"
@@ -23,6 +22,7 @@ void SGraphNodeTransformConstant::Construct(const FArguments& InArgs, UEdGraphNo
 	SCustomizableObjectNode::Construct({}, InGraphNode);
 }
 
+
 void SGraphNodeTransformConstant::SetDefaultTitleAreaWidget(TSharedRef<SOverlay> DefaultTitleAreaWidget)
 {
 	// Collapsing arrow of the title area
@@ -49,6 +49,7 @@ void SGraphNodeTransformConstant::SetDefaultTitleAreaWidget(TSharedRef<SOverlay>
 	];
 }
 
+
 void SGraphNodeTransformConstant::CreateBelowPinControls(TSharedPtr<SVerticalBox> MainBox)
 {
 	constexpr float TransformBoxMinWidth = 250.0f;
@@ -143,22 +144,26 @@ void SGraphNodeTransformConstant::CreateBelowPinControls(TSharedPtr<SVerticalBox
 	];
 }
 
+
 void SGraphNodeTransformConstant::OnExpressionPreviewChanged(const ECheckBoxState NewCheckedState)
 {
 	NodeTransformConstant->bCollapsed = (NewCheckedState != ECheckBoxState::Checked);
 	UpdateGraphNode();
 }
 
+
 ECheckBoxState SGraphNodeTransformConstant::IsExpressionPreviewChecked() const
 {
 	return NodeTransformConstant->bCollapsed ? ECheckBoxState::Unchecked : ECheckBoxState::Checked;
 }
 
+
 const FSlateBrush* SGraphNodeTransformConstant::GetExpressionPreviewArrow() const
 {
 	return FCustomizableObjectEditorStyle::Get().GetBrush(NodeTransformConstant->bCollapsed ? TEXT("Nodes.ArrowDown") : TEXT("Nodes.ArrowUp"));
 }
 
+
 void SGraphNodeTransformConstant::OnLocationChanged(double Value, ETextCommit::Type, EAxis::Type Axis)
 {
 	FVector Location = NodeTransformConstant->Value.GetLocation();
@@ -171,6 +176,7 @@ void SGraphNodeTransformConstant::OnLocationChanged(double Value, ETextCommit::T
 	}
 }
 
+
 void SGraphNodeTransformConstant::OnRotationChanged(double Value, ETextCommit::Type, EAxis::Type Axis)
 {
 	FRotator Rotation = NodeTransformConstant->Value.Rotator();
@@ -184,6 +190,7 @@ void SGraphNodeTransformConstant::OnRotationChanged(double Value, ETextCommit::T
 	}
 }
 
+
 void SGraphNodeTransformConstant::OnScaleChanged(double Value, ETextCommit::Type, EAxis::Type Axis)
 {
 	FVector Scale = NodeTransformConstant->Value.GetScale3D();
@@ -206,25 +213,47 @@ FText UCustomizableObjectNodeTransformConstant::GetNodeTitle(ENodeTitleType::Typ
 	return LOCTEXT("Transform_Constant", "Transform Constant");
 }
 
+
 FLinearColor UCustomizableObjectNodeTransformConstant::GetNodeTitleColor() const
 {
 	const UEdGraphSchema_CustomizableObject* Schema = GetDefault<UEdGraphSchema_CustomizableObject>();
 	return Schema->GetPinTypeColor(Schema->PC_Transform);
 }
 
+
 FText UCustomizableObjectNodeTransformConstant::GetTooltipText() const
 {
 	return LOCTEXT("Transform_Constant_Tooltip", "Define a constant transform value.");
 }
 
+
 void UCustomizableObjectNodeTransformConstant::AllocateDefaultPins(UCustomizableObjectNodeRemapPins* RemapPins)
 {
-	const UEdGraphSchema_CustomizableObject* Schema = GetDefault<UEdGraphSchema_CustomizableObject>();
-
-	UEdGraphPin* ValuePin = CustomCreatePin(EGPD_Output, Schema->PC_Transform, FName("Value"));
+	const FName Type = UEdGraphSchema_CustomizableObject::PC_Transform;
+	const FName PinName = UEdGraphSchema_CustomizableObject::GetPinCategoryName(Type);
+	const FText PinFriendlyName = UEdGraphSchema_CustomizableObject::GetPinCategoryFriendlyName(Type);
+	
+	UEdGraphPin* ValuePin = CustomCreatePin(EGPD_Output, Type, PinName);
+	ValuePin->PinFriendlyName = PinFriendlyName;
 	ValuePin->bDefaultValueIsIgnored = true;
 }
 
+
+void UCustomizableObjectNodeTransformConstant::BackwardsCompatibleFixup(int32 CustomizableObjectCustomVersion)
+{
+	Super::BackwardsCompatibleFixup(CustomizableObjectCustomVersion);
+
+	if (CustomizableObjectCustomVersion == FCustomizableObjectCustomVersion::UpdatedNodesPinName2)
+	{
+		if (UEdGraphPin* Pin = FindPin(TEXT("Value")))
+		{
+			Pin->PinName = TEXT("Transform");
+			Pin->PinFriendlyName =  LOCTEXT("Transform_Pin_Category", "Transform");
+		}
+	}
+}
+
+
 TSharedPtr<SGraphNode> UCustomizableObjectNodeTransformConstant::CreateVisualWidget()
 {
 	return SNew(SGraphNodeTransformConstant, this);
diff --git a/Engine/Plugins/Mutable/Source/CustomizableObjectEditor/Private/MuCOE/Nodes/CustomizableObjectNodeTransformConstant.h b/Engine/Plugins/Mutable/Source/CustomizableObjectEditor/Private/MuCOE/Nodes/CustomizableObjectNodeTransformConstant.h
--- a/Engine/Plugins/Mutable/Source/CustomizableObjectEditor/Private/MuCOE/Nodes/CustomizableObjectNodeTransformConstant.h
+++ b/Engine/Plugins/Mutable/Source/CustomizableObjectEditor/Private/MuCOE/Nodes/CustomizableObjectNodeTransformConstant.h
@@ -63,17 +63,18 @@ public:
 	FTransform Value = FTransform::Identity;
 
 	// Begin EdGraphNode interface
-	FText GetNodeTitle(ENodeTitleType::Type TitleType) const override;
-	FLinearColor GetNodeTitleColor() const override;
-	FText GetTooltipText() const override;
+	virtual FText GetNodeTitle(ENodeTitleType::Type TitleType) const override;
+	virtual FLinearColor GetNodeTitleColor() const override;
+	virtual FText GetTooltipText() const override;
 	// End EdGraphNode interface
 
 	// UCustomizableObjectNode interface
-	void AllocateDefaultPins(UCustomizableObjectNodeRemapPins* RemapPins) override;
-	bool IsAffectedByLOD() const override { return false; }
+	virtual void AllocateDefaultPins(UCustomizableObjectNodeRemapPins* RemapPins) override;
+	virtual bool IsAffectedByLOD() const override { return false; }
+	virtual void BackwardsCompatibleFixup(int32 CustomizableObjectCustomVersion) override;
 
 	// Creates the SGraph Node widget for the Color Editor
-	TSharedPtr<SGraphNode> CreateVisualWidget() override;
+	virtual TSharedPtr<SGraphNode> CreateVisualWidget() override;
 
 	// Determines if the Node is collapsed or not
 	bool bCollapsed = true;
diff --git a/Engine/Plugins/Mutable/Source/CustomizableObjectEditor/Private/MuCOE/Nodes/CustomizableObjectNodeTransformParameter.cpp b/Engine/Plugins/Mutable/Source/CustomizableObjectEditor/Private/MuCOE/Nodes/CustomizableObjectNodeTransformParameter.cpp
--- a/Engine/Plugins/Mutable/Source/CustomizableObjectEditor/Private/MuCOE/Nodes/CustomizableObjectNodeTransformParameter.cpp
+++ b/Engine/Plugins/Mutable/Source/CustomizableObjectEditor/Private/MuCOE/Nodes/CustomizableObjectNodeTransformParameter.cpp
@@ -2,6 +2,7 @@
 
 #include "MuCOE/Nodes/CustomizableObjectNodeTransformParameter.h"
 
+#include "MuCO/CustomizableObjectCustomVersion.h"
 #include "MuCOE/EdGraphSchema_CustomizableObject.h"
 
 
@@ -24,17 +25,20 @@ FText UCustomizableObjectNodeTransformParameter::GetNodeTitle(ENodeTitleType::Ty
 	}
 }
 
+
 FLinearColor UCustomizableObjectNodeTransformParameter::GetNodeTitleColor() const
 {
 	const UEdGraphSchema_CustomizableObject* Schema = GetDefault<UEdGraphSchema_CustomizableObject>();
 	return Schema->GetPinTypeColor(Schema->PC_Transform);
 }
 
+
 FText UCustomizableObjectNodeTransformParameter::GetTooltipText() const
 {
 	return LOCTEXT("Transform_Parameter_Tooltip", "Expose a runtime modifiable transform parameter from the Customizable Object.");
 }
 
+
 void UCustomizableObjectNodeTransformParameter::OnRenameNode(const FString& NewName)
 {
 	if (!NewName.IsEmpty())
@@ -43,12 +47,31 @@ void UCustomizableObjectNodeTransformParameter::OnRenameNode(const FString& NewN
 	}
 }
 
-void UCustomizableObjectNodeTransformParameter::AllocateDefaultPins(UCustomizableObjectNodeRemapPins* RemapPins)
+
+void UCustomizableObjectNodeTransformParameter::BackwardsCompatibleFixup(int32 CustomizableObjectCustomVersion)
 {
-	const UEdGraphSchema_CustomizableObject* Schema = GetDefault<UEdGraphSchema_CustomizableObject>();
+	Super::BackwardsCompatibleFixup(CustomizableObjectCustomVersion);
 
-	UEdGraphPin* ValuePin = CustomCreatePin(EGPD_Output, Schema->PC_Transform, FName("Value"));
+	if (CustomizableObjectCustomVersion == FCustomizableObjectCustomVersion::UpdatedNodesPinName2)
+	{
+		if (UEdGraphPin* Pin = FindPin(TEXT("Value")))
+		{
+			Pin->PinName = TEXT("Transform");
+			Pin->PinFriendlyName = LOCTEXT("Transform_Pin_Category", "Transform");
+		}
+	}
+}
+
+
+void UCustomizableObjectNodeTransformParameter::AllocateDefaultPins(UCustomizableObjectNodeRemapPins* RemapPins)
+{
+	const FName Type = UEdGraphSchema_CustomizableObject::PC_Transform;
+	const FName PinName = UEdGraphSchema_CustomizableObject::GetPinCategoryName(Type);
+	const FText PinFriendlyName = UEdGraphSchema_CustomizableObject::GetPinCategoryFriendlyName(Type);
+	
+	UEdGraphPin* ValuePin = CustomCreatePin(EGPD_Output, Type, PinName);
 	ValuePin->bDefaultValueIsIgnored = true;
+	ValuePin->PinFriendlyName = PinFriendlyName;
 }
 
 #undef LOCTEXT_NAMESPACE
diff --git a/Engine/Plugins/Mutable/Source/CustomizableObjectEditor/Private/MuCOE/Nodes/CustomizableObjectNodeTransformParameter.h b/Engine/Plugins/Mutable/Source/CustomizableObjectEditor/Private/MuCOE/Nodes/CustomizableObjectNodeTransformParameter.h
--- a/Engine/Plugins/Mutable/Source/CustomizableObjectEditor/Private/MuCOE/Nodes/CustomizableObjectNodeTransformParameter.h
+++ b/Engine/Plugins/Mutable/Source/CustomizableObjectEditor/Private/MuCOE/Nodes/CustomizableObjectNodeTransformParameter.h
@@ -25,14 +25,15 @@ public:
 	FMutableParamUIMetadata ParamUIMetadata;
 
 	// Begin EdGraphNode interface
-	FText GetNodeTitle(ENodeTitleType::Type TitleType) const override;
-	FLinearColor GetNodeTitleColor() const override;
-	FText GetTooltipText() const override;
-	void OnRenameNode(const FString& NewName) override;
-	bool GetCanRenameNode() const override { return true; }
+	virtual FText GetNodeTitle(ENodeTitleType::Type TitleType) const override;
+	virtual FLinearColor GetNodeTitleColor() const override;
+	virtual FText GetTooltipText() const override;
+	virtual void OnRenameNode(const FString& NewName) override;
+	virtual bool GetCanRenameNode() const override { return true; }
+	virtual void BackwardsCompatibleFixup(int32 CustomizableObjectCustomVersion) override;
 
 	// UCustomizableObjectNode interface
-	void AllocateDefaultPins(UCustomizableObjectNodeRemapPins* RemapPins) override;
-	bool IsAffectedByLOD() const override { return false; }
+	virtual void AllocateDefaultPins(UCustomizableObjectNodeRemapPins* RemapPins) override;
+	virtual bool IsAffectedByLOD() const override { return false; }
 	
 };
