[Mutable] Created CustomizableObjectMacroLibrary asset and editor.
+ Created the new asset to allow mutable macros.
+ Created new editor to edit mutable macros.
+ Created nodes to instantiate mutable macros into mutable graphs.

#rnx

diff --git a/Engine/Plugins/Mutable/Source/CustomizableObject/Internal/MuCO/ICustomizableObjectEditorModule.h b/Engine/Plugins/Mutable/Source/CustomizableObject/Internal/MuCO/ICustomizableObjectEditorModule.h
--- a/Engine/Plugins/Mutable/Source/CustomizableObject/Internal/MuCO/ICustomizableObjectEditorModule.h
+++ b/Engine/Plugins/Mutable/Source/CustomizableObject/Internal/MuCO/ICustomizableObjectEditorModule.h
@@ -27,6 +27,7 @@ extern const FName CustomizableObjectInstanceEditorAppIdentifier;
 extern const FName CustomizableObjectPopulationEditorAppIdentifier;
 extern const FName CustomizableObjectPopulationClassEditorAppIdentifier;
 extern const FName CustomizableObjectDebuggerAppIdentifier;
+extern const FName CustomizableObjectMacroLibraryEditorAppIdentifier;
 
 /**
  * Customizable object editor module interface
diff --git a/Engine/Plugins/Mutable/Source/CustomizableObjectEditor/Private/MuCOE/CustomizableObjectEditor.cpp b/Engine/Plugins/Mutable/Source/CustomizableObjectEditor/Private/MuCOE/CustomizableObjectEditor.cpp
--- a/Engine/Plugins/Mutable/Source/CustomizableObjectEditor/Private/MuCOE/CustomizableObjectEditor.cpp
+++ b/Engine/Plugins/Mutable/Source/CustomizableObjectEditor/Private/MuCOE/CustomizableObjectEditor.cpp
@@ -444,9 +444,11 @@ void FCustomizableObjectEditor::CreateGraphEditorWidget(UEdGraph* InGraph)
 	FGraphAppearanceInfo AppearanceInfo;	
 	AppearanceInfo.CornerText = LOCTEXT("ApperanceCornerText", "MUTABLE");
 
+	// TODO(Max) Each editor has to bind its own events
 	SGraphEditor::FGraphEditorEvents InEvents;
 	InEvents.OnSelectionChanged = SGraphEditor::FOnSelectionChanged::CreateSP( this, &FCustomizableObjectEditor::OnSelectedGraphNodesChanged );
 	InEvents.OnTextCommitted = FOnNodeTextCommitted::CreateSP(this, &FCustomizableObjectEditor::OnNodeTitleCommitted);
+	InEvents.OnNodeDoubleClicked = FSingleNodeEvent::CreateSP(this, &FCustomizableObjectEditor::OnNodeDoubleClicked);// TODO(Max): Remove this once this editor starts using the new toolkit
 
 	// Make full graph editor
 	GraphEditor = SNew(SGraphEditor)
@@ -1693,6 +1695,15 @@ void FCustomizableObjectEditor::OnSelectedGraphNodesChanged(const FGraphPanelSel
 }
 
 
+void FCustomizableObjectEditor::OnNodeDoubleClicked(UEdGraphNode* Node) //TODO(Max): Remove this once this editor uses the new toolkit
+{
+	if (Node->CanJumpToDefinition())
+	{
+		Node->JumpToDefinition();
+	}
+}
+
+
 void FCustomizableObjectEditor::NotifyPostChange( const FPropertyChangedEvent& PropertyChangedEvent, FProperty* PropertyThatChanged  )
 {
 	// Is it a source graph node?
diff --git a/Engine/Plugins/Mutable/Source/CustomizableObjectEditor/Private/MuCOE/CustomizableObjectEditor.h b/Engine/Plugins/Mutable/Source/CustomizableObjectEditor/Private/MuCOE/CustomizableObjectEditor.h
--- a/Engine/Plugins/Mutable/Source/CustomizableObjectEditor/Private/MuCOE/CustomizableObjectEditor.h
+++ b/Engine/Plugins/Mutable/Source/CustomizableObjectEditor/Private/MuCOE/CustomizableObjectEditor.h
@@ -172,6 +172,7 @@ public:
 	void DuplicateSelectedNodes();
 	bool CanDuplicateSelectedNodes() const;
 	void OnSelectedGraphNodesChanged(const FGraphPanelSelectionSet& NewSelection);
+	void OnNodeDoubleClicked(UEdGraphNode* Node); //TODO(Max): Remove this once this editor uses the new toolkit
 
 	/**
 	 * Called when a node's title is committed for a rename
diff --git a/Engine/Plugins/Mutable/Source/CustomizableObjectEditor/Private/MuCOE/CustomizableObjectEditorModule.cpp b/Engine/Plugins/Mutable/Source/CustomizableObjectEditor/Private/MuCOE/CustomizableObjectEditorModule.cpp
--- a/Engine/Plugins/Mutable/Source/CustomizableObjectEditor/Private/MuCOE/CustomizableObjectEditorModule.cpp
+++ b/Engine/Plugins/Mutable/Source/CustomizableObjectEditor/Private/MuCOE/CustomizableObjectEditorModule.cpp
@@ -27,12 +27,15 @@
 #include "MuCOE/CustomizableObjectIdentifierCustomization.h"
 #include "MuCOE/CustomizableObjectInstanceEditor.h"
 #include "MuCOE/CustomizableObjectInstanceFactory.h"
+#include "MuCOE/CustomizableObjectMacroLibrary/CustomizableObjectMacroDetails.h"
 #include "MuCOE/CustomizableObjectVersionBridge.h"
 #include "MuCOE/CustomizableObjectNodeObjectGroupDetails.h"
 #include "MuCOE/Nodes/CustomizableObjectNodeCopyMaterial.h"
 #include "MuCOE/Nodes/CustomizableObjectNodeExternalPin.h"
 #include "MuCOE/Nodes/CustomizableObjectNodeExternalPinDetails.h"
 #include "MuCOE/Nodes/CustomizableObjectNodeLayoutBlocks.h"
+#include "MuCOE/Nodes/CustomizableObjectNodeMacroInstance.h"
+#include "MuCOE/Nodes/CustomizableObjectNodeMacroInstanceDetails.h"
 #include "MuCOE/Nodes/CustomizableObjectNodeMaterial.h"
 #include "MuCOE/Nodes/CustomizableObjectNodeMeshSectionDetails.h"
 #include "MuCOE/Nodes/CustomizableObjectNodeModifierEditMeshSection.h"
@@ -66,6 +69,7 @@
 #include "MuCOE/Nodes/CustomizableObjectNodeStaticMesh.h"
 #include "MuCOE/Nodes/CustomizableObjectNodeTable.h"
 #include "MuCOE/Nodes/CustomizableObjectNodeTableDetails.h"
+#include "MuCOE/Nodes/CustomizableObjectNodeTunnel.h"
 #include "MuCOE/Nodes/CustomizableObjectNodeComponentMeshDetails.h"
 #include "MuCOE/Nodes/CustomizableObjectNodeVariation.h"
 #include "MuCOE/Nodes/CustomizableObjectNodeTexture.h"
@@ -106,6 +110,7 @@ class UObject;
 const FName CustomizableObjectEditorAppIdentifier = FName(TEXT("CustomizableObjectEditorApp"));
 const FName CustomizableObjectInstanceEditorAppIdentifier = FName(TEXT("CustomizableObjectInstanceEditorApp"));
 const FName CustomizableObjectDebuggerAppIdentifier = FName(TEXT("CustomizableObjectDebuggerApp"));
+const FName CustomizableObjectMacroLibraryEditorAppIdentifier = FName(TEXT("CustomizableObjectMacroLibraryApp"));
 
 #define LOCTEXT_NAMESPACE "MutableSettings"
 
@@ -226,11 +231,13 @@ void FCustomizableObjectEditorModule::StartupModule()
 	RegisterCustomDetails(PropertyModule, UCustomizableObjectNodeStaticMesh::StaticClass(), FOnGetDetailCustomizationInstance::CreateStatic(&FCustomizableObjectNodeDetails::MakeInstance));
 	RegisterCustomDetails(PropertyModule, UCustomizableObjectNodeTable::StaticClass(), FOnGetDetailCustomizationInstance::CreateStatic(&FCustomizableObjectNodeTableDetails::MakeInstance));
 	RegisterCustomDetails(PropertyModule, UCustomizableObjectNodeComponentMesh::StaticClass(), FOnGetDetailCustomizationInstance::CreateStatic(&FCustomizableObjectNodeComponentMeshDetails::MakeInstance));
+	RegisterCustomDetails(PropertyModule, UCustomizableObjectNodeMacroInstance::StaticClass(), FOnGetDetailCustomizationInstance::CreateStatic(&FCustomizableObjectNodeMacroInstanceDetails::MakeInstance));
 
 	// Other Objects
 	RegisterCustomDetails(PropertyModule, UCustomizableObject::StaticClass(), FOnGetDetailCustomizationInstance::CreateStatic(&FCustomizableObjectDetails::MakeInstance));
 	RegisterCustomDetails(PropertyModule, UCustomizableObjectInstance::StaticClass(), FOnGetDetailCustomizationInstance::CreateStatic(&FCustomizableInstanceDetails::MakeInstance));
 	RegisterCustomDetails(PropertyModule, UCustomSettings::StaticClass(), FOnGetDetailCustomizationInstance::CreateStatic(&FCustomizableObjectCustomSettingsDetails::MakeInstance));
+	RegisterCustomDetails(PropertyModule, UCustomizableObjectMacro::StaticClass(), FOnGetDetailCustomizationInstance::CreateStatic(&FCustomizableObjectMacroDetails::MakeInstance));
 
 	// We need to cache those two FNames as if we want to get them on "ShutdownModule" we get "None" FNames and an ASAN error in Linux.
 	MeshReshapeBoneReferenceUStructName = FMeshReshapeBoneReference::StaticStruct()->GetFName();
diff --git a/Engine/Plugins/Mutable/Source/CustomizableObjectEditor/Private/MuCOE/CustomizableObjectGraph.cpp b/Engine/Plugins/Mutable/Source/CustomizableObjectEditor/Private/MuCOE/CustomizableObjectGraph.cpp
--- a/Engine/Plugins/Mutable/Source/CustomizableObjectEditor/Private/MuCOE/CustomizableObjectGraph.cpp
+++ b/Engine/Plugins/Mutable/Source/CustomizableObjectEditor/Private/MuCOE/CustomizableObjectGraph.cpp
@@ -6,6 +6,7 @@
 #include "MuCO/CustomizableObjectCustomVersion.h"
 #include "MuCOE/EdGraphSchema_CustomizableObject.h"
 #include "MuCOE/Nodes/CustomizableObjectNodeObject.h"
+#include "MuCOE/Nodes/CustomizableObjectNodeTunnel.h"
 
 class UObject;
 
@@ -177,23 +178,53 @@ void UCustomizableObjectGraph::PostDuplicate(bool bDuplicateForPIE)
 
 void UCustomizableObjectGraph::AddEssentialGraphNodes()
 {
-	// Check whether the graph has a base node and create one if it doesn't since it is required
-	bool bGraphHasBase = false;
-
-	for (const TObjectPtr<UEdGraphNode>& AuxNode : Nodes)
+	if (!IsMacro())
 	{
-		UCustomizableObjectNodeObject* CustomizableObjectNodeObject = Cast<UCustomizableObjectNodeObject>(AuxNode);
+		// Check whether the graph has a base node and create one if it doesn't since it is required
+		bool bGraphHasBase = false;
 
-		if (CustomizableObjectNodeObject && CustomizableObjectNodeObject->bIsBase)
+		for (const TObjectPtr<UEdGraphNode>& AuxNode : Nodes)
 		{
-			bGraphHasBase = true;
-			break;
+			UCustomizableObjectNodeObject* CustomizableObjectNodeObject = Cast<UCustomizableObjectNodeObject>(AuxNode);
+
+			if (CustomizableObjectNodeObject && CustomizableObjectNodeObject->bIsBase)
+			{
+				bGraphHasBase = true;
+				break;
+			}
 		}
-	}
 
-	if (!bGraphHasBase)
+		if (!bGraphHasBase)
+		{
+			UCustomizableObjectNodeObject* NodeTemplate = NewObject<UCustomizableObjectNodeObject>();
+			FCustomizableObjectSchemaAction_NewNode::CreateNode(this, nullptr, FVector2D::ZeroVector, Cast<UEdGraphNode>(NodeTemplate));
+		}
+	}
+	else
 	{
-		UCustomizableObjectNodeObject* NodeTemplate = NewObject<UCustomizableObjectNodeObject>();
-		FCustomizableObjectSchemaAction_NewNode::CreateNode(this, nullptr, FVector2D::ZeroVector, Cast<UEdGraphNode>(NodeTemplate));
+		UCustomizableObjectMacro* ParentMacro = Cast<UCustomizableObjectMacro>(GetOuter());
+		check(ParentMacro);
+
+		UCustomizableObjectNodeTunnel* InputNode = NewObject<UCustomizableObjectNodeTunnel>();
+		InputNode->bIsInputNode = true;
+		InputNode->ParentMacro = ParentMacro;
+
+		UCustomizableObjectNodeTunnel* OutputNode = NewObject<UCustomizableObjectNodeTunnel>();
+		OutputNode->bIsInputNode = false;
+		OutputNode->ParentMacro = ParentMacro;
+
+		FVector2D InputNodePos = FVector2D(-100.0f, 0.0f);
+		FVector2D OutputNodePos = FVector2D(100.0f, 0.0f);
+		FCustomizableObjectSchemaAction_NewNode::CreateNode(this, nullptr, InputNodePos, Cast<UEdGraphNode>(InputNode));
+		FCustomizableObjectSchemaAction_NewNode::CreateNode(this, nullptr, OutputNodePos, Cast<UEdGraphNode>(OutputNode));
 	}
 }
+
+
+bool UCustomizableObjectGraph::IsMacro() const
+{
+	UObject* Outer = GetOuter();
+	check(Outer);
+
+	return Cast<UCustomizableObjectMacro>(Outer) != nullptr;
+}
diff --git a/Engine/Plugins/Mutable/Source/CustomizableObjectEditor/Private/MuCOE/CustomizableObjectGraph.h b/Engine/Plugins/Mutable/Source/CustomizableObjectEditor/Private/MuCOE/CustomizableObjectGraph.h
--- a/Engine/Plugins/Mutable/Source/CustomizableObjectEditor/Private/MuCOE/CustomizableObjectGraph.h
+++ b/Engine/Plugins/Mutable/Source/CustomizableObjectEditor/Private/MuCOE/CustomizableObjectGraph.h
@@ -8,7 +8,6 @@
 
 class UObject;
 
-
 UCLASS()
 class CUSTOMIZABLEOBJECTEDITOR_API UCustomizableObjectGraph : public UEdGraph
 {
@@ -34,7 +33,9 @@ public:
 	void BackwardsCompatibleFixup(int32 CustomizableObjectCustomVersion);
 
 	void PostBackwardsCompatibleFixup();
-	
+
+	bool IsMacro() const;
+
 private:
 
 	// Request Node Id Update Map
diff --git a/Engine/Plugins/Mutable/Source/CustomizableObjectEditor/Private/MuCOE/CustomizableObjectGraphEditorToolkit.cpp b/Engine/Plugins/Mutable/Source/CustomizableObjectEditor/Private/MuCOE/CustomizableObjectGraphEditorToolkit.cpp
new file mode 100644
--- /dev/null
+++ b/Engine/Plugins/Mutable/Source/CustomizableObjectEditor/Private/MuCOE/CustomizableObjectGraphEditorToolkit.cpp
@@ -0,0 +1,635 @@
+// Copyright Epic Games, Inc. All Rights Reserved.
+
+#include "MuCOE/CustomizableObjectGraphEditorToolkit.h"
+
+#include "Framework/Commands/GenericCommands.h"
+#include "Framework/Commands/UIAction.h"
+#include "HAL/PlatformApplicationMisc.h"
+#include "MaterialGraph/MaterialGraphNode.h"
+#include "MaterialGraph/MaterialGraphNode_Comment.h"
+#include "MuCOE/CustomizableObjectEditorNodeContextCommands.h"
+#include "MuCOE/CustomizableObjectEditorLogger.h"
+#include "MuCOE/CustomizableObjectGraph.h"
+#include "MuCOE/Nodes/CustomizableObjectNode.h"
+#include "SAssetSearchBox.h"
+#include "SNodePanel.h"
+#include "Widgets/Input/SSearchBox.h"
+
+#define LOCTEXT_NAMESPACE "CustomizableObjectGraphEditorToolkit"
+
+bool FCustomizableObjectGraphEditorToolkit::CanPasteNodes() const 
+{
+	if (!GraphEditor.IsValid() || !GraphEditor->GetCurrentGraph())
+	{
+		return false;
+	}
+
+	FString ClipboardContent;
+	FPlatformApplicationMisc::ClipboardPaste(ClipboardContent);
+
+	return FEdGraphUtilities::CanImportNodesFromText(GraphEditor->GetCurrentGraph(), ClipboardContent);
+}
+
+void FCustomizableObjectGraphEditorToolkit::PasteNodes()
+{
+	if (!GraphEditor.IsValid())
+	{
+		return;
+	}
+
+	PasteNodesHere(GraphEditor->GetPasteLocation());
+}
+
+void FCustomizableObjectGraphEditorToolkit::PasteNodesHere(const FVector2D& Location) 
+{
+	if (!GraphEditor.IsValid() || !GraphEditor->GetCurrentGraph())
+	{
+		return;
+	}
+
+	// Clear the selection set (newly pasted stuff will be selected)
+	GraphEditor->ClearSelectionSet();
+
+	// Grab the text to paste from the clipboard.
+	FString TextToImport;
+	FPlatformApplicationMisc::ClipboardPaste(TextToImport);
+
+	// Import the nodes
+	TSet<UEdGraphNode*> PastedNodes;
+	FEdGraphUtilities::ImportNodesFromText(GraphEditor->GetCurrentGraph(), TextToImport, PastedNodes);
+
+	//Average position of nodes so we can move them while still maintaining relative distances to each other
+	FVector2D AvgNodePosition(0.0f, 0.0f);
+
+	for (TSet<UEdGraphNode*>::TIterator It(PastedNodes); It; ++It)
+	{
+		UEdGraphNode* Node = *It;
+		AvgNodePosition.X += Node->NodePosX;
+		AvgNodePosition.Y += Node->NodePosY;
+	}
+
+	if (PastedNodes.Num() > 0)
+	{
+		float InvNumNodes = 1.0f / float(PastedNodes.Num());
+		AvgNodePosition.X *= InvNumNodes;
+		AvgNodePosition.Y *= InvNumNodes;
+	}
+
+	for (TSet<UEdGraphNode*>::TIterator It(PastedNodes); It; ++It)
+	{
+		UEdGraphNode* Node = *It;
+
+		// Select the newly pasted stuff
+		GraphEditor->SetNodeSelection(Node, true);
+
+		Node->NodePosX = (Node->NodePosX - AvgNodePosition.X) + Location.X;
+		Node->NodePosY = (Node->NodePosY - AvgNodePosition.Y) + Location.Y;
+
+		Node->SnapToGrid(SNodePanel::GetSnapGridSize());
+
+		// Give new node a different Guid from the old one
+		Node->CreateNewGuid();
+	}
+
+	for (UEdGraphNode* PastedNode : PastedNodes)
+	{
+		if (UCustomizableObjectNode* TypedNode = Cast<UCustomizableObjectNode>(PastedNode))
+		{
+			TypedNode->PostBackwardsCompatibleFixup();
+		}
+	}
+
+	// Update UI
+	GraphEditor->NotifyGraphChanged();
+
+	//TODO(Max)
+	//CustomizableObject->PostEditChange();
+	GraphEditor->GetCurrentGraph()->MarkPackageDirty();
+}
+
+void FCustomizableObjectGraphEditorToolkit::SelectNode(const UEdGraphNode* Node) 
+{
+	if (!GraphEditor.IsValid())
+	{
+		return;
+	}
+
+	GraphEditor->JumpToNode(Node);
+}
+
+void FCustomizableObjectGraphEditorToolkit::DeleteSelectedNodes()
+{
+	if (!GraphEditor.IsValid())
+	{
+		return;
+	}
+
+	const FScopedTransaction Transaction(LOCTEXT("UEdGraphSchema_CustomizableObject", "Delete Nodes"));
+
+	const FGraphPanelSelectionSet SelectedNodes = GraphEditor->GetSelectedNodes();
+	GraphEditor->ClearSelectionSet();
+
+	for (FGraphPanelSelectionSet::TConstIterator NodeIt(SelectedNodes); NodeIt; ++NodeIt)
+	{
+		if (UEdGraphNode* Node = Cast<UEdGraphNode>(*NodeIt))
+		{
+			if (Node->CanUserDeleteNode())
+			{
+				if (const UEdGraph* GraphObj = Node->GetGraph())
+				{
+					if (const UEdGraphSchema* Schema = GraphObj->GetSchema())
+					{
+						Schema->BreakNodeLinks(*Node);  // Required to notify to all connected nodes (UEdGraphNode::PinConnectionListChanged() and UEdGraphNode::PinConnectionListChanged(...))
+					}
+				}
+
+				Node->DestroyNode();
+			}
+		}
+	}
+}
+
+bool FCustomizableObjectGraphEditorToolkit::CanDeleteNodes() const
+{
+	if (GraphEditor.IsValid() && GraphEditor->GetSelectedNodes().Num() > 0)
+	{
+		for (auto Itr = GraphEditor->GetSelectedNodes().CreateConstIterator(); Itr; ++Itr)
+		{
+			UCustomizableObjectNode* Node = Cast<UCustomizableObjectNode>(*Itr);
+
+			if (Node && !Node->CanUserDeleteNode())
+			{
+				return false;
+			}
+		}
+
+		return true;
+	}
+
+	return false;
+}
+
+void FCustomizableObjectGraphEditorToolkit::DuplicateSelectedNodes()
+{
+	CopySelectedNodes();
+	PasteNodes();
+}
+
+bool FCustomizableObjectGraphEditorToolkit::CanDuplicateSelectedNodes() const
+{
+	return CanCopyNodes();
+}
+
+void FCustomizableObjectGraphEditorToolkit::OnNodeTitleCommitted(const FText& NewText, ETextCommit::Type CommitInfo, UEdGraphNode* NodeBeingChanged)
+{
+	if (NodeBeingChanged)
+	{
+		//TODO(Max)
+		//const FScopedTransaction Transaction(LOCTEXT("RenameNode", "Rename Node"));
+		//NodeBeingChanged->Modify();
+		NodeBeingChanged->OnRenameNode(NewText.ToString());
+	}
+}
+
+
+void FCustomizableObjectGraphEditorToolkit::OnNodeDoubleClicked(UEdGraphNode* Node)
+{
+	if (Node->CanJumpToDefinition())
+	{
+		Node->JumpToDefinition();
+	}
+}
+
+
+void FCustomizableObjectGraphEditorToolkit::CopySelectedNodes()
+{
+	// Export the selected nodes and place the text on the clipboard
+	const FGraphPanelSelectionSet SelectedNodes = GraphEditor->GetSelectedNodes();
+
+	FString ExportedText;
+
+	for (FGraphPanelSelectionSet::TConstIterator SelectedIter(SelectedNodes); SelectedIter; ++SelectedIter)
+	{
+		if (UEdGraphNode* Node = Cast<UEdGraphNode>(*SelectedIter))
+		{
+			Node->PrepareForCopying();
+		}
+	}
+
+	FEdGraphUtilities::ExportNodesToText(SelectedNodes, /*out*/ ExportedText);
+	FPlatformApplicationMisc::ClipboardCopy(*ExportedText);
+
+	//TODO(Max): This is weird... review it
+	// Make sure Material remains the owner of the copied nodes
+	for (FGraphPanelSelectionSet::TConstIterator SelectedIter(SelectedNodes); SelectedIter; ++SelectedIter)
+	{
+		if (UMaterialGraphNode* Node = Cast<UMaterialGraphNode>(*SelectedIter))
+		{
+			Node->PostCopyNode();
+		}
+		else if (UMaterialGraphNode_Comment* Comment = Cast<UMaterialGraphNode_Comment>(*SelectedIter))
+		{
+			Comment->PostCopyNode();
+		}
+	}
+}
+
+bool FCustomizableObjectGraphEditorToolkit::CanCopyNodes() const
+{
+	// If any of the nodes can be duplicated then we should allow copying
+	const FGraphPanelSelectionSet SelectedNodes = GraphEditor->GetSelectedNodes();
+	for (FGraphPanelSelectionSet::TConstIterator SelectedIter(SelectedNodes); SelectedIter; ++SelectedIter)
+	{
+		UEdGraphNode* Node = Cast<UEdGraphNode>(*SelectedIter);
+		if ((Node != NULL) && Node->CanDuplicateNode())
+		{
+			return true;
+		}
+	}
+
+	return false;
+}
+
+void FCustomizableObjectGraphEditorToolkit::CutSelectedNodes()
+{
+	CopySelectedNodes();
+	// Cut should only delete nodes that can be duplicated
+	DeleteSelectedNodes();
+}
+
+bool FCustomizableObjectGraphEditorToolkit::CanCutNodes() const
+{
+	return CanCopyNodes() && CanDeleteNodes();
+}
+
+void FCustomizableObjectGraphEditorToolkit::OnRenameNode()
+{
+	if (GraphEditor.IsValid())
+	{
+		const FGraphPanelSelectionSet& SelectedNodes = GraphEditor->GetSelectedNodes();
+
+		for (FGraphPanelSelectionSet::TConstIterator NodeIt(SelectedNodes); NodeIt; ++NodeIt)
+		{
+			// Rename only the first valid selected node
+			UEdGraphNode* SelectedNode = Cast<UEdGraphNode>(*NodeIt);
+			if (SelectedNode && SelectedNode->GetCanRenameNode())
+			{
+				GraphEditor->IsNodeTitleVisible(SelectedNode, true);
+				break;
+			}
+		}
+	}
+}
+
+bool FCustomizableObjectGraphEditorToolkit::CanRenameNodes() const
+{
+	if (GraphEditor.IsValid())
+	{
+		const FGraphPanelSelectionSet& SelectedNodes = GraphEditor->GetSelectedNodes();
+
+		for (FGraphPanelSelectionSet::TConstIterator SelectedIter(SelectedNodes); SelectedIter; ++SelectedIter)
+		{
+			if (UEdGraphNode* Node = Cast<UEdGraphNode>(*SelectedIter))
+			{
+				return Node->GetCanRenameNode();
+			}
+		}
+	}
+
+	return false;
+}
+
+void FCustomizableObjectGraphEditorToolkit::CreateCommentBoxFromKey()
+{
+	CreateCommentBox(GraphEditor->GetPasteLocation());
+}
+
+UEdGraphNode* FCustomizableObjectGraphEditorToolkit::CreateCommentBox(const FVector2D& NodePos)
+{
+	if (!GraphEditor.IsValid() || !GraphEditor->GetCurrentGraph())
+	{
+		return nullptr;
+	}
+
+	//const FScopedTransaction Transaction(LOCTEXT("UEdGraphSchema_CustomizableObject", "Add Comment Box"));
+	UEdGraphNode_Comment* CommentTemplate = NewObject<UEdGraphNode_Comment>();
+	UEdGraphNode_Comment* NewComment = nullptr;
+
+	{
+		//TODO(Max)
+		//CustomizableObject->Modify();
+
+		// const FGraphPanelSelectionSet& SelectionSet = GraphEditor->GetSelectedNodes();
+		FSlateRect Bounds;
+		FVector2D Location;
+		FVector2D Size;
+
+		if (GraphEditor->GetBoundsForSelectedNodes(Bounds, 50.f))
+		{
+			Location.X = Bounds.Left;
+			Location.Y = Bounds.Top;
+			Size = Bounds.GetSize();
+		}
+		else
+		{
+			Location.X = NodePos.X;
+			Location.Y = NodePos.Y;
+			Size.X = 400;
+			Size.Y = 100;
+		}
+
+		NewComment = FEdGraphSchemaAction_NewNode::SpawnNodeFromTemplate<UEdGraphNode_Comment>(GraphEditor->GetCurrentGraph(), CommentTemplate, NodePos, true);
+		NewComment->NodePosX = Location.X;
+		NewComment->NodePosY = Location.Y;
+		NewComment->NodeWidth = Size.X;
+		NewComment->NodeHeight = Size.Y;
+		NewComment->NodeComment = FString(TEXT("Comment"));
+	}
+
+	//TODO(Max)
+	GraphEditor->NotifyGraphChanged();
+	GraphEditor->GetCurrentGraph()->MarkPackageDirty();
+
+	return NewComment;
+}
+
+void FCustomizableObjectGraphEditorToolkit::OnEnterText(const FText& NewText, ETextCommit::Type TextType)
+{
+	if (TextType != ETextCommit::OnEnter)
+	{
+		return;
+	}
+
+	if (!GraphEditor)
+	{
+		return;
+	}
+
+	const UEdGraph* Graph = GraphEditor->GetCurrentGraph();
+	if (!Graph)
+	{
+		return;
+	}
+
+	bool bFound = false;
+
+	const FString FindString = NewText.ToString();
+
+	for (TObjectPtr<UEdGraphNode> Node : Graph->Nodes)
+	{
+		if (!Node)
+		{
+			continue;
+		}
+
+		// Node names are not in the reflection system
+		const FString NodeName = Node->GetNodeTitle(ENodeTitleType::FullTitle).ToString().Replace(TEXT("\n"), TEXT(" "));
+		if (NodeName.Contains(NewText.ToString(), ESearchCase::IgnoreCase))
+		{
+			LogSearchResult(*Node, "Node", bFound, NodeName);
+			bFound = true;
+		}
+
+		// Pins are not in the reflection system
+		for (const UEdGraphPin* Pin : Node->GetAllPins())
+		{
+			const FString PinFriendlyName = Pin->PinFriendlyName.ToString();
+			if (PinFriendlyName.Contains(FindString))
+			{
+				LogSearchResult(*Node, "Pin", bFound, PinFriendlyName);
+				bFound = true;
+			}
+		}
+
+		// Find anything marked as a UPROPERTY
+		for (TFieldIterator<FProperty> It(Node->GetClass()); It; ++It)
+		{
+			FindProperty(*It, Node, FindString, *Node, bFound);
+		}
+	}
+
+	const FText Text = bFound ?
+		LOCTEXT("SearchCompleted", "Search completed") :
+		FText::FromString("No Results for: " + NewText.ToString());
+
+	FCustomizableObjectEditorLogger::CreateLog(Text)
+		.Category(ELoggerCategory::GraphSearch)
+		.CustomNotification()
+		.Log();
+}
+
+void FCustomizableObjectGraphEditorToolkit::FindProperty(const FProperty* Property, const void* InContainer, const FString& FindString, const UObject& Context, bool& bFound)
+{
+	if (!Property || !InContainer)
+	{
+		return;
+	}
+
+	const FString PropertyName = Property->GetDisplayNameText().ToString();
+	if (PropertyName.Contains(FindString))
+	{
+		LogSearchResult(Context, "Property Name", bFound, *PropertyName);
+		bFound = true;
+	}
+
+	for (int32 Index = 0; Index < Property->ArrayDim; ++Index)
+	{
+		const uint8* ValuePtr = Property->ContainerPtrToValuePtr<uint8>(InContainer, Index);
+
+		if (const FStrProperty* StringProperty = CastField<FStrProperty>(Property))
+		{
+			const FString* StringResult = StringProperty->GetPropertyValuePtr(ValuePtr);
+			if (StringResult->Contains(FindString))
+			{
+				LogSearchResult(Context, "Property Value", bFound, *StringResult);
+				bFound = true;
+			}
+		}
+		else if (const FEnumProperty* EnumProperty = CastField<FEnumProperty>(Property))
+		{
+			const UEnum* EnumResult = EnumProperty->GetEnum();
+
+			const FString StringResult = EnumResult->GetDisplayNameTextByIndex(*ValuePtr).ToString();
+			if (StringResult.Contains(FindString))
+			{
+				LogSearchResult(Context, "Property Value", bFound, StringResult);
+				bFound = true;
+			}
+		}
+		else if (const FSoftObjectProperty* SoftObjectProperty = CastField<FSoftObjectProperty>(Property))
+		{
+			const FString ObjectPath = SoftObjectProperty->GetPropertyValuePtr(ValuePtr)->ToString();
+			if (ObjectPath.Contains(FindString))
+			{
+				LogSearchResult(Context, "Property Value", bFound, ObjectPath);
+				bFound = true;
+			}
+		}
+		else if (const FObjectPropertyBase* ObjectProperty = CastField<FObjectPropertyBase>(Property))
+		{
+			if (const UObject* ObjectValue = ObjectProperty->GetObjectPropertyValue(ValuePtr))
+			{
+				const FString Name = ObjectValue->GetName();
+
+				if (ObjectValue->GetName().Contains(FindString))
+				{
+					LogSearchResult(Context, "Property Value", bFound, Name);
+					bFound = true;
+				}
+			}
+		}
+		else if (const FStructProperty* StructProperty = CastField<FStructProperty>(Property))
+		{
+			for (TFieldIterator<FProperty> It(StructProperty->Struct); It; ++It)
+			{
+				FindProperty(*It, ValuePtr, FindString, Context, bFound);
+			}
+		}
+		else if (const FArrayProperty* ArrayProperty = CastField<FArrayProperty>(Property))
+		{
+			FScriptArrayHelper ArrayHelper(ArrayProperty, ValuePtr);
+			for (int32 ValueIdx = 0; ValueIdx < ArrayHelper.Num(); ++ValueIdx)
+			{
+				FindProperty(ArrayProperty->Inner, ArrayHelper.GetRawPtr(ValueIdx), FindString, Context, bFound);
+			}
+		}
+		else if (const FSetProperty* SetProperty = CastField<FSetProperty>(Property))
+		{
+			FScriptSetHelper SetHelper(SetProperty, ValuePtr);
+			for (FScriptSetHelper::FIterator SetIt = SetHelper.CreateIterator(); SetIt; ++SetIt)
+			{
+				FindProperty(SetProperty->ElementProp, SetHelper.GetElementPtr(SetIt), FindString, Context, bFound);
+			}
+		}
+		else if (const FMapProperty* MapProperty = CastField<FMapProperty>(Property))
+		{
+			FScriptMapHelper MapHelper(MapProperty, ValuePtr);
+			for (FScriptMapHelper::FIterator MapIt = MapHelper.CreateIterator(); MapIt; ++MapIt)
+			{
+				const uint8* MapValuePtr = MapHelper.GetPairPtr(MapIt);
+				FindProperty(MapProperty->KeyProp, MapValuePtr, FindString, Context, bFound);
+				FindProperty(MapProperty->ValueProp, MapValuePtr, FindString, Context, bFound);
+			}
+		}
+	}
+}
+
+void FCustomizableObjectGraphEditorToolkit::LogSearchResult(const UObject& Context, const FString& Type, bool bIsFirst, const FString& Result) const
+{
+	if (!bIsFirst)
+	{
+		FCustomizableObjectEditorLogger::CreateLog(LOCTEXT("SearchResults", "Search Results:"))
+			.Notification(false)
+			.Log();
+	}
+
+	FCustomizableObjectEditorLogger::CreateLog(FText::FromString(Type + ": " + Result))
+		.Context(Context)
+		.BaseObject()
+		.Notification(false)
+		.Log();
+}
+
+
+/** Create new tab for the supplied graph - don't call this directly, call SExplorer->FindTabForGraph.*/
+void FCustomizableObjectGraphEditorToolkit::CreateGraphEditorWidget(UEdGraph* InGraph, SGraphEditor::FGraphEditorEvents& InEvents)
+{
+	UCustomizableObjectGraph* CustomizableObjectGraph = Cast<UCustomizableObjectGraph>(InGraph);
+	check(CustomizableObjectGraph);
+
+	GraphEditorCommands = MakeShareable(new FUICommandList);
+
+	TSharedRef<SWidget> TitleBarWidget =
+		SNew(SHorizontalBox)
+		+ SHorizontalBox::Slot()
+		.HAlign(HAlign_Left)
+		.FillWidth(10.f)
+		.Padding(5.f)
+		[
+			SNew(SSearchBox)
+				.HintText(LOCTEXT("Search", "Search..."))
+				.ToolTipText(LOCTEXT("Search Nodes, Properties or Values that contain the inserted words", "Search Nodes, Properties or Values that contain the inserted words"))
+				.OnTextCommitted(FOnTextCommitted::CreateSP(this, &FCustomizableObjectGraphEditorToolkit::OnEnterText))
+				.SelectAllTextWhenFocused(true)
+		];
+
+	// Create the appearance info
+	FGraphAppearanceInfo AppearanceInfo;
+	AppearanceInfo.CornerText = CustomizableObjectGraph->IsMacro() ? LOCTEXT("ApperanceCornerText", "MUTABLE MACRO") : LOCTEXT("ApperanceCornerText", "MUTABLE");
+
+	// Add toolkit common events
+	InEvents.OnTextCommitted = FOnNodeTextCommitted::CreateSP(this, &FCustomizableObjectGraphEditorToolkit::OnNodeTitleCommitted);
+	InEvents.OnNodeDoubleClicked = FSingleNodeEvent::CreateSP(this, &FCustomizableObjectGraphEditorToolkit::OnNodeDoubleClicked);
+
+	// Make full graph editor
+	GraphEditor = SNew(SGraphEditor)
+		.AdditionalCommands(GraphEditorCommands)
+		.Appearance(AppearanceInfo)
+		.GraphToEdit(InGraph)
+		.GraphEvents(InEvents)
+		.TitleBar(TitleBarWidget)
+		.ShowGraphStateOverlay(false); // Removes graph state overlays (border and text) such as "SIMULATING" and "READ-ONLY"
+
+	// Ensure commands are registered
+	FGraphEditorCommands::Register();
+	FCustomizableObjectEditorNodeContextCommands::Register();
+
+	// Editing commands
+	GraphEditorCommands->MapAction(FGenericCommands::Get().Delete,
+		FExecuteAction::CreateSP(this, &FCustomizableObjectGraphEditorToolkit::DeleteSelectedNodes),
+		FCanExecuteAction::CreateSP(this, &FCustomizableObjectGraphEditorToolkit::CanDeleteNodes));
+
+	GraphEditorCommands->MapAction(FGenericCommands::Get().Copy,
+		FExecuteAction::CreateSP(this, &FCustomizableObjectGraphEditorToolkit::CopySelectedNodes),
+		FCanExecuteAction::CreateSP(this, &FCustomizableObjectGraphEditorToolkit::CanCopyNodes));
+
+	GraphEditorCommands->MapAction(FGenericCommands::Get().Paste,
+		FExecuteAction::CreateSP(this, &FCustomizableObjectGraphEditorToolkit::PasteNodes),
+		FCanExecuteAction::CreateSP(this, &FCustomizableObjectGraphEditorToolkit::CanPasteNodes));
+
+	GraphEditorCommands->MapAction(FGenericCommands::Get().Cut,
+		FExecuteAction::CreateSP(this, &FCustomizableObjectGraphEditorToolkit::CutSelectedNodes),
+		FCanExecuteAction::CreateSP(this, &FCustomizableObjectGraphEditorToolkit::CanCutNodes));
+
+	GraphEditorCommands->MapAction(FGenericCommands::Get().Duplicate,
+		FExecuteAction::CreateSP(this, &FCustomizableObjectGraphEditorToolkit::DuplicateSelectedNodes),
+		FCanExecuteAction::CreateSP(this, &FCustomizableObjectGraphEditorToolkit::CanDuplicateSelectedNodes));
+
+	GraphEditorCommands->MapAction(FGenericCommands::Get().Rename,
+		FExecuteAction::CreateSP(this, &FCustomizableObjectGraphEditorToolkit::OnRenameNode),
+		FCanExecuteAction::CreateSP(this, &FCustomizableObjectGraphEditorToolkit::CanRenameNodes));
+
+	GraphEditorCommands->MapAction(FCustomizableObjectEditorNodeContextCommands::Get().CreateComment,
+		FExecuteAction::CreateSP(this, &FCustomizableObjectGraphEditorToolkit::CreateCommentBoxFromKey));
+
+	// Alignment Commands
+	GraphEditorCommands->MapAction(FGraphEditorCommands::Get().AlignNodesTop,
+		FExecuteAction::CreateSP(GraphEditor.Get(), &SGraphEditor::OnAlignTop));
+
+	GraphEditorCommands->MapAction(FGraphEditorCommands::Get().AlignNodesMiddle,
+		FExecuteAction::CreateSP(GraphEditor.Get(), &SGraphEditor::OnAlignMiddle));
+
+	GraphEditorCommands->MapAction(FGraphEditorCommands::Get().AlignNodesBottom,
+		FExecuteAction::CreateSP(GraphEditor.Get(), &SGraphEditor::OnAlignBottom));
+
+	GraphEditorCommands->MapAction(FGraphEditorCommands::Get().AlignNodesLeft,
+		FExecuteAction::CreateSP(GraphEditor.Get(), &SGraphEditor::OnAlignLeft));
+
+	GraphEditorCommands->MapAction(FGraphEditorCommands::Get().AlignNodesCenter,
+		FExecuteAction::CreateSP(GraphEditor.Get(), &SGraphEditor::OnAlignCenter));
+
+	GraphEditorCommands->MapAction(FGraphEditorCommands::Get().AlignNodesRight,
+		FExecuteAction::CreateSP(GraphEditor.Get(), &SGraphEditor::OnAlignRight));
+
+	GraphEditorCommands->MapAction(FGraphEditorCommands::Get().StraightenConnections,
+		FExecuteAction::CreateSP(GraphEditor.Get(), &SGraphEditor::OnStraightenConnections));
+
+	// Distribution Commands
+	GraphEditorCommands->MapAction(FGraphEditorCommands::Get().DistributeNodesHorizontally,
+		FExecuteAction::CreateSP(GraphEditor.Get(), &SGraphEditor::OnDistributeNodesH));
+
+	GraphEditorCommands->MapAction(FGraphEditorCommands::Get().DistributeNodesVertically,
+		FExecuteAction::CreateSP(GraphEditor.Get(), &SGraphEditor::OnDistributeNodesV));
+}
+
+#undef LOCTEXT_NAMESPACE
\ No newline at end of file
diff --git a/Engine/Plugins/Mutable/Source/CustomizableObjectEditor/Private/MuCOE/CustomizableObjectMacroLibrary/AssetDefinition_CustomizableObjectMacroLibrary.cpp b/Engine/Plugins/Mutable/Source/CustomizableObjectEditor/Private/MuCOE/CustomizableObjectMacroLibrary/AssetDefinition_CustomizableObjectMacroLibrary.cpp
new file mode 100644
--- /dev/null
+++ b/Engine/Plugins/Mutable/Source/CustomizableObjectEditor/Private/MuCOE/CustomizableObjectMacroLibrary/AssetDefinition_CustomizableObjectMacroLibrary.cpp
@@ -0,0 +1,72 @@
+// Copyright Epic Games, Inc. All Rights Reserved.
+
+#include "MuCOE/CustomizableObjectMacroLibrary/AssetDefinition_CustomizableObjectMacroLibrary.h"
+
+#include "MuCOE/CustomizableObjectMacroLibrary/CustomizableObjectMacroLibrary.h"
+#include "MuCOE/CustomizableObjectMacroLibrary/CustomizableObjectMacroLibraryEditor.h"
+
+#define LOCTEXT_NAMESPACE "AssetDefinition_CustomizableObjectMacroLibrary"
+
+FText UAssetDefinition_CustomizableObjectMacroLibrary::GetAssetDisplayName() const
+{ 
+	return NSLOCTEXT("AssetTypeActions", "AssetTypeActions_CustomizableObjectMacroLibrary", "Customizable Object Macro Library"); 
+}
+
+
+FLinearColor UAssetDefinition_CustomizableObjectMacroLibrary::GetAssetColor() const
+{
+	return FLinearColor(FColor(100, 100, 100)); 
+}
+
+
+TSoftClassPtr<UObject> UAssetDefinition_CustomizableObjectMacroLibrary::GetAssetClass() const
+{
+	return UCustomizableObjectMacroLibrary::StaticClass(); 
+}
+
+
+EAssetCommandResult UAssetDefinition_CustomizableObjectMacroLibrary::PerformAssetDiff(const FAssetDiffArgs& DiffArgs) const
+{
+	return EAssetCommandResult::Unhandled; 
+}
+
+
+EAssetCommandResult UAssetDefinition_CustomizableObjectMacroLibrary::OpenAssets(const FAssetOpenArgs& OpenArgs) const
+{
+	const EToolkitMode::Type Mode = OpenArgs.ToolkitHost.IsValid() ? EToolkitMode::WorldCentric : EToolkitMode::Standalone;
+
+	for (UCustomizableObjectMacroLibrary* Object : OpenArgs.LoadObjects<UCustomizableObjectMacroLibrary>())
+	{
+		const TSharedPtr<FCustomizableObjectMacroLibraryEditor> Editor = MakeShared<FCustomizableObjectMacroLibraryEditor>();
+		Editor->InitEditor(Mode, OpenArgs.ToolkitHost, Object);
+	}
+
+	return EAssetCommandResult::Handled;
+}
+
+
+TConstArrayView<FAssetCategoryPath> UAssetDefinition_CustomizableObjectMacroLibrary::GetAssetCategories() const
+{
+	static const std::initializer_list<FAssetCategoryPath> Categories =
+	{
+		// Asset can be found inside the Mutable submenu 
+		NSLOCTEXT("AssetTypeActions", "Mutable", "Mutable")
+	};
+
+	return Categories;
+}
+
+
+EAssetCommandResult UAssetDefinition_CustomizableObjectMacroLibrary::ActivateAssets(const FAssetActivateArgs& ActivateArgs) const
+{
+	return EAssetCommandResult::Unhandled;
+}
+
+
+FAssetOpenSupport UAssetDefinition_CustomizableObjectMacroLibrary::GetAssetOpenSupport(const FAssetOpenSupportArgs& OpenSupportArgs) const
+{
+	return FAssetOpenSupport(EAssetOpenMethod::Edit, true, EToolkitMode::Standalone);
+}
+
+#undef LOCTEXT_NAMESPACE
+
diff --git a/Engine/Plugins/Mutable/Source/CustomizableObjectEditor/Private/MuCOE/CustomizableObjectMacroLibrary/AssetDefinition_CustomizableObjectMacroLibrary.h b/Engine/Plugins/Mutable/Source/CustomizableObjectEditor/Private/MuCOE/CustomizableObjectMacroLibrary/AssetDefinition_CustomizableObjectMacroLibrary.h
new file mode 100644
--- /dev/null
+++ b/Engine/Plugins/Mutable/Source/CustomizableObjectEditor/Private/MuCOE/CustomizableObjectMacroLibrary/AssetDefinition_CustomizableObjectMacroLibrary.h
@@ -0,0 +1,25 @@
+// Copyright Epic Games, Inc. All Rights Reserved.
+#pragma once
+
+#include "AssetDefinitionDefault.h"
+#include "AssetDefinition_CustomizableObjectMacroLibrary.generated.h"
+
+
+UCLASS()
+class UAssetDefinition_CustomizableObjectMacroLibrary final : public UAssetDefinitionDefault
+{
+	GENERATED_BODY()
+
+public:
+
+	// UAssetDefinition Interface
+	virtual FText GetAssetDisplayName() const override;
+	virtual FLinearColor GetAssetColor() const override;
+	virtual TSoftClassPtr<UObject> GetAssetClass() const override;
+	virtual EAssetCommandResult PerformAssetDiff(const FAssetDiffArgs& DiffArgs) const override;
+	virtual EAssetCommandResult OpenAssets(const FAssetOpenArgs& OpenArgs) const override;
+	virtual TConstArrayView<FAssetCategoryPath> GetAssetCategories() const override;
+	virtual EAssetCommandResult ActivateAssets(const FAssetActivateArgs& ActivateArgs) const override;
+	virtual FAssetOpenSupport GetAssetOpenSupport(const FAssetOpenSupportArgs& OpenSupportArgs) const override;
+
+};
diff --git a/Engine/Plugins/Mutable/Source/CustomizableObjectEditor/Private/MuCOE/CustomizableObjectMacroLibrary/CustomizableObjectMacroDetails.cpp b/Engine/Plugins/Mutable/Source/CustomizableObjectEditor/Private/MuCOE/CustomizableObjectMacroLibrary/CustomizableObjectMacroDetails.cpp
new file mode 100644
--- /dev/null
+++ b/Engine/Plugins/Mutable/Source/CustomizableObjectEditor/Private/MuCOE/CustomizableObjectMacroLibrary/CustomizableObjectMacroDetails.cpp
@@ -0,0 +1,338 @@
+// Copyright Epic Games, Inc. All Rights Reserved.
+
+#include "MuCOE/CustomizableObjectMacroLibrary/CustomizableObjectMacroDetails.h"
+
+#include "DetailCategoryBuilder.h"
+#include "DetailLayoutBuilder.h"
+#include "DetailWidgetRow.h"
+#include "IDetailsView.h"
+#include "MuCOE/EdGraphSchema_CustomizableObject.h"
+#include "MuCOE/Nodes/CustomizableObjectNodeTunnel.h"
+#include "PropertyCustomizationHelpers.h"
+#include "SListViewSelectorDropdownMenu.h"
+#include "Widgets/Input/SButton.h"
+#include "Widgets/Input/SComboButton.h"
+#include "Widgets/SBoxPanel.h"
+#include "Widgets/Views/SListView.h"
+#include "Widgets/Views/STableRow.h"
+
+
+#define LOCTEXT_NAMESPACE "CustomizableObjectMacroDetails"
+
+TSharedRef<IDetailCustomization> FCustomizableObjectMacroDetails::MakeInstance()
+{
+	return MakeShareable(new FCustomizableObjectMacroDetails);
+}
+
+
+void FCustomizableObjectMacroDetails::CustomizeDetails(const TSharedPtr<IDetailLayoutBuilder>& DetailBuilder)
+{
+	DetailBuilderPtr = DetailBuilder;
+
+	TSharedPtr<IDetailsView> DetailsView = DetailBuilder->GetDetailsViewSharedPtr();
+
+	Macro = Cast<UCustomizableObjectMacro>(DetailsView->GetSelectedObjects()[0].Get());
+
+	if (!Macro || !Macro->Graph)
+	{
+		return;
+	}
+
+	TArray<UCustomizableObjectNodeTunnel*> IONodes;
+	TObjectPtr<UCustomizableObjectNodeTunnel> InputNode;
+	TObjectPtr<UCustomizableObjectNodeTunnel> OutputNode;
+
+	Macro->Graph->GetNodesOfClass<UCustomizableObjectNodeTunnel>(IONodes);
+	check(IONodes.Num() == 2);
+
+	for (UCustomizableObjectNodeTunnel* IONode : IONodes)
+	{
+		if (IONode->bIsInputNode)
+		{
+			InputNode = IONode;
+		}
+		else
+		{
+			OutputNode = IONode;
+		}
+	}
+
+	IDetailCategoryBuilder& GraphCategory = DetailBuilder->EditCategory("Macro");
+	IDetailCategoryBuilder& InputsCategory = DetailBuilder->EditCategory("Inputs");
+	IDetailCategoryBuilder& OutputsCategory = DetailBuilder->EditCategory("Outputs");
+
+	InputsCategory.HeaderContent(GenerateCategoryButtonWidget(ECOMacroIOType::COMVT_Input, InputNode));
+	OutputsCategory.HeaderContent(GenerateCategoryButtonWidget(ECOMacroIOType::COMVT_Output, OutputNode));
+
+	GenerateVariableList(InputsCategory, InputNode, ECOMacroIOType::COMVT_Input);
+	GenerateVariableList(OutputsCategory, OutputNode, ECOMacroIOType::COMVT_Output);
+}
+
+
+void FCustomizableObjectMacroDetails::GenerateVariableList(IDetailCategoryBuilder& IOCategory, UCustomizableObjectNodeTunnel* IONode, ECOMacroIOType IOType)
+{
+	if (!Macro || !Macro->Graph)
+	{
+		return;
+	}
+
+	const UEdGraphSchema_CustomizableObject* Schema = Cast<UEdGraphSchema_CustomizableObject>(Macro->Graph->GetSchema());
+	check(Schema);
+
+	if (!IONode->Pins.Num())
+	{
+		// Add a text widget to let the user know to hit the + icon to add parameters.
+		IOCategory.AddCustomRow(FText::GetEmpty()).WholeRowContent()
+		.MaxDesiredWidth(980.f)
+		[
+			SNew(SHorizontalBox)
+			+ SHorizontalBox::Slot()
+			.VAlign(VAlign_Center)
+			.Padding(0.0f, 0.0f, 4.0f, 0.0f)
+			.AutoWidth()
+			[
+				SNew(STextBlock)
+				.Text(LOCTEXT("NoArgumentsAddedForBlueprint", "Press the + icon above to add a new Variable"))
+				.Font(IDetailLayoutBuilder::GetDetailFont())
+			]
+		];
+
+		return;
+	}
+
+	for (UCustomizableObjectMacroInputOutput* InputOutput : Macro->InputOutputs)
+	{
+		if (InputOutput->Type == IOType)
+		{
+			IOCategory.AddCustomRow(FText::FromName(InputOutput->Name))
+			.NameContent()
+			[
+				SNew(SEditableTextBox)
+				.Text(this, &FCustomizableObjectMacroDetails::GetVariableName, InputOutput)
+				.OnTextChanged(this, &FCustomizableObjectMacroDetails::OnVariableNameChanged, InputOutput)
+				.OnTextCommitted(this, &FCustomizableObjectMacroDetails::OnVariableNameCommitted, InputOutput, IONode)
+				.Font(IDetailLayoutBuilder::GetDetailFont())
+			]
+			.ValueContent()
+			[
+				SNew(SCOMacroPinTypeSelector)
+				.Variable(InputOutput)
+				.IONode(IONode)
+				.OnVariableRemoved(this, &FCustomizableObjectMacroDetails::OnRemoveVariable)
+			];
+		}
+	}
+}
+
+
+TSharedRef<SWidget> FCustomizableObjectMacroDetails::GenerateCategoryButtonWidget(ECOMacroIOType IOType, UCustomizableObjectNodeTunnel* IONode)
+{
+	FText TooltipText = IONode->bIsInputNode ? LOCTEXT("MacroInputTooltip", "Create a new Input variable") : LOCTEXT("MacroOutputTooltip", "Create a new Output variable");
+
+	return SNew(SHorizontalBox)
+		+ SHorizontalBox::Slot()
+		.HAlign(HAlign_Right)
+		[
+			SNew(SButton)
+				.ButtonStyle(FAppStyle::Get(), "SimpleButton")
+				.ContentPadding(FMargin(1, 0))
+				.OnClicked(this, &FCustomizableObjectMacroDetails::AddNewVariable, IOType, IONode)
+				.HAlign(HAlign_Right)
+				.ToolTipText(TooltipText)
+				.VAlign(VAlign_Center)
+				[
+					SNew(SImage)
+						.Image(FAppStyle::Get().GetBrush("Icons.PlusCircle"))
+						.ColorAndOpacity(FSlateColor::UseForeground())
+				]
+		];
+}
+
+
+FReply FCustomizableObjectMacroDetails::AddNewVariable(ECOMacroIOType VarType, UCustomizableObjectNodeTunnel* IONode)
+{
+	if (Macro && IONode)
+	{
+		if (Macro->AddVariable(VarType))
+		{
+			IONode->ReconstructNode();
+			DetailBuilderPtr.Pin()->ForceRefreshDetails();
+		}
+
+		return FReply::Handled();
+	}
+
+	return FReply::Unhandled();
+}
+
+
+FText FCustomizableObjectMacroDetails::GetVariableName(UCustomizableObjectMacroInputOutput* Variable) const
+{
+	if (Variable)
+	{
+		return FText::FromName(Variable->Name);
+	}
+
+	return FText();
+}
+
+
+void FCustomizableObjectMacroDetails::OnVariableNameChanged(const FText& InNewText, UCustomizableObjectMacroInputOutput* Variable)
+{
+	if (Variable)
+	{
+		//TODO(Max): Check for repeated names and notify somehow
+		Variable->Name = *InNewText.ToString();
+	}
+}
+
+
+void FCustomizableObjectMacroDetails::OnVariableNameCommitted(const FText& InNewText, ETextCommit::Type InTextCommit, UCustomizableObjectMacroInputOutput* Variable, UCustomizableObjectNodeTunnel* IONode)
+{
+	if (Variable && IONode)
+	{
+		//TODO(Max): Check for repeated names and notify somehow
+		Variable->Name = *InNewText.ToString();
+		IONode->ReconstructNode();
+	}
+}
+
+
+void FCustomizableObjectMacroDetails::OnRemoveVariable(UCustomizableObjectNodeTunnel* IONode)
+{
+	if (IONode)
+	{
+		IONode->ReconstructNode();
+	}
+
+	DetailBuilderPtr.Pin()->ForceRefreshDetails();
+}
+
+
+// PinTypeSelector Widget --------------
+void SCOMacroPinTypeSelector::Construct(const FArguments& InArgs)
+{
+	Variable = InArgs._Variable;
+	IONode = InArgs._IONode;
+	Macro = Cast< UCustomizableObjectMacro>(Variable.GetOuter());
+	OnVariableRemoved = InArgs._OnVariableRemoved;
+
+	if (!Macro || !Macro->Graph || !Variable || !IONode)
+	{
+		ChildSlot
+		[
+			SNew(STextBlock).Text(LOCTEXT("NullMacroErrorPin","Variable Error!"))
+		];
+
+		return;
+	}
+
+	const UEdGraphSchema_CustomizableObject* Schema = Cast<UEdGraphSchema_CustomizableObject>(Macro->Graph->GetSchema());
+	check(Schema);
+
+	for (const FName& PinType : Schema->SupportedMacroPinTypes)
+	{
+		TSharedPtr<FPinNameRowData> NewSourceData = MakeShareable(new FPinNameRowData());
+		NewSourceData->PinCategory = PinType;
+		NewSourceData->PinFriendlyName = Schema->GetPinCategoryName(PinType);
+
+		ListViewSource.Add(NewSourceData);
+	}
+	
+	TSharedPtr<SListView<TSharedPtr<FPinNameRowData>>> PinTypesListView = SNew(SListView<TSharedPtr<FPinNameRowData>>)
+	.ListItemsSource(&ListViewSource)
+	.SelectionMode(ESelectionMode::Single)
+	.OnMouseButtonClick_Lambda([this](TSharedPtr<FPinNameRowData> Item)
+	{
+		if (Item)
+		{
+			Variable->PinCategoryType = Item->PinCategory;
+			IONode->ReconstructNode();
+		}
+	})
+	.ListViewStyle(&FAppStyle::Get().GetWidgetStyle<FTableViewStyle>("SimpleListView"))
+	.OnGenerateRow(this, &SCOMacroPinTypeSelector::OnGenerateRow);
+
+	ChildSlot
+	[
+		SNew(SHorizontalBox)
+		+ SHorizontalBox::Slot()
+		.AutoWidth()
+		[
+			// Combo button that summons the dropdown menu
+			SNew(SComboButton)
+			.ButtonContent()
+			[
+				SNew(STextBlock)
+				.Font(FAppStyle::GetFontStyle(TEXT("PropertyWindow.NormalFont")))
+				.Text(this, &SCOMacroPinTypeSelector::GetVariablePinTypeText)
+			]
+			.MenuContent()
+			[
+				SNew(SListViewSelectorDropdownMenu<TSharedPtr<FPinNameRowData>>, nullptr, PinTypesListView)
+				[
+					SNew(SBorder)
+					.BorderImage(FAppStyle::GetBrush("Menu.Background"))
+					.Padding(2)
+					[
+						SNew(SBox)
+						[
+							PinTypesListView.ToSharedRef()
+						]
+					]
+				]
+			]
+			.IsFocusable(true)
+			.ContentPadding(2.0f)
+		]
+
+		+ SHorizontalBox::Slot()
+		.AutoWidth()
+		[
+			PropertyCustomizationHelpers::MakeClearButton(
+				FSimpleDelegate::CreateSP(this, &SCOMacroPinTypeSelector::OnRemoveClicked), LOCTEXT("RemoveVariableTooltip", "Remove Variable."))
+		]
+	];
+}
+
+
+TSharedRef<ITableRow> SCOMacroPinTypeSelector::OnGenerateRow(TSharedPtr<FPinNameRowData> RowData, const TSharedRef<STableViewBase>& OwnerTable)
+{
+	TSharedPtr<STableRow<TSharedPtr<FString>>> ListTableRow;
+
+	SAssignNew(ListTableRow, STableRow<TSharedPtr<FString>>, OwnerTable)
+	[
+		SNew(SVerticalBox)
+			+ SVerticalBox::Slot()
+			.AutoHeight()
+			.Padding(20.0f, 1.0f, 20.0f, 1.0f)
+			[
+				SNew(STextBlock)
+					.Text(FText::FromName(RowData.IsValid() ? RowData->PinFriendlyName : FName("Invalid")))
+			]
+	];
+
+	return ListTableRow.ToSharedRef();
+}
+
+
+FText SCOMacroPinTypeSelector::GetVariablePinTypeText() const
+{
+	const UEdGraphSchema_CustomizableObject* Schema = Cast<UEdGraphSchema_CustomizableObject>(Macro->Graph->GetSchema());
+	check(Schema);
+
+	return Variable ? FText::FromName(Schema->GetPinCategoryName(Variable->PinCategoryType)) : FText();
+}
+
+
+void SCOMacroPinTypeSelector::OnRemoveClicked()
+{
+	if (Macro)
+	{
+		Macro->RemoveVariable(Variable);
+		OnVariableRemoved.ExecuteIfBound(IONode);
+	}
+}
+
+
+#undef LOCTEXT_NAMESPACE
diff --git a/Engine/Plugins/Mutable/Source/CustomizableObjectEditor/Private/MuCOE/CustomizableObjectMacroLibrary/CustomizableObjectMacroDetails.h b/Engine/Plugins/Mutable/Source/CustomizableObjectEditor/Private/MuCOE/CustomizableObjectMacroLibrary/CustomizableObjectMacroDetails.h
new file mode 100644
--- /dev/null
+++ b/Engine/Plugins/Mutable/Source/CustomizableObjectEditor/Private/MuCOE/CustomizableObjectMacroLibrary/CustomizableObjectMacroDetails.h
@@ -0,0 +1,107 @@
+// Copyright Epic Games, Inc. All Rights Reserved.
+
+#pragma once
+
+#include "IDetailCustomization.h"
+#include "Widgets/SCompoundWidget.h"
+#include "Widgets/Views/ITableRow.h"
+
+class IDetailLayoutBuilder;
+class IDetailCategoryBuilder;
+class ITableRow;
+class STableViewBase;
+class UCustomizableObjectMacro;
+class UCustomizableObjectMacroInputOutput;
+class UCustomizableObjectNodeTunnel;
+
+
+enum class ECOMacroIOType : uint8;
+
+struct FPinNameRowData
+{
+	FName PinFriendlyName;
+	FName PinCategory;
+};
+
+/** Widget to edit the type of a macro's variable. */
+class SCOMacroPinTypeSelector : public SCompoundWidget
+{
+	DECLARE_DELEGATE_OneParam(FOnVariableRemoved, UCustomizableObjectNodeTunnel*);
+
+	SLATE_BEGIN_ARGS(SCOMacroPinTypeSelector) {}
+		SLATE_ARGUMENT(TObjectPtr<UCustomizableObjectMacroInputOutput>, Variable)
+		SLATE_ARGUMENT(TObjectPtr<UCustomizableObjectNodeTunnel>, IONode)
+		SLATE_EVENT(FOnVariableRemoved, OnVariableRemoved)
+	SLATE_END_ARGS()
+
+	void Construct(const FArguments& InArgs);
+
+	/** Generates the widget of a List view row. */
+	TSharedRef<ITableRow> OnGenerateRow(TSharedPtr<FPinNameRowData> RowData, const TSharedRef<STableViewBase>& OwnerTable);
+
+	/** Returns the name of the pin type of the variable.*/
+	FText GetVariablePinTypeText() const;
+
+	/** Callback of the remove variable button. */
+	void OnRemoveClicked();
+
+private:
+
+	/** Pointer to the variablethas is being edited */
+	TObjectPtr<UCustomizableObjectMacroInputOutput> Variable;
+
+	/** Pointer to the macro where the variable belongs. */
+	TObjectPtr<UCustomizableObjectMacro> Macro;
+
+	/** Pointer to the node that will expose this variable to the graph. */
+	TObjectPtr<UCustomizableObjectNodeTunnel> IONode;
+
+	/** Array that contains all the pin types showed in the list view widget.*/
+	TArray<TSharedPtr<FPinNameRowData>> ListViewSource;
+
+	/** Callback to indicate that a variable has been removed */
+	FOnVariableRemoved OnVariableRemoved;
+};
+
+
+/** Details View of UCustomizableObjectMacros */
+class FCustomizableObjectMacroDetails : public IDetailCustomization
+{
+public:
+
+	static TSharedRef<IDetailCustomization> MakeInstance();
+
+	// ILayoutDetails interface
+	/** Do not use. Add details customization in the other CustomizeDetails signature. */
+	virtual void CustomizeDetails(IDetailLayoutBuilder& DetailBuilder) override {};
+	virtual void CustomizeDetails(const TSharedPtr<IDetailLayoutBuilder>& DetailBuilder) override;
+
+private:
+
+	/** Creates a widget list to create and visualize macro variables. */
+	void GenerateVariableList(IDetailCategoryBuilder& InputsCategory, UCustomizableObjectNodeTunnel* IONode, ECOMacroIOType VariableType);
+	
+	/** Creates the button that allows to add new macro variables. */
+	TSharedRef<SWidget> GenerateCategoryButtonWidget(ECOMacroIOType VariableType, UCustomizableObjectNodeTunnel* IONode);
+	
+	/** Returns the name of the specified variable. Needed to update automatically the text if the name of the variable changes. */
+	FText GetVariableName(UCustomizableObjectMacroInputOutput* Variable) const;
+	
+	/** Editable Text Box Callbacks */
+	void OnVariableNameChanged(const FText& InNewText, UCustomizableObjectMacroInputOutput* Variable);
+	void OnVariableNameCommitted(const FText& InNewText, ETextCommit::Type InTextCommit, UCustomizableObjectMacroInputOutput* Variable, UCustomizableObjectNodeTunnel* IONode);
+
+	/** Callback to comunicate the editor that the name of a variable has changes. Updates the variable List too. */
+	void OnRemoveVariable(UCustomizableObjectNodeTunnel* IONode);
+
+	/** Adds a new variable to a Macro. */
+	FReply AddNewVariable(ECOMacroIOType VarType, UCustomizableObjectNodeTunnel* IONode);
+
+private:
+
+	/** Details builder pointer */
+	TWeakPtr<IDetailLayoutBuilder> DetailBuilderPtr;
+
+	/** Pointer to the macro thas is being edited */
+	TObjectPtr<UCustomizableObjectMacro> Macro;
+};
diff --git a/Engine/Plugins/Mutable/Source/CustomizableObjectEditor/Private/MuCOE/CustomizableObjectMacroLibrary/CustomizableObjectMacroLibrary.cpp b/Engine/Plugins/Mutable/Source/CustomizableObjectEditor/Private/MuCOE/CustomizableObjectMacroLibrary/CustomizableObjectMacroLibrary.cpp
new file mode 100644
--- /dev/null
+++ b/Engine/Plugins/Mutable/Source/CustomizableObjectEditor/Private/MuCOE/CustomizableObjectMacroLibrary/CustomizableObjectMacroLibrary.cpp
@@ -0,0 +1,97 @@
+// Copyright Epic Games, Inc. All Rights Reserved.
+
+#include "MuCOE/CustomizableObjectMacroLibrary/CustomizableObjectMacroLibrary.h"
+
+#include "MuCOE/CustomizableObjectGraph.h"
+#include "MuCOE/Nodes/CustomizableObjectNodeTunnel.h"
+#include "MuCOE/EdGraphSchema_CustomizableObject.h"
+
+#define LOCTEXT_NAMESPACE "CustomizableObjectMacroLibrary"
+
+UCustomizableObjectMacro* UCustomizableObjectMacroLibrary::AddMacro()
+{
+	UCustomizableObjectMacro* NewMacro = nullptr;
+
+	FName BaseName = "NewMacro";
+	FName MacroName = MakeUniqueObjectName(this, UCustomizableObjectMacro::StaticClass(), FName(BaseName));
+	
+	NewMacro = NewObject<UCustomizableObjectMacro>(this, MacroName, RF_Transactional | RF_Public);
+	UCustomizableObjectGraph* NewGraph = NewObject<UCustomizableObjectGraph>(NewMacro, NAME_None, RF_Transactional);
+
+	NewGraph->AddEssentialGraphNodes();
+
+	NewMacro->Graph = NewGraph;
+	NewMacro->Name = MacroName;
+	Macros.Add(NewMacro);
+
+	return NewMacro;
+}
+
+
+void UCustomizableObjectMacroLibrary::RemoveMacro(UCustomizableObjectMacro* MacroToRemove)
+{
+	if (Macros.Contains(MacroToRemove))
+	{
+		Macros.Remove(MacroToRemove);
+	}
+}
+
+
+UCustomizableObjectMacroInputOutput* UCustomizableObjectMacro::AddVariable(ECOMacroIOType VarType)
+{
+	const UEdGraphSchema_CustomizableObject* Schema = Cast<UEdGraphSchema_CustomizableObject>(Graph->GetSchema());
+	check(Schema);
+
+	FName BaseName = "NewVar";
+	FName BaseType = Schema->PC_Mesh;
+	FName VariableName = MakeUniqueObjectName(this, UCustomizableObjectMacroInputOutput::StaticClass(), FName(BaseName));
+
+	UCustomizableObjectMacroInputOutput* NewVariable = NewObject<UCustomizableObjectMacroInputOutput>(this, VariableName, RF_Transactional);
+	NewVariable->Type = VarType;
+	NewVariable->PinCategoryType = BaseType;
+	NewVariable->Name = VariableName;
+	NewVariable->UniqueId = FGuid::NewGuid();
+
+	InputOutputs.Add(NewVariable);
+
+	return NewVariable;
+}
+
+
+void UCustomizableObjectMacro::RemoveVariable(UCustomizableObjectMacroInputOutput* Variable)
+{
+	if (InputOutputs.Contains(Variable))
+	{
+		InputOutputs.Remove(Variable);
+	}
+}
+
+const UCustomizableObjectNodeTunnel* UCustomizableObjectMacro::GetIONode(ECOMacroIOType Type) const
+{
+	const UCustomizableObjectNodeTunnel* OutNode = nullptr;
+
+	if (Graph)
+	{
+		TArray<const UCustomizableObjectNodeTunnel*> IONodes;
+		Graph->GetNodesOfClass<const UCustomizableObjectNodeTunnel>(IONodes);
+		check(IONodes.Num() == 2);
+
+		for (const UCustomizableObjectNodeTunnel* IONode : IONodes)
+		{
+			if (IONode->bIsInputNode && Type == ECOMacroIOType::COMVT_Input)
+			{
+				OutNode = IONode;
+				break;
+			}
+			else if(!IONode->bIsInputNode && Type == ECOMacroIOType::COMVT_Output)
+			{
+				OutNode = IONode;
+				break;
+			}
+		}
+	}
+
+	return OutNode;
+}
+
+#undef LOCTEXT_NAMESPACE
diff --git a/Engine/Plugins/Mutable/Source/CustomizableObjectEditor/Private/MuCOE/CustomizableObjectMacroLibrary/CustomizableObjectMacroLibrary.h b/Engine/Plugins/Mutable/Source/CustomizableObjectEditor/Private/MuCOE/CustomizableObjectMacroLibrary/CustomizableObjectMacroLibrary.h
new file mode 100644
--- /dev/null
+++ b/Engine/Plugins/Mutable/Source/CustomizableObjectEditor/Private/MuCOE/CustomizableObjectMacroLibrary/CustomizableObjectMacroLibrary.h
@@ -0,0 +1,103 @@
+// Copyright Epic Games, Inc. All Rights Reserved.
+
+#pragma once
+
+#include "CustomizableObjectMacroLibrary.generated.h"
+
+#ifndef ENABLE_MUTABLE_MACROS
+	#define ENABLE_MUTABLE_MACROS 0
+#endif // !ENABLE_MUTABLE_MACROS
+
+class UEdGraph;
+class UCustomizableObjectNodeTunnel;
+
+UENUM()
+enum class ECOMacroIOType : uint8
+{
+	COMVT_Input,
+	COMVT_Output
+};
+
+UCLASS()
+class UCustomizableObjectMacroInputOutput : public UObject
+{
+	GENERATED_BODY()
+
+public:
+
+	/** Name of the variable and pin that represents it. */
+	UPROPERTY()
+	FName Name;
+
+	/** Whether it is an input or output variable. */
+	UPROPERTY()
+	ECOMacroIOType Type = ECOMacroIOType::COMVT_Input;
+
+	/** Type of the variable and pin. */
+	UPROPERTY()
+	FName PinCategoryType;
+
+	/** Unique id that identifies this variable. Usefull for pin reconstruction and name repetitions. */
+	UPROPERTY()
+	FGuid UniqueId;
+};
+
+UCLASS()
+class CUSTOMIZABLEOBJECTEDITOR_API UCustomizableObjectMacro : public UObject
+{
+	GENERATED_BODY()
+
+public:
+
+	/** Adds a new variable to the specified Macro. */
+	UCustomizableObjectMacroInputOutput* AddVariable(ECOMacroIOType VarType);
+
+	/** Removes a variable from the specified Macro. */
+	void RemoveVariable(UCustomizableObjectMacroInputOutput* Variable);
+
+	/** Returns the input or output node */
+	const UCustomizableObjectNodeTunnel* GetIONode(ECOMacroIOType Type) const;
+
+public:
+
+	/** Name of the Macro. */
+	UPROPERTY(EditAnywhere, Category = Macro)
+	FName Name;
+
+	/** Description of what this macro does. */
+	UPROPERTY(EditAnywhere, Category = Macro)
+	FString Description = "Macro Description";
+
+	/** Container of all input and output variables of this Macro. */
+	UPROPERTY()
+	TArray<TObjectPtr<UCustomizableObjectMacroInputOutput>> InputOutputs;
+
+	/** Graph of the Macro. */
+	UPROPERTY()
+	TObjectPtr<UEdGraph> Graph;
+
+	//TODO(Max): Add a Callback to refresh all nodes that instantiate this macro when this macro changes
+};
+
+/** A Macro Library is an asset that stores reusable mutable graphs. */
+UCLASS()
+class CUSTOMIZABLEOBJECTEDITOR_API UCustomizableObjectMacroLibrary : public UObject
+{
+	GENERATED_BODY()
+
+public:
+
+	UCustomizableObjectMacroLibrary() {};
+
+	/** Creates a new macro with the basic nodes and adds it to the library. */
+	UCustomizableObjectMacro* AddMacro();
+
+	/** Removes the specified macro from the library. */
+	void RemoveMacro(UCustomizableObjectMacro* MacroToRemove);
+
+public:
+
+	/** List of macros. */
+	UPROPERTY();
+	TArray<TObjectPtr<UCustomizableObjectMacro>> Macros;
+};
diff --git a/Engine/Plugins/Mutable/Source/CustomizableObjectEditor/Private/MuCOE/CustomizableObjectMacroLibrary/CustomizableObjectMacroLibraryEditor.cpp b/Engine/Plugins/Mutable/Source/CustomizableObjectEditor/Private/MuCOE/CustomizableObjectMacroLibrary/CustomizableObjectMacroLibraryEditor.cpp
new file mode 100644
--- /dev/null
+++ b/Engine/Plugins/Mutable/Source/CustomizableObjectEditor/Private/MuCOE/CustomizableObjectMacroLibrary/CustomizableObjectMacroLibraryEditor.cpp
@@ -0,0 +1,353 @@
+// Copyright Epic Games, Inc. All Rights Reserved.
+
+#include "CustomizableObjectMacroLibraryEditor.h"
+
+#include "Editor/Transactor.h"
+#include "Framework/Commands/GenericCommands.h"
+#include "Misc/ITransaction.h"
+#include "MuCOE/CustomizableObjectEditorModule.h"
+#include "MuCOE/CustomizableObjectGraph.h"
+#include "MuCOE/Widgets/SCustomizableObjectMacroLibraryList.h"
+#include "MuCOE/Nodes/CustomizableObjectNodeTunnel.h"
+#include "PropertyEditorModule.h"
+#include "Widgets/Docking/SDockTab.h"
+#include "Widgets/SNullWidget.h"
+
+
+#define LOCTEXT_NAMESPACE "FCustomizableObjectMacroLibraryEditor"
+
+/** Tab Ids */
+const FName FCustomizableObjectMacroLibraryEditor::DetailsTabId(TEXT("CustomizableObjectMacroLibrary_Details"));
+const FName FCustomizableObjectMacroLibraryEditor::MacroSelectorTabId(TEXT("CustomizableObjectMacroLibrary_MacroSelector"));
+const FName FCustomizableObjectMacroLibraryEditor::GraphTabId(TEXT("CustomizableObjectMacroLibrary_Graph"));
+
+
+void FCustomizableObjectMacroLibraryEditor::InitEditor(const EToolkitMode::Type Mode, const TSharedPtr<IToolkitHost>& InitToolkitHost, UObject* ObjectToEdit)
+{
+	MacroLibrary = Cast<UCustomizableObjectMacroLibrary>(ObjectToEdit);
+
+	if (MacroLibrary->Macros.Num())
+	{
+		SelectedMacro = MacroLibrary->Macros[0];
+	}
+
+	// Tab Generation
+	const TSharedRef<FTabManager::FLayout> StandaloneDefaultLayout = FTabManager::NewLayout("Standalone_CustomizableObjectMacroLibraryEditor_Layout_v0.1")
+		->AddArea
+		(
+			FTabManager::NewPrimaryArea()->SetOrientation(Orient_Horizontal)
+			->Split
+			(
+				FTabManager::NewStack()
+				->SetSizeCoefficient(0.6f)
+				->SetHideTabWell(true)
+				->AddTab(GraphTabId, ETabState::OpenedTab)
+			)
+			->Split
+			(
+				FTabManager::NewSplitter()->SetOrientation(Orient_Vertical)
+				->SetSizeCoefficient(0.4f)
+				->Split
+				(
+					FTabManager::NewStack()
+					->SetSizeCoefficient(0.5f)
+					->AddTab(DetailsTabId, ETabState::OpenedTab)
+				)
+				->Split
+				(
+					FTabManager::NewStack()
+					->SetSizeCoefficient(0.5f)
+					->AddTab(MacroSelectorTabId, ETabState::OpenedTab)
+				)
+			)
+		);
+
+	const bool bCreateDefaultStandaloneMenu = true;
+	const bool bCreateDefaultToolbar = true;
+	FAssetEditorToolkit::InitAssetEditor(Mode, InitToolkitHost, CustomizableObjectMacroLibraryEditorAppIdentifier, StandaloneDefaultLayout, bCreateDefaultStandaloneMenu, bCreateDefaultToolbar, ObjectToEdit);
+}
+
+
+void FCustomizableObjectMacroLibraryEditor::RegisterTabSpawners(const TSharedRef<class FTabManager>& InTabManager)
+{
+	WorkspaceMenuCategory = InTabManager->AddLocalWorkspaceMenuCategory(LOCTEXT("WorkspaceMenu_CustomizableObjectEditor", "Customizable Object Editor"));
+	TSharedRef<FWorkspaceItem> WorkspaceMenuCategoryRef = WorkspaceMenuCategory.ToSharedRef();
+
+	FAssetEditorToolkit::RegisterTabSpawners(InTabManager);
+
+	InTabManager->RegisterTabSpawner(DetailsTabId,
+		FOnSpawnTab::CreateSP(this, &FCustomizableObjectMacroLibraryEditor::SpawnTab_Details))
+		.SetDisplayName(LOCTEXT("DetailsTab", "Details"))
+		.SetGroup(WorkspaceMenuCategoryRef);
+
+	InTabManager->RegisterTabSpawner(MacroSelectorTabId,
+		FOnSpawnTab::CreateSP(this, &FCustomizableObjectMacroLibraryEditor::SpawnTab_MacroSelector))
+		.SetDisplayName(LOCTEXT("MacroSelectorTab", "MacroSelector"))
+		.SetGroup(WorkspaceMenuCategoryRef);
+
+	InTabManager->RegisterTabSpawner(GraphTabId,
+		FOnSpawnTab::CreateSP(this, &FCustomizableObjectMacroLibraryEditor::SpawnTab_Graph))
+		.SetDisplayName(LOCTEXT("GraphTab", "Graph"))
+		.SetGroup(WorkspaceMenuCategoryRef);
+}
+
+
+void FCustomizableObjectMacroLibraryEditor::UnregisterTabSpawners(const TSharedRef<class FTabManager>& InTabManager)
+{
+	InTabManager->UnregisterTabSpawner(DetailsTabId);
+	InTabManager->UnregisterTabSpawner(MacroSelectorTabId);
+	InTabManager->UnregisterTabSpawner(GraphTabId);
+}
+
+
+TSharedRef<SDockTab> FCustomizableObjectMacroLibraryEditor::SpawnTab_Details(const FSpawnTabArgs& Args)
+{
+	check(Args.GetTabId() == DetailsTabId);
+
+	if (!DetailsView.IsValid())
+	{
+		FPropertyEditorModule& PropPlugin = FModuleManager::LoadModuleChecked<FPropertyEditorModule>("PropertyEditor");
+
+		FDetailsViewArgs DetailsViewArgs;
+		DetailsViewArgs.NotifyHook = this;
+		DetailsViewArgs.NameAreaSettings = FDetailsViewArgs::ENameAreaSettings::HideNameArea;
+		DetailsViewArgs.bAllowSearch = true;
+		DetailsViewArgs.bShowObjectLabel = false;
+		DetailsViewArgs.bShowScrollBar = false;
+
+		DetailsView = PropPlugin.CreateDetailView(DetailsViewArgs);
+		DetailsView->SetObject(SelectedMacro, true);
+	}
+
+	return SNew(SDockTab)
+	.Label(LOCTEXT("DetailsDockTab", "Details"))
+	.TabColorScale(GetTabColorScale())
+	[
+		DetailsView.ToSharedRef()
+	];
+}
+
+
+TSharedRef<SDockTab> FCustomizableObjectMacroLibraryEditor::SpawnTab_MacroSelector(const FSpawnTabArgs& Args)
+{
+	check(Args.GetTabId() == MacroSelectorTabId);
+
+	if (!MacroSelector.IsValid())
+	{
+		SAssignNew(MacroSelector, SCustomizableObjectMacroLibraryList)
+		.MacroLibrary(MacroLibrary)
+		.OnAddMacroButtonClicked(FExecuteAction::CreateSP(this, &FCustomizableObjectMacroLibraryEditor::OnAddMacroButtonClicked))
+		.OnSelectMacro(SCustomizableObjectMacroLibraryList::FOnSelectMacroDelegate::CreateSP(this, &FCustomizableObjectMacroLibraryEditor::OnMacroSelectionChanged))
+		.OnRemoveMacro(SCustomizableObjectMacroLibraryList::FOnSelectMacroDelegate::CreateSP(this, &FCustomizableObjectMacroLibraryEditor::OnRemoveMacroButtonClicked));
+	}
+
+	return SNew(SDockTab)
+	.Label(LOCTEXT("MacroLibraryDockTab", "Macro Selector"))
+	.TabColorScale(GetTabColorScale())
+	[
+		MacroSelector.ToSharedRef()
+	];
+}
+
+
+TSharedRef<SDockTab> FCustomizableObjectMacroLibraryEditor::SpawnTab_Graph(const FSpawnTabArgs& Args)
+{
+	check(Args.GetTabId() == GraphTabId);
+
+	if (!GraphEditor.IsValid() && SelectedMacro)
+	{
+		// Add Editor custom events
+		SGraphEditor::FGraphEditorEvents InEvents;
+		InEvents.OnSelectionChanged = SGraphEditor::FOnSelectionChanged::CreateSP(this, &FCustomizableObjectMacroLibraryEditor::OnSelectedGraphNodesChanged);
+
+		CreateGraphEditorWidget(SelectedMacro->Graph, InEvents);
+	}
+
+	GraphTab = SNew(SDockTab)
+		.Label(LOCTEXT("GraphDockTab", "Graph"))
+		.TabColorScale(GetTabColorScale());
+
+	if (GraphEditor.IsValid())
+	{
+		GraphTab->SetContent(GraphEditor.ToSharedRef());
+	}
+	else
+	{
+		GraphTab->SetContent(SNullWidget::NullWidget);
+	}
+
+	return GraphTab.ToSharedRef();
+}
+
+
+void FCustomizableObjectMacroLibraryEditor::OnSelectedGraphNodesChanged(const FGraphPanelSelectionSet& NewSelection)
+{
+	if (DetailsView.IsValid())
+	{
+		if (NewSelection.Num() == 1)
+		{
+			if (Cast<UCustomizableObjectNodeTunnel>(NewSelection.Array()[0]) && SelectedMacro)
+			{
+				DetailsView->SetObject(SelectedMacro, true);
+			}
+			else
+			{
+				DetailsView->SetObject(NewSelection.Array()[0], true);
+			}
+		}
+		else
+		{
+			DetailsView->SetObjects(NewSelection.Array(), true);
+		}
+	}
+}
+
+
+FName FCustomizableObjectMacroLibraryEditor::GetToolkitFName() const
+{
+	return FName("CustomizableObjectMacroLibraryEditor");
+}
+
+
+FText FCustomizableObjectMacroLibraryEditor::GetBaseToolkitName() const
+{
+	return LOCTEXT("ToolkitName", "Customizable Object Macro Library Editor");
+}
+
+
+FText FCustomizableObjectMacroLibraryEditor::GetToolkitName() const
+{
+	if (GetEditingObjects().Num() == 1)
+	{
+		return FAssetEditorToolkit::GetToolkitName();
+	}
+
+	return GetBaseToolkitName();
+}
+
+
+FString FCustomizableObjectMacroLibraryEditor::GetWorldCentricTabPrefix() const
+{
+	return LOCTEXT("WorldCentricTabPrefix", "CustomizableObjectMacroLibraryEditor ").ToString();
+}
+
+
+FLinearColor FCustomizableObjectMacroLibraryEditor::GetWorldCentricTabColorScale() const
+{
+	return FLinearColor(0.0f, 0.0f, 0.2f, 0.5f);
+}
+
+
+void FCustomizableObjectMacroLibraryEditor::NotifyPostChange(const FPropertyChangedEvent& PropertyChangedEvent, FProperty* PropertyThatChanged)
+{
+	check(PropertyThatChanged);
+
+	// Refresh macros listview if a macro has been renamed
+	if (PropertyThatChanged->NamePrivate == "Name")
+	{
+		OnMacroNameChanged();
+	}
+}
+
+
+void FCustomizableObjectMacroLibraryEditor::PostUndo(bool bSuccess) 
+{
+	if (bSuccess)
+	{
+		//TODO(Max): Select the correct macro after an undo.
+		/*FTransactionContext TransactionContext = GEditor->GetTransactionName();
+		UE_LOG(LogTemp, Warning, TEXT("Transaction: %s"), *TransactionContext.Title.ToString());*/
+
+		MacroSelector->GenerateRowView();
+	}
+}
+
+
+void FCustomizableObjectMacroLibraryEditor::PostRedo(bool bSuccess)
+{
+	if (bSuccess)
+	{
+		//TODO(Max): Select the correct macro after a redo.
+		if (MacroSelector.IsValid())
+		{
+			MacroSelector->GenerateRowView();
+		}
+	}
+}
+
+
+void FCustomizableObjectMacroLibraryEditor::OnAddMacroButtonClicked()
+{	
+	if (MacroLibrary)
+	{
+		FScopedTransaction LocalTransaction(LOCTEXT("AddMAcroScopTransaction", "Add Macro"));
+		MacroLibrary->Modify();
+
+		UCustomizableObjectMacro* NewMacro = MacroLibrary->AddMacro();
+		
+		SetSelectedMacro(NewMacro);
+	}
+}
+
+
+void FCustomizableObjectMacroLibraryEditor::OnRemoveMacroButtonClicked(UCustomizableObjectMacro* MacroToRemove)
+{
+	if (MacroLibrary)
+	{
+		if (SelectedMacro == MacroToRemove)
+		{
+			SetSelectedMacro(nullptr);
+		}
+
+		MacroLibrary->RemoveMacro(MacroToRemove);
+	}
+}
+
+
+void FCustomizableObjectMacroLibraryEditor::OnMacroNameChanged()
+{
+	if (MacroSelector.IsValid())
+	{
+		MacroSelector->GenerateRowView();
+	}
+}
+
+
+void FCustomizableObjectMacroLibraryEditor::OnMacroSelectionChanged(UCustomizableObjectMacro* NewSelection)
+{
+	SetSelectedMacro(NewSelection);
+}
+
+
+void FCustomizableObjectMacroLibraryEditor::SetSelectedMacro(UCustomizableObjectMacro* NewSelection)
+{
+	if (!NewSelection)
+	{
+		// Removing references
+		SelectedMacro = nullptr;
+		GraphEditor.Reset();
+		GraphTab->SetContent(SNullWidget::NullWidget);
+		DetailsView->SetObject(nullptr, true);
+	}
+	else
+	{
+		if (NewSelection != SelectedMacro)
+		{
+			// Generate a new graph editor
+			GraphEditor.Reset();
+			SelectedMacro = NewSelection;
+
+			SGraphEditor::FGraphEditorEvents InEvents;
+			InEvents.OnSelectionChanged = SGraphEditor::FOnSelectionChanged::CreateSP(this, &FCustomizableObjectMacroLibraryEditor::OnSelectedGraphNodesChanged);
+			InEvents.OnNodeDoubleClicked = FSingleNodeEvent::CreateSP(this, &FCustomizableObjectMacroLibraryEditor::OnNodeDoubleClicked);
+
+			CreateGraphEditorWidget(NewSelection->Graph, InEvents);
+
+			GraphTab->SetContent(GraphEditor.ToSharedRef());
+		}
+	}
+
+	DetailsView->SetObject(SelectedMacro, true);
+}
+
+#undef LOCTEXT_NAMESPACE
diff --git a/Engine/Plugins/Mutable/Source/CustomizableObjectEditor/Private/MuCOE/CustomizableObjectMacroLibrary/CustomizableObjectMacroLibraryEditor.h b/Engine/Plugins/Mutable/Source/CustomizableObjectEditor/Private/MuCOE/CustomizableObjectMacroLibrary/CustomizableObjectMacroLibraryEditor.h
new file mode 100644
--- /dev/null
+++ b/Engine/Plugins/Mutable/Source/CustomizableObjectEditor/Private/MuCOE/CustomizableObjectMacroLibrary/CustomizableObjectMacroLibraryEditor.h
@@ -0,0 +1,84 @@
+// Copyright Epic Games, Inc. All Rights Reserved.
+
+#pragma once
+
+#include "EditorUndoClient.h"
+#include "MuCOE/CustomizableObjectGraphEditorToolkit.h"
+#include "MuCOE/CustomizableObjectMacroLibrary/CustomizableObjectMacroLibrary.h"
+
+class SCustomizableObjectMacroLibraryList;
+class SDockTab;
+
+/** Macro Library Editor. */
+class FCustomizableObjectMacroLibraryEditor : public FCustomizableObjectGraphEditorToolkit, public FSelfRegisteringEditorUndoClient
+{
+public:
+
+	void InitEditor(const EToolkitMode::Type Mode, const TSharedPtr<IToolkitHost>& InitToolkitHost, UObject* ObjectToEdit);
+
+	/** Destructor */
+	virtual ~FCustomizableObjectMacroLibraryEditor() {};
+
+	/** IToolkit interface */
+	virtual void RegisterTabSpawners(const TSharedRef<class FTabManager>& InTabManager) override;
+	virtual void UnregisterTabSpawners(const TSharedRef<class FTabManager>& InTabManager) override;
+	virtual FName GetToolkitFName() const override;
+	virtual FText GetBaseToolkitName() const override;
+	virtual FText GetToolkitName() const override;
+	virtual FString GetWorldCentricTabPrefix() const override;
+	virtual FLinearColor GetWorldCentricTabColorScale() const override;
+
+	/** NotifyHook Interface */
+	virtual void NotifyPostChange(const FPropertyChangedEvent& PropertyChangedEvent, FProperty* PropertyThatChanged) override;
+
+	/** FEditorUndoClient Interface */
+	virtual void PostUndo(bool bSuccess) override;
+	virtual void PostRedo(bool bSuccess) override;
+
+	/** Sets which macro is edited in the editor. */
+	void SetSelectedMacro(UCustomizableObjectMacro* NewSelection);
+
+private:
+
+	/** Tab Spawners */
+	TSharedRef<SDockTab> SpawnTab_Details(const FSpawnTabArgs& Args);
+	TSharedRef<SDockTab> SpawnTab_MacroSelector(const FSpawnTabArgs& Args);
+	TSharedRef<SDockTab> SpawnTab_Graph(const FSpawnTabArgs& Args);
+
+	/** Graph Editor callbacks */
+	void OnSelectedGraphNodesChanged(const FGraphPanelSelectionSet& NewSelection);
+
+	/** Macro List View callbacks */
+	void OnAddMacroButtonClicked();
+	void OnRemoveMacroButtonClicked(UCustomizableObjectMacro* NewSelection);
+	void OnMacroSelectionChanged(UCustomizableObjectMacro* NewSelection);
+	
+	/** Editor Details View callbacks */
+	void OnMacroNameChanged();
+
+public:
+
+	/** Tab Ids */
+	static const FName DetailsTabId;
+	static const FName MacroSelectorTabId;
+	static const FName GraphTabId;
+
+private:
+
+	/** Macro Library being edited */
+	TObjectPtr<UCustomizableObjectMacroLibrary> MacroLibrary = nullptr;
+	
+	/** The currently viewed Macro Graph. */
+	TObjectPtr<UCustomizableObjectMacro> SelectedMacro;
+
+	/** Details view */
+	TSharedPtr<IDetailsView> DetailsView;
+
+	/** Macro List View widget */
+	TSharedPtr<SCustomizableObjectMacroLibraryList> MacroSelector;
+
+	/** Pointer to the tab that contains the graph editor. Needed to refresh its content. */
+	TSharedPtr<SDockTab> GraphTab;
+
+};
+
diff --git a/Engine/Plugins/Mutable/Source/CustomizableObjectEditor/Private/MuCOE/CustomizableObjectMacroLibrary/CustomizableObjectMacroLibraryEditorActions.cpp b/Engine/Plugins/Mutable/Source/CustomizableObjectEditor/Private/MuCOE/CustomizableObjectMacroLibrary/CustomizableObjectMacroLibraryEditorActions.cpp
new file mode 100644
--- /dev/null
+++ b/Engine/Plugins/Mutable/Source/CustomizableObjectEditor/Private/MuCOE/CustomizableObjectMacroLibrary/CustomizableObjectMacroLibraryEditorActions.cpp
@@ -0,0 +1 @@
+// Copyright Epic Games, Inc. All Rights Reserved.
diff --git a/Engine/Plugins/Mutable/Source/CustomizableObjectEditor/Private/MuCOE/CustomizableObjectMacroLibrary/CustomizableObjectMacroLibraryEditorActions.h b/Engine/Plugins/Mutable/Source/CustomizableObjectEditor/Private/MuCOE/CustomizableObjectMacroLibrary/CustomizableObjectMacroLibraryEditorActions.h
new file mode 100644
--- /dev/null
+++ b/Engine/Plugins/Mutable/Source/CustomizableObjectEditor/Private/MuCOE/CustomizableObjectMacroLibrary/CustomizableObjectMacroLibraryEditorActions.h
@@ -0,0 +1 @@
+// Copyright Epic Games, Inc. All Rights Reserved.
diff --git a/Engine/Plugins/Mutable/Source/CustomizableObjectEditor/Private/MuCOE/CustomizableObjectMacroLibrary/CustomizableObjectMacroLibraryFactory.cpp b/Engine/Plugins/Mutable/Source/CustomizableObjectEditor/Private/MuCOE/CustomizableObjectMacroLibrary/CustomizableObjectMacroLibraryFactory.cpp
new file mode 100644
--- /dev/null
+++ b/Engine/Plugins/Mutable/Source/CustomizableObjectEditor/Private/MuCOE/CustomizableObjectMacroLibrary/CustomizableObjectMacroLibraryFactory.cpp
@@ -0,0 +1,30 @@
+// Copyright Epic Games, Inc. All Rights Reserved.
+
+#include "MuCOE/CustomizableObjectMacroLibrary/CustomizableObjectMacroLibraryFactory.h"
+
+#include "MuCOE/CustomizableObjectGraph.h"
+#include "MuCOE/CustomizableObjectMacroLibrary/CustomizableObjectMacroLibrary.h"
+
+UCustomizableObjectMacroLibraryFactory::UCustomizableObjectMacroLibraryFactory() 
+{
+#if ENABLE_MUTABLE_MACROS
+	bCreateNew = true;
+#else
+	bCreateNew = false;
+#endif
+
+	SupportedClass = UCustomizableObjectMacroLibrary::StaticClass();
+	bEditAfterNew = true;
+}
+
+UObject* UCustomizableObjectMacroLibraryFactory::FactoryCreateNew(UClass* Class, UObject* InParent, FName Name, EObjectFlags Flags, UObject* Context, FFeedbackContext* Warn)
+{
+	UCustomizableObjectMacroLibrary* NewMacroLibrary = NewObject<UCustomizableObjectMacroLibrary>(InParent, Class, Name, Flags, Context);
+	
+	if (NewMacroLibrary)
+	{
+		NewMacroLibrary->AddMacro();
+	}
+
+	return NewMacroLibrary;
+}
diff --git a/Engine/Plugins/Mutable/Source/CustomizableObjectEditor/Private/MuCOE/CustomizableObjectMacroLibrary/CustomizableObjectMacroLibraryFactory.h b/Engine/Plugins/Mutable/Source/CustomizableObjectEditor/Private/MuCOE/CustomizableObjectMacroLibrary/CustomizableObjectMacroLibraryFactory.h
new file mode 100644
--- /dev/null
+++ b/Engine/Plugins/Mutable/Source/CustomizableObjectEditor/Private/MuCOE/CustomizableObjectMacroLibrary/CustomizableObjectMacroLibraryFactory.h
@@ -0,0 +1,22 @@
+// Copyright Epic Games, Inc. All Rights Reserved.
+
+#pragma once
+
+#include "Factories/Factory.h"
+#include "CustomizableObjectMacroLibraryFactory.generated.h"
+
+
+UCLASS()
+class UCustomizableObjectMacroLibraryFactory : public UFactory
+{
+	GENERATED_BODY()
+
+public:
+
+	UCustomizableObjectMacroLibraryFactory();
+
+	// UFactory Interface
+	UObject* FactoryCreateNew(UClass* Class, UObject* InParent, FName Name, EObjectFlags Flags, UObject* Context, FFeedbackContext* Warn);
+
+};
+
diff --git a/Engine/Plugins/Mutable/Source/CustomizableObjectEditor/Private/MuCOE/EdGraphSchema_CustomizableObject.cpp b/Engine/Plugins/Mutable/Source/CustomizableObjectEditor/Private/MuCOE/EdGraphSchema_CustomizableObject.cpp
--- a/Engine/Plugins/Mutable/Source/CustomizableObjectEditor/Private/MuCOE/EdGraphSchema_CustomizableObject.cpp
+++ b/Engine/Plugins/Mutable/Source/CustomizableObjectEditor/Private/MuCOE/EdGraphSchema_CustomizableObject.cpp
@@ -13,6 +13,7 @@
 #include "Materials/MaterialInterface.h"
 #include "MuCO/CustomizableObjectExtension.h"
 #include "MuCO/ICustomizableObjectModule.h"
+#include "MuCOE/CustomizableObjectGraph.h"
 #include "MuCOE/ICustomizableObjectEditor.h"
 #include "MuCOE/Nodes/CustomizableObjectNodeAnimationPose.h"
 #include "MuCOE/Nodes/CustomizableObjectNodeColorArithmeticOp.h"
@@ -32,6 +33,7 @@
 #include "MuCOE/Nodes/CustomizableObjectNodeFloatVariation.h"
 #include "MuCOE/Nodes/CustomizableObjectNodeFloatArithmeticOp.h"
 #include "MuCOE/Nodes/CustomizableObjectNodeGroupProjectorParameter.h"
+#include "MuCOE/Nodes/CustomizableObjectNodeMacroInstance.h"
 #include "MuCOE/Nodes/CustomizableObjectNodeMaterialVariation.h"
 #include "MuCOE/Nodes/CustomizableObjectNodeMaterialSwitch.h"
 #include "MuCOE/Nodes/CustomizableObjectNodeModifierExtendMeshSection.h"
@@ -61,6 +63,7 @@
 #include "MuCOE/Nodes/CustomizableObjectNodeTable.h"
 #include "MuCOE/Nodes/CustomizableObjectNodeTexture.h"
 #include "MuCOE/Nodes/CustomizableObjectNodePassThroughTexture.h"
+#include "MuCOE/Nodes/CustomizableObjectNodeStringConstant.h"
 #include "MuCOE/Nodes/CustomizableObjectNodeTextureBinarise.h"
 #include "MuCOE/Nodes/CustomizableObjectNodeTextureColourMap.h"
 #include "MuCOE/Nodes/CustomizableObjectNodeTextureFromChannels.h"
@@ -263,7 +266,11 @@ const FName UEdGraphSchema_CustomizableObject::PC_MaterialAsset("materialAsset")
 const FName UEdGraphSchema_CustomizableObject::PC_Wildcard("wildcard");
 const FName UEdGraphSchema_CustomizableObject::PC_PoseAsset("poseAsset");
 const FName UEdGraphSchema_CustomizableObject::PC_Transform("transform");
+const FName UEdGraphSchema_CustomizableObject::PC_String("string");
 
+// Add more pin types to this array if needed
+const TArray<FName> UEdGraphSchema_CustomizableObject::SupportedMacroPinTypes = { PC_Component, PC_Material, PC_Modifier, PC_Mesh, PC_Image,
+PC_PassThroughImage, PC_Projector, PC_GroupProjector, PC_Color, PC_Float, PC_Enum, PC_Stack, PC_MaterialAsset, PC_PoseAsset, PC_Transform, PC_String };
 
 // Node categories
 const FText UEdGraphSchema_CustomizableObject::NC_Experimental(FText::FromString("Experimental"));
@@ -365,6 +372,16 @@ void UEdGraphSchema_CustomizableObject::GetGraphContextActions(FGraphContextMenu
 	// return true if Filter is passed.
 	const auto Filter = [bDisableFilter](UCustomizableObjectNode* TemplateNode, FGraphContextMenuBuilder& ContextMenuBuilder) -> bool
 	{
+		const UCustomizableObjectGraph* CustomizableObjectGraph = Cast<UCustomizableObjectGraph>(ContextMenuBuilder.CurrentGraph);
+
+		if (CustomizableObjectGraph && TemplateNode && CustomizableObjectGraph->IsMacro())
+		{
+			if (!TemplateNode->IsNodeSupportedInMacros())
+			{
+				return false;
+			}
+		}
+		
 		if (!ContextMenuBuilder.FromPin || bDisableFilter)
 		{
 			return true;
@@ -412,6 +429,11 @@ void UEdGraphSchema_CustomizableObject::GetGraphContextActions(FGraphContextMenu
 			ContextMenuBuilder.CreateTemplateNode<UCustomizableObjectNodeMaterialSwitch>(),
 			ContextMenuBuilder.CreateTemplateNode<UCustomizableObjectNodeCopyMaterial>(),
 			ContextMenuBuilder.CreateTemplateNode<UCustomizableObjectNodeTable>(),
+
+#if	ENABLE_MUTABLE_MACROS
+			ContextMenuBuilder.CreateTemplateNode<UCustomizableObjectNodeMacroInstance>(),
+			ContextMenuBuilder.CreateTemplateNode<UCustomizableObjectNodeStringConstant>(),
+#endif
 		};
 
 		AddNewNodeCategoryActionsFiltered(ObjectTemplateNodes, ContextMenuBuilder, TEXT("Object"), GeneralGrouping, Filter);
@@ -795,6 +817,10 @@ FLinearColor UEdGraphSchema_CustomizableObject::GetPinTypeColor(const FName& Typ
 	{
 		return FLinearColor(FColorList::Orange); 
 	}
+	else if (TypeString == PC_String)
+	{
+		return FLinearColor(0.700000f, 0.010000f, 0.660000f, 1.000000f); // Dark Pink
+	}
 
 	for (const FRegisteredCustomizableObjectPinType& PinType : ICustomizableObjectModule::Get().GetExtendedPinTypes())
 	{
@@ -1399,6 +1425,10 @@ FName UEdGraphSchema_CustomizableObject::GetPinCategoryName(const FName& PinCate
 	{
 		return TEXT("Transform");
 	}
+	else if (PinCategory == UEdGraphSchema_CustomizableObject::PC_String)
+	{
+		return TEXT("String");
+	}
 	else
 	{
 		for (const FRegisteredCustomizableObjectPinType& PinType : ICustomizableObjectModule::Get().GetExtendedPinTypes())
@@ -1494,6 +1524,10 @@ FText UEdGraphSchema_CustomizableObject::GetPinCategoryFriendlyName(const FName&
 	{
 		return LOCTEXT("Transform_Pin_Category", "Transform");
 	}
+	else if (PinCategory == UEdGraphSchema_CustomizableObject::PC_String)
+	{
+		return LOCTEXT("String_Pin_Category", "String");
+	}
 	else
 	{
 		for (const FRegisteredCustomizableObjectPinType& PinType : ICustomizableObjectModule::Get().GetExtendedPinTypes())
diff --git a/Engine/Plugins/Mutable/Source/CustomizableObjectEditor/Private/MuCOE/ICustomizableObjectEditor.h b/Engine/Plugins/Mutable/Source/CustomizableObjectEditor/Private/MuCOE/ICustomizableObjectEditor.h
--- a/Engine/Plugins/Mutable/Source/CustomizableObjectEditor/Private/MuCOE/ICustomizableObjectEditor.h
+++ b/Engine/Plugins/Mutable/Source/CustomizableObjectEditor/Private/MuCOE/ICustomizableObjectEditor.h
@@ -27,6 +27,7 @@ public:
 
 	virtual UEdGraphNode* CreateCommentBox(const FVector2D& NodePos) { return nullptr; }
 
+	//TODO(Max): Keep this method
 	virtual void UpdateGraphNodeProperties() {}
 
 	// Make sure the graph editor selects a specific node
@@ -34,6 +35,7 @@ public:
 
 	virtual void ReconstructNode(UEdGraphNode* Node) {}
 
+	//TODO(Max): Keep this method
 	/** Reconstructs all child the nodes that match the given type.
 	 * @param StartNode Root node to start the graph traversal. This one also will be reconstructed.
 	 * @param NodeType Node types to reconstruct. */
diff --git a/Engine/Plugins/Mutable/Source/CustomizableObjectEditor/Private/MuCOE/Nodes/CustomizableObjectNodeExposePin.h b/Engine/Plugins/Mutable/Source/CustomizableObjectEditor/Private/MuCOE/Nodes/CustomizableObjectNodeExposePin.h
--- a/Engine/Plugins/Mutable/Source/CustomizableObjectEditor/Private/MuCOE/Nodes/CustomizableObjectNodeExposePin.h
+++ b/Engine/Plugins/Mutable/Source/CustomizableObjectEditor/Private/MuCOE/Nodes/CustomizableObjectNodeExposePin.h
@@ -39,6 +39,7 @@ public:
 	// UCustomizableObjectNode interface
 	virtual bool CanConnect(const UEdGraphPin* InOwnedInputPin, const UEdGraphPin* InOutputPin, bool& bOutIsOtherNodeBlocklisted, bool& bOutArePinsCompatible) const override;
 	virtual void AllocateDefaultPins(UCustomizableObjectNodeRemapPins* RemapPins) override;
+	virtual bool IsNodeSupportedInMacros() override { return false; }
 
 	/** Return the Expose Pin Node expose pin name. */
 	FString GetNodeName() const;
diff --git a/Engine/Plugins/Mutable/Source/CustomizableObjectEditor/Private/MuCOE/Nodes/CustomizableObjectNodeExternalPin.h b/Engine/Plugins/Mutable/Source/CustomizableObjectEditor/Private/MuCOE/Nodes/CustomizableObjectNodeExternalPin.h
--- a/Engine/Plugins/Mutable/Source/CustomizableObjectEditor/Private/MuCOE/Nodes/CustomizableObjectNodeExternalPin.h
+++ b/Engine/Plugins/Mutable/Source/CustomizableObjectEditor/Private/MuCOE/Nodes/CustomizableObjectNodeExternalPin.h
@@ -43,6 +43,7 @@ public:
 	virtual void UpdateReferencedNodeId(const FGuid& NewGuid) override;
 	virtual void PreEditChange(FProperty* PropertyAboutToChange) override;
 	virtual void PostEditChangeProperty(FPropertyChangedEvent& PropertyChangedEvent) override;
+	virtual bool IsNodeSupportedInMacros() override { return false; }
 	
 	// Own interface
 	
diff --git a/Engine/Plugins/Mutable/Source/CustomizableObjectEditor/Private/MuCOE/Nodes/CustomizableObjectNodeMacroInstance.cpp b/Engine/Plugins/Mutable/Source/CustomizableObjectEditor/Private/MuCOE/Nodes/CustomizableObjectNodeMacroInstance.cpp
new file mode 100644
--- /dev/null
+++ b/Engine/Plugins/Mutable/Source/CustomizableObjectEditor/Private/MuCOE/Nodes/CustomizableObjectNodeMacroInstance.cpp
@@ -0,0 +1,162 @@
+// Copyright Epic Games, Inc. All Rights Reserved.
+
+#include "MuCOE/Nodes/CustomizableObjectNodeMacroInstance.h"
+
+#include "Editor.h"
+#include "MuCOE/CustomizableObjectMacroLibrary/CustomizableObjectMacroLibraryEditor.h"
+#include "MuCOE/EdGraphSchema_CustomizableObject.h"
+#include "Subsystems/AssetEditorSubsystem.h"
+
+#define LOCTEXT_NAMESPACE "CustomizableObjectNodeMacroInstance"
+
+
+bool UCustomizableObjectNodeMacroInstanceRemapPins::Equal(const UCustomizableObjectNode& Node, const UEdGraphPin& OldPin, const UEdGraphPin& NewPin) const
+{
+	const UCustomizableObjectNodeMacroInstancePinData* PinDataOldPin = Cast<UCustomizableObjectNodeMacroInstancePinData>(Node.GetPinData(OldPin));
+	const UCustomizableObjectNodeMacroInstancePinData* PinDataNewPin = Cast<UCustomizableObjectNodeMacroInstancePinData>(Node.GetPinData(NewPin));
+
+	return PinDataOldPin->VariableId == PinDataNewPin->VariableId && OldPin.PinType.PinCategory == NewPin.PinType.PinCategory;
+}
+
+
+void UCustomizableObjectNodeMacroInstanceRemapPins::RemapPins(const UCustomizableObjectNode& Node, const TArray<UEdGraphPin*>& OldPins, const TArray<UEdGraphPin*>& NewPins, TMap<UEdGraphPin*, UEdGraphPin*>& PinsToRemap, TArray<UEdGraphPin*>& PinsToOrphan)
+{
+	for (UEdGraphPin* OldPin : OldPins)
+	{
+		bool bFound = false;
+
+		for (UEdGraphPin* NewPin : NewPins)
+		{
+			if (Equal(Node, *OldPin, *NewPin))
+			{
+				bFound = true;
+
+				PinsToRemap.Add(OldPin, NewPin);
+				break;
+			}
+		}
+
+		if (!bFound && OldPin->LinkedTo.Num())
+		{
+			PinsToOrphan.Add(OldPin);
+		}
+	}
+}
+
+
+UCustomizableObjectNodeMacroInstanceRemapPins* UCustomizableObjectNodeMacroInstance::CreateRemapPinsDefault() const
+{
+	return NewObject<UCustomizableObjectNodeMacroInstanceRemapPins>();
+}
+
+
+FText UCustomizableObjectNodeMacroInstance::GetNodeTitle(ENodeTitleType::Type TitleType) const
+{
+	if (TitleType == ENodeTitleType::ListView || !ParentMacro)
+	{
+		return LOCTEXT("MacroInstanceNodeTitle", "Macro Instance");
+	}
+	else
+	{
+		FFormatNamedArguments Args;
+		Args.Add(TEXT("MacroInstanceName"), FText::FromName(ParentMacro->Name));
+
+		return FText::Format(LOCTEXT("MeshSection_Title", "{MacroInstanceName}\nMacro Instance"), Args);
+	}
+}
+
+
+FText UCustomizableObjectNodeMacroInstance::GetTooltipText() const
+{
+	if (ParentMacroLibrary && ParentMacro)
+	{
+		FString Tooltip = ParentMacro->Name.ToString() + ":\n" + ParentMacro->Description;
+		return FText::FromString(Tooltip);
+	}
+
+	return LOCTEXT("MacroInstanceNodeTooltip", "Macro Instance");
+}
+
+
+FLinearColor UCustomizableObjectNodeMacroInstance::GetNodeTitleColor() const
+{
+	return FLinearColor(0.15f, 0.15f, 0.15f);
+}
+
+
+bool UCustomizableObjectNodeMacroInstance::GetCanRenameNode() const
+{
+	return false;
+}
+
+
+bool UCustomizableObjectNodeMacroInstance::CanUserDeleteNode() const
+{
+	return true;
+}
+
+
+UObject* UCustomizableObjectNodeMacroInstance::GetJumpTargetForDoubleClick() const
+{
+	if (ParentMacroLibrary && ParentMacro)
+	{
+		return ParentMacroLibrary;
+	}
+
+	return nullptr;
+}
+
+
+bool UCustomizableObjectNodeMacroInstance::CanJumpToDefinition() const
+{
+	return ParentMacroLibrary && ParentMacro;
+}
+
+
+void UCustomizableObjectNodeMacroInstance::JumpToDefinition() const
+{
+	// Open the editor for the Macro library
+	GEditor->GetEditorSubsystem<UAssetEditorSubsystem>()->OpenEditorForAsset(GetJumpTargetForDoubleClick());
+
+	// Find the editor we just opened
+	IAssetEditorInstance* AssetEditor = GEditor->GetEditorSubsystem<UAssetEditorSubsystem>()->FindEditorForAsset(ParentMacroLibrary, false);
+
+	// Set the Parent Macro to edit
+	if (FCustomizableObjectMacroLibraryEditor* Editor = static_cast<FCustomizableObjectMacroLibraryEditor*>(AssetEditor))
+	{
+		Editor->SetSelectedMacro(ParentMacro);
+	}
+}
+
+
+void UCustomizableObjectNodeMacroInstance::AllocateDefaultPins(UCustomizableObjectNodeRemapPins* RemapPins)
+{
+	if (ParentMacroLibrary && ParentMacro)
+	{
+		for (const UCustomizableObjectMacroInputOutput* Variable : ParentMacro->InputOutputs)
+		{
+			EEdGraphPinDirection PinDirection = Variable->Type == ECOMacroIOType::COMVT_Input ? EEdGraphPinDirection::EGPD_Input : EEdGraphPinDirection::EGPD_Output;
+			UCustomizableObjectNodeMacroInstancePinData* PinData = NewObject<UCustomizableObjectNodeMacroInstancePinData>(this);
+			PinData->VariableId = Variable->UniqueId;
+
+			const FName PinName = Variable->Name;
+			UEdGraphPin* Pin = CustomCreatePin(PinDirection, Variable->PinCategoryType, PinName, PinData);
+		}
+	}
+}
+
+
+void UCustomizableObjectNodeMacroInstance::PostEditChangeProperty(FPropertyChangedEvent& PropertyChangedEvent)
+{
+	Super::PostEditChangeProperty(PropertyChangedEvent);
+
+	const FName PropertyName = PropertyChangedEvent.MemberProperty ? PropertyChangedEvent.MemberProperty->GetFName() : NAME_None;
+
+	if (PropertyName == GET_MEMBER_NAME_CHECKED(UCustomizableObjectNodeMacroInstance, ParentMacroLibrary))
+	{
+		ParentMacro = nullptr;
+		ReconstructNode();
+	}
+}
+
+#undef LOCTEXT_NAMESPACE
\ No newline at end of file
diff --git a/Engine/Plugins/Mutable/Source/CustomizableObjectEditor/Private/MuCOE/Nodes/CustomizableObjectNodeMacroInstance.h b/Engine/Plugins/Mutable/Source/CustomizableObjectEditor/Private/MuCOE/Nodes/CustomizableObjectNodeMacroInstance.h
new file mode 100644
--- /dev/null
+++ b/Engine/Plugins/Mutable/Source/CustomizableObjectEditor/Private/MuCOE/Nodes/CustomizableObjectNodeMacroInstance.h
@@ -0,0 +1,73 @@
+// Copyright Epic Games, Inc. All Rights Reserved.
+
+#pragma once
+
+#include "MuCOE/CustomizableObjectMacroLibrary/CustomizableObjectMacroLibrary.h"
+#include "MuCOE/Nodes/CustomizableObjectNode.h"
+#include "MuCOE/RemapPins/CustomizableObjectNodeRemapPinsByName.h"
+
+#include "CustomizableObjectNodeMacroInstance.generated.h"
+
+
+/** Base class for all Macro Instance Pins. */
+UCLASS()
+class CUSTOMIZABLEOBJECTEDITOR_API UCustomizableObjectNodeMacroInstancePinData : public UCustomizableObjectNodePinData
+{
+	GENERATED_BODY()
+
+public:
+
+	/** Id of the variable associated to a Macro Instance node pin */
+	UPROPERTY()
+	FGuid VariableId;
+
+};
+
+
+UCLASS()
+class UCustomizableObjectNodeMacroInstanceRemapPins : public UCustomizableObjectNodeRemapPinsByName
+{
+	GENERATED_BODY()
+public:
+
+	// Specific method to decide when two pins are equal
+	virtual bool Equal(const UCustomizableObjectNode& Node, const UEdGraphPin& OldPin, const UEdGraphPin& NewPin) const override;
+
+	// Method to use in the RemapPins step of the node reconstruction process
+	virtual void RemapPins(const UCustomizableObjectNode& Node, const TArray<UEdGraphPin*>& OldPins, const TArray<UEdGraphPin*>& NewPins, TMap<UEdGraphPin*, UEdGraphPin*>& PinsToRemap, TArray<UEdGraphPin*>& PinsToOrphan) override;
+};
+
+
+
+UCLASS()
+class CUSTOMIZABLEOBJECTEDITOR_API UCustomizableObjectNodeMacroInstance : public UCustomizableObjectNode
+{
+public:
+	GENERATED_BODY()
+
+	/** Macro Library where the macro to instantiate belongs. */
+	UPROPERTY(EditAnywhere, Category = MacroInstance)
+	TObjectPtr<UCustomizableObjectMacroLibrary> ParentMacroLibrary;
+
+	/** Macro that represent that instantiates the node. */
+	UPROPERTY(EditAnywhere, Category = MacroInstance)
+	TObjectPtr<UCustomizableObjectMacro> ParentMacro;
+
+public:
+
+	// EdGraphNode interface
+	virtual FText GetNodeTitle(ENodeTitleType::Type TitleType) const override;
+	virtual FText GetTooltipText() const override;
+	virtual FLinearColor GetNodeTitleColor() const override;
+	virtual bool GetCanRenameNode() const override;
+	virtual bool CanUserDeleteNode() const override;
+	UObject* GetJumpTargetForDoubleClick() const override;
+	virtual bool CanJumpToDefinition() const override;
+	virtual void JumpToDefinition() const override;
+
+	// UCustomizableObjectNode interface
+	void AllocateDefaultPins(UCustomizableObjectNodeRemapPins* RemapPins) override;
+	virtual UCustomizableObjectNodeMacroInstanceRemapPins* CreateRemapPinsDefault() const override;
+	virtual void PostEditChangeProperty(FPropertyChangedEvent& PropertyChangedEvent) override;
+
+};
\ No newline at end of file
diff --git a/Engine/Plugins/Mutable/Source/CustomizableObjectEditor/Private/MuCOE/Nodes/CustomizableObjectNodeMacroInstanceDetails.cpp b/Engine/Plugins/Mutable/Source/CustomizableObjectEditor/Private/MuCOE/Nodes/CustomizableObjectNodeMacroInstanceDetails.cpp
new file mode 100644
--- /dev/null
+++ b/Engine/Plugins/Mutable/Source/CustomizableObjectEditor/Private/MuCOE/Nodes/CustomizableObjectNodeMacroInstanceDetails.cpp
@@ -0,0 +1,125 @@
+// Copyright Epic Games, Inc. All Rights Reserved.
+
+#include "MuCOE/Nodes/CustomizableObjectNodeMacroInstanceDetails.h"
+
+#include "DetailCategoryBuilder.h"
+#include "DetailLayoutBuilder.h"
+#include "DetailWidgetRow.h"
+#include "IDetailsView.h"
+#include "Misc/Attribute.h"
+#include "Widgets/Text/STextBlock.h"
+#include "Widgets/Input/SComboBox.h"
+
+#define LOCTEXT_NAMESPACE "CustomizableObjectNodeMacroInstanceDetails"
+
+
+TSharedRef<IDetailCustomization> FCustomizableObjectNodeMacroInstanceDetails::MakeInstance()
+{
+	return MakeShareable(new FCustomizableObjectNodeMacroInstanceDetails);
+}
+
+
+void FCustomizableObjectNodeMacroInstanceDetails::CustomizeDetails(IDetailLayoutBuilder& DetailBuilder)
+{
+	FCustomizableObjectNodeDetails::CustomizeDetails(DetailBuilder);
+
+	TSharedPtr<const IDetailsView> DetailsView = DetailBuilder.GetDetailsViewSharedPtr();
+	if (DetailsView && DetailsView->GetSelectedObjects().Num())
+	{
+		Node = Cast<UCustomizableObjectNodeMacroInstance>(DetailsView->GetSelectedObjects()[0].Get());
+	}
+
+	IDetailCategoryBuilder& MacroInstanceCategory = DetailBuilder.EditCategory("MacroInstance");
+
+	if (!Node)
+	{
+		MacroInstanceCategory.AddCustomRow(FText::FromString("MacroInstaceDetailsError"))
+		[
+			SNew(STextBlock).Text(LOCTEXT("MacroInstaceDetailsErrorMessage", "Error: Node not found."))
+		];
+
+		return;
+	}
+
+	TSharedPtr<IPropertyHandle> PropertyHandle = DetailBuilder.GetProperty(GET_MEMBER_NAME_STRING_CHECKED(UCustomizableObjectNodeMacroInstance, ParentMacro));
+
+	SelectedSource = GenerateComboboxSource();
+
+	// Macro selector widget
+	DetailBuilder.EditDefaultProperty(PropertyHandle)->CustomWidget()
+	.OverrideResetToDefault(FResetToDefaultOverride::Hide())
+	.IsValueEnabled(TAttribute<bool>::CreateSPLambda(this, [this]() { return Node->ParentMacroLibrary != nullptr; }))
+	.NameContent()
+	[
+		SNew(STextBlock).Text(LOCTEXT("ParentMacroText", "Parent Macro"))
+	]
+	.ValueContent()
+	[
+		SNew(SComboBox<TSharedPtr<FMacroSelectorItem>>)
+		.OptionsSource(&ComboboxSource)
+		.OnComboBoxOpening_Lambda([this]() {GenerateComboboxSource(); })
+		.InitiallySelectedItem(SelectedSource)
+		.OnGenerateWidget_Lambda([](TSharedPtr<FMacroSelectorItem> Item)
+		{
+			return SNew(STextBlock).Text(FText::FromName(Item->Macro->Name));
+		})
+		.OnSelectionChanged_Lambda([this](TSharedPtr<FMacroSelectorItem> Item, ESelectInfo::Type)
+		{
+			if (Item.IsValid())
+			{
+				SelectedSource = Item;
+				Node->ParentMacro = Item->Macro;
+				Node->ReconstructNode();
+			}
+		})
+		.Content()
+		[
+			SNew(STextBlock).MinDesiredWidth(200)
+			.Font(IDetailLayoutBuilder::GetDetailFont())
+			.Text(this, &FCustomizableObjectNodeMacroInstanceDetails::GetSelectedMacroName)
+		]
+	];
+}
+
+
+TSharedPtr<FMacroSelectorItem> FCustomizableObjectNodeMacroInstanceDetails::GenerateComboboxSource()
+{
+	TSharedPtr<FMacroSelectorItem> CurrentSelected = nullptr;
+	ComboboxSource.Empty();
+
+	if (Node && Node->ParentMacroLibrary)
+	{
+		for (UCustomizableObjectMacro* Macro : Node->ParentMacroLibrary->Macros)
+		{
+			if (Macro)
+			{
+				TSharedPtr<FMacroSelectorItem> MacroSelector = MakeShareable(new FMacroSelectorItem());
+				MacroSelector->Macro = Macro;
+
+				ComboboxSource.Add(MacroSelector);
+
+				if (Node->ParentMacro && Node->ParentMacro == Macro)
+				{
+					CurrentSelected = MacroSelector;
+				}
+			}
+		}
+	}
+
+	return CurrentSelected;
+}
+
+
+FText FCustomizableObjectNodeMacroInstanceDetails::GetSelectedMacroName() const
+{
+	FText Text = FText::FromString("- Nothing Selected -");
+
+	if (Node->ParentMacro)
+	{
+		Text = FText::FromName(Node->ParentMacro->Name);
+	}
+
+	return Text;
+}
+
+#undef LOCTEXT_NAMESPACE
\ No newline at end of file
diff --git a/Engine/Plugins/Mutable/Source/CustomizableObjectEditor/Private/MuCOE/Nodes/CustomizableObjectNodeMacroInstanceDetails.h b/Engine/Plugins/Mutable/Source/CustomizableObjectEditor/Private/MuCOE/Nodes/CustomizableObjectNodeMacroInstanceDetails.h
new file mode 100644
--- /dev/null
+++ b/Engine/Plugins/Mutable/Source/CustomizableObjectEditor/Private/MuCOE/Nodes/CustomizableObjectNodeMacroInstanceDetails.h
@@ -0,0 +1,49 @@
+// Copyright Epic Games, Inc. All Rights Reserved.
+
+#pragma once
+
+#include "CustomizableObjectNodeDetails.h"
+#include "MuCOE/Nodes/CustomizableObjectNodeMacroInstance.h"
+
+class UCustomizableObjectMacro;
+
+/** Data container for the macro selector combobox */
+struct FMacroSelectorItem
+{
+	TObjectPtr<UCustomizableObjectMacro> Macro;
+};
+
+
+/** Details View of UCustomizableObjectNodeMacroInstance */
+class FCustomizableObjectNodeMacroInstanceDetails : public FCustomizableObjectNodeDetails
+{
+public:
+
+	static TSharedRef<IDetailCustomization> MakeInstance();
+
+	// ILayoutDetails interface
+	/** Do not use. Add details customization in the other CustomizeDetails signature. */
+	virtual void CustomizeDetails(IDetailLayoutBuilder& DetailBuilder) override;
+
+	/** Generates the soruce of the combobox used to select a Macro. */
+	TSharedPtr<FMacroSelectorItem> GenerateComboboxSource();
+
+private:
+
+	FText GetSelectedMacroName() const;
+
+private:
+
+	/** Pointer to the Macro Instance Node. */
+	TObjectPtr<UCustomizableObjectNodeMacroInstance> Node;
+
+	/** Combobox Source array. */
+	TArray<TSharedPtr<FMacroSelectorItem>> ComboboxSource;
+
+	/** Current selected Macro in the Combobox. */
+	TSharedPtr<FMacroSelectorItem> SelectedSource;
+
+};
+
+
+#undef LOCTEXT_NAMESPACE
\ No newline at end of file
diff --git a/Engine/Plugins/Mutable/Source/CustomizableObjectEditor/Private/MuCOE/Nodes/CustomizableObjectNodeStringConstant.cpp b/Engine/Plugins/Mutable/Source/CustomizableObjectEditor/Private/MuCOE/Nodes/CustomizableObjectNodeStringConstant.cpp
new file mode 100644
--- /dev/null
+++ b/Engine/Plugins/Mutable/Source/CustomizableObjectEditor/Private/MuCOE/Nodes/CustomizableObjectNodeStringConstant.cpp
@@ -0,0 +1,44 @@
+// Copyright Epic Games, Inc. All Rights Reserved.
+
+#include "MuCOE/Nodes/CustomizableObjectNodeStringConstant.h"
+
+#include "Editor.h"
+#include "MuCOE/EdGraphSchema_CustomizableObject.h"
+#include "Subsystems/AssetEditorSubsystem.h"
+
+#define LOCTEXT_NAMESPACE "CustomizableObjectNodeStringConstant"
+
+
+FText UCustomizableObjectNodeStringConstant::GetNodeTitle(ENodeTitleType::Type TitleType) const
+{
+	return LOCTEXT("ConstantStringNodeTitle", "String");
+}
+
+
+FText UCustomizableObjectNodeStringConstant::GetTooltipText() const
+{
+	return LOCTEXT("ConstantStringNodeTooltip", "String Node");
+}
+
+
+FLinearColor UCustomizableObjectNodeStringConstant::GetNodeTitleColor() const
+{
+	return FLinearColor(0.7f, 0.01f, 0.66f);
+}
+
+
+bool UCustomizableObjectNodeStringConstant::GetCanRenameNode() const
+{
+	return false;
+}
+
+
+void UCustomizableObjectNodeStringConstant::AllocateDefaultPins(UCustomizableObjectNodeRemapPins* RemapPins)
+{
+	const UEdGraphSchema_CustomizableObject* Schema = GetDefault<UEdGraphSchema_CustomizableObject>();
+
+	UEdGraphPin* ValuePin = CustomCreatePin(EGPD_Output, Schema->PC_String, FName("Value"));
+	ValuePin->bDefaultValueIsIgnored = true;
+}
+
+#undef LOCTEXT_NAMESPACE
\ No newline at end of file
diff --git a/Engine/Plugins/Mutable/Source/CustomizableObjectEditor/Private/MuCOE/Nodes/CustomizableObjectNodeStringConstant.h b/Engine/Plugins/Mutable/Source/CustomizableObjectEditor/Private/MuCOE/Nodes/CustomizableObjectNodeStringConstant.h
new file mode 100644
--- /dev/null
+++ b/Engine/Plugins/Mutable/Source/CustomizableObjectEditor/Private/MuCOE/Nodes/CustomizableObjectNodeStringConstant.h
@@ -0,0 +1,31 @@
+// Copyright Epic Games, Inc. All Rights Reserved.
+
+#pragma once
+
+#include "MuCOE/Nodes/CustomizableObjectNode.h"
+
+#include "CustomizableObjectNodeStringConstant.generated.h"
+
+
+UCLASS()
+class CUSTOMIZABLEOBJECTEDITOR_API UCustomizableObjectNodeStringConstant : public UCustomizableObjectNode
+{
+public:
+	GENERATED_BODY()
+
+	/** Macro Library where the macro to instantiate belongs. */
+	UPROPERTY(EditAnywhere, Category = CustomizableObject)
+	FString Value;
+
+public:
+
+	// EdGraphNode interface
+	virtual FText GetNodeTitle(ENodeTitleType::Type TitleType) const override;
+	virtual FText GetTooltipText() const override;
+	virtual FLinearColor GetNodeTitleColor() const override;
+	virtual bool GetCanRenameNode() const override;
+
+	// UCustomizableObjectNode interface
+	void AllocateDefaultPins(UCustomizableObjectNodeRemapPins* RemapPins) override;
+
+};
\ No newline at end of file
diff --git a/Engine/Plugins/Mutable/Source/CustomizableObjectEditor/Private/MuCOE/Nodes/CustomizableObjectNodeTunnel.cpp b/Engine/Plugins/Mutable/Source/CustomizableObjectEditor/Private/MuCOE/Nodes/CustomizableObjectNodeTunnel.cpp
new file mode 100644
--- /dev/null
+++ b/Engine/Plugins/Mutable/Source/CustomizableObjectEditor/Private/MuCOE/Nodes/CustomizableObjectNodeTunnel.cpp
@@ -0,0 +1,105 @@
+// Copyright Epic Games, Inc. All Rights Reserved.
+
+#include "MuCOE/Nodes/CustomizableObjectNodeTunnel.h"
+
+#include "MuCOE/EdGraphSchema_CustomizableObject.h"
+
+#define LOCTEXT_NAMESPACE "CustomizableObjectNodeTunnel"
+
+
+bool UCustomizableObjectNodeTunnelRemapPins::Equal(const UCustomizableObjectNode& Node, const UEdGraphPin& OldPin, const UEdGraphPin& NewPin) const
+{
+	const UCustomizableObjectNodeTunnelPinData* PinDataOldPin = Cast<UCustomizableObjectNodeTunnelPinData>(Node.GetPinData(OldPin));
+	const UCustomizableObjectNodeTunnelPinData* PinDataNewPin = Cast<UCustomizableObjectNodeTunnelPinData>(Node.GetPinData(NewPin));
+
+	return PinDataOldPin->VariableId == PinDataNewPin->VariableId && OldPin.PinType.PinCategory == NewPin.PinType.PinCategory;
+}
+
+
+void UCustomizableObjectNodeTunnelRemapPins::RemapPins(const UCustomizableObjectNode& Node, const TArray<UEdGraphPin*>& OldPins, const TArray<UEdGraphPin*>& NewPins, TMap<UEdGraphPin*, UEdGraphPin*>& PinsToRemap, TArray<UEdGraphPin*>& PinsToOrphan)
+{
+	for (UEdGraphPin* OldPin : OldPins)
+	{
+		bool bFound = false;
+
+		for (UEdGraphPin* NewPin : NewPins)
+		{
+			if (Equal(Node, *OldPin, *NewPin))
+			{
+				bFound = true;
+
+				PinsToRemap.Add(OldPin, NewPin);
+				break;
+			}
+		}
+
+		if (!bFound && OldPin->LinkedTo.Num())
+		{
+			PinsToOrphan.Add(OldPin);
+		}
+	}
+}
+
+
+UCustomizableObjectNodeTunnelRemapPins* UCustomizableObjectNodeTunnel::CreateRemapPinsDefault() const
+{
+	return NewObject<UCustomizableObjectNodeTunnelRemapPins>();
+}
+
+
+FText UCustomizableObjectNodeTunnel::GetNodeTitle(ENodeTitleType::Type TitleType) const
+{
+	if (bIsInputNode)
+	{
+		return LOCTEXT("TunnelInput_Text", "Inputs");
+	}
+	else
+	{
+		return LOCTEXT("TunnelOutput_Text", "Outputs");
+	}
+}
+
+
+FText UCustomizableObjectNodeTunnel::GetTooltipText() const
+{
+	return LOCTEXT("TunnelTooltipText", "Node Tunnel");
+}
+
+
+FLinearColor UCustomizableObjectNodeTunnel::GetNodeTitleColor() const
+{
+	return FLinearColor(0.15f, 0.15f, 0.15f);
+}
+
+
+bool UCustomizableObjectNodeTunnel::GetCanRenameNode() const
+{
+	return false;
+}
+
+
+bool UCustomizableObjectNodeTunnel::CanUserDeleteNode() const
+{
+	return false;
+}
+
+
+void UCustomizableObjectNodeTunnel::AllocateDefaultPins(UCustomizableObjectNodeRemapPins* RemapPins)
+{
+	for (const UCustomizableObjectMacroInputOutput* Variable : ParentMacro->InputOutputs)
+	{
+		ECOMacroIOType NodeType = bIsInputNode ? ECOMacroIOType::COMVT_Input : ECOMacroIOType::COMVT_Output;
+		EEdGraphPinDirection PinDirection = bIsInputNode ? EEdGraphPinDirection::EGPD_Output : EEdGraphPinDirection::EGPD_Input;
+
+		if (Variable->Type == NodeType)
+		{
+			UCustomizableObjectNodeTunnelPinData* PinData = NewObject<UCustomizableObjectNodeTunnelPinData>(this);
+			PinData->VariableId = Variable->UniqueId;
+
+			const FName PinName = Variable->Name;
+			UEdGraphPin* Pin = CustomCreatePin(PinDirection, Variable->PinCategoryType, PinName, PinData);
+		}
+	}
+}
+
+#undef LOCTEXT_NAMESPACE
\ No newline at end of file
diff --git a/Engine/Plugins/Mutable/Source/CustomizableObjectEditor/Private/MuCOE/Nodes/CustomizableObjectNodeTunnel.h b/Engine/Plugins/Mutable/Source/CustomizableObjectEditor/Private/MuCOE/Nodes/CustomizableObjectNodeTunnel.h
new file mode 100644
--- /dev/null
+++ b/Engine/Plugins/Mutable/Source/CustomizableObjectEditor/Private/MuCOE/Nodes/CustomizableObjectNodeTunnel.h
@@ -0,0 +1,70 @@
+// Copyright Epic Games, Inc. All Rights Reserved.
+
+#pragma once
+
+#include "MuCOE/CustomizableObjectMacroLibrary/CustomizableObjectMacroLibrary.h"
+#include "MuCOE/Nodes/CustomizableObjectNode.h"
+#include "MuCOE/RemapPins/CustomizableObjectNodeRemapPinsByName.h"
+
+#include "CustomizableObjectNodeTunnel.generated.h"
+
+
+/** Base class for all Tunnel Pins. */
+UCLASS()
+class CUSTOMIZABLEOBJECTEDITOR_API UCustomizableObjectNodeTunnelPinData : public UCustomizableObjectNodePinData
+{
+	GENERATED_BODY()
+
+public:
+
+	/** Id of the variable associated to a tunnel node pin */
+	UPROPERTY()
+	FGuid VariableId;
+
+};
+
+
+UCLASS()
+class UCustomizableObjectNodeTunnelRemapPins : public UCustomizableObjectNodeRemapPinsByName
+{
+	GENERATED_BODY()
+public:
+
+	// Specific method to decide when two pins are equal
+	virtual bool Equal(const UCustomizableObjectNode& Node, const UEdGraphPin& OldPin, const UEdGraphPin& NewPin) const override;
+
+	// Method to use in the RemapPins step of the node reconstruction process
+	virtual void RemapPins(const UCustomizableObjectNode& Node, const TArray<UEdGraphPin*>& OldPins, const TArray<UEdGraphPin*>& NewPins, TMap<UEdGraphPin*, UEdGraphPin*>& PinsToRemap, TArray<UEdGraphPin*>& PinsToOrphan) override;
+};
+
+
+
+UCLASS()
+class CUSTOMIZABLEOBJECTEDITOR_API UCustomizableObjectNodeTunnel : public UCustomizableObjectNode
+{
+public:
+
+	/** Indicates if this node is an input or output node. */
+	UPROPERTY()
+	bool bIsInputNode = false;
+
+	/** Pointer to the macro where this node belongs. */
+	UPROPERTY()
+	TObjectPtr<UCustomizableObjectMacro> ParentMacro;
+
+public:
+
+	GENERATED_BODY()
+
+	// EdGraphNode interface
+	virtual FText GetNodeTitle(ENodeTitleType::Type TitleType) const override;
+	virtual FText GetTooltipText() const override;
+	virtual FLinearColor GetNodeTitleColor() const override;
+	virtual bool GetCanRenameNode() const override;
+	virtual bool CanUserDeleteNode() const;
+
+	// UCustomizableObjectNode interface
+	void AllocateDefaultPins(UCustomizableObjectNodeRemapPins* RemapPins) override;
+	virtual UCustomizableObjectNodeTunnelRemapPins* CreateRemapPinsDefault() const override;
+
+};
\ No newline at end of file
diff --git a/Engine/Plugins/Mutable/Source/CustomizableObjectEditor/Private/MuCOE/Widgets/SCustomizableObjectMacroLibraryList.cpp b/Engine/Plugins/Mutable/Source/CustomizableObjectEditor/Private/MuCOE/Widgets/SCustomizableObjectMacroLibraryList.cpp
new file mode 100644
--- /dev/null
+++ b/Engine/Plugins/Mutable/Source/CustomizableObjectEditor/Private/MuCOE/Widgets/SCustomizableObjectMacroLibraryList.cpp
@@ -0,0 +1,174 @@
+// Copyright Epic Games, Inc. All Rights Reserved.
+
+#include "MuCOE/Widgets/SCustomizableObjectMacroLibraryList.h"
+
+#include "Misc/MessageDialog.h"
+#include "MuCOE/CustomizableObjectGraph.h"
+#include "PropertyCustomizationHelpers.h"
+#include "SPositiveActionButton.h"
+#include "Widgets/Layout/SScrollBorder.h"
+#include "Widgets/Views/STreeView.h"
+
+
+#define LOCTEXT_NAMESPACE "SCustomizableObjectMacroLibraryList"
+
+/** Row on the tree representing the bone structure of the mu::FMesh. It represents the UI side of the BoneDefinition */
+class SMacroLibraryTreeRow : public STableRow<TSharedPtr<FMacroTreeEntry>>
+{
+public:
+	DECLARE_DELEGATE_OneParam(FOnRemoveMacroButtonClickedDelegate, UCustomizableObjectMacro*)
+
+	SLATE_BEGIN_ARGS(SMacroLibraryTreeRow) {}
+		SLATE_ARGUMENT(FOnRemoveMacroButtonClickedDelegate, OnRemoveMacro)
+	SLATE_END_ARGS()
+
+	void Construct(const FArguments& Args, const TSharedRef<STableViewBase>& InOwnerTableView, const TSharedPtr<FMacroTreeEntry>& InRowItem)
+	{
+		RowItem = InRowItem;
+		OnRemoveMacro = Args._OnRemoveMacro;
+
+		FText MacroName = RowItem->Macro.IsValid() ? FText::FromName(RowItem->Macro->Name) : FText::FromString("#Error# - Null Macro");
+
+		this->ChildSlot
+			[
+				SNew(SHorizontalBox)
+				+ SHorizontalBox::Slot()
+				.Padding(5.0f, 5.0f, 5.0f, 5.0f)
+				.AutoWidth()
+				[
+					SNew(SImage).Image(FAppStyle::GetBrush(TEXT("GraphEditor.Macro_16x")))
+				]
+
+				+ SHorizontalBox::Slot()
+				.Padding(5.0f, 5.0f, 0.0f, 5.0f)
+				.AutoWidth()
+				[
+					SNew(STextBlock).Text(MacroName)
+				]
+
+				+ SHorizontalBox::Slot()
+				.FillWidth(1.f)
+				[
+					SNew(SSpacer)
+				]
+
+				+ SHorizontalBox::Slot()
+				.HAlign(EHorizontalAlignment::HAlign_Right)
+				.AutoWidth()
+				[
+					PropertyCustomizationHelpers::MakeClearButton(
+						FSimpleDelegate::CreateSP(this, &SMacroLibraryTreeRow::RemoveMacro), LOCTEXT("RemoveMacroTooltip", "Remove Macro."))
+				]
+			];
+
+		STableRow< TSharedPtr<FMacroTreeEntry> >::ConstructInternal(
+			STableRow::FArguments()
+			.Style(FAppStyle::Get(), "TableView.Row")
+			.ShowSelection(true)
+			, InOwnerTableView);
+	}
+
+	void RemoveMacro()
+	{
+		if (RowItem && RowItem->Macro.IsValid())
+		{
+			FText Msg = LOCTEXT("RemoveMacroTextWindow", "Are you sure you want to remove the Macro?");
+
+			if (FMessageDialog::Open(EAppMsgType::OkCancel, Msg) == EAppReturnType::Ok)
+			{
+				OnRemoveMacro.ExecuteIfBound(RowItem->Macro.Get());
+			}
+		}
+	}
+
+private:
+
+	TSharedPtr<FMacroTreeEntry> RowItem;
+	FOnRemoveMacroButtonClickedDelegate OnRemoveMacro;
+};
+
+void SCustomizableObjectMacroLibraryList::Construct(const FArguments& InArgs)
+{
+	MacroLibrary = InArgs._MacroLibrary;
+	OnAddMacroButtonClicked = InArgs._OnAddMacroButtonClicked;
+	OnSelectMacro = InArgs._OnSelectMacro;
+	OnRemoveMacro = InArgs._OnRemoveMacro;
+
+	GenerateRowView();
+
+	SAssignNew(ListView, SListView<TSharedPtr<FMacroTreeEntry>>)
+		.SelectionMode(ESelectionMode::Single)
+		.ListItemsSource(&ListViewSource)
+		.OnGenerateRow_Lambda([this](TSharedPtr<FMacroTreeEntry> InItem, const TSharedRef<STableViewBase>& InOwnerTable) -> TSharedRef<ITableRow>
+		{
+			TSharedRef<SMacroLibraryTreeRow> Row = SNew(SMacroLibraryTreeRow, InOwnerTable, InItem)
+			.OnRemoveMacro(SMacroLibraryTreeRow::FOnRemoveMacroButtonClickedDelegate::CreateSP(this, &SCustomizableObjectMacroLibraryList::OnRemoveCurrentMacro));
+
+			return Row;
+		})
+		.OnSelectionChanged_Lambda([this](TSharedPtr<FMacroTreeEntry> InItem, ESelectInfo::Type)
+		{
+			if (InItem.IsValid() && InItem->Macro.IsValid())
+			{
+				OnSelectMacro.ExecuteIfBound(InItem->Macro.Get());
+			}
+		});
+
+		this->ChildSlot
+		[
+			SNew(SVerticalBox)
+			+SVerticalBox::Slot()
+			.AutoHeight()
+			[
+				SNew(SPositiveActionButton)
+				.OnClicked_Lambda([this]() 
+				{
+					if (OnAddMacroButtonClicked.IsBound())
+					{
+						OnAddMacroButtonClicked.Execute();
+						GenerateRowView();
+						ListView->RequestListRefresh();
+						return FReply::Handled();
+					}
+
+					return FReply::Unhandled();
+				})
+			]
+
+			+ SVerticalBox::Slot()
+			.AutoHeight()
+			[
+				ListView.ToSharedRef()
+			]
+		];
+}
+
+
+void SCustomizableObjectMacroLibraryList::GenerateRowView()
+{
+	ListViewSource.Empty();
+
+	for (UCustomizableObjectMacro* Macro : MacroLibrary->Macros)
+	{
+		TSharedPtr<FMacroTreeEntry> NewGraphDefinition = MakeShareable(new FMacroTreeEntry());
+		NewGraphDefinition->Macro = Macro;
+
+		ListViewSource.Add(NewGraphDefinition);
+	}
+	
+	if (ListView.IsValid())
+	{
+		ListView->RequestListRefresh();
+	}
+}
+
+void SCustomizableObjectMacroLibraryList::OnRemoveCurrentMacro(UCustomizableObjectMacro* MacroToDelete)
+{
+	if (MacroToDelete)
+	{
+		OnRemoveMacro.ExecuteIfBound(MacroToDelete);
+		GenerateRowView();
+	}
+}
+
+#undef LOCTEXT_NAMESPACE
diff --git a/Engine/Plugins/Mutable/Source/CustomizableObjectEditor/Private/MuCOE/Widgets/SCustomizableObjectMacroLibraryList.h b/Engine/Plugins/Mutable/Source/CustomizableObjectEditor/Private/MuCOE/Widgets/SCustomizableObjectMacroLibraryList.h
new file mode 100644
--- /dev/null
+++ b/Engine/Plugins/Mutable/Source/CustomizableObjectEditor/Private/MuCOE/Widgets/SCustomizableObjectMacroLibraryList.h
@@ -0,0 +1,54 @@
+// Copyright Epic Games, Inc. All Rights Reserved.
+
+#pragma once
+
+#include "MuCOE/CustomizableObjectMacroLibrary/CustomizableObjectMacroLibrary.h"
+#include "Widgets/SCompoundWidget.h"
+
+template <typename ItemType> class SListView;
+
+struct FMacroTreeEntry
+{
+	TWeakObjectPtr<UCustomizableObjectMacro> Macro;
+};
+
+class SCustomizableObjectMacroLibraryList : public SCompoundWidget
+{
+public:
+	DECLARE_DELEGATE(FOnAddMacroButtonClickedDelegate)
+	DECLARE_DELEGATE_OneParam(FOnSelectMacroDelegate, UCustomizableObjectMacro*)
+	DECLARE_DELEGATE_OneParam(FOnRemoveMacroDelegate, UCustomizableObjectMacro*)
+
+	SLATE_BEGIN_ARGS(SCustomizableObjectMacroLibraryList) {}
+		SLATE_ARGUMENT(TObjectPtr<UCustomizableObjectMacroLibrary>, MacroLibrary)
+		SLATE_ARGUMENT(FOnAddMacroButtonClickedDelegate, OnAddMacroButtonClicked)
+		SLATE_ARGUMENT(FOnSelectMacroDelegate, OnSelectMacro)
+		SLATE_ARGUMENT(FOnRemoveMacroDelegate, OnRemoveMacro)
+	SLATE_END_ARGS()
+
+	// SWidget interface
+	void Construct(const FArguments& InArgs);
+
+	/** Generates and Updates the List View widget and its source. */
+	void GenerateRowView();
+
+	/** Callback to regenerate the List View when a Macro has been removed. It is also used to comunicate the Macro Library editor. */
+	void OnRemoveCurrentMacro(UCustomizableObjectMacro* MacroToDelete);
+
+private:
+
+	/** Pointer to the Macro Library that contains the macros. */
+	TObjectPtr<UCustomizableObjectMacroLibrary> MacroLibrary;
+
+	/** Pointer to the List View widget. Needed to refresh it when something changes. */
+	TSharedPtr<SListView<TSharedPtr<FMacroTreeEntry>>> ListView;
+
+	/** Array that contains all the macros of the Library in a format for the List View*/
+	TArray<TSharedPtr<FMacroTreeEntry>> ListViewSource;
+
+	// List view action callbacks
+	FOnAddMacroButtonClickedDelegate OnAddMacroButtonClicked;
+	FOnSelectMacroDelegate OnSelectMacro;
+	FOnRemoveMacroDelegate OnRemoveMacro;
+};
+
diff --git a/Engine/Plugins/Mutable/Source/CustomizableObjectEditor/Public/MuCOE/CustomizableObjectGraphEditorToolkit.h b/Engine/Plugins/Mutable/Source/CustomizableObjectEditor/Public/MuCOE/CustomizableObjectGraphEditorToolkit.h
new file mode 100644
--- /dev/null
+++ b/Engine/Plugins/Mutable/Source/CustomizableObjectEditor/Public/MuCOE/CustomizableObjectGraphEditorToolkit.h
@@ -0,0 +1,75 @@
+// Copyright Epic Games, Inc. All Rights Reserved.
+
+#pragma once
+
+#include "Misc/NotifyHook.h"
+#include "Toolkits/AssetEditorToolkit.h"
+#include "BlueprintUtilities.h"
+#include "EdGraph/EdGraph.h"
+#include "EdGraph/EdGraphNode.h"
+#include "EdGraphUtilities.h"
+#include "GraphEditor.h"
+#include "GraphEditorActions.h"
+#include "ScopedTransaction.h"
+
+class UClass;
+class UCustomizableObjectNode;
+
+// Public interface to Customizable Object Editor
+class FCustomizableObjectGraphEditorToolkit : public FAssetEditorToolkit, public FNotifyHook
+{
+public:
+	
+	/** Searches a node that contains the inserted word */
+	void OnEnterText(const FText& NewText, ETextCommit::Type TextType);
+	void FindProperty(const FProperty* Property, const void* InContainer, const FString& FindString, const UObject& Context, bool& bFound);
+	void LogSearchResult(const UObject& Context, const FString& Type, bool bIsFirst, const FString& Result) const;
+
+
+	// Graph Node Actions --------------------------
+	void SelectNode(const UEdGraphNode* Node);
+		
+	void DeleteSelectedNodes();
+	bool CanDeleteNodes() const;
+	
+	void DuplicateSelectedNodes();
+	bool CanDuplicateSelectedNodes() const;
+		
+	void OnNodeTitleCommitted(const FText& NewText, ETextCommit::Type CommitInfo, UEdGraphNode* NodeBeingChanged);
+	
+	void CopySelectedNodes();
+	bool CanCopyNodes() const;
+	
+	void PasteNodes();
+	void PasteNodesHere(const FVector2D& Location);
+	bool CanPasteNodes() const;
+	
+	void CutSelectedNodes();
+	bool CanCutNodes() const;
+
+	void OnRenameNode();
+	bool CanRenameNodes() const;
+
+	void CreateCommentBoxFromKey();
+	UEdGraphNode* CreateCommentBox(const FVector2D& NodePos);
+
+	void OnNodeDoubleClicked(UEdGraphNode* Node);
+
+	// ------------------------------------------
+
+	/** Creates a new graph editor widget 
+	* @param InGraph is the graph that will be represented in the graph editor.
+	* @param InEvents are custom events occurring in/on the graph.
+	*/
+	void CreateGraphEditorWidget(UEdGraph* InGraph, SGraphEditor::FGraphEditorEvents& InEvents);
+
+public:
+
+	/** Pointer to the graph editor widget */
+	TSharedPtr<SGraphEditor> GraphEditor;
+
+private:
+
+	/** Command list for the graph editor */
+	TSharedPtr<FUICommandList> GraphEditorCommands;
+};
diff --git a/Engine/Plugins/Mutable/Source/CustomizableObjectEditor/Public/MuCOE/EdGraphSchema_CustomizableObject.h b/Engine/Plugins/Mutable/Source/CustomizableObjectEditor/Public/MuCOE/EdGraphSchema_CustomizableObject.h
--- a/Engine/Plugins/Mutable/Source/CustomizableObjectEditor/Public/MuCOE/EdGraphSchema_CustomizableObject.h
+++ b/Engine/Plugins/Mutable/Source/CustomizableObjectEditor/Public/MuCOE/EdGraphSchema_CustomizableObject.h
@@ -44,6 +44,9 @@ public:
 	static const FName PC_Wildcard;
 	static const FName PC_PoseAsset;
 	static const FName PC_Transform;
+	static const FName PC_String;
+
+	static const TArray<FName> SupportedMacroPinTypes;
 
 	// Node categories
 	static const FText NC_Experimental;
diff --git a/Engine/Plugins/Mutable/Source/CustomizableObjectEditor/Public/MuCOE/Nodes/CustomizableObjectNode.h b/Engine/Plugins/Mutable/Source/CustomizableObjectEditor/Public/MuCOE/Nodes/CustomizableObjectNode.h
--- a/Engine/Plugins/Mutable/Source/CustomizableObjectEditor/Public/MuCOE/Nodes/CustomizableObjectNode.h
+++ b/Engine/Plugins/Mutable/Source/CustomizableObjectEditor/Public/MuCOE/Nodes/CustomizableObjectNode.h
@@ -277,6 +277,11 @@ public:
 	{
 		return *CastChecked<T>(PinsDataId[Pin.PinId]);
 	}
+
+	/** returs true if this node can be added to a COMacro Graph */
+	virtual bool IsNodeSupportedInMacros() { return true; } 
+
+	bool IsMacroNode() const;
 	
 	UPROPERTY()
 	FPostEditChangePropertyDelegate PostEditChangePropertyDelegate;
