[Mutable] Remove unused custom allocators.

diff --git a/Engine/Plugins/Mutable/Source/CustomizableObjectEditor/Private/MuCOE/CustomizableObjectEditorModule.cpp b/Engine/Plugins/Mutable/Source/CustomizableObjectEditor/Private/MuCOE/CustomizableObjectEditorModule.cpp
--- a/Engine/Plugins/Mutable/Source/CustomizableObjectEditor/Private/MuCOE/CustomizableObjectEditorModule.cpp
+++ b/Engine/Plugins/Mutable/Source/CustomizableObjectEditor/Private/MuCOE/CustomizableObjectEditorModule.cpp
@@ -150,18 +150,6 @@ void DeleteUnusedMutableStreamedDataEditorFiles()
 	}
 }
 
-static void* CustomMalloc(std::size_t Size_t, uint32_t Alignment)
-{
-	return FMemory::Malloc(Size_t, Alignment);
-}
-
-
-static void CustomFree(void* mem)
-{
-	return FMemory::Free(mem);
-}
-
-
 int32 ConvertOptimizationLevel(ECustomizableObjectOptimizationLevel OptimizationLevel)
 {
 	switch (OptimizationLevel)
