[mutable] Enable concurrency in CO cooking. Remove special code path for non-concurrent optimization.
#rnx

diff --git a/Engine/Plugins/Mutable/Source/CustomizableObjectEditor/Private/MuCOE/CustomizableObjectCompileRunnable.cpp b/Engine/Plugins/Mutable/Source/CustomizableObjectEditor/Private/MuCOE/CustomizableObjectCompileRunnable.cpp
--- a/Engine/Plugins/Mutable/Source/CustomizableObjectEditor/Private/MuCOE/CustomizableObjectCompileRunnable.cpp
+++ b/Engine/Plugins/Mutable/Source/CustomizableObjectEditor/Private/MuCOE/CustomizableObjectCompileRunnable.cpp
@@ -33,12 +33,6 @@
 #define UE_MUTABLE_CORE_REGION	TEXT("Mutable Core")
 
 
-TAutoConsoleVariable<bool> CVarMutableCompilerConcurrency(
-	TEXT("mutable.ForceCompilerConcurrency"),
-	true,
-	TEXT("Force the use of multithreading when compiling CustomizableObjects both in editor and cook commandlets."),
-	ECVF_Default);
-
 TAutoConsoleVariable<bool> CVarMutableCompilerDiskCache(
 	TEXT("mutable.ForceCompilerDiskCache"),
 	false,
@@ -104,14 +98,6 @@ uint32 FCustomizableObjectCompileRunnable::Run()
 	// Translate CO compile options into mu::CompilerOptions
 	mu::Ptr<mu::CompilerOptions> CompilerOptions = new mu::CompilerOptions();
 
-	bool bUseConcurrency = !Options.bIsCooking;
-	if (CVarMutableCompilerConcurrency->GetBool())
-	{
-		bUseConcurrency = true;
-	}
-
-	CompilerOptions->SetUseConcurrency(bUseConcurrency);
-
 	bool bUseDiskCache = Options.bUseDiskCompilation;
 	if (CVarMutableCompilerDiskCache->GetBool())
 	{
diff --git a/Engine/Plugins/Mutable/Source/MutableTools/Private/MuT/CodeOptimiser.cpp b/Engine/Plugins/Mutable/Source/MutableTools/Private/MuT/CodeOptimiser.cpp
--- a/Engine/Plugins/Mutable/Source/MutableTools/Private/MuT/CodeOptimiser.cpp
+++ b/Engine/Plugins/Mutable/Source/MutableTools/Private/MuT/CodeOptimiser.cpp
@@ -789,192 +789,10 @@ namespace mu
 				});
 		};
 
-		bool bUseConcurrency = InOptions->bUseConcurrency;
-		//bUseConcurrency = false;
 
-		if (bUseConcurrency)
-		{
-			// Launch the tasks.
-			UE::Tasks::FTask LaunchTask = UE::Tasks::Launch(TEXT("ConstantGeneratorLaunchTasks"), 
-				[&ConstantSubgraphs, &GetRequisites, Pass, InOptions]()
-				{
-					MUTABLE_CPUPROFILER_SCOPE(ConstantGenerator_LaunchTasks);
-
-					FImageOperator ImOp = FImageOperator::GetDefault(InOptions->ImageFormatFunc);
-
-					// Traverse list of constants to generate. It is ordered in a bottom-up way.
-					int32 SubgraphCount = ConstantSubgraphs.Num();
-					for (int32 OrderIndex = 0; OrderIndex < SubgraphCount; ++OrderIndex)
-					{
-						int32 Index = SubgraphCount - 1 - OrderIndex;
-
-						Ptr<ASTOp> SubgraphRoot = ConstantSubgraphs[Index].Root;
-						UE::Tasks::FTaskEvent& SubgraphCompletionEvent = ConstantSubgraphs[Index].CompletedEvent;
-
-						bool bIsReference = false;
-						EOpType SubgraphType = SubgraphRoot->GetOpType();
-
-						if (SubgraphType == EOpType::IM_REFERENCE
-							||
-							SubgraphType == EOpType::IM_CONSTANT)
-						{
-							uint32 ImageID = 0;
-							if (SubgraphType == EOpType::IM_REFERENCE)							
-							{
-								bIsReference = true;
-								const ASTOpReferenceResource* Typed = static_cast<const ASTOpReferenceResource*>(SubgraphRoot.get());
-								ImageID = Typed->ID;
-							}
-							else if(SubgraphType == EOpType::IM_CONSTANT)
-							{
-								const ASTOpConstantResource* Typed = static_cast<const ASTOpConstantResource*>(SubgraphRoot.get());
-								const FImage* Value = static_cast<const FImage*>( Typed->GetValue().Get() );
-								bIsReference = Value->IsReference();
-								if (bIsReference)
-								{
-									ImageID = Value->GetReferencedTexture();
-								}
-							}
-
-							// Instead of generating the constant we resolve the reference, which also replaces the ASTOp.
-							if (bIsReference)
-							{
-								TSharedPtr< TSharedPtr<FImage> > ResolveImage = MakeShared<TSharedPtr<FImage>>();
-
-								constexpr bool bRunImmediatlyIfPossible = false;
-								UE::Tasks::FTask ReferenceCompletion = InOptions->OptimisationOptions.ReferencedImageResourceProvider(ImageID, ResolveImage, bRunImmediatlyIfPossible);
-
-								UE::Tasks::FTask CompleteTask = UE::Tasks::Launch(TEXT("MutableResolveComplete"),
-									[SubgraphRoot, InOptions, ResolveImage]()
-									{
-										Ptr<ASTOpConstantResource> ConstantOp;
-										{
-											MUTABLE_CPUPROFILER_SCOPE(MutableResolveComplete_CreateConstant);
-											ConstantOp = new ASTOpConstantResource;
-											ConstantOp->Type = EOpType::IM_CONSTANT;
-											{
-												MUTABLE_CPUPROFILER_SCOPE(GetSourceDataDescriptor);
-												ConstantOp->SourceDataDescriptor = SubgraphRoot->GetSourceDataDescriptor();
-											}
-											ConstantOp->SetValue(*ResolveImage, InOptions->OptimisationOptions.DiskCacheContext);
-										}
-										{
-											MUTABLE_CPUPROFILER_SCOPE(MutableResolveComplete_Replace);
-											ASTOp::Replace(SubgraphRoot, ConstantOp);
-										}
-									},
-									ReferenceCompletion,
-									LowLevelTasks::ETaskPriority::BackgroundNormal);
-
-								SubgraphCompletionEvent.AddPrerequisites(CompleteTask);
-							}
-
-						}
-
-						if (SubgraphType == EOpType::ME_REFERENCE
-							||
-							SubgraphType == EOpType::ME_CONSTANT)
-						{
-							uint32 MeshID = 0;
-							FString MeshMorph;
-							if (SubgraphType == EOpType::ME_REFERENCE)
-							{
-								bIsReference = true;
-								const ASTOpReferenceResource* Typed = static_cast<const ASTOpReferenceResource*>(SubgraphRoot.get());
-								MeshID = Typed->ID;
-							}
-							else if (SubgraphType == EOpType::ME_CONSTANT)
-							{
-								const ASTOpConstantResource* Typed = static_cast<const ASTOpConstantResource*>(SubgraphRoot.get());
-								const FMesh* Value = static_cast<const FMesh*>(Typed->GetValue().Get());
-								bIsReference = Value->IsReference();
-								if (bIsReference)
-								{
-									MeshID = Value->GetReferencedMesh();
-									MeshMorph = Value->GetReferencedMorph();
-								}
-							}
-
-							// Instead of generating the constant we resolve the reference, which also replaces the ASTOp.
-							if (bIsReference)
-							{
-								TSharedPtr< TSharedPtr<FMesh> > ResolveMesh = MakeShared<TSharedPtr<FMesh>>();
-
-								constexpr bool bRunImmediatlyIfPossible = false;
-								UE::Tasks::FTask ReferenceCompletion = InOptions->OptimisationOptions.ReferencedMeshResourceProvider(MeshID, MeshMorph, ResolveMesh, bRunImmediatlyIfPossible);
-
-								UE::Tasks::FTask CompleteTask = UE::Tasks::Launch(TEXT("MutableResolveComplete"),
-									[SubgraphRoot, InOptions, ResolveMesh]()
-									{
-										Ptr<ASTOpConstantResource> ConstantOp;
-										{
-											MUTABLE_CPUPROFILER_SCOPE(MutableResolveComplete_CreateConstant);
-											ConstantOp = new ASTOpConstantResource;
-											ConstantOp->Type = EOpType::ME_CONSTANT;
-											{
-												MUTABLE_CPUPROFILER_SCOPE(GetSourceDataDescriptor);
-												ConstantOp->SourceDataDescriptor = SubgraphRoot->GetSourceDataDescriptor();
-											}
-											ConstantOp->SetValue(*ResolveMesh, InOptions->OptimisationOptions.DiskCacheContext);
-										}
-										{
-											MUTABLE_CPUPROFILER_SCOPE(MutableResolveComplete_Replace);
-											ASTOp::Replace(SubgraphRoot, ConstantOp);
-										}
-									},
-									ReferenceCompletion,
-									LowLevelTasks::ETaskPriority::BackgroundNormal);
-
-								SubgraphCompletionEvent.AddPrerequisites(CompleteTask);
-							}
-						}
-
-						if (!bIsReference)
-						{
-							// Scan for requisites
-							TArray< UE::Tasks::FTask, TInlineAllocator<8> > Requisites;
-							GetRequisites(SubgraphRoot, Requisites);
-
-							TUniquePtr<FConstantTask> Task(new FConstantTask(SubgraphRoot, InOptions, Pass));
-
-							// Launch the preparation on the AST-modification pipe
-							UE::Tasks::FTask CompleteTask = UE::Tasks::Launch(TEXT("MutableConstant"), [TaskPtr = MoveTemp(Task), ImOp]()
-								{
-									MUTABLE_CPUPROFILER_SCOPE(MutableConstantPrepare);
-
-									// We need the clone because linking modifies ASTOp state and also to be safe for concurrency.
-									TaskPtr->SourceCloned = ASTOp::DeepClone(TaskPtr->Source);
-
-									TaskPtr->Run(ImOp);
-
-									ASTOp::Replace(TaskPtr->Source, TaskPtr->Result);
-									TaskPtr->Result = nullptr;
-									TaskPtr->Source = nullptr;
-								},
-								Requisites,
-								LowLevelTasks::ETaskPriority::BackgroundHigh);
-
-							SubgraphCompletionEvent.AddPrerequisites(CompleteTask);
-						}
-
-						ConstantSubgraphs[Index].Root = nullptr;
-						SubgraphCompletionEvent.Trigger();
-						
-						UE::Tasks::AddNested(SubgraphCompletionEvent);
-					}
-
-				});
-
-			// Wait for pending tasks
-			{
-				MUTABLE_CPUPROFILER_SCOPE(Waiting);
-				LaunchTask.Wait();
-			}
-		}
-
-		else
-		{
-			// Non-concurrent version: Prepare the tasks.
+		// Launch the tasks.
+		UE::Tasks::FTask LaunchTask = UE::Tasks::Launch(TEXT("ConstantGeneratorLaunchTasks"), 
+			[&ConstantSubgraphs, &GetRequisites, Pass, InOptions]()
 			{
 				MUTABLE_CPUPROFILER_SCOPE(ConstantGenerator_LaunchTasks);
 
@@ -997,16 +815,16 @@ namespace mu
 						SubgraphType == EOpType::IM_CONSTANT)
 					{
 						uint32 ImageID = 0;
-						if (SubgraphType == EOpType::IM_REFERENCE)
+						if (SubgraphType == EOpType::IM_REFERENCE)							
 						{
 							bIsReference = true;
 							const ASTOpReferenceResource* Typed = static_cast<const ASTOpReferenceResource*>(SubgraphRoot.get());
 							ImageID = Typed->ID;
 						}
-						else if (SubgraphType == EOpType::IM_CONSTANT)
+						else if(SubgraphType == EOpType::IM_CONSTANT)
 						{
 							const ASTOpConstantResource* Typed = static_cast<const ASTOpConstantResource*>(SubgraphRoot.get());
-							const FImage* Value = static_cast<const FImage*>(Typed->GetValue().Get());
+							const FImage* Value = static_cast<const FImage*>( Typed->GetValue().Get() );
 							bIsReference = Value->IsReference();
 							if (bIsReference)
 							{
@@ -1017,21 +835,36 @@ namespace mu
 						// Instead of generating the constant we resolve the reference, which also replaces the ASTOp.
 						if (bIsReference)
 						{
+							TSharedPtr< TSharedPtr<FImage> > ResolveImage = MakeShared<TSharedPtr<FImage>>();
+
+							constexpr bool bRunImmediatlyIfPossible = false;
+							UE::Tasks::FTask ReferenceCompletion = InOptions->OptimisationOptions.ReferencedImageResourceProvider(ImageID, ResolveImage, bRunImmediatlyIfPossible);
 
-							auto ImmediateCompleteFunc = [SubgraphRoot, InOptions, ImageID]()
+							UE::Tasks::FTask CompleteTask = UE::Tasks::Launch(TEXT("MutableResolveComplete"),
+								[SubgraphRoot, InOptions, ResolveImage]()
 								{
-									TSharedPtr< TSharedPtr<FImage> > ResolveImage = MakeShared<TSharedPtr<FImage>>();
-									constexpr bool bRunImmediatlyIfPossible = true;
-									UE::Tasks::FTask ReferenceCompletion = InOptions->OptimisationOptions.ReferencedImageResourceProvider(ImageID, ResolveImage, bRunImmediatlyIfPossible);
-									ReferenceCompletion.Wait();
-									Ptr<ASTOpConstantResource> ConstantOp = new ASTOpConstantResource;
-									ConstantOp->SourceDataDescriptor = SubgraphRoot->GetSourceDataDescriptor();
-									ConstantOp->Type = EOpType::IM_CONSTANT;
-									ConstantOp->SetValue(*ResolveImage, InOptions->OptimisationOptions.DiskCacheContext);
-									ASTOp::Replace(SubgraphRoot, ConstantOp);
-								};
-							ConstantSubgraphs[Index].NonConcurrentTask = ImmediateCompleteFunc;
+									Ptr<ASTOpConstantResource> ConstantOp;
+									{
+										MUTABLE_CPUPROFILER_SCOPE(MutableResolveComplete_CreateConstant);
+										ConstantOp = new ASTOpConstantResource;
+										ConstantOp->Type = EOpType::IM_CONSTANT;
+										{
+											MUTABLE_CPUPROFILER_SCOPE(GetSourceDataDescriptor);
+											ConstantOp->SourceDataDescriptor = SubgraphRoot->GetSourceDataDescriptor();
+										}
+										ConstantOp->SetValue(*ResolveImage, InOptions->OptimisationOptions.DiskCacheContext);
+									}
+									{
+										MUTABLE_CPUPROFILER_SCOPE(MutableResolveComplete_Replace);
+										ASTOp::Replace(SubgraphRoot, ConstantOp);
+									}
+								},
+								ReferenceCompletion,
+								LowLevelTasks::ETaskPriority::BackgroundNormal);
+
+							SubgraphCompletionEvent.AddPrerequisites(CompleteTask);
 						}
+
 					}
 
 					if (SubgraphType == EOpType::ME_REFERENCE
@@ -1061,19 +894,34 @@ namespace mu
 						// Instead of generating the constant we resolve the reference, which also replaces the ASTOp.
 						if (bIsReference)
 						{
-							auto ImmediateCompleteFunc = [SubgraphRoot, InOptions, MeshID, MeshMorph]()
+							TSharedPtr< TSharedPtr<FMesh> > ResolveMesh = MakeShared<TSharedPtr<FMesh>>();
+
+							constexpr bool bRunImmediatlyIfPossible = false;
+							UE::Tasks::FTask ReferenceCompletion = InOptions->OptimisationOptions.ReferencedMeshResourceProvider(MeshID, MeshMorph, ResolveMesh, bRunImmediatlyIfPossible);
+
+							UE::Tasks::FTask CompleteTask = UE::Tasks::Launch(TEXT("MutableResolveComplete"),
+								[SubgraphRoot, InOptions, ResolveMesh]()
 								{
-									TSharedPtr< TSharedPtr<FMesh> > ResolveMesh = MakeShared<TSharedPtr<FMesh>>();
-									constexpr bool bRunImmediatlyIfPossible = true;
-									UE::Tasks::FTask ReferenceCompletion = InOptions->OptimisationOptions.ReferencedMeshResourceProvider(MeshID, MeshMorph, ResolveMesh, bRunImmediatlyIfPossible);
-									ReferenceCompletion.Wait();
-									Ptr<ASTOpConstantResource> ConstantOp = new ASTOpConstantResource;
-									ConstantOp->SourceDataDescriptor = SubgraphRoot->GetSourceDataDescriptor();
-									ConstantOp->Type = EOpType::ME_CONSTANT;
-									ConstantOp->SetValue(*ResolveMesh, InOptions->OptimisationOptions.DiskCacheContext);
-									ASTOp::Replace(SubgraphRoot, ConstantOp);
-								};
-							ConstantSubgraphs[Index].NonConcurrentTask = ImmediateCompleteFunc;
+									Ptr<ASTOpConstantResource> ConstantOp;
+									{
+										MUTABLE_CPUPROFILER_SCOPE(MutableResolveComplete_CreateConstant);
+										ConstantOp = new ASTOpConstantResource;
+										ConstantOp->Type = EOpType::ME_CONSTANT;
+										{
+											MUTABLE_CPUPROFILER_SCOPE(GetSourceDataDescriptor);
+											ConstantOp->SourceDataDescriptor = SubgraphRoot->GetSourceDataDescriptor();
+										}
+										ConstantOp->SetValue(*ResolveMesh, InOptions->OptimisationOptions.DiskCacheContext);
+									}
+									{
+										MUTABLE_CPUPROFILER_SCOPE(MutableResolveComplete_Replace);
+										ASTOp::Replace(SubgraphRoot, ConstantOp);
+									}
+								},
+								ReferenceCompletion,
+								LowLevelTasks::ETaskPriority::BackgroundNormal);
+
+							SubgraphCompletionEvent.AddPrerequisites(CompleteTask);
 						}
 					}
 
@@ -1084,86 +932,39 @@ namespace mu
 						GetRequisites(SubgraphRoot, Requisites);
 
 						TUniquePtr<FConstantTask> Task(new FConstantTask(SubgraphRoot, InOptions, Pass));
-						FConstantTask* TaskPtr = Task.Get();
 
-						ConstantSubgraphs[Index].Requisites = Requisites;
-						ConstantSubgraphs[Index].TaskData = MoveTemp(Task);
-						ConstantSubgraphs[Index].NonConcurrentTask = [TaskPtr, ImOp]()
+						// Launch the preparation on the AST-modification pipe
+						UE::Tasks::FTask CompleteTask = UE::Tasks::Launch(TEXT("MutableConstant"), [TaskPtr = MoveTemp(Task), ImOp]()
 							{
+								MUTABLE_CPUPROFILER_SCOPE(MutableConstantPrepare);
+
+								// We need the clone because linking modifies ASTOp state and also to be safe for concurrency.
 								TaskPtr->SourceCloned = ASTOp::DeepClone(TaskPtr->Source);
+
 								TaskPtr->Run(ImOp);
-								ASTOp::Replace(TaskPtr->Source, TaskPtr->Result);
-							};
-					}
-					
-					ConstantSubgraphs[Index].Root = nullptr;
-				}
-			}
 
-			auto TryRunFunc = [&ConstantSubgraphs](int32 Index) -> bool
-				{
-					bool bCanRun = true;
-					for (const UE::Tasks::FTask& Req : ConstantSubgraphs[Index].Requisites)
-					{
-						if (!Req.IsCompleted())
-						{
-							bCanRun = false;
-							break;
-						}
-					}
+								ASTOp::Replace(TaskPtr->Source, TaskPtr->Result);
+								TaskPtr->Result = nullptr;
+								TaskPtr->Source = nullptr;
+							},
+							Requisites,
+							LowLevelTasks::ETaskPriority::BackgroundHigh);
 
-					if (bCanRun)
-					{
-						ConstantSubgraphs[Index].NonConcurrentTask();
-						ConstantSubgraphs[Index].NonConcurrentTask.Reset();
-						ConstantSubgraphs[Index].TaskData.Reset();
-						ConstantSubgraphs[Index].CompletedEvent.Trigger();
+						SubgraphCompletionEvent.AddPrerequisites(CompleteTask);
 					}
 
-					return bCanRun;
-				};
-
-			// Resolve tasks top-down while as long as they have the prerequisites completed.
-			int32 SubgraphCount = ConstantSubgraphs.Num();
-			TArray<int32> Pending;
-			for (int32 OrderIndex = 0; OrderIndex < SubgraphCount; ++OrderIndex)
-			{
-				int32 Index = SubgraphCount - 1 - OrderIndex;
-
-				bool bRan = TryRunFunc(Index);
-
-				if (bRan)
-				{
-					// We may have unlocked one of the pending tasks. Review them all.
-					bool bModified = true;
-					while (bModified)
-					{
-						bModified = false;
-						for (int32 PendingIndexIndex = 0; PendingIndexIndex < Pending.Num(); )
-						{
-							int32 PendingIndex = Pending[PendingIndexIndex];
-							bRan = TryRunFunc(PendingIndex);
-							if (bRan)
-							{
-								bModified = true;
-								Pending.RemoveAtSwap(PendingIndexIndex);
-								break;
-							}
-							else
-							{
-								++PendingIndexIndex;
-							}
-						}
-					}
-				}
-				else
-				{
-					// We couldn't run the task yet, remember it for later.
-					Pending.Add(Index);
+					ConstantSubgraphs[Index].Root = nullptr;
+					SubgraphCompletionEvent.Trigger();
+						
+					UE::Tasks::AddNested(SubgraphCompletionEvent);
 				}
-			}
 
-			check(Pending.IsEmpty());
+			});
+
+		// Wait for pending tasks
+		{
+			MUTABLE_CPUPROFILER_SCOPE(Waiting);
+			LaunchTask.Wait();
 		}
 
 		bool bSomethingModified = ConstantSubgraphs.Num() > 0;
diff --git a/Engine/Plugins/Mutable/Source/MutableTools/Private/MuT/Compiler.cpp b/Engine/Plugins/Mutable/Source/MutableTools/Private/MuT/Compiler.cpp
--- a/Engine/Plugins/Mutable/Source/MutableTools/Private/MuT/Compiler.cpp
+++ b/Engine/Plugins/Mutable/Source/MutableTools/Private/MuT/Compiler.cpp
@@ -252,13 +252,6 @@ namespace mu
 	}
 
 
-	//---------------------------------------------------------------------------------------------
-	void CompilerOptions::SetUseConcurrency(bool bEnabled)
-	{
-		m_pD->bUseConcurrency = bEnabled;
-	}
-
-
     //---------------------------------------------------------------------------------------------
     void CompilerOptions::SetOptimisationMaxIteration( int32 MaxIterations )
     {
diff --git a/Engine/Plugins/Mutable/Source/MutableTools/Private/MuT/CompilerPrivate.h b/Engine/Plugins/Mutable/Source/MutableTools/Private/MuT/CompilerPrivate.h
--- a/Engine/Plugins/Mutable/Source/MutableTools/Private/MuT/CompilerPrivate.h
+++ b/Engine/Plugins/Mutable/Source/MutableTools/Private/MuT/CompilerPrivate.h
@@ -66,9 +66,6 @@ namespace mu
         int32 ImageCompressionQuality = 0;
 		int32 ImageTiling = 0 ;
 
-		/** If this flag is enabled, the compiler can use concurrency to reduce compile time at the cost of higher CPU and memory usage. */
-		bool bUseConcurrency = false;
-
 		bool bIgnoreStates = false;
 		bool bLog = false;
 
diff --git a/Engine/Plugins/Mutable/Source/MutableTools/Public/MuT/Compiler.h b/Engine/Plugins/Mutable/Source/MutableTools/Public/MuT/Compiler.h
--- a/Engine/Plugins/Mutable/Source/MutableTools/Public/MuT/Compiler.h
+++ b/Engine/Plugins/Mutable/Source/MutableTools/Public/MuT/Compiler.h
@@ -53,9 +53,6 @@ namespace mu
         //!
         void SetIgnoreStates( bool bIgnore );
 
-		/** Enable concurrent compilation. It's faster, but uses more CPU and memory. */
-		void SetUseConcurrency(bool bEnabled);
-
         //! If enabled, the disk will be used as temporary memory. This will make the compilation
         //! process very slow, but will be able to compile very large models.
         void SetUseDiskCache( bool bEnabled );
