[Mutable] Overlay Material support for Mesh Components.

#rnx

diff --git a/Engine/Plugins/Mutable/Source/CustomizableObject/Internal/MuCO/CustomizableObjectInstancePrivate.h b/Engine/Plugins/Mutable/Source/CustomizableObject/Internal/MuCO/CustomizableObjectInstancePrivate.h
--- a/Engine/Plugins/Mutable/Source/CustomizableObject/Internal/MuCO/CustomizableObjectInstancePrivate.h
+++ b/Engine/Plugins/Mutable/Source/CustomizableObject/Internal/MuCO/CustomizableObjectInstancePrivate.h
@@ -152,6 +152,9 @@ struct FCustomizableInstanceComponentData
 
 	UPROPERTY(Transient)
 	TArray<TObjectPtr<UMaterialInterface>> OverrideMaterials;
+
+	UPROPERTY(Transient)
+	TObjectPtr<UMaterialInterface> OverlayMaterial;
 };
 
 USTRUCT()
diff --git a/Engine/Plugins/Mutable/Source/CustomizableObject/Internal/MuCO/CustomizableObjectSystemPrivate.h b/Engine/Plugins/Mutable/Source/CustomizableObject/Internal/MuCO/CustomizableObjectSystemPrivate.h
--- a/Engine/Plugins/Mutable/Source/CustomizableObject/Internal/MuCO/CustomizableObjectSystemPrivate.h
+++ b/Engine/Plugins/Mutable/Source/CustomizableObject/Internal/MuCO/CustomizableObjectSystemPrivate.h
@@ -400,6 +400,9 @@ struct FInstanceUpdateData
 		/** Range in the LODs array */
 		uint16 FirstLOD = 0;
 		uint16 LODCount = 0;
+
+		/** Overlay material. Only one per mesh, if any. */
+		int32 OverlayMaterial = INDEX_NONE;
 	};
 
 	TArray<FComponent> Components;
diff --git a/Engine/Plugins/Mutable/Source/CustomizableObject/Private/MuCO/CustomizableObjectInstance.cpp b/Engine/Plugins/Mutable/Source/CustomizableObject/Private/MuCO/CustomizableObjectInstance.cpp
--- a/Engine/Plugins/Mutable/Source/CustomizableObject/Private/MuCO/CustomizableObjectInstance.cpp
+++ b/Engine/Plugins/Mutable/Source/CustomizableObject/Private/MuCO/CustomizableObjectInstance.cpp
@@ -322,6 +322,7 @@ void UCustomizableInstancePrivate::PrepareForUpdate(const TSharedRef<FUpdateCont
 		ComponentData.PhysicsAssets.PhysicsAssetsToMerge.Empty();
 		ComponentData.ClothingPhysicsAssetsToStream.Empty();
 		ComponentData.StreamedResourceIndex.Empty();
+		ComponentData.OverlayMaterial = nullptr;
 
 #if WITH_EDITORONLY_DATA
 		ComponentData.MeshPartPaths.Empty();
@@ -5410,6 +5411,25 @@ UE::Tasks::FTask UCustomizableInstancePrivate::LoadAdditionalAssetsAndData(const
 		}
 	}
 
+	for (const FInstanceUpdateData::FComponent& Component : OperationData->InstanceUpdateData.Components)
+	{
+		if (Component.OverlayMaterial == INDEX_NONE || ObjectToInstanceIndexMap.Contains(Component.OverlayMaterial))
+		{
+			continue;
+		}
+
+		TSoftObjectPtr<UMaterialInterface> AssetPtr = ModelResources.Materials.IsValidIndex(Component.OverlayMaterial) ? ModelResources.Materials[Component.OverlayMaterial] : nullptr;
+		UMaterialInterface* LoadedMaterial = AssetPtr.Get();
+
+		const int32 ReferencedMaterialsIndex = ReferencedMaterials.Add(LoadedMaterial);
+		ObjectToInstanceIndexMap.Add(Component.OverlayMaterial, ReferencedMaterialsIndex);
+
+		if (!LoadedMaterial && !AssetPtr.IsNull())
+		{
+			AssetsToStream.Add(AssetPtr.ToSoftObjectPath());
+		}
+	}
+
 
 	// Load Skeletons required by the SubMeshes of the newly generated Mesh, will be merged later
 	for (int32 InstanceComponentIndex = 0; InstanceComponentIndex < OperationData->NumInstanceComponents; ++InstanceComponentIndex)
@@ -6401,10 +6421,28 @@ void UCustomizableInstancePrivate::BuildMaterials(const TSharedRef<FUpdateContex
 		// It is not safe to replace the materials of a SkeletalMesh whose resources are initialized. Use overrides instead.
 		const bool bUseOverrideMaterialsOnly = !bIsTransientMesh || (OperationData->bUseMeshCache && SkeletalMesh->GetResourceForRendering()->IsInitialized());
 
+		UMaterialInterface* OverlayMaterial = nullptr;
+
 		FCustomizableInstanceComponentData* ComponentData = GetComponentData(ObjectComponentIndex);
 		if (ComponentData)
 		{
 			ComponentData->OverrideMaterials.Reset();
+			ComponentData->OverlayMaterial = nullptr;
+
+			if (const uint32* ReferencedMaterialIndex = ObjectToInstanceIndexMap.Find(Component.OverlayMaterial))
+			{
+				if (ReferencedMaterials.IsValidIndex(*ReferencedMaterialIndex))
+				{
+					ComponentData->OverlayMaterial = ReferencedMaterials[*ReferencedMaterialIndex];
+					OverlayMaterial = ComponentData->OverlayMaterial;
+				}
+			}
+		}
+
+		if (!bUseOverrideMaterialsOnly)
+		{
+			RecreateRenderStateOnInstanceComponent[InstanceComponentIndex] |= SkeletalMesh->GetOverlayMaterial() != OverlayMaterial;
+			SkeletalMesh->SetOverlayMaterial(OverlayMaterial);
 		}
 
 		TArray<FSkeletalMaterial> Materials;
@@ -7029,7 +7067,7 @@ void UCustomizableInstancePrivate::BuildMaterials(const TSharedRef<FUpdateContex
 			}
 
 			// Force recreate render state after replacing the materials to avoid a crash in the render pipeline if the old materials are GCed while in use.
-			RecreateRenderStateOnInstanceComponent[InstanceComponentIndex] = SkeletalMesh->GetResourceForRendering()->IsInitialized() && SkeletalMesh->GetMaterials() != Materials;
+			RecreateRenderStateOnInstanceComponent[InstanceComponentIndex] |= SkeletalMesh->GetResourceForRendering()->IsInitialized() && SkeletalMesh->GetMaterials() != Materials;
 
 			SkeletalMesh->SetMaterials(Materials);
 
diff --git a/Engine/Plugins/Mutable/Source/CustomizableObject/Private/MuCO/CustomizableObjectSystem.cpp b/Engine/Plugins/Mutable/Source/CustomizableObject/Private/MuCO/CustomizableObjectSystem.cpp
--- a/Engine/Plugins/Mutable/Source/CustomizableObject/Private/MuCO/CustomizableObjectSystem.cpp
+++ b/Engine/Plugins/Mutable/Source/CustomizableObject/Private/MuCO/CustomizableObjectSystem.cpp
@@ -1276,11 +1276,21 @@ void UpdateSkeletalMesh(const TSharedRef<FUpdateContextPrivate>& Context)
 			}
 		}
 
-		if (OldOverridenMaterials != Parent->OverrideMaterials)
+		bAttachedParentUpdated |= OldOverridenMaterials != Parent->OverrideMaterials;
+
+		UMaterialInterface* OverlayMaterial = nullptr;
+		if (bUseOverrideMaterials /* bUseOverlayMaterials*/)
 		{
-			bAttachedParentUpdated = true;
+			if (FCustomizableInstanceComponentData* ComponentData = CustomizableObjectInstance->GetPrivate()->GetComponentData(Usage->GetComponentName()))
+			{
+				OverlayMaterial = ComponentData->OverlayMaterial;
+			}
 		}
 
+		bAttachedParentUpdated |= Parent->GetOverlayMaterial() != OverlayMaterial;
+		Parent->SetOverlayMaterial(OverlayMaterial);
+
+
 		if (CustomizableObjectInstancePrivateData->HasCOInstanceFlags(ReplacePhysicsAssets) &&
 			SkeletalMesh &&
 			Parent->GetWorld())
@@ -2828,6 +2838,8 @@ namespace impl
 					continue;
 				}
 
+				Component.OverlayMaterial = Instance->GetOverlayMaterial(InstanceComponentIndex);
+
 				for (int32 MutableLODIndex = 0; MutableLODIndex < Component.LODCount; ++MutableLODIndex)
 				{
 					// If the LOD is not generated we still add an empty one to keep indexes aligned.
diff --git a/Engine/Plugins/Mutable/Source/CustomizableObjectEditor/Private/MuCOE/GenerateMutableSource/GenerateMutableSourceComponent.cpp b/Engine/Plugins/Mutable/Source/CustomizableObjectEditor/Private/MuCOE/GenerateMutableSource/GenerateMutableSourceComponent.cpp
--- a/Engine/Plugins/Mutable/Source/CustomizableObjectEditor/Private/MuCOE/GenerateMutableSource/GenerateMutableSourceComponent.cpp
+++ b/Engine/Plugins/Mutable/Source/CustomizableObjectEditor/Private/MuCOE/GenerateMutableSource/GenerateMutableSourceComponent.cpp
@@ -24,6 +24,8 @@
 #include "MuT/NodeComponentSwitch.h"
 #include "MuT/NodeComponentVariation.h"
 #include "MuT/NodeLOD.h"
+#include "MuT/NodeScalar.h"
+#include "MuT/NodeScalarConstant.h"
 #include "MuT/NodeSurfaceNew.h"
 #include "MuT/NodeMeshConstant.h"
 #include "MuT/NodeObjectNew.h"
@@ -142,6 +144,20 @@ mu::Ptr<mu::NodeComponent> GenerateMutableSourceComponent(const UEdGraphPin* Pin
 		NodeComponentNew->Id = GenerationContext.ComponentNames.Find(TypedComponentMesh->ComponentName);
 		NodeComponentNew->SetMessageContext(Node);
 		
+		UEdGraphPin* MaterialAssetPin = TypedComponentMesh->GetOverlayMaterialAssetPin();
+		if (const UEdGraphPin* ConnectedPin = MaterialAssetPin ? FollowInputPin(*MaterialAssetPin) : nullptr)
+		{
+			GenerationContext.CurrentMaterialTableParameterId = TEXT("__MutableMaterialId") + ConnectedPin->PinId.ToString();
+			mu::Ptr<mu::NodeScalar> ScalarNode = GenerateMutableSourceFloat(ConnectedPin, GenerationContext);
+			NodeComponentNew->OverlayMaterial = ScalarNode;
+		}
+		else if(UMaterialInterface* OverlayMaterial = GenerationContext.LoadObject(TypedComponentMesh->GetOverlayMaterial()))
+		{
+			mu::Ptr<mu::NodeScalarConstant> ScalarNode = new mu::NodeScalarConstant();
+			ScalarNode->SetValue(GenerationContext.ReferencedMaterials.AddUnique(OverlayMaterial));
+			NodeComponentNew->OverlayMaterial = ScalarNode;
+		}
+
 		Result = NodeComponentNew;
 
 		GenerationContext.CurrentMeshComponent = TypedComponentMesh->ComponentName;
diff --git a/Engine/Plugins/Mutable/Source/CustomizableObjectEditor/Private/MuCOE/GenerateMutableSource/GenerateMutableSourceSurface.cpp b/Engine/Plugins/Mutable/Source/CustomizableObjectEditor/Private/MuCOE/GenerateMutableSource/GenerateMutableSourceSurface.cpp
--- a/Engine/Plugins/Mutable/Source/CustomizableObjectEditor/Private/MuCOE/GenerateMutableSource/GenerateMutableSourceSurface.cpp
+++ b/Engine/Plugins/Mutable/Source/CustomizableObjectEditor/Private/MuCOE/GenerateMutableSource/GenerateMutableSourceSurface.cpp
@@ -934,7 +934,10 @@ mu::Ptr<mu::NodeSurface> GenerateMutableSourceSurface(const UEdGraphPin * Pin, F
 			if (bTableMaterialPinLinked)
 			{
 				if (const UEdGraphPin* ConnectedPin = FollowInputPin(*MaterialPin))
-				{
+				{					
+					// The 'None' option will use the CurrentReferencedMaterialIndex.
+					check(GenerationContext.CurrentReferencedMaterialIndex != INDEX_NONE);
+					
 					// Adding pin ID to support multiple material columns
 					MaterialName += MaterialPin->PinId.ToString();
 
diff --git a/Engine/Plugins/Mutable/Source/CustomizableObjectEditor/Private/MuCOE/GenerateMutableSource/GenerateMutableSourceTable.cpp b/Engine/Plugins/Mutable/Source/CustomizableObjectEditor/Private/MuCOE/GenerateMutableSource/GenerateMutableSourceTable.cpp
--- a/Engine/Plugins/Mutable/Source/CustomizableObjectEditor/Private/MuCOE/GenerateMutableSource/GenerateMutableSourceTable.cpp
+++ b/Engine/Plugins/Mutable/Source/CustomizableObjectEditor/Private/MuCOE/GenerateMutableSource/GenerateMutableSourceTable.cpp
@@ -1316,15 +1316,6 @@ bool GenerateNoneRow(const UCustomizableObjectNodeTable& TableNode, const UEdGra
 			ColumnId = MutableTable->AddColumn(ColumnName, mu::ETableColumnType::Scalar);
 		}
 
-		// Get the material of the linked mesh section
-		if (!GenerationContext.ReferencedMaterials.IsValidIndex(GenerationContext.CurrentReferencedMaterialIndex))
-		{
-			FText Message = FText::Format(LOCTEXT("NoneRowWarning_InvalidMaterialIndex", "Error creating [None] row for Table Node {0}. Invalid Material Index."), FText::FromString(TableNode.ParameterName));
-			GenerationContext.Log(Message, &TableNode, EMessageSeverity::Error);
-
-			return false;
-		}
-
 		MutableTable->SetCell(ColumnId, RowId, GenerationContext.CurrentReferencedMaterialIndex);
 	}
 	else if (Pin->PinType.PinCategory == Schema->PC_Color)
diff --git a/Engine/Plugins/Mutable/Source/CustomizableObjectEditor/Private/MuCOE/Nodes/CustomizableObjectNodeComponentMesh.cpp b/Engine/Plugins/Mutable/Source/CustomizableObjectEditor/Private/MuCOE/Nodes/CustomizableObjectNodeComponentMesh.cpp
--- a/Engine/Plugins/Mutable/Source/CustomizableObjectEditor/Private/MuCOE/Nodes/CustomizableObjectNodeComponentMesh.cpp
+++ b/Engine/Plugins/Mutable/Source/CustomizableObjectEditor/Private/MuCOE/Nodes/CustomizableObjectNodeComponentMesh.cpp
@@ -48,6 +48,23 @@ void UCustomizableObjectNodeComponentMesh::BackwardsCompatibleFixup(int32 Custom
 }
 
 
+void UCustomizableObjectNodeComponentMesh::AllocateDefaultPins(UCustomizableObjectNodeRemapPins* RemapPins)
+{
+	const UEdGraphSchema_CustomizableObject* Schema = GetDefault<UEdGraphSchema_CustomizableObject>();
+
+	{
+		const FString PinFriendlyName = TEXT("Overlay Material");
+		const FString PinName = PinFriendlyName + FString(TEXT("_Input_Pin"));
+		UEdGraphPin* OverlayMaterialPin = CustomCreatePin(EGPD_Input, Schema->PC_MaterialAsset, FName(*PinName));
+		OverlayMaterialPin->PinFriendlyName = FText::FromString(PinFriendlyName);
+		OverlayMaterialPin->bDefaultValueIsIgnored = true;
+		OverlayMaterialPin->PinToolTip = "Pin for an Overlay Material from a Table Node";
+	}
+	
+	Super::AllocateDefaultPins(RemapPins);
+}
+
+
 FText UCustomizableObjectNodeComponentMesh::GetNodeTitle(ENodeTitleType::Type TitleType) const
 {
 	FText NodeTitle;
@@ -80,5 +97,27 @@ void UCustomizableObjectNodeComponentMesh::OnRenameNode(const FString& NewName)
 }
 
 
+TSoftObjectPtr<UMaterialInterface> UCustomizableObjectNodeComponentMesh::GetOverlayMaterial() const
+{
+	return OverlayMaterial;
+}
+
+
+UEdGraphPin* UCustomizableObjectNodeComponentMesh::GetOverlayMaterialAssetPin() const
+{
+	const FString PinFriendlyName = TEXT("Overlay Material");
+	const FString PinName = PinFriendlyName + FString(TEXT("_Input_Pin"));
+
+	if (UEdGraphPin* Pin = FindPin(PinName))
+	{
+		return Pin;
+	}
+	else
+	{
+		return FindPin(PinFriendlyName);
+	}
+}
+
+
 #undef LOCTEXT_NAMESPACE
 
diff --git a/Engine/Plugins/Mutable/Source/CustomizableObjectEditor/Private/MuCOE/Nodes/CustomizableObjectNodeComponentMesh.h b/Engine/Plugins/Mutable/Source/CustomizableObjectEditor/Private/MuCOE/Nodes/CustomizableObjectNodeComponentMesh.h
--- a/Engine/Plugins/Mutable/Source/CustomizableObjectEditor/Private/MuCOE/Nodes/CustomizableObjectNodeComponentMesh.h
+++ b/Engine/Plugins/Mutable/Source/CustomizableObjectEditor/Private/MuCOE/Nodes/CustomizableObjectNodeComponentMesh.h
@@ -7,6 +7,7 @@
 
 #include "CustomizableObjectNodeComponentMesh.generated.h"
 
+class UMaterialInterface;
 
 UENUM()
 enum class ECustomizableObjectSelectionOverride : uint8
@@ -55,9 +56,14 @@ public:
 	virtual void BackwardsCompatibleFixup(int32 CustomizableObjectCustomVersion) override;
 
 	// UEdGraphNode interface
+	virtual void AllocateDefaultPins(UCustomizableObjectNodeRemapPins* RemapPins) override;
 	virtual FText GetNodeTitle(ENodeTitleType::Type TitleType) const override;
 	virtual void OnRenameNode(const FString& NewName) override;
 
+	// Own interface
+	TSoftObjectPtr<UMaterialInterface> GetOverlayMaterial() const;
+	UEdGraphPin* GetOverlayMaterialAssetPin() const;
+
 	UPROPERTY(EditAnywhere, Category = ComponentMesh)
 	FName ComponentName = "Default name";
 	
@@ -74,6 +80,9 @@ public:
 	UPROPERTY(EditAnywhere, Category = ComponentMesh)
 	TArray<FLODReductionSettings> LODReductionSettings;
 
+	UPROPERTY(EditAnywhere, Category = ComponentMesh)
+	TSoftObjectPtr<UMaterialInterface> OverlayMaterial;
+
 	/** Details selected LOD. */
 	int32 SelectedLOD = 0;
 };
diff --git a/Engine/Plugins/Mutable/Source/CustomizableObjectEditor/Private/MuCOE/SMutableGraphViewer.cpp b/Engine/Plugins/Mutable/Source/CustomizableObjectEditor/Private/MuCOE/SMutableGraphViewer.cpp
--- a/Engine/Plugins/Mutable/Source/CustomizableObjectEditor/Private/MuCOE/SMutableGraphViewer.cpp
+++ b/Engine/Plugins/Mutable/Source/CustomizableObjectEditor/Private/MuCOE/SMutableGraphViewer.cpp
@@ -572,6 +572,8 @@ void SMutableGraphViewer::GetChildrenForInfo(TSharedPtr<FMutableGraphTreeElement
 		{
 			AddChildFunc(ComponentVar->LODs[LODIndex].get(), FString::Printf(TEXT("LOD [%d]"), LODIndex));
 		}
+
+		AddChildFunc(ComponentVar->OverlayMaterial.get(), TEXT("OVERLAY MATERIAL"));
 	}
 
 	else if (ParentNode->GetType() == mu::NodeComponentEdit::GetStaticType())
diff --git a/Engine/Plugins/Mutable/Source/MutableRuntime/Internal/MuR/Operations.h b/Engine/Plugins/Mutable/Source/MutableRuntime/Internal/MuR/Operations.h
--- a/Engine/Plugins/Mutable/Source/MutableRuntime/Internal/MuR/Operations.h
+++ b/Engine/Plugins/Mutable/Source/MutableRuntime/Internal/MuR/Operations.h
@@ -305,6 +305,9 @@ namespace mu
 		//! Add extension data to an instance
 		IN_ADDEXTENSIONDATA,
 
+		//! Add overlay material to an instance
+		IN_ADDOVERLAYMATERIAL,
+
         //-----------------------------------------------------------------------------------------
         // Layout operations
         //-----------------------------------------------------------------------------------------
@@ -965,6 +968,14 @@ namespace mu
 			ADDRESS ExtensionDataName;
 		};
 
+		struct InstanceAddOverlayMaterialArgs
+		{
+			// This is a reference to an op that produces the Instance that the material will
+			// be added to.
+			ADDRESS Instance;
+			ADDRESS OverlayMaterialId;
+		};
+
         struct LayoutPackArgs
         {
             ADDRESS Source;
diff --git a/Engine/Plugins/Mutable/Source/MutableRuntime/Private/MuR/CodeRunner.cpp b/Engine/Plugins/Mutable/Source/MutableRuntime/Private/MuR/CodeRunner.cpp
--- a/Engine/Plugins/Mutable/Source/MutableRuntime/Private/MuR/CodeRunner.cpp
+++ b/Engine/Plugins/Mutable/Source/MutableRuntime/Private/MuR/CodeRunner.cpp
@@ -791,11 +791,9 @@ namespace mu
 					check(Args.Instance);
 					check(Args.ExtensionData);
 
-					TArray<FScheduledOp> Dependencies;
-					Dependencies.Emplace(Args.Instance, Item);
-					Dependencies.Emplace(Args.ExtensionData, Item);
-
-					AddOp(FScheduledOp(Item.At, Item, 1), Dependencies);
+					AddOp(FScheduledOp(Item.At, Item, 1),
+						FScheduledOp(Args.Instance, Item),
+						FScheduledOp(Args.ExtensionData, Item));
 
 					break;
 				}
@@ -826,6 +824,48 @@ namespace mu
 			
 			break;
 		}
+		case EOpType::IN_ADDOVERLAYMATERIAL:
+		{
+			OP::InstanceAddOverlayMaterialArgs Args = Program.GetOpArgs<OP::InstanceAddOverlayMaterialArgs>(Item.At);
+			switch (Item.Stage)
+			{
+			case 0:
+			{
+				// Must pass in a Material op
+				check(Args.OverlayMaterialId);
+
+				AddOp(FScheduledOp(Item.At, Item, 1),
+					FScheduledOp(Args.Instance, Item),
+					FScheduledOp(Args.OverlayMaterialId, Item));
+
+				break;
+			}
+
+			case 1:
+			{
+				TSharedPtr<const FInstance> Base = LoadInstance(FCacheAddress(Args.Instance, Item));
+				TSharedPtr<FInstance> Result;
+				if (!Base)
+				{
+					Result = MakeShared<FInstance>();
+				}
+				else
+				{
+					Result = mu::CloneOrTakeOver<FInstance>(Base);
+				}
+
+				Result->GetPrivate()->AddOverlayMaterial(0, LoadScalar(FCacheAddress(Args.OverlayMaterialId, Item)));
+
+				StoreInstance(Item, Result);
+				break;
+			}
+
+			default:
+				check(false);
+			}
+
+			break;
+		}
 
         default:
                 check(false);
diff --git a/Engine/Plugins/Mutable/Source/MutableRuntime/Private/MuR/Instance.cpp b/Engine/Plugins/Mutable/Source/MutableRuntime/Private/MuR/Instance.cpp
--- a/Engine/Plugins/Mutable/Source/MutableRuntime/Private/MuR/Instance.cpp
+++ b/Engine/Plugins/Mutable/Source/MutableRuntime/Private/MuR/Instance.cpp
@@ -379,6 +379,12 @@ namespace mu
         return NAME_None;
     }
 
+	int32 FInstance::GetOverlayMaterial(int32 ComponentIndex) const
+	{
+		check(m_pD->Components.IsValidIndex(ComponentIndex));
+		check(m_pD->Components[ComponentIndex].OverlayMaterialId < int32(1 << 22));
+		return static_cast<int32>(m_pD->Components[ComponentIndex].OverlayMaterialId);
+	}
 
     //---------------------------------------------------------------------------------------------
 	int32 FInstance::GetExtensionDataCount() const
@@ -564,6 +570,17 @@ namespace mu
     }
 
 
+	void FInstance::Private::AddOverlayMaterial(int32 ComponentIndex, float OverlayMaterialId)
+	{
+		while (ComponentIndex >= Components.Num())
+		{
+			AddComponent();
+		}
+
+		Components[ComponentIndex].OverlayMaterialId = OverlayMaterialId;
+	}
+
+
     //---------------------------------------------------------------------------------------------
 	void FInstance::Private::AddExtensionData(const TSharedPtr<const FExtensionData>& Data, FName Name)
 	{
diff --git a/Engine/Plugins/Mutable/Source/MutableRuntime/Private/MuR/InstancePrivate.h b/Engine/Plugins/Mutable/Source/MutableRuntime/Private/MuR/InstancePrivate.h
--- a/Engine/Plugins/Mutable/Source/MutableRuntime/Private/MuR/InstancePrivate.h
+++ b/Engine/Plugins/Mutable/Source/MutableRuntime/Private/MuR/InstancePrivate.h
@@ -73,6 +73,8 @@ namespace mu
     {
 		uint16 Id;
 
+		float OverlayMaterialId = -1.f;
+
 		TArray<FInstanceLOD, TInlineAllocator<4>> LODs;
     };
 
@@ -104,7 +106,8 @@ namespace mu
         int32 AddVector( int32 ComponentIndex, int32 LODIndex, int32 SurfaceIndex, const FVector4f&, FName Name);
         int32 AddScalar( int32 ComponentIndex, int32 LODIndex, int32 SurfaceIndex, float, FName Name);
         int32 AddString( int32 ComponentIndex, int32 LODIndex, int32 SurfaceIndex, const FString& Value, FName Name);
-		
+		void AddOverlayMaterial( int32 ComponentIndex, float OverlayMaterialId);
+
 		// Data must be non-null
 		void AddExtensionData(const TSharedPtr<const FExtensionData>& Data, FName Name);
     };
diff --git a/Engine/Plugins/Mutable/Source/MutableRuntime/Private/MuR/Operations.cpp b/Engine/Plugins/Mutable/Source/MutableRuntime/Private/MuR/Operations.cpp
--- a/Engine/Plugins/Mutable/Source/MutableRuntime/Private/MuR/Operations.cpp
+++ b/Engine/Plugins/Mutable/Source/MutableRuntime/Private/MuR/Operations.cpp
@@ -131,6 +131,7 @@ namespace mu
 		{ EDataType::Instance,		},	// IN_ADDCOMPONENT
 		{ EDataType::Instance,		},	// IN_ADDLOD
 		{ EDataType::Instance,		},	// IN_ADDEXTENSIONDATA
+		{ EDataType::Instance,		},	// IN_ADDOVERLAYMATERIAL
 
 		{ EDataType::Layout,		},	// LA_PACK
 		{ EDataType::Layout,		},	// LA_MERGE
@@ -790,6 +791,16 @@ namespace mu
 			break;
 		}
 
+		case EOpType::IN_ADDOVERLAYMATERIAL:
+		{
+			const OP::InstanceAddOverlayMaterialArgs Args = program.GetOpArgs<OP::InstanceAddOverlayMaterialArgs>(at);
+
+			f(Args.Instance);
+			f(Args.OverlayMaterialId);
+
+			break;
+		}
+
         //-------------------------------------------------------------------------------------
         case EOpType::LA_PACK:
         {
diff --git a/Engine/Plugins/Mutable/Source/MutableRuntime/Public/MuR/Instance.h b/Engine/Plugins/Mutable/Source/MutableRuntime/Public/MuR/Instance.h
--- a/Engine/Plugins/Mutable/Source/MutableRuntime/Public/MuR/Instance.h
+++ b/Engine/Plugins/Mutable/Source/MutableRuntime/Public/MuR/Instance.h
@@ -183,6 +183,10 @@ namespace mu
         //! \param StringIndex Index of the string, from 0 to GetStringCount()-1
         FName GetStringName( int32 ComponentIndex, int32 LODIndex, int32 SurfaceIndex, int32 StringIndex ) const;
 
+		//! Get the index of the overlay material.
+		//! \param ComponentIndex Index of the component, from 0 to GetComponentCount()-1
+		int32 GetOverlayMaterial( int32 ComponentIndex ) const;
+
 		//! Get the number of ExtensionData values in a component
 		int32 GetExtensionDataCount() const;
 
diff --git a/Engine/Plugins/Mutable/Source/MutableTools/Private/MuT/ASTOpAddOverlayMaterial.cpp b/Engine/Plugins/Mutable/Source/MutableTools/Private/MuT/ASTOpAddOverlayMaterial.cpp
new file mode 100644
--- /dev/null
+++ b/Engine/Plugins/Mutable/Source/MutableTools/Private/MuT/ASTOpAddOverlayMaterial.cpp
@@ -0,0 +1,95 @@
+// Copyright Epic Games, Inc. All Rights Reserved.
+
+#include "MuT/ASTOpAddOverlayMaterial.h"
+
+#include "HAL/PlatformMath.h"
+#include "MuR/ModelPrivate.h"
+#include "MuR/RefCounted.h"
+
+
+namespace mu
+{
+
+
+//-------------------------------------------------------------------------------------------------
+//-------------------------------------------------------------------------------------------------
+//-------------------------------------------------------------------------------------------------
+	ASTOpAddOverlayMaterial::ASTOpAddOverlayMaterial()
+	: Instance(this)
+	, OverlayMaterialId(this)
+{
+}
+
+
+	ASTOpAddOverlayMaterial::~ASTOpAddOverlayMaterial()
+{
+	// Explicit call needed to avoid recursive destruction
+	ASTOp::RemoveChildren();
+}
+
+bool ASTOpAddOverlayMaterial::IsEqual(const ASTOp& OtherUntyped) const
+{
+	if (OtherUntyped.GetOpType() == GetOpType())
+	{
+		const ASTOpAddOverlayMaterial* Other = static_cast<const ASTOpAddOverlayMaterial*>(&OtherUntyped);
+		return Instance == Other->Instance
+			&& OverlayMaterialId == Other->OverlayMaterialId;
+	}
+
+	return false;
+}
+
+uint64 ASTOpAddOverlayMaterial::Hash() const
+{
+	uint64 Result = std::hash<uint64>()(uint64(EOpType::IN_ADDOVERLAYMATERIAL));
+	
+	hash_combine(Result, Instance.child().get());
+	hash_combine(Result, OverlayMaterialId.child().get());
+	
+	return Result;
+}
+
+mu::Ptr<ASTOp> ASTOpAddOverlayMaterial::Clone(MapChildFuncRef MapChild) const
+{
+	Ptr<ASTOpAddOverlayMaterial> NewInstance = new ASTOpAddOverlayMaterial();
+
+	NewInstance->Instance = MapChild(Instance.child());
+	NewInstance->OverlayMaterialId = MapChild(OverlayMaterialId.child());
+
+	return NewInstance;
+}
+
+void ASTOpAddOverlayMaterial::ForEachChild(const TFunctionRef<void(ASTChild&)> F)
+{
+	F(Instance);
+	F(OverlayMaterialId);
+}
+
+void ASTOpAddOverlayMaterial::Link(FProgram& Program, FLinkerOptions*)
+{
+	// Already linked?
+	if (linkedAddress)
+	{
+		return;
+	}
+
+	OP::InstanceAddOverlayMaterialArgs Args;
+	FMemory::Memzero(Args);
+
+	if (Instance)
+	{
+		Args.Instance = Instance->linkedAddress;
+	}
+
+	if (OverlayMaterialId)
+	{
+		Args.OverlayMaterialId = OverlayMaterialId->linkedAddress;
+	}
+	
+	linkedAddress = (OP::ADDRESS)Program.OpAddress.Num();
+	Program.OpAddress.Add((uint32_t)Program.ByteCode.Num());
+	AppendCode(Program.ByteCode, EOpType::IN_ADDOVERLAYMATERIAL);
+	AppendCode(Program.ByteCode, Args);
+}
+
+}
diff --git a/Engine/Plugins/Mutable/Source/MutableTools/Private/MuT/ASTOpAddOverlayMaterial.h b/Engine/Plugins/Mutable/Source/MutableTools/Private/MuT/ASTOpAddOverlayMaterial.h
new file mode 100644
--- /dev/null
+++ b/Engine/Plugins/Mutable/Source/MutableTools/Private/MuT/ASTOpAddOverlayMaterial.h
@@ -0,0 +1,43 @@
+// Copyright Epic Games, Inc. All Rights Reserved.
+
+#pragma once
+
+#include "HAL/Platform.h"
+#include "Math/NumericLimits.h"
+#include "MuR/Operations.h"
+#include "MuR/Ptr.h"
+#include "MuT/AST.h"
+
+
+namespace mu
+{
+struct FProgram;
+
+
+//---------------------------------------------------------------------------------------------
+//! Adds an overlay material to an Instance
+//---------------------------------------------------------------------------------------------
+class ASTOpAddOverlayMaterial final : public ASTOp
+{
+public:
+
+	ASTChild Instance;
+	ASTChild OverlayMaterialId;
+
+public:
+
+	ASTOpAddOverlayMaterial();
+	ASTOpAddOverlayMaterial(const ASTOpAddOverlayMaterial&) = delete;
+	~ASTOpAddOverlayMaterial();
+
+	EOpType GetOpType() const override { return EOpType::IN_ADDOVERLAYMATERIAL; }
+	uint64 Hash() const override;
+	void ForEachChild(const TFunctionRef<void(ASTChild&)> F) override;
+	bool IsEqual(const ASTOp& OtherUntyped) const override;
+	Ptr<ASTOp> Clone(MapChildFuncRef MapChild) const override;
+	void Link(FProgram& Program, FLinkerOptions* Options) override;
+};
+
+
+}
+
diff --git a/Engine/Plugins/Mutable/Source/MutableTools/Private/MuT/CodeGenerator.cpp b/Engine/Plugins/Mutable/Source/MutableTools/Private/MuT/CodeGenerator.cpp
--- a/Engine/Plugins/Mutable/Source/MutableTools/Private/MuT/CodeGenerator.cpp
+++ b/Engine/Plugins/Mutable/Source/MutableTools/Private/MuT/CodeGenerator.cpp
@@ -16,6 +16,7 @@
 #include "MuR/Platform.h"
 #include "MuT/ASTOpAddLOD.h"
 #include "MuT/ASTOpAddExtensionData.h"
+#include "MuT/ASTOpAddOverlayMaterial.h"
 #include "MuT/ASTOpConditional.h"
 #include "MuT/ASTOpConstantBool.h"
 #include "MuT/ASTOpConstantResource.h"
@@ -610,10 +611,25 @@ namespace mu
 			}
 		}
 
+		Ptr<ASTOp> LastInstOp = LODsOp;
+
+		if (NodeScalarPtr pScalarNode = InNode->OverlayMaterial)
+		{
+			Ptr<ASTOpAddOverlayMaterial> AddOverlayMaterialOp = new ASTOpAddOverlayMaterial();
+			AddOverlayMaterialOp->Instance = LastInstOp;
+
+			// Scalar
+			FScalarGenerationResult ScalarResult;
+			GenerateScalar(ScalarResult, Options, pScalarNode);
+			AddOverlayMaterialOp->OverlayMaterialId = ScalarResult.op;
+
+			LastInstOp = AddOverlayMaterialOp;
+		}
+
 		Ptr<ASTOpInstanceAdd> InstanceOp = new ASTOpInstanceAdd();
 		InstanceOp->type = EOpType::IN_ADDCOMPONENT;
 		InstanceOp->instance = Options.BaseInstance;
-		InstanceOp->value = LODsOp;
+		InstanceOp->value = LastInstOp;
 		InstanceOp->ExternalId = node.Id;
 
 		Result.Op = InstanceOp;
diff --git a/Engine/Plugins/Mutable/Source/MutableTools/Private/MuT/Compiler.cpp b/Engine/Plugins/Mutable/Source/MutableTools/Private/MuT/Compiler.cpp
--- a/Engine/Plugins/Mutable/Source/MutableTools/Private/MuT/Compiler.cpp
+++ b/Engine/Plugins/Mutable/Source/MutableTools/Private/MuT/Compiler.cpp
@@ -169,6 +169,7 @@ namespace mu
 		{ false,	{ 0, 0, 0, 0, 0, 0, 0, 0, 0, 0, 0, 0, 0, 0, 0, 0, 0, 0, 0, 0, 0, 0, 0, 0, 0, 0, 0, 0, 0 }		},	// IN_ADDCOMPONENT
 		{ false,	{ 0, 0, 0, 0, 0, 0, 0, 0, 0, 0, 0, 0, 0, 0, 0, 0, 0, 0, 0, 0, 0, 0, 0, 0, 0, 0, 0, 0, 0 }		},	// IN_ADDLOD
 		{ false,	{ 0, 0, 0, 0, 0, 0, 0, 0, 0, 0, 0, 0, 0, 0, 0, 0, 0, 0, 0, 0, 0, 0, 0, 0, 0, 0, 0, 0, 0 }		},	// IN_ADDEXTENSIONDATA
+		{ false,	{ 0, 0, 0, 0, 0, 0, 0, 0, 0, 0, 0, 0, 0, 0, 0, 0, 0, 0, 0, 0, 0, 0, 0, 0, 0, 0, 0, 0, 0 }		},	// IN_ADDOVERLAYMATERIAL
 
 		{ true,		{ 0, 0, 0, 0, 0, 0, 0, 0, 0, 0, 0, 0, 0, 0, 0, 0, 0, 0, 0, 0, 0, 0, 0, 0, 0, 0, 0, 0, 0 }		},	// LA_PACK
 		{ true,		{ 0, 0, 0, 0, 0, 0, 0, 0, 0, 0, 0, 0, 0, 0, 0, 0, 0, 0, 0, 0, 0, 0, 0, 0, 0, 0, 0, 0, 0 }		},	// LA_MERGE
diff --git a/Engine/Plugins/Mutable/Source/MutableTools/Private/MuT/ErrorLog.cpp b/Engine/Plugins/Mutable/Source/MutableTools/Private/MuT/ErrorLog.cpp
--- a/Engine/Plugins/Mutable/Source/MutableTools/Private/MuT/ErrorLog.cpp
+++ b/Engine/Plugins/Mutable/Source/MutableTools/Private/MuT/ErrorLog.cpp
@@ -315,6 +315,7 @@ namespace mu
         TEXT("IN_ADDCOMPONENT  "),
         TEXT("IN_ADDLOD        "),
 		TEXT("IN_ADDEXTENSIDATA"),
+		TEXT("IN_ADDOVERLAYMATERIAL"),
 		
 		TEXT("LA_PACK          "),
 		TEXT("LA_MERGE         "),
diff --git a/Engine/Plugins/Mutable/Source/MutableTools/Public/MuT/NodeComponentNew.h b/Engine/Plugins/Mutable/Source/MutableTools/Public/MuT/NodeComponentNew.h
--- a/Engine/Plugins/Mutable/Source/MutableTools/Public/MuT/NodeComponentNew.h
+++ b/Engine/Plugins/Mutable/Source/MutableTools/Public/MuT/NodeComponentNew.h
@@ -7,7 +7,7 @@
 #include "MuR/RefCounted.h"
 #include "MuT/Node.h"
 #include "MuT/NodeComponent.h"
-
+#include "MuT/NodeScalar.h"
 
 namespace mu
 {
@@ -19,6 +19,8 @@ namespace mu
 		/** Externally managed id assigned to this component. */
 		uint16 Id = 0;
 
+		Ptr<NodeScalar> OverlayMaterial;
+
 	public:
 
 		//-----------------------------------------------------------------------------------------
