[mutable] Fix the mutable debugger "mesh constant" UI not fitlering related operations in the tree.
#rnx

diff --git a/Engine/Plugins/Mutable/Source/CustomizableObjectEditor/Private/MuCOE/SMutableCodeViewer.cpp b/Engine/Plugins/Mutable/Source/CustomizableObjectEditor/Private/MuCOE/SMutableCodeViewer.cpp
--- a/Engine/Plugins/Mutable/Source/CustomizableObjectEditor/Private/MuCOE/SMutableCodeViewer.cpp
+++ b/Engine/Plugins/Mutable/Source/CustomizableObjectEditor/Private/MuCOE/SMutableCodeViewer.cpp
@@ -2550,7 +2550,26 @@ bool SMutableCodeViewer::IsConstantResourceUsedByOperation(const int32 IndexOnCo
 			{
 				if (OperationType == mu::EOpType::ME_CONSTANT)
 				{
-					bResourceLocated = IndexOnConstantsArray == InProgram.GetOpArgs<mu::OP::MeshConstantArgs>(OperationAddress).value;
+					mu::OP::MeshConstantArgs Args = InProgram.GetOpArgs<mu::OP::MeshConstantArgs>(OperationAddress);
+					mu::FConstantResourceIndex ResourceIndex = *reinterpret_cast<const mu::FConstantResourceIndex*>(&Args.value);
+					
+					if (!ResourceIndex.Streamable)
+					{
+						bResourceLocated = IndexOnConstantsArray == ResourceIndex.Index;
+					}
+					else
+					{
+						int32 DebuggerIndexOffset = InProgram.ConstantMeshesPermanent.Num();
+						for (const TPair<uint32, TSharedPtr<const mu::FMesh>>& Entry : InProgram.ConstantMeshesStreamed)
+						{
+							if (Entry.Key == ResourceIndex.Index)
+							{
+								bResourceLocated = IndexOnConstantsArray == DebuggerIndexOffset;
+								break;
+							}
+							++DebuggerIndexOffset;
+						}
+					}
 				}
 				break;
 			}
