[Mutable] Reduced the frequency of TickWarnings.

diff --git a/Engine/Plugins/Mutable/Source/CustomizableObject/Private/MuCO/CustomizableObjectSystem.cpp b/Engine/Plugins/Mutable/Source/CustomizableObject/Private/MuCO/CustomizableObjectSystem.cpp
--- a/Engine/Plugins/Mutable/Source/CustomizableObject/Private/MuCO/CustomizableObjectSystem.cpp
+++ b/Engine/Plugins/Mutable/Source/CustomizableObject/Private/MuCO/CustomizableObjectSystem.cpp
@@ -4741,7 +4741,7 @@ void UCustomizableObjectSystemPrivate::OnMutableEnabledChanged(IConsoleVariable*
 
 		if (!SystemPrivate->TickWarningsDelegateHandle.IsValid())
 		{
-			SystemPrivate->TickWarningsDelegateHandle = FTSTicker::GetCoreTicker().AddTicker(FTickerDelegate::CreateStatic(&TickWarnings), 0.f);
+			SystemPrivate->TickWarningsDelegateHandle = FTSTicker::GetCoreTicker().AddTicker(FTickerDelegate::CreateStatic(&TickWarnings), OnScreenWarningsTickerTime);
 		}
 #endif // !UE_SERVER
 	}
