[Mutable] Removed unnecessary "#undef" and fixed a LOCTEXT conflict.

#rnx

diff --git a/Engine/Plugins/Mutable/Source/CustomizableObjectEditor/Private/MuCOE/CustomizableObjectGraphEditorToolkit.cpp b/Engine/Plugins/Mutable/Source/CustomizableObjectEditor/Private/MuCOE/CustomizableObjectGraphEditorToolkit.cpp
--- a/Engine/Plugins/Mutable/Source/CustomizableObjectEditor/Private/MuCOE/CustomizableObjectGraphEditorToolkit.cpp
+++ b/Engine/Plugins/Mutable/Source/CustomizableObjectEditor/Private/MuCOE/CustomizableObjectGraphEditorToolkit.cpp
@@ -555,7 +555,7 @@ void FCustomizableObjectGraphEditorToolkit::CreateGraphEditorWidget(UEdGraph* In
 
 	// Create the appearance info
 	FGraphAppearanceInfo AppearanceInfo;
-	AppearanceInfo.CornerText = CustomizableObjectGraph->IsMacro() ? LOCTEXT("ApperanceCornerText", "MUTABLE MACRO") : LOCTEXT("ApperanceCornerText", "MUTABLE");
+	AppearanceInfo.CornerText = CustomizableObjectGraph->IsMacro() ? LOCTEXT("ApperanceCornerMacroText", "MUTABLE MACRO") : LOCTEXT("ApperanceCornerText", "MUTABLE");
 
 	// Add toolkit common events
 	InEvents.OnTextCommitted = FOnNodeTextCommitted::CreateSP(this, &FCustomizableObjectGraphEditorToolkit::OnNodeTitleCommitted);
diff --git a/Engine/Plugins/Mutable/Source/CustomizableObjectEditor/Private/MuCOE/Nodes/CustomizableObjectNodeMacroInstanceDetails.h b/Engine/Plugins/Mutable/Source/CustomizableObjectEditor/Private/MuCOE/Nodes/CustomizableObjectNodeMacroInstanceDetails.h
--- a/Engine/Plugins/Mutable/Source/CustomizableObjectEditor/Private/MuCOE/Nodes/CustomizableObjectNodeMacroInstanceDetails.h
+++ b/Engine/Plugins/Mutable/Source/CustomizableObjectEditor/Private/MuCOE/Nodes/CustomizableObjectNodeMacroInstanceDetails.h
@@ -44,6 +44,3 @@ private:
 	TSharedPtr<FMacroSelectorItem> SelectedSource;
 
 };
-
-
-#undef LOCTEXT_NAMESPACE
\ No newline at end of file
