[Mutable] Added non-dynamic callback to FCompileParams.

diff --git a/Engine/Plugins/Mutable/Source/CustomizableObject/Internal/MuCO/CustomizableObjectPrivate.h b/Engine/Plugins/Mutable/Source/CustomizableObject/Internal/MuCO/CustomizableObjectPrivate.h
--- a/Engine/Plugins/Mutable/Source/CustomizableObject/Internal/MuCO/CustomizableObjectPrivate.h
+++ b/Engine/Plugins/Mutable/Source/CustomizableObject/Internal/MuCO/CustomizableObjectPrivate.h
@@ -1466,6 +1466,7 @@ public:
 	TArray<FText> Errors;
 
 	FCompileDelegate Callback;
+	FCompileNativeDelegate CallbackNative;
 
 	bool bAsync = true;
 	bool bSkipIfCompiled = false;
diff --git a/Engine/Plugins/Mutable/Source/CustomizableObject/Private/MuCO/CustomizableObject.cpp b/Engine/Plugins/Mutable/Source/CustomizableObject/Private/MuCO/CustomizableObject.cpp
--- a/Engine/Plugins/Mutable/Source/CustomizableObject/Private/MuCO/CustomizableObject.cpp
+++ b/Engine/Plugins/Mutable/Source/CustomizableObject/Private/MuCO/CustomizableObject.cpp
@@ -1940,6 +1940,7 @@ void UCustomizableObject::Compile(const FCompileParams& Params)
 		CallbackParams.bCompiled = IsCompiled();
 
 		Params.Callback.ExecuteIfBound(CallbackParams);
+		Params.CallbackNative.ExecuteIfBound(CallbackParams);
 	}
 }
 
diff --git a/Engine/Plugins/Mutable/Source/CustomizableObject/Public/MuCO/CustomizableObject.h b/Engine/Plugins/Mutable/Source/CustomizableObject/Public/MuCO/CustomizableObject.h
--- a/Engine/Plugins/Mutable/Source/CustomizableObject/Public/MuCO/CustomizableObject.h
+++ b/Engine/Plugins/Mutable/Source/CustomizableObject/Public/MuCO/CustomizableObject.h
@@ -161,6 +161,8 @@ struct FCompileCallbackParams
 
 DECLARE_DYNAMIC_DELEGATE_OneParam(FCompileDelegate, const FCompileCallbackParams&, Params);
 
+DECLARE_DELEGATE_OneParam(FCompileNativeDelegate, const FCompileCallbackParams&);
+
 
 USTRUCT(BlueprintType)
 struct FCompileParams
@@ -196,6 +198,9 @@ struct FCompileParams
 	/** Called once the compilation request has finished (successfully or not). */
 	UPROPERTY(BlueprintReadWrite, Category = Compile)
 	FCompileDelegate Callback;
+
+	/** See Callback. */
+	FCompileNativeDelegate CallbackNative;
 };
 
 
diff --git a/Engine/Plugins/Mutable/Source/CustomizableObjectEditor/Private/MuCOE/CustomizableObjectCompiler.cpp b/Engine/Plugins/Mutable/Source/CustomizableObjectEditor/Private/MuCOE/CustomizableObjectCompiler.cpp
--- a/Engine/Plugins/Mutable/Source/CustomizableObjectEditor/Private/MuCOE/CustomizableObjectCompiler.cpp
+++ b/Engine/Plugins/Mutable/Source/CustomizableObjectEditor/Private/MuCOE/CustomizableObjectCompiler.cpp
@@ -143,6 +143,7 @@ void FCustomizableObjectCompiler::Compile(const TSharedRef<FCompilationRequest>&
 		Params.bRequestFailed = true;
     	
 		InCompileRequest->Callback.ExecuteIfBound(Params);
+		InCompileRequest->CallbackNative.ExecuteIfBound(Params);
 		return;
 	}
 	
@@ -153,6 +154,7 @@ void FCustomizableObjectCompiler::Compile(const TSharedRef<FCompilationRequest>&
 		Params.bSkipped = true;
     	
 		InCompileRequest->Callback.ExecuteIfBound(Params);
+		InCompileRequest->CallbackNative.ExecuteIfBound(Params);
 		return;
 	}
 
@@ -169,6 +171,7 @@ void FCustomizableObjectCompiler::Compile(const TSharedRef<FCompilationRequest>&
 			Params.bCompiled = Object->IsCompiled();
 
 			InCompileRequest->Callback.ExecuteIfBound(Params);
+			InCompileRequest->CallbackNative.ExecuteIfBound(Params);
 			return;
 		}
 	}
@@ -336,6 +339,7 @@ void FCustomizableObjectCompiler::EnqueueCompileRequest(const TSharedRef<FCompil
 		Params.bRequestFailed = true;
 
 		CompileRequest->Callback.ExecuteIfBound(Params);
+		CompileRequest->CallbackNative.ExecuteIfBound(Params);
 		return;
 	}
 	
@@ -346,6 +350,7 @@ void FCustomizableObjectCompiler::EnqueueCompileRequest(const TSharedRef<FCompil
 		Params.bCompiled = CompileRequest->GetCustomizableObject()->IsCompiled();
 
 		CompileRequest->Callback.ExecuteIfBound(Params);
+		CompileRequest->CallbackNative.ExecuteIfBound(Params);
 		return;
 	}
 		
@@ -363,6 +368,7 @@ void FCustomizableObjectCompiler::EnqueueCompileRequest(const TSharedRef<FCompil
 			Params.bCompiled = CompileRequest->GetCustomizableObject()->IsCompiled();
 
 			CompileRequest->Callback.ExecuteIfBound(Params);
+			CompileRequest->CallbackNative.ExecuteIfBound(Params);
 			return;
 		}
 
@@ -1102,7 +1108,8 @@ void FCustomizableObjectCompiler::CompleteRequest(ECompilationStatePrivate State
 	}
 		
 	CurrentRequest->Callback.ExecuteIfBound(Params);
-	
+	CurrentRequest->CallbackNative.ExecuteIfBound(Params);
+
 	// Request completed, reset pointers and state
 	CurrentObject = nullptr;
 	CurrentRequest.Reset();
@@ -1140,6 +1147,7 @@ bool FCustomizableObjectCompiler::TryPopCompileRequest()
 			Params.bRequestFailed = true;
     	
 			CompileRequest->Callback.ExecuteIfBound(Params);
+			CompileRequest->CallbackNative.ExecuteIfBound(Params);
 			continue;
 		}
 		
diff --git a/Engine/Plugins/Mutable/Source/CustomizableObjectEditor/Private/MuCOE/CustomizableObjectEditorModule.cpp b/Engine/Plugins/Mutable/Source/CustomizableObjectEditor/Private/MuCOE/CustomizableObjectEditorModule.cpp
--- a/Engine/Plugins/Mutable/Source/CustomizableObjectEditor/Private/MuCOE/CustomizableObjectEditorModule.cpp
+++ b/Engine/Plugins/Mutable/Source/CustomizableObjectEditor/Private/MuCOE/CustomizableObjectEditorModule.cpp
@@ -933,6 +933,7 @@ void FCustomizableObjectEditorModule::CompileCustomizableObject(UCustomizableObj
 		CallbackParams.bCompiled = Object.IsCompiled();
 		
 		Params.Callback.ExecuteIfBound(CallbackParams);
+		Params.CallbackNative.ExecuteIfBound(CallbackParams);
 		return;
 	}
 
@@ -945,6 +946,7 @@ void FCustomizableObjectEditorModule::CompileCustomizableObject(UCustomizableObj
 	CompileRequest->Options.OptimizationLevel = ConvertOptimizationLevel(Params.OptimizationLevel);
 	CompileRequest->Options.bGatherReferences = Params.bGatherReferneces;
 	CompileRequest->Callback = Params.Callback;
+	CompileRequest->CallbackNative = Params.CallbackNative;
 	
 	if (Params.CompileOnlySelectedInstance)
 	{
