[Mutable] Move FCompileRequest outside of CO module.

diff --git a/Engine/Plugins/Mutable/Source/CustomizableObject/Internal/MuCO/CustomizableObjectPrivate.h b/Engine/Plugins/Mutable/Source/CustomizableObject/Internal/MuCO/CustomizableObjectPrivate.h
--- a/Engine/Plugins/Mutable/Source/CustomizableObject/Internal/MuCO/CustomizableObjectPrivate.h
+++ b/Engine/Plugins/Mutable/Source/CustomizableObject/Internal/MuCO/CustomizableObjectPrivate.h
@@ -1430,52 +1430,6 @@ public:
 };
 
 
-#if WITH_EDITOR
-struct CUSTOMIZABLEOBJECT_API FCompilationRequest
-{
-	FCompilationRequest(UCustomizableObject& CustomizableObject);
-
-	UCustomizableObject* GetCustomizableObject();
-	
-	void SetDerivedDataCachePolicy(UE::DerivedData::ECachePolicy InCachePolicy);
-	UE::DerivedData::ECachePolicy GetDerivedDataCachePolicy() const;
-
-	void BuildDerivedDataCacheKey();
-	UE::DerivedData::FCacheKey GetDerivedDataCacheKey() const;
-
-	void SetCompilationState(ECompilationStatePrivate InState, ECompilationResultPrivate InResult);
-
-	ECompilationStatePrivate GetCompilationState() const;
-	ECompilationResultPrivate GetCompilationResult() const;
-	
-	bool operator==(const FCompilationRequest& Other) const;
-
-private:
-	TWeakObjectPtr<UCustomizableObject> CustomizableObject;
-
-	ECompilationStatePrivate State = ECompilationStatePrivate::None;
-	ECompilationResultPrivate Result = ECompilationResultPrivate::Unknown;
-	
-	UE::DerivedData::ECachePolicy DDCPolicy = UE::DerivedData::ECachePolicy::None;
-	UE::DerivedData::FCacheKey DDCKey;
-
-public:
-	FCompilationOptions Options;
-
-	TArray<FText> Warnings;
-	TArray<FText> Errors;
-
-	FCompileDelegate Callback;
-	FCompileNativeDelegate CallbackNative;
-
-	bool bAsync = true;
-	bool bSkipIfCompiled = false;
-	bool bSkipIfOutOfDate = false;
-	bool bSilentCompilation = true;
-};
-#endif
-
-
 UCLASS()
 class CUSTOMIZABLEOBJECT_API UCustomizableObjectPrivate : public UObject
 {
@@ -1582,9 +1536,6 @@ public:
 	FGuid GetVersionId() const;
 
 	void SaveEmbeddedData(FArchive& Ar);
-
-	// Compile the object for a specific platform - Compile for Cook Customizable Object
-	void CompileForTargetPlatform(UCustomizableObject& CustomizableObject, const ITargetPlatform& TargetPlatform);
 	
 	// Add a profile that stores the values of the parameters used by the CustomInstance.
 	FReply AddNewParameterProfile(FString Name, class UCustomizableObjectInstance& CustomInstance);
@@ -1694,9 +1645,6 @@ public:
 
 	/** Unique Identifier - Deterministic. Used to locate Model and Streamable data on disk. Should not be modified. */
 	FGuid Identifier;
-
-	/** Cook requests. */
-	TArray<TSharedRef<FCompilationRequest>> CompileRequests;
 	
 	ECompilationResultPrivate CompilationResult = ECompilationResultPrivate::Unknown;
 	
diff --git a/Engine/Plugins/Mutable/Source/CustomizableObject/Internal/MuCO/ICustomizableObjectEditorModule.h b/Engine/Plugins/Mutable/Source/CustomizableObject/Internal/MuCO/ICustomizableObjectEditorModule.h
--- a/Engine/Plugins/Mutable/Source/CustomizableObject/Internal/MuCO/ICustomizableObjectEditorModule.h
+++ b/Engine/Plugins/Mutable/Source/CustomizableObject/Internal/MuCO/ICustomizableObjectEditorModule.h
@@ -29,6 +29,10 @@ extern const FName CustomizableObjectPopulationClassEditorAppIdentifier;
 extern const FName CustomizableObjectDebuggerAppIdentifier;
 extern const FName CustomizableObjectMacroLibraryEditorAppIdentifier;
 
+
+#define MODULE_NAME_COE "CustomizableObjectEditor"
+
+
 /**
  * Customizable object editor module interface
  */
@@ -43,13 +47,13 @@ public:
 			return nullptr;
 		}
 		
-		return FModuleManager::LoadModulePtr<ICustomizableObjectEditorModule>("CustomizableObjectEditor");
+		return FModuleManager::LoadModulePtr<ICustomizableObjectEditorModule>(MODULE_NAME_COE);
 	}
 	
 	static ICustomizableObjectEditorModule& GetChecked()
 	{
 		check(!IsRunningGame())		// This module is editor-only. DO NOT try to access it during gameplay
-		return FModuleManager::LoadModuleChecked<ICustomizableObjectEditorModule>("CustomizableObjectEditor");
+		return FModuleManager::LoadModuleChecked<ICustomizableObjectEditorModule>(MODULE_NAME_COE);
 	}
 
 	virtual TSharedPtr<FExtensibilityManager> GetCustomizableObjectEditorToolBarExtensibilityManager() { return nullptr; }
@@ -91,13 +95,9 @@ public:
 	 */
 	virtual void BakeCustomizableObjectInstance(UCustomizableObjectInstance* InTargetInstance, const FBakingConfiguration& InBakingConfig) = 0;
 
-	// TODO GMT Remove from this interface (public header). Should only be a FCustomizableObjectEditorModule function. Also move FCompilationRequest to COE module.
-	/** Request for a given customizable object to be compiled. Async compile requests will be queued and processed sequentially.
-	 * @param InCompilationRequest - Request to compile an object.
-	 * @param bForceRequest - Queue request even if already in the pending list. */
-	virtual void CompileCustomizableObject(const TSharedRef<FCompilationRequest>& InCompilationRequest, bool bForceRequest = false) = 0;
-
-	virtual void CompileCustomizableObject(UCustomizableObject& Object, const FCompileParams& Params) = 0;
+	/** Compile the given Customizable Object.
+	 * If calling inside the Customizable Object Editor module, consider using ICustomizableObjectEditorModulePrivate::EnqueueCompileRequest. */
+	virtual void CompileCustomizableObject(UCustomizableObject& Object, const FCompileParams* Params, bool bSilentCompilation) = 0;
 
 	virtual int32 Tick(bool bBlocking) = 0;
 
@@ -117,5 +117,9 @@ public:
 	virtual void PostBackwardsCompatibleFixup(UEdGraph& Graph) = 0;
 
 	virtual bool IsCompiling(const UCustomizableObject& Object) const = 0;
+
+	virtual void BeginCacheForCookedPlatformData(UCustomizableObject& Object, const ITargetPlatform* TargetPlatform) = 0;
+	
+	virtual bool IsCachedCookedPlatformDataLoaded(UCustomizableObject& Object, const ITargetPlatform* TargetPlatform) = 0;
 };
  
diff --git a/Engine/Plugins/Mutable/Source/CustomizableObject/Private/MuCO/CustomizableObject.cpp b/Engine/Plugins/Mutable/Source/CustomizableObject/Private/MuCO/CustomizableObject.cpp
--- a/Engine/Plugins/Mutable/Source/CustomizableObject/Private/MuCO/CustomizableObject.cpp
+++ b/Engine/Plugins/Mutable/Source/CustomizableObject/Private/MuCO/CustomizableObject.cpp
@@ -454,40 +454,15 @@ void UCustomizableObject::PostRename(UObject * OldOuter, const FName OldName)
 
 void UCustomizableObject::BeginCacheForCookedPlatformData(const ITargetPlatform* TargetPlatform)
 {
-	if (!TargetPlatform)
-	{
-		return;
-	}
-
-	const TSharedRef<FCompilationRequest>* CompileRequest = GetPrivate()->CompileRequests.FindByPredicate(
-		[&TargetPlatform](const TSharedPtr<FCompilationRequest>& Request) { return Request->Options.TargetPlatform == TargetPlatform; });
-	
-	if (CompileRequest)
-	{
-		return;
-	}
-
-	// Compile and save in the CachedPlatformsData map
-	GetPrivate()->CompileForTargetPlatform(*this, *TargetPlatform);
+	ICustomizableObjectEditorModule& Module = ICustomizableObjectEditorModule::GetChecked();
+	Module.BeginCacheForCookedPlatformData(*this, TargetPlatform);
 }
 
 
 bool UCustomizableObject::IsCachedCookedPlatformDataLoaded(const ITargetPlatform* TargetPlatform) 
 {
-	if (!TargetPlatform)
-	{
-		return true;
-	}
-
-	const TSharedRef<FCompilationRequest>* CompileRequest = GetPrivate()->CompileRequests.FindByPredicate(
-		[&TargetPlatform](const TSharedRef<FCompilationRequest>& Request) { return Request->Options.TargetPlatform == TargetPlatform; });
-	
-	if (CompileRequest)
-	{
-		return CompileRequest->Get().GetCompilationState() == ECompilationStatePrivate::Completed;
-	}
-
-	return true;
+	ICustomizableObjectEditorModule& Module = ICustomizableObjectEditorModule::GetChecked();
+	return Module.IsCachedCookedPlatformDataLoaded(*this, TargetPlatform);
 }
 
 
@@ -710,43 +685,6 @@ void UCustomizableObjectPrivate::LoadCompiledDataFromDisk()
 }
 
 
-void UCustomizableObjectPrivate::CompileForTargetPlatform(UCustomizableObject& CustomizableObject, const ITargetPlatform& TargetPlatform)
-{
-	ICustomizableObjectEditorModule* EditorModule = ICustomizableObjectEditorModule::Get();
-	if (!EditorModule || !EditorModule->IsRootObject(CustomizableObject))
-	{
-		SetIsChildObject(true);
-		return;
-	}
-
-	const bool bAsync = CVarMutableAsyncCook.GetValueOnAnyThread();
-
-	TSharedRef<FCompilationRequest> CompileRequest = MakeShared<FCompilationRequest>(CustomizableObject);
-	CompileRequest->bAsync = bAsync;
-	CompileRequest->Options.OptimizationLevel = UE_MUTABLE_MAX_OPTIMIZATION; // Force max optimization when packaging.
-	CompileRequest->Options.TextureCompression = ECustomizableObjectTextureCompression::HighQuality;
-	CompileRequest->Options.bIsCooking = true;
-	CompileRequest->Options.bUseBulkData = CVarMutableUseBulkData.GetValueOnAnyThread();
-	CompileRequest->Options.TargetPlatform = &TargetPlatform;
-
-	const int32 DDCUsage = CVarMutableDerivedDataCacheUsage.GetValueOnAnyThread();
-	UE::DerivedData::ECachePolicy DefaultCachePolicy = UE::DerivedData::ECachePolicy::None;
-	if (DDCUsage == 1)
-	{
-		DefaultCachePolicy = UE::DerivedData::ECachePolicy::Local;
-	}
-	else if (DDCUsage == 2)
-	{
-		DefaultCachePolicy = UE::DerivedData::ECachePolicy::Default;
-	}
-	CompileRequest->SetDerivedDataCachePolicy(DefaultCachePolicy);
-
-	CompileRequests.Add(CompileRequest);
-
-	EditorModule->CompileCustomizableObject(CompileRequest, true);
-}
-
-
 bool UCustomizableObject::ConditionalAutoCompile()
 {
 	check(IsInGameThread());
@@ -797,10 +735,7 @@ bool UCustomizableObject::ConditionalAutoCompile()
 		}
 		else
 		{
-			TSharedRef<FCompilationRequest> CompileRequest = MakeShared<FCompilationRequest>(*this);
-			CompileRequest->bSilentCompilation = true;
-			
-			EditorModule->CompileCustomizableObject(CompileRequest);
+			EditorModule->CompileCustomizableObject(*this, nullptr, true);
 		}
 	}
 
@@ -1931,7 +1866,7 @@ void UCustomizableObject::Compile(const FCompileParams& Params)
 {
 	if (ICustomizableObjectEditorModule* EditorModule = ICustomizableObjectEditorModule::Get())
 	{
-		EditorModule->CompileCustomizableObject(*this, Params);
+		EditorModule->CompileCustomizableObject(*this, &Params, false);
 	}
 	else
 	{
@@ -3726,72 +3661,6 @@ FArchive& operator<<(FArchive& Ar, FMutableRefSkeletalMeshData& Data)
 }
 #endif 
 
-#if WITH_EDITOR
-FCompilationRequest::FCompilationRequest(UCustomizableObject& InCustomizableObject)
-{
-	CustomizableObject = &InCustomizableObject;
-	Options = InCustomizableObject.GetPrivate()->GetCompileOptions();
-	DDCPolicy = UE::DerivedData::ECachePolicy::None;
-}
-
-
-UCustomizableObject* FCompilationRequest::GetCustomizableObject()
-{
-	return CustomizableObject.Get();
-}
-
-
-void FCompilationRequest::SetDerivedDataCachePolicy(UE::DerivedData::ECachePolicy InCachePolicy)
-{
-	DDCPolicy = InCachePolicy;
-	Options.bQueryCompiledDatafromDDC = EnumHasAnyFlags(InCachePolicy, UE::DerivedData::ECachePolicy::Query);
-	Options.bStoreCompiledDataInDDC = EnumHasAnyFlags(InCachePolicy, UE::DerivedData::ECachePolicy::Store);
-}
-
-
-UE::DerivedData::ECachePolicy FCompilationRequest::GetDerivedDataCachePolicy() const
-{
-	return DDCPolicy;
-}
-
-void FCompilationRequest::BuildDerivedDataCacheKey()
-{
-	if (UCustomizableObject* Object = CustomizableObject.Get())
-	{
-		DDCKey = Object->GetPrivate()->GetDerivedDataCacheKeyForOptions(Options);
-	}
-}
-
-
-UE::DerivedData::FCacheKey FCompilationRequest::GetDerivedDataCacheKey() const
-{
-	return DDCKey;
-}
-
-
-void FCompilationRequest::SetCompilationState(ECompilationStatePrivate InState, ECompilationResultPrivate InResult)
-{
-	State = InState;
-	Result = InResult;
-}
-
-
-ECompilationStatePrivate FCompilationRequest::GetCompilationState() const
-{
-	return State;
-}
-
-
-ECompilationResultPrivate FCompilationRequest::GetCompilationResult() const
-{
-	return Result;
-}
-
-
-bool FCompilationRequest::operator==(const FCompilationRequest& Other) const
-{
-	return CustomizableObject == Other.CustomizableObject && Options.TargetPlatform == Other.Options.TargetPlatform;
-}
-#endif
 
 #undef LOCTEXT_NAMESPACE
+
diff --git a/Engine/Plugins/Mutable/Source/CustomizableObject/Private/MuCO/CustomizableObjectInstance.cpp b/Engine/Plugins/Mutable/Source/CustomizableObject/Private/MuCO/CustomizableObjectInstance.cpp
--- a/Engine/Plugins/Mutable/Source/CustomizableObject/Private/MuCO/CustomizableObjectInstance.cpp
+++ b/Engine/Plugins/Mutable/Source/CustomizableObject/Private/MuCO/CustomizableObjectInstance.cpp
@@ -1049,6 +1049,7 @@ void UCustomizableObjectInstance::UpdateSkeletalMeshAsyncResult(FInstanceUpdateN
 	SystemPrivate->EnqueueUpdateSkeletalMesh(Context);
 }
 
+
 void UCustomizableInstancePrivate::TickUpdateCloseCustomizableObjects(UCustomizableObjectInstance& Public, FMutableInstanceUpdateMap& InOutRequestedUpdates)
 {
 	UCustomizableObject* Object = Public.GetCustomizableObject();
@@ -1063,9 +1064,7 @@ void UCustomizableInstancePrivate::TickUpdateCloseCustomizableObjects(UCustomiza
 	{
 		if (ICustomizableObjectEditorModule* EditorModule = ICustomizableObjectEditorModule::Get())
 		{
-			TSharedRef<FCompilationRequest> CompileRequest = MakeShared<FCompilationRequest>(*Object);
-			CompileRequest->bSilentCompilation = true;
-			EditorModule->CompileCustomizableObject(CompileRequest);
+			EditorModule->CompileCustomizableObject(*Object, nullptr, true);
 		}
 	}
 #endif
diff --git a/Engine/Plugins/Mutable/Source/CustomizableObject/Public/MuCO/CustomizableObject.h b/Engine/Plugins/Mutable/Source/CustomizableObject/Public/MuCO/CustomizableObject.h
--- a/Engine/Plugins/Mutable/Source/CustomizableObject/Public/MuCO/CustomizableObject.h
+++ b/Engine/Plugins/Mutable/Source/CustomizableObject/Public/MuCO/CustomizableObject.h
@@ -35,6 +35,13 @@ CUSTOMIZABLEOBJECT_API DECLARE_LOG_CATEGORY_EXTERN(LogMutable, Log, All);
 DECLARE_MULTICAST_DELEGATE(FPostCompileDelegate)
 
 
+CUSTOMIZABLEOBJECT_API extern TAutoConsoleVariable<bool> CVarMutableUseBulkData;
+
+CUSTOMIZABLEOBJECT_API extern TAutoConsoleVariable<int32> CVarMutableDerivedDataCacheUsage;
+
+CUSTOMIZABLEOBJECT_API extern TAutoConsoleVariable<bool> CVarMutableAsyncCook;
+
+
 USTRUCT()
 struct FFParameterOptionsTags
 {
diff --git a/Engine/Plugins/Mutable/Source/CustomizableObjectEditor/Internal/ICustomizableObjectEditorModulePrivate.h b/Engine/Plugins/Mutable/Source/CustomizableObjectEditor/Internal/ICustomizableObjectEditorModulePrivate.h
new file mode 100644
--- /dev/null
+++ b/Engine/Plugins/Mutable/Source/CustomizableObjectEditor/Internal/ICustomizableObjectEditorModulePrivate.h
@@ -0,0 +1,33 @@
+﻿// Copyright Epic Games, Inc. All Rights Reserved.
+
+#pragma once
+
+#include "Modules/ModuleInterface.h"
+#include "MuCO/ICustomizableObjectEditorModule.h"
+
+
+/** Interface only accessible from the Customizable Object Editor module.
+ *
+ * TODO Move to private. Currently can not be moved due to MutableValidation module. */
+class ICustomizableObjectEditorModulePrivate : public ICustomizableObjectEditorModule
+{
+public:
+	static ICustomizableObjectEditorModulePrivate* Get()
+	{
+		// Prevent access to this module if the game is being played (in Standalone mode for example)
+		if (IsRunningGame())
+		{
+			return nullptr;
+		}
+		
+		return FModuleManager::LoadModulePtr<ICustomizableObjectEditorModulePrivate>(MODULE_NAME_COE);
+	}
+	
+	static ICustomizableObjectEditorModulePrivate& GetChecked()
+	{
+		check(!IsRunningGame()) // This module is editor-only. DO NOT try to access it during gameplay
+		return FModuleManager::LoadModuleChecked<ICustomizableObjectEditorModulePrivate>(MODULE_NAME_COE);
+	}
+	
+	virtual void EnqueueCompileRequest(const TSharedRef<FCompilationRequest>& InCompilationRequest, bool bForceRequest = false) = 0;
+};
\ No newline at end of file
diff --git a/Engine/Plugins/Mutable/Source/CustomizableObjectEditor/Internal/MuCOE/CompileRequest.h b/Engine/Plugins/Mutable/Source/CustomizableObjectEditor/Internal/MuCOE/CompileRequest.h
new file mode 100644
--- /dev/null
+++ b/Engine/Plugins/Mutable/Source/CustomizableObjectEditor/Internal/MuCOE/CompileRequest.h
@@ -0,0 +1,52 @@
+﻿// Copyright Epic Games, Inc. All Rights Reserved.
+
+#pragma once
+#include "MuCO/CustomizableObject.h"
+#include "MuCO/CustomizableObjectCompilerTypes.h"
+
+class UCustomizableObject;
+
+
+struct CUSTOMIZABLEOBJECTEDITOR_API FCompilationRequest
+{
+	FCompilationRequest(UCustomizableObject& CustomizableObject);
+
+	UCustomizableObject* GetCustomizableObject();
+	
+	void SetDerivedDataCachePolicy(UE::DerivedData::ECachePolicy InCachePolicy);
+	UE::DerivedData::ECachePolicy GetDerivedDataCachePolicy() const;
+
+	void BuildDerivedDataCacheKey();
+	UE::DerivedData::FCacheKey GetDerivedDataCacheKey() const;
+
+	void SetCompilationState(ECompilationStatePrivate InState, ECompilationResultPrivate InResult);
+
+	ECompilationStatePrivate GetCompilationState() const;
+	ECompilationResultPrivate GetCompilationResult() const;
+	
+	bool operator==(const FCompilationRequest& Other) const;
+
+private:
+	TWeakObjectPtr<UCustomizableObject> CustomizableObject;
+
+	ECompilationStatePrivate State = ECompilationStatePrivate::None;
+	ECompilationResultPrivate Result = ECompilationResultPrivate::Unknown;
+	
+	UE::DerivedData::ECachePolicy DDCPolicy = UE::DerivedData::ECachePolicy::None;
+	UE::DerivedData::FCacheKey DDCKey;
+
+public:
+	FCompilationOptions Options;
+
+	TArray<FText> Warnings;
+	TArray<FText> Errors;
+
+	FCompileDelegate Callback;
+	FCompileNativeDelegate CallbackNative;
+
+	bool bAsync = true;
+	bool bSkipIfCompiled = false;
+	bool bSkipIfOutOfDate = false;
+	bool bSilentCompilation = true;
+};
+
diff --git a/Engine/Plugins/Mutable/Source/CustomizableObjectEditor/Private/MuCOE/AssetDefinition_CustomizableObject.cpp b/Engine/Plugins/Mutable/Source/CustomizableObjectEditor/Private/MuCOE/AssetDefinition_CustomizableObject.cpp
--- a/Engine/Plugins/Mutable/Source/CustomizableObjectEditor/Private/MuCOE/AssetDefinition_CustomizableObject.cpp
+++ b/Engine/Plugins/Mutable/Source/CustomizableObjectEditor/Private/MuCOE/AssetDefinition_CustomizableObject.cpp
@@ -10,6 +10,7 @@
 #include "CustomizableObjectEditor.h"
 #include "Editor.h"
 #include "IContentBrowserSingleton.h"
+#include "ICustomizableObjectEditorModulePrivate.h"
 #include "Misc/MessageDialog.h"
 #include "MuCO/CustomizableObject.h"
 #include "MuCO/CustomizableObjectInstance.h"
@@ -18,6 +19,7 @@
 #include "Subsystems/AssetEditorSubsystem.h"
 #include "ToolMenus.h"
 #include "MuCO/LoadUtils.h"
+#include "MuCOE/CompileRequest.h"
 
 
 #define LOCTEXT_NAMESPACE "CustomizableObjectEditor"
@@ -255,7 +257,7 @@ namespace MenuExtension_CustomizableObject
 		
 		for (const TSharedRef<FCompilationRequest>& CompilationRequest : CompilationRequests)
 		{
-			ICustomizableObjectEditorModule::GetChecked().CompileCustomizableObject(CompilationRequest);
+			ICustomizableObjectEditorModulePrivate::GetChecked().EnqueueCompileRequest(CompilationRequest);
 		}
 	}
 
diff --git a/Engine/Plugins/Mutable/Source/CustomizableObjectEditor/Private/MuCOE/CompileRequest.cpp b/Engine/Plugins/Mutable/Source/CustomizableObjectEditor/Private/MuCOE/CompileRequest.cpp
new file mode 100644
--- /dev/null
+++ b/Engine/Plugins/Mutable/Source/CustomizableObjectEditor/Private/MuCOE/CompileRequest.cpp
@@ -0,0 +1,77 @@
+﻿// Copyright Epic Games, Inc. All Rights Reserved.
+
+
+#include "MuCOE/CompileRequest.h"
+
+#include "DerivedDataCachePolicy.h"
+#include "MuCO/CustomizableObject.h"
+#include "MuCO/CustomizableObjectCompilerTypes.h"
+#include "MuCO/CustomizableObjectPrivate.h"
+
+
+FCompilationRequest::FCompilationRequest(UCustomizableObject& InCustomizableObject)
+{
+	CustomizableObject = &InCustomizableObject;
+	Options = InCustomizableObject.GetPrivate()->GetCompileOptions();
+	DDCPolicy = UE::DerivedData::ECachePolicy::None;
+}
+
+
+UCustomizableObject* FCompilationRequest::GetCustomizableObject()
+{
+	return CustomizableObject.Get();
+}
+
+
+void FCompilationRequest::SetDerivedDataCachePolicy(UE::DerivedData::ECachePolicy InCachePolicy)
+{
+	DDCPolicy = InCachePolicy;
+	Options.bQueryCompiledDatafromDDC = EnumHasAnyFlags(InCachePolicy, UE::DerivedData::ECachePolicy::Query);
+	Options.bStoreCompiledDataInDDC = EnumHasAnyFlags(InCachePolicy, UE::DerivedData::ECachePolicy::Store);
+}
+
+
+UE::DerivedData::ECachePolicy FCompilationRequest::GetDerivedDataCachePolicy() const
+{
+	return DDCPolicy;
+}
+
+void FCompilationRequest::BuildDerivedDataCacheKey()
+{
+	if (UCustomizableObject* Object = CustomizableObject.Get())
+	{
+		DDCKey = Object->GetPrivate()->GetDerivedDataCacheKeyForOptions(Options);
+	}
+}
+
+
+UE::DerivedData::FCacheKey FCompilationRequest::GetDerivedDataCacheKey() const
+{
+	return DDCKey;
+}
+
+
+void FCompilationRequest::SetCompilationState(ECompilationStatePrivate InState, ECompilationResultPrivate InResult)
+{
+	State = InState;
+	Result = InResult;
+}
+
+
+ECompilationStatePrivate FCompilationRequest::GetCompilationState() const
+{
+	return State;
+}
+
+
+ECompilationResultPrivate FCompilationRequest::GetCompilationResult() const
+{
+	return Result;
+}
+
+
+bool FCompilationRequest::operator==(const FCompilationRequest& Other) const
+{
+	return CustomizableObject == Other.CustomizableObject && Options.TargetPlatform == Other.Options.TargetPlatform;
+}
+
diff --git a/Engine/Plugins/Mutable/Source/CustomizableObjectEditor/Private/MuCOE/CustomizableObjectCompileRunnable.cpp b/Engine/Plugins/Mutable/Source/CustomizableObjectEditor/Private/MuCOE/CustomizableObjectCompileRunnable.cpp
--- a/Engine/Plugins/Mutable/Source/CustomizableObjectEditor/Private/MuCOE/CustomizableObjectCompileRunnable.cpp
+++ b/Engine/Plugins/Mutable/Source/CustomizableObjectEditor/Private/MuCOE/CustomizableObjectCompileRunnable.cpp
@@ -9,6 +9,7 @@
 #include "MuCO/CustomizableObject.h"
 #include "MuCO/CustomizableObjectSystem.h"
 #include "MuCO/CustomizableObjectSystemPrivate.h"
+#include "MuCOE/CompileRequest.h"
 #include "MuR/Model.h"
 #include "MuT/Compiler.h"
 #include "MuT/ErrorLog.h"
diff --git a/Engine/Plugins/Mutable/Source/CustomizableObjectEditor/Private/MuCOE/CustomizableObjectCompileRunnable.h b/Engine/Plugins/Mutable/Source/CustomizableObjectEditor/Private/MuCOE/CustomizableObjectCompileRunnable.h
--- a/Engine/Plugins/Mutable/Source/CustomizableObjectEditor/Private/MuCOE/CustomizableObjectCompileRunnable.h
+++ b/Engine/Plugins/Mutable/Source/CustomizableObjectEditor/Private/MuCOE/CustomizableObjectCompileRunnable.h
@@ -18,6 +18,7 @@
 
 #include <atomic>
 
+struct FCompilationRequest;
 class ITargetPlatform;
 
 class FCustomizableObjectCompileRunnable : public FRunnable
diff --git a/Engine/Plugins/Mutable/Source/CustomizableObjectEditor/Private/MuCOE/CustomizableObjectCompiler.cpp b/Engine/Plugins/Mutable/Source/CustomizableObjectEditor/Private/MuCOE/CustomizableObjectCompiler.cpp
--- a/Engine/Plugins/Mutable/Source/CustomizableObjectEditor/Private/MuCOE/CustomizableObjectCompiler.cpp
+++ b/Engine/Plugins/Mutable/Source/CustomizableObjectEditor/Private/MuCOE/CustomizableObjectCompiler.cpp
@@ -26,6 +26,7 @@
 #include "Interfaces/ITargetPlatform.h"
 #include "Misc/App.h"
 #include "MuCO/CustomizableObjectSystemPrivate.h"
+#include "MuCOE/CompileRequest.h"
 #include "MuR/Model.h"
 #include "MuR/ModelPrivate.h"
 #include "MuCOE/CustomizableObjectVersionBridge.h"
diff --git a/Engine/Plugins/Mutable/Source/CustomizableObjectEditor/Private/MuCOE/CustomizableObjectCompiler.h b/Engine/Plugins/Mutable/Source/CustomizableObjectEditor/Private/MuCOE/CustomizableObjectCompiler.h
--- a/Engine/Plugins/Mutable/Source/CustomizableObjectEditor/Private/MuCOE/CustomizableObjectCompiler.h
+++ b/Engine/Plugins/Mutable/Source/CustomizableObjectEditor/Private/MuCOE/CustomizableObjectCompiler.h
@@ -15,6 +15,7 @@
 
 #include "Framework/Notifications/NotificationManager.h"
 
+struct FCompilationRequest;
 class FCustomizableObjectCompileRunnable;
 class FCustomizableObjectSaveDDRunnable;
 class FReferenceCollector;
diff --git a/Engine/Plugins/Mutable/Source/CustomizableObjectEditor/Private/MuCOE/CustomizableObjectEditor.cpp b/Engine/Plugins/Mutable/Source/CustomizableObjectEditor/Private/MuCOE/CustomizableObjectEditor.cpp
--- a/Engine/Plugins/Mutable/Source/CustomizableObjectEditor/Private/MuCOE/CustomizableObjectEditor.cpp
+++ b/Engine/Plugins/Mutable/Source/CustomizableObjectEditor/Private/MuCOE/CustomizableObjectEditor.cpp
@@ -69,6 +69,7 @@
 #include "Widgets/Notifications/SNotificationList.h"
 #include "SMutableObjectViewer.h"
 #include "MuCO/CustomizableObjectCustomVersion.h"
+#include "MuCOE/CompileRequest.h"
 #include "MuCOE/CustomizableObjectInstanceFactory.h"
 
 class FAdvancedPreviewScene;
@@ -1461,7 +1462,7 @@ void FCustomizableObjectEditor::CompileObject(bool bOnlySelectedParameters, bool
 		CompileRequest->Options.ParamNamesToSelectedOptions = GetCompileOnlySelectedParameters(*GetPreviewInstance());
 	}
 	
-	ICustomizableObjectEditorModule::GetChecked().CompileCustomizableObject(CompileRequest);
+	ICustomizableObjectEditorModulePrivate::GetChecked().EnqueueCompileRequest(CompileRequest);
 }
 
 
@@ -2299,7 +2300,7 @@ void FCustomizableObjectEditor::OnCustomizableObjectStatusChanged(FCustomizableO
 			CompileRequest->bSkipIfOutOfDate = true;
 			CompileRequest->bSilentCompilation = false;
 			
-			ICustomizableObjectEditorModule::GetChecked().CompileCustomizableObject(CompileRequest);
+			ICustomizableObjectEditorModulePrivate::GetChecked().EnqueueCompileRequest(CompileRequest);
 		}
 	}
 }
diff --git a/Engine/Plugins/Mutable/Source/CustomizableObjectEditor/Private/MuCOE/CustomizableObjectEditorFunctionLibrary.cpp b/Engine/Plugins/Mutable/Source/CustomizableObjectEditor/Private/MuCOE/CustomizableObjectEditorFunctionLibrary.cpp
--- a/Engine/Plugins/Mutable/Source/CustomizableObjectEditor/Private/MuCOE/CustomizableObjectEditorFunctionLibrary.cpp
+++ b/Engine/Plugins/Mutable/Source/CustomizableObjectEditor/Private/MuCOE/CustomizableObjectEditorFunctionLibrary.cpp
@@ -7,10 +7,10 @@
 #include "CustomizableObjectEditorModule.h"
 #include "IContentBrowserSingleton.h"
 #include "AssetRegistry/AssetRegistryModule.h"
-#include "MuCOE/CustomizableObjectCompiler.h"
 #include "MuCO/CustomizableObject.h"
 #include "MuCO/CustomizableObjectSystemPrivate.h"
 #include "MuCO/ICustomizableObjectEditorModule.h"
+#include "MuCOE/CompileRequest.h"
 #include "MuCOE/CustomizableObjectFactory.h"
 #include "UObject/SavePackage.h"
 
@@ -36,7 +36,7 @@ ECustomizableObjectCompilationState UCustomizableObjectEditorFunctionLibrary::Co
 	CompileRequest->Options.TextureCompression = InTextureCompression;
 	CompileRequest->Options.bGatherReferences = bGatherReferences;
 	
-	ICustomizableObjectEditorModule::GetChecked().CompileCustomizableObject(CompileRequest);
+	ICustomizableObjectEditorModulePrivate::GetChecked().EnqueueCompileRequest(CompileRequest);
 
 	check(CompileRequest->GetCompilationState() == ECompilationStatePrivate::Completed);
 
diff --git a/Engine/Plugins/Mutable/Source/CustomizableObjectEditor/Private/MuCOE/CustomizableObjectEditorModule.cpp b/Engine/Plugins/Mutable/Source/CustomizableObjectEditor/Private/MuCOE/CustomizableObjectEditorModule.cpp
--- a/Engine/Plugins/Mutable/Source/CustomizableObjectEditor/Private/MuCOE/CustomizableObjectEditorModule.cpp
+++ b/Engine/Plugins/Mutable/Source/CustomizableObjectEditor/Private/MuCOE/CustomizableObjectEditorModule.cpp
@@ -92,6 +92,7 @@
 #include "Editor.h"
 #include "MuCO/CustomizableObjectInstancePrivate.h"
 #include "MuCO/CustomizableObjectSystemPrivate.h"
+#include "MuCOE/CompileRequest.h"
 #include "MuCOE/CustomizableObjectGraph.h"
 #include "Nodes/CustomizableObjectNodeComponentMeshDetails.h"
 #include "Nodes/CustomizableObjectNodeModifierTransformInMeshDetails.h"
@@ -106,13 +107,14 @@ class ICustomizableObjectInstanceEditor;
 class IToolkitHost;
 class UObject;
 
+#define LOCTEXT_NAMESPACE "MutableSettings"
+
 
 const FName CustomizableObjectEditorAppIdentifier = FName(TEXT("CustomizableObjectEditorApp"));
 const FName CustomizableObjectInstanceEditorAppIdentifier = FName(TEXT("CustomizableObjectInstanceEditorApp"));
 const FName CustomizableObjectDebuggerAppIdentifier = FName(TEXT("CustomizableObjectDebuggerApp"));
 const FName CustomizableObjectMacroLibraryEditorAppIdentifier = FName(TEXT("CustomizableObjectMacroLibraryApp"));
 
-#define LOCTEXT_NAMESPACE "MutableSettings"
 
 /** Max timespan in days before a Saved/MutableStreamedDataEditor file is deleted. */
 constexpr int32 MaxAccessTimespan = 30;
@@ -917,43 +919,52 @@ void FCustomizableObjectEditorModule::PostBackwardsCompatibleFixup(UEdGraph& Gra
 }
 
 
-void FCustomizableObjectEditorModule::CompileCustomizableObject(const TSharedRef<FCompilationRequest>& InCompilationRequest, bool bForceRequest)
+void FCustomizableObjectEditorModule::EnqueueCompileRequest(const TSharedRef<FCompilationRequest>& InCompilationRequest, bool bForceRequest)
 {
 	Compiler.EnqueueCompileRequest(InCompilationRequest, bForceRequest);
 }
 
 
-void FCustomizableObjectEditorModule::CompileCustomizableObject(UCustomizableObject& Object, const FCompileParams& Params)
+void FCustomizableObjectEditorModule::CompileCustomizableObject(UCustomizableObject& Object, const FCompileParams* Params, bool bSilentCompilation)
 {
 	UCustomizableObjectSystem* System = UCustomizableObjectSystem::GetInstance();
 	if (!System)
 	{
-		FCompileCallbackParams CallbackParams;
-		CallbackParams.bRequestFailed = true;
-		CallbackParams.bCompiled = Object.IsCompiled();
+		if (Params)
+		{
+			FCompileCallbackParams CallbackParams;
+			CallbackParams.bRequestFailed = true;
+			CallbackParams.bCompiled = Object.IsCompiled();
+		
+			Params->Callback.ExecuteIfBound(CallbackParams);
+			Params->CallbackNative.ExecuteIfBound(CallbackParams);
+		}
 		
-		Params.Callback.ExecuteIfBound(CallbackParams);
-		Params.CallbackNative.ExecuteIfBound(CallbackParams);
 		return;
 	}
 
 	TSharedRef<FCompilationRequest> CompileRequest = MakeShared<FCompilationRequest>(Object);
-	CompileRequest->bAsync = Params.bAsync;
-	CompileRequest->bSkipIfCompiled = Params.bSkipIfCompiled;
-	CompileRequest->bSkipIfOutOfDate = Params.bSkipIfOutOfDate;
-	CompileRequest->bSilentCompilation = true;
-	CompileRequest->Options.TextureCompression = Params.TextureCompression;
-	CompileRequest->Options.OptimizationLevel = ConvertOptimizationLevel(Params.OptimizationLevel);
-	CompileRequest->Options.bGatherReferences = Params.bGatherReferneces;
-	CompileRequest->Callback = Params.Callback;
-	CompileRequest->CallbackNative = Params.CallbackNative;
-	
-	if (Params.CompileOnlySelectedInstance)
+
+	if (Params)
 	{
-		CompileRequest->Options.ParamNamesToSelectedOptions = GetCompileOnlySelectedParameters(*Params.CompileOnlySelectedInstance);
+		CompileRequest->bAsync = Params->bAsync;
+		CompileRequest->bSkipIfCompiled = Params->bSkipIfCompiled;
+		CompileRequest->bSkipIfOutOfDate = Params->bSkipIfOutOfDate;
+		CompileRequest->Options.TextureCompression = Params->TextureCompression;
+		CompileRequest->Options.OptimizationLevel = ConvertOptimizationLevel(Params->OptimizationLevel);
+		CompileRequest->Options.bGatherReferences = Params->bGatherReferneces;
+		CompileRequest->Callback = Params->Callback;
+		CompileRequest->CallbackNative = Params->CallbackNative;
+
+		if (Params->CompileOnlySelectedInstance)
+		{
+			CompileRequest->Options.ParamNamesToSelectedOptions = GetCompileOnlySelectedParameters(*Params->CompileOnlySelectedInstance);
+		}
 	}
+
+	CompileRequest->bSilentCompilation = bSilentCompilation;
 	
-	CompileCustomizableObject(CompileRequest);
+	EnqueueCompileRequest(CompileRequest);
 }
 
 
@@ -964,6 +975,86 @@ int32 FCustomizableObjectEditorModule::Tick(bool bBlocking)
 }
 
 
+void FCustomizableObjectEditorModule::BeginCacheForCookedPlatformData(UCustomizableObject& Object, const ITargetPlatform* TargetPlatform)
+{
+	if (!TargetPlatform)
+	{
+		return;
+	}
+
+	TArray<TSharedRef<FCompilationRequest>>& ObjectCompileRequests = CookCompileRequests.FindOrAdd(&Object);
+	
+	const bool Exists = ObjectCompileRequests.ContainsByPredicate([&TargetPlatform](const TSharedPtr<FCompilationRequest>& Request)
+	{
+		return Request->Options.TargetPlatform == TargetPlatform;
+	});
+
+	if (Exists)
+	{
+		return;
+	}
+
+	if (!IsRootObject(Object))
+	{
+		Object.GetPrivate()->SetIsChildObject(true);
+		return;
+	}
+
+	const bool bAsync = CVarMutableAsyncCook.GetValueOnAnyThread();
+
+	TSharedRef<FCompilationRequest> CompileRequest = MakeShared<FCompilationRequest>(Object);
+	CompileRequest->bAsync = bAsync;
+	CompileRequest->Options.OptimizationLevel = UE_MUTABLE_MAX_OPTIMIZATION; // Force max optimization when packaging.
+	CompileRequest->Options.TextureCompression = ECustomizableObjectTextureCompression::HighQuality;
+	CompileRequest->Options.bIsCooking = true;
+	CompileRequest->Options.bUseBulkData = CVarMutableUseBulkData.GetValueOnAnyThread();
+	CompileRequest->Options.TargetPlatform = TargetPlatform;
+
+	const int32 DDCUsage = CVarMutableDerivedDataCacheUsage.GetValueOnAnyThread();
+	UE::DerivedData::ECachePolicy DefaultCachePolicy = UE::DerivedData::ECachePolicy::None;
+	if (DDCUsage == 1)
+	{
+		DefaultCachePolicy = UE::DerivedData::ECachePolicy::Local;
+	}
+	else if (DDCUsage == 2)
+	{
+		DefaultCachePolicy = UE::DerivedData::ECachePolicy::Default;
+	}
+	CompileRequest->SetDerivedDataCachePolicy(DefaultCachePolicy);
+
+	ObjectCompileRequests.Add(CompileRequest);
+
+	EnqueueCompileRequest(CompileRequest, true);
+}
+
+
+bool FCustomizableObjectEditorModule::IsCachedCookedPlatformDataLoaded(UCustomizableObject& Object, const ITargetPlatform* TargetPlatform)
+{
+	if (!TargetPlatform)
+	{
+		return true;
+	}
+
+	TArray<TSharedRef<FCompilationRequest>>* ObjectCompileRequests = CookCompileRequests.Find(&Object);
+	if (!ObjectCompileRequests)
+	{
+		return true;
+	}
+	
+	const TSharedRef<FCompilationRequest>* CompileRequest = ObjectCompileRequests->FindByPredicate([&TargetPlatform](const TSharedRef<FCompilationRequest>& Request)
+	{
+		return Request->Options.TargetPlatform == TargetPlatform;
+	});
+	
+	if (CompileRequest)
+	{
+		return CompileRequest->Get().GetCompilationState() == ECompilationStatePrivate::Completed;
+	}
+
+	return true;
+}
+
+
 void FCustomizableObjectEditorModule::CancelCompileRequests()
 {
 	Compiler.ForceFinishCompilation();
@@ -1034,7 +1125,7 @@ void FCustomizableObjectEditorModule::OnPreBeginPIE(const bool bIsSimulatingInEd
 		{
 			for (const TSharedRef<FCompilationRequest>& Request : Requests)
 			{
-				CompileCustomizableObject(Request);
+				EnqueueCompileRequest(Request);
 			}
 		}
 	}
diff --git a/Engine/Plugins/Mutable/Source/CustomizableObjectEditor/Private/MuCOE/CustomizableObjectEditorModule.h b/Engine/Plugins/Mutable/Source/CustomizableObjectEditor/Private/MuCOE/CustomizableObjectEditorModule.h
--- a/Engine/Plugins/Mutable/Source/CustomizableObjectEditor/Private/MuCOE/CustomizableObjectEditorModule.h
+++ b/Engine/Plugins/Mutable/Source/CustomizableObjectEditor/Private/MuCOE/CustomizableObjectEditorModule.h
@@ -2,6 +2,7 @@
 
 #pragma once
 
+#include "ICustomizableObjectEditorModulePrivate.h"
 #include "MuCO/ICustomizableObjectEditorModule.h"
 #include "MuCO/CustomizableObjectCompilerTypes.h"
 
@@ -33,7 +34,7 @@ TMap<FString, FString> GetCompileOnlySelectedParameters(const UCustomizableObjec
 /**
  * StaticMesh editor module
  */
-class FCustomizableObjectEditorModule : public ICustomizableObjectEditorModule
+class FCustomizableObjectEditorModule : public ICustomizableObjectEditorModulePrivate
 {
 public:
 	// IModuleInterface interface
@@ -56,14 +57,16 @@ public:
 	virtual void CancelCompileRequests() override;
 	virtual int32 GetNumCompileRequests() override;
 	virtual bool IsCompiling(const UCustomizableObject& Object) const override;
-	virtual void CompileCustomizableObject(const TSharedRef<FCompilationRequest>& InCompilationRequest, bool bForceRequest = false) override;
-	virtual void CompileCustomizableObject(UCustomizableObject& Object, const FCompileParams& Params) override;
-
+	virtual void CompileCustomizableObject(UCustomizableObject& Object, const FCompileParams* Params, bool bSilentCompilation) override;
 	virtual int32 Tick(bool bBlocking) override;
-
+	virtual void BeginCacheForCookedPlatformData(UCustomizableObject& Object, const ITargetPlatform* TargetPlatform) override;
+	virtual bool IsCachedCookedPlatformDataLoaded(UCustomizableObject& Object, const ITargetPlatform* TargetPlatform) override;
 	virtual TSharedPtr<FExtensibilityManager> GetCustomizableObjectEditorToolBarExtensibilityManager() override { return CustomizableObjectEditor_ToolBarExtensibilityManager; }
 	virtual TSharedPtr<FExtensibilityManager> GetCustomizableObjectEditorMenuExtensibilityManager() override { return CustomizableObjectEditor_MenuExtensibilityManager; }
 
+	// ICustomizableObjectEditorModulePrivate interface
+	virtual void EnqueueCompileRequest(const TSharedRef<FCompilationRequest>& InCompilationRequest, bool bForceRequest = false) override;
+
 private:	
 	TSharedPtr<FExtensibilityManager> CustomizableObjectEditor_ToolBarExtensibilityManager;
 	TSharedPtr<FExtensibilityManager> CustomizableObjectEditor_MenuExtensibilityManager;
@@ -95,4 +98,7 @@ private:
 	FName MeshReshapeBoneReferenceUStructName;
 
 	FName BoneToRemoveUStructName;
+
+	/** Cook requests. */
+	TMap<TWeakObjectPtr<UCustomizableObject>, TArray<TSharedRef<FCompilationRequest>>> CookCompileRequests;
 };
\ No newline at end of file
diff --git a/Engine/Plugins/Mutable/Source/MutableValidation/Private/MuV/AssetValidator_CustomizableObjects.cpp b/Engine/Plugins/Mutable/Source/MutableValidation/Private/MuV/AssetValidator_CustomizableObjects.cpp
--- a/Engine/Plugins/Mutable/Source/MutableValidation/Private/MuV/AssetValidator_CustomizableObjects.cpp
+++ b/Engine/Plugins/Mutable/Source/MutableValidation/Private/MuV/AssetValidator_CustomizableObjects.cpp
@@ -4,6 +4,7 @@
 
 #include "DataValidationModule.h"
 #include "Editor.h"
+#include "ICustomizableObjectEditorModulePrivate.h"
 #include "MutableValidationSettings.h"
 #include "AssetRegistry/AssetData.h"
 #include "Delegates/DelegateSignatureImpl.inl"
@@ -15,6 +16,7 @@
 #include "MuCO/CustomizableObjectCompilerTypes.h"
 #include "MuCO/CustomizableObjectPrivate.h"
 #include "MuCO/ICustomizableObjectEditorModule.h"
+#include "MuCOE/CompileRequest.h"
 #include "UObject/NameTypes.h"
 #include "UObject/Object.h"
 
@@ -147,7 +149,7 @@ EDataValidationResult UAssetValidator_CustomizableObjects::IsCustomizableObjectV
 		// Run Sync compilation -> Warning : Potentially long operation -------------
 		TSharedRef<FCompilationRequest> CompileRequest = MakeShared<FCompilationRequest>(*RootObject);
 		CompileRequest->bAsync = false;
-		ICustomizableObjectEditorModule::GetChecked().CompileCustomizableObject(CompileRequest);
+		ICustomizableObjectEditorModulePrivate::GetChecked().EnqueueCompileRequest(CompileRequest);
 		// --------------------------------------------------------------------------
 		
 		// Get compilation errors and warnings
diff --git a/Engine/Plugins/Mutable/Source/MutableValidation/Private/MuV/CustomizableObjectCompilationUtility.cpp b/Engine/Plugins/Mutable/Source/MutableValidation/Private/MuV/CustomizableObjectCompilationUtility.cpp
--- a/Engine/Plugins/Mutable/Source/MutableValidation/Private/MuV/CustomizableObjectCompilationUtility.cpp
+++ b/Engine/Plugins/Mutable/Source/MutableValidation/Private/MuV/CustomizableObjectCompilationUtility.cpp
@@ -2,6 +2,7 @@
 
 #include "CustomizableObjectCompilationUtility.h"
 
+#include "ICustomizableObjectEditorModulePrivate.h"
 #include "ScopedLogSection.h"
 #include "Commandlets/Commandlet.h"
 #include "MuCO/CustomizableObject.h"
@@ -10,10 +11,10 @@
 #include "MuCO/CustomizableObjectCompilerTypes.h"
 #include "MuCO/ICustomizableObjectEditorModule.h"
 #include "Interfaces/ITargetPlatform.h"
+#include "MuCOE/CompileRequest.h"
 #include "MuCOE/CustomizableObjectEditorFunctionLibrary.h"
 
 
-
 bool FCustomizableObjectCompilationUtility::CompileCustomizableObject(UCustomizableObject* InCustomizableObject, const bool bShouldLogMutableLogs /* = true */, const FCompilationOptions* InCompilationOptionsOverride  /* nullptr */)
 {
 	LLM_SCOPE_BYNAME(TEXT("FCustomizableObjectCompilationUtility/Compile"));
@@ -94,7 +95,7 @@ bool FCustomizableObjectCompilationUtility::CompileCustomizableObject(UCustomiza
 		CompileRequest->bSilentCompilation = false;
 		CompileRequest->Options = CompilationOptions;
 		
-		ICustomizableObjectEditorModule::GetChecked().CompileCustomizableObject(CompileRequest);
+		ICustomizableObjectEditorModulePrivate::GetChecked().EnqueueCompileRequest(CompileRequest);
 		
 		// Wait while the compilation takes place
 		const double CompilationStartSeconds = FPlatformTime::Seconds();
