[Mutable] Deprecate Find<type>ParameterNameIndex.

User should not interact with ParameterIndex. Replaced with Contains<type>Parameter.

diff --git a/Engine/Plugins/Mutable/Source/CustomizableObject/Internal/MuCO/CustomizableObjectInstancePrivate.h b/Engine/Plugins/Mutable/Source/CustomizableObject/Internal/MuCO/CustomizableObjectInstancePrivate.h
--- a/Engine/Plugins/Mutable/Source/CustomizableObject/Internal/MuCO/CustomizableObjectInstancePrivate.h
+++ b/Engine/Plugins/Mutable/Source/CustomizableObject/Internal/MuCO/CustomizableObjectInstancePrivate.h
@@ -352,6 +352,21 @@ public:
 	/** Return true if the instance is not locked and if it's compiled. */
 	bool CanUpdateInstance() const;
 
+	/** Finds in IntParameters a parameter with name ParamName, returns the index if found, -1 otherwise. */
+	int32 FindIntParameterNameIndex(const FString& ParamName) const;
+
+	/** Finds in FloatParameters a parameter with name ParamName, returns the index if found, -1 otherwise. */
+	int32 FindFloatParameterNameIndex(const FString& ParamName) const;
+
+	/** Finds in BoolParameters a parameter with name ParamName, returns the index if found, -1 otherwise. */
+	int32 FindBoolParameterNameIndex(const FString& ParamName) const;
+
+	/** Finds in VectorParameters a parameter with name ParamName, returns the index if found, -1 otherwise. */
+	int32 FindVectorParameterNameIndex(const FString& ParamName) const;
+
+	/** Finds in ProjectorParameters a parameter with name ParamName, returns the index if found, -1 otherwise. */
+	int32 FindProjectorParameterNameIndex(const FString& ParamName) const;
+	
 	UCustomizableObjectInstance* GetPublic() const;
 
 	// If any components are using this instance, they will store the min of their distances to the player here every frame for LOD purposes
diff --git a/Engine/Plugins/Mutable/Source/CustomizableObject/Private/MuCO/CustomizableObjectInstance.cpp b/Engine/Plugins/Mutable/Source/CustomizableObject/Private/MuCO/CustomizableObjectInstance.cpp
--- a/Engine/Plugins/Mutable/Source/CustomizableObject/Private/MuCO/CustomizableObjectInstance.cpp
+++ b/Engine/Plugins/Mutable/Source/CustomizableObject/Private/MuCO/CustomizableObjectInstance.cpp
@@ -2811,6 +2811,42 @@ FCustomizableObjectProjector UCustomizableObjectInstance::GetProjector(const FSt
 }
 
 
+bool UCustomizableObjectInstance::ContainsIntParameter(const FString& ParameterName) const
+{
+	return Descriptor.FindTypedParameterIndex(ParameterName, EMutableParameterType::Int) != INDEX_NONE;
+}
+
+
+bool UCustomizableObjectInstance::ContainsFloatParameter(const FString& ParameterName) const
+{
+	return Descriptor.FindTypedParameterIndex(ParameterName, EMutableParameterType::Float) != INDEX_NONE;
+}
+
+
+bool UCustomizableObjectInstance::ContainsBoolParameter(const FString& ParameterName) const
+{
+	return Descriptor.FindTypedParameterIndex(ParameterName, EMutableParameterType::Bool) != INDEX_NONE;
+}
+
+
+bool UCustomizableObjectInstance::ContainsVectorParameter(const FString& ParameterName) const
+{
+	return Descriptor.FindTypedParameterIndex(ParameterName, EMutableParameterType::Color) != INDEX_NONE;
+}
+
+
+bool UCustomizableObjectInstance::ContainsProjectorParameter(const FString& ParameterName) const
+{
+	return Descriptor.FindTypedParameterIndex(ParameterName, EMutableParameterType::Projector) != INDEX_NONE;
+}
+
+
+bool UCustomizableObjectInstance::ContainsTransformParameter(const FString& ParameterName) const
+{
+	return Descriptor.FindTypedParameterIndex(ParameterName, EMutableParameterType::Transform) != INDEX_NONE;
+}
+
+
 int32 UCustomizableObjectInstance::FindIntParameterNameIndex(const FString& ParamName) const
 {
 	return Descriptor.FindTypedParameterIndex(ParamName, EMutableParameterType::Int);
@@ -2840,6 +2876,37 @@ int32 UCustomizableObjectInstance::FindProjectorParameterNameIndex(const FString
 	return Descriptor.FindTypedParameterIndex(ParamName, EMutableParameterType::Projector);
 }
 
+
+int32 UCustomizableInstancePrivate::FindIntParameterNameIndex(const FString& ParamName) const
+{
+	return GetPublic()->Descriptor.FindTypedParameterIndex(ParamName, EMutableParameterType::Int);
+}
+
+
+int32 UCustomizableInstancePrivate::FindFloatParameterNameIndex(const FString& ParamName) const
+{
+	return GetPublic()->Descriptor.FindTypedParameterIndex(ParamName, EMutableParameterType::Float);
+}
+
+
+int32 UCustomizableInstancePrivate::FindBoolParameterNameIndex(const FString& ParamName) const
+{
+	return GetPublic()->Descriptor.FindTypedParameterIndex(ParamName, EMutableParameterType::Bool);
+}
+
+
+int32 UCustomizableInstancePrivate::FindVectorParameterNameIndex(const FString& ParamName) const
+{
+	return GetPublic()->Descriptor.FindTypedParameterIndex(ParamName, EMutableParameterType::Color);
+}
+
+
+int32 UCustomizableInstancePrivate::FindProjectorParameterNameIndex(const FString& ParamName) const
+{
+	return GetPublic()->Descriptor.FindTypedParameterIndex(ParamName, EMutableParameterType::Projector);
+}
+
+
 void UCustomizableObjectInstance::SetRandomValues()
 {
 	Descriptor.SetRandomValues();
diff --git a/Engine/Plugins/Mutable/Source/CustomizableObject/Public/MuCO/CustomizableObjectInstance.h b/Engine/Plugins/Mutable/Source/CustomizableObject/Public/MuCO/CustomizableObjectInstance.h
--- a/Engine/Plugins/Mutable/Source/CustomizableObject/Public/MuCO/CustomizableObjectInstance.h
+++ b/Engine/Plugins/Mutable/Source/CustomizableObject/Public/MuCO/CustomizableObjectInstance.h
@@ -501,26 +501,50 @@ public:
 	/** Get the current projector for the parameter with the given name. */
 	FCustomizableObjectProjector GetProjector(const FString& ParamName, int32 RangeIndex) const;
 	
-	/** Finds in IntParameters a parameter with name ParamName, returns the index if found, -1 otherwise. */
+	/** Return true if the Int Parameter exists. */
+	UFUNCTION(BlueprintCallable, Category = CustomizableObjectInstance)
+	bool ContainsIntParameter(const FString& ParameterName) const;
+
+	/** Return true if the Float Parameter exists. */
+	UFUNCTION(BlueprintCallable, Category = CustomizableObjectInstance)
+	bool ContainsFloatParameter(const FString& ParameterName) const;
+
+	/** Return true if the Bool Parameter exists. */
+	UFUNCTION(BlueprintCallable, Category = CustomizableObjectInstance)
+	bool ContainsBoolParameter(const FString& ParameterName) const;
+
+	/** Return true if the Vector Parameter exists. */
+	UFUNCTION(BlueprintCallable, Category = CustomizableObjectInstance)
+	bool ContainsVectorParameter(const FString& ParameterName) const;
+
+	/** Return true if the Projector Parameter exists. */
+	UFUNCTION(BlueprintCallable, Category = CustomizableObjectInstance)
+	bool ContainsProjectorParameter(const FString& ParameterName) const;
+
+	/** Return true if the Transform Parameter exists. */
+	UFUNCTION(BlueprintCallable, Category = CustomizableObjectInstance)
+	bool ContainsTransformParameter(const FString& ParameterName) const;
+
+	UE_DEPRECATED(5.6, "Use ContainsIntParameter instead")
 	UFUNCTION(BlueprintCallable, Category = CustomizableObjectInstance)
 	int32 FindIntParameterNameIndex(const FString& ParamName) const;
 
-	/** Finds in FloatParameters a parameter with name ParamName, returns the index if found, -1 otherwise. */
+	UE_DEPRECATED(5.6, "Use ContainsFloatParameter instead")
 	UFUNCTION(BlueprintCallable, Category = CustomizableObjectInstance)
 	int32 FindFloatParameterNameIndex(const FString& ParamName) const;
 
-	/** Finds in BoolParameters a parameter with name ParamName, returns the index if found, -1 otherwise. */
+	UE_DEPRECATED(5.6, "Use ContainsBoolParameter instead")
 	UFUNCTION(BlueprintCallable, Category = CustomizableObjectInstance)
 	int32 FindBoolParameterNameIndex(const FString& ParamName) const;
 
-	/** Finds in VectorParameters a parameter with name ParamName, returns the index if found, -1 otherwise. */
+	UE_DEPRECATED(5.6, "Use ContainsVectorParameter instead")
 	UFUNCTION(BlueprintCallable, Category = CustomizableObjectInstance)
 	int32 FindVectorParameterNameIndex(const FString& ParamName) const;
 
-	/** Finds in ProjectorParameters a parameter with name ParamName, returns the index if found, -1 otherwise. */
+	UE_DEPRECATED(5.6, "Use ContainsProjectorParameter instead")
 	UFUNCTION(BlueprintCallable, Category = CustomizableObjectInstance)
 	int32 FindProjectorParameterNameIndex(const FString& ParamName) const;
-
+	
 	/** Increases the range of values of the integer with ParamName and returns the index of the new integer value, -1 otherwise.
 	  * The added value is initialized with the first integer option and is the last one of the range. */
 	UFUNCTION(BlueprintCallable, Category = CustomizableObjectInstance)
diff --git a/Engine/Plugins/Mutable/Source/CustomizableObjectEditor/Private/MuCOE/CustomizableInstanceDetails.cpp b/Engine/Plugins/Mutable/Source/CustomizableObjectEditor/Private/MuCOE/CustomizableInstanceDetails.cpp
--- a/Engine/Plugins/Mutable/Source/CustomizableObjectEditor/Private/MuCOE/CustomizableInstanceDetails.cpp
+++ b/Engine/Plugins/Mutable/Source/CustomizableObjectEditor/Private/MuCOE/CustomizableInstanceDetails.cpp
@@ -1855,7 +1855,7 @@ FReply FCustomizableInstanceDetails::OnProjectorSelectChanged(const FString Para
 FReply FCustomizableInstanceDetails::OnProjectorCopyTransform(const FString ParamName, const int32 RangeIndex) const
 {
 	const int32 ParameterIndexInObject = CustomInstance->GetCustomizableObject()->GetPrivate()->FindParameter(ParamName);
-	const int32 ProjectorParamIndex = CustomInstance->FindProjectorParameterNameIndex(ParamName);
+	const int32 ProjectorParamIndex = CustomInstance->GetPrivate()->FindProjectorParameterNameIndex(ParamName);
 
 	if ((ParameterIndexInObject >= 0) && (ProjectorParamIndex >= 0))
 	{
@@ -1940,7 +1940,7 @@ TSharedRef<SWidget> FCustomizableInstanceDetails::GenerateMultidimensionalProjec
 {
 	const TSharedPtr<ICustomizableObjectInstanceEditor> Editor = GetEditorChecked();
 	const TArray<FCustomizableObjectProjectorParameterValue>& ProjectorParameters = CustomInstance->GetPrivate()->GetDescriptor().GetProjectorParameters();
-	const int32 ProjectorParamIndex = CustomInstance->FindProjectorParameterNameIndex(ParamName);
+	const int32 ProjectorParamIndex = CustomInstance->GetPrivate()->FindProjectorParameterNameIndex(ParamName);
 
 	check(ProjectorParamIndex < ProjectorParameters.Num());
 
