[mutable] Created simple ARSearch commandlet to be run prior to the mutable editor tests

#rnx

diff --git a/Engine/Plugins/Mutable/Source/MutableValidation/Private/MuV/CacheAssetRegistryCommandlet.cpp b/Engine/Plugins/Mutable/Source/MutableValidation/Private/MuV/CacheAssetRegistryCommandlet.cpp
new file mode 100644
--- /dev/null
+++ b/Engine/Plugins/Mutable/Source/MutableValidation/Private/MuV/CacheAssetRegistryCommandlet.cpp
@@ -0,0 +1,21 @@
+// Copyright Epic Games, Inc. All Rights Reserved.
+
+#include "MuV/CacheAssetRegistryCommandlet.h"
+
+#include "ValidationUtils.h"
+#include "Misc/CommandLine.h"
+
+int32 UCacheAssetRegistryCommandlet::Main(const FString& Params)
+{
+	// Ensure we are not running this commandlet in vain as we want the cache to be filled
+	if (FParse::Param(FCommandLine::Get(), TEXT("NoAssetRegistryCacheWrite")))
+	{
+		UE_LOG(LogTemp, Error, TEXT("The Asset Registry data is not going to be cached due to the arg 'NoAssetRegistryCacheWrite' being present in the commandline."));
+		return 1;
+	}
+	
+	// Perform a blocking search to ensure all assets used by mutable are reachable using the AssetRegistry
+	PrepareAssetRegistry();
+
+	return 0;
+}
diff --git a/Engine/Plugins/Mutable/Source/MutableValidation/Public/MuV/CacheAssetRegistryCommandlet.h b/Engine/Plugins/Mutable/Source/MutableValidation/Public/MuV/CacheAssetRegistryCommandlet.h
new file mode 100644
--- /dev/null
+++ b/Engine/Plugins/Mutable/Source/MutableValidation/Public/MuV/CacheAssetRegistryCommandlet.h
@@ -0,0 +1,17 @@
+// Copyright Epic Games, Inc. All Rights Reserved.
+
+#pragma once
+
+#include "Commandlets/Commandlet.h"
+
+#include "CacheAssetRegistryCommandlet.generated.h"
+
+/** Simple commandlet used to fill-up the asset registry cache for the execution of the mutable tests */
+UCLASS()
+class UCacheAssetRegistryCommandlet : public UCommandlet
+{
+	GENERATED_BODY()
+
+public:
+	virtual int32 Main(const FString& Params) override;
+};
