[mutable] Updated pin names
- Updated the input and output pin names of multiple CO nodes for cases where the current name did not add info or was incoherent with nodes of a similar functionality.
- Removed option to add comment on the drop of a connection line.
- Removed unnecessary overrides of methods
- Added required "virtual" keyword in multiple method overrides.
- Fixed node serialization issue.

#rnx

diff --git a/Engine/Plugins/Mutable/Source/CustomizableObject/Internal/MuCO/CustomizableObjectCustomVersion.h b/Engine/Plugins/Mutable/Source/CustomizableObject/Internal/MuCO/CustomizableObjectCustomVersion.h
--- a/Engine/Plugins/Mutable/Source/CustomizableObject/Internal/MuCO/CustomizableObjectCustomVersion.h
+++ b/Engine/Plugins/Mutable/Source/CustomizableObject/Internal/MuCO/CustomizableObjectCustomVersion.h
@@ -174,6 +174,8 @@ struct CUSTOMIZABLEOBJECT_API FCustomizableObjectCustomVersion
 
 		UpdatedNodesPinName2,
 
+		UpdatedNodesPinName3,
+
 		// -----<new versions can be added above this line>-------------------------------------------------
 		VersionPlusOne,
 		LatestVersion = VersionPlusOne - 1
diff --git a/Engine/Plugins/Mutable/Source/CustomizableObjectEditor/Private/MuCOE/EdGraphSchema_CustomizableObject.cpp b/Engine/Plugins/Mutable/Source/CustomizableObjectEditor/Private/MuCOE/EdGraphSchema_CustomizableObject.cpp
--- a/Engine/Plugins/Mutable/Source/CustomizableObjectEditor/Private/MuCOE/EdGraphSchema_CustomizableObject.cpp
+++ b/Engine/Plugins/Mutable/Source/CustomizableObjectEditor/Private/MuCOE/EdGraphSchema_CustomizableObject.cpp
@@ -365,9 +365,9 @@ void UEdGraphSchema_CustomizableObject::GetGraphContextActions(FGraphContextMenu
 {
 	//const UCustomizableObjectGraph* CustomizableObjectGraph = CastChecked<UCustomizableObjectGraph>(ContextMenuBuilder.CurrentGraph);
 	//TSharedPtr<FCustomizableObjectSchemaAction_NewNode> Action;
-	const int32 GeneralGrouping = 3;
+	constexpr int32 GeneralGrouping = 3;
 
-	const bool bDisableFilter = false;
+	constexpr  bool bDisableFilter = false;
 
 	// return true if Filter is passed.
 	const auto Filter = [bDisableFilter](UCustomizableObjectNode* TemplateNode, FGraphContextMenuBuilder& ContextMenuBuilder) -> bool
@@ -648,7 +648,8 @@ void UEdGraphSchema_CustomizableObject::GetGraphContextActions(FGraphContextMenu
 			AddNewNodeActionFiltered(TemplateNode, ContextMenuBuilder, Category.ToString(), GeneralGrouping, Filter);
 		}
 	}
-	
+
+	if (!ContextMenuBuilder.FromPin)
 	{
 		UEdGraphNode_Comment* Node = NewObject<UEdGraphNode_Comment>();
 		TSharedPtr<FCustomizableObjectSchemaAction_NewNode> Action = AddNewNodeAction(ContextMenuBuilder, FString(), Node->GetNodeTitle(ENodeTitleType::ListView), FText(), 1);
diff --git a/Engine/Plugins/Mutable/Source/CustomizableObjectEditor/Private/MuCOE/Nodes/CustomizableObjectNode.cpp b/Engine/Plugins/Mutable/Source/CustomizableObjectEditor/Private/MuCOE/Nodes/CustomizableObjectNode.cpp
--- a/Engine/Plugins/Mutable/Source/CustomizableObjectEditor/Private/MuCOE/Nodes/CustomizableObjectNode.cpp
+++ b/Engine/Plugins/Mutable/Source/CustomizableObjectEditor/Private/MuCOE/Nodes/CustomizableObjectNode.cpp
@@ -40,6 +40,24 @@ bool UCustomizableObjectNode::IsSingleOutputNode() const
 }
 
 
+UEdGraphPin* UCustomizableObjectNode::CustomCreatePinSimple(EEdGraphPinDirection Direction, const FName& Category, bool bIsArray)
+{
+	const FName NewPinName = UEdGraphSchema_CustomizableObject::GetPinCategoryName(Category);
+	const FText NewPinFriendlyName = UEdGraphSchema_CustomizableObject::GetPinCategoryFriendlyName(Category);
+
+	UEdGraphPin* Pin = CreatePin(Direction, Category, NewPinName);
+	check(Pin);
+	Pin->PinFriendlyName = NewPinFriendlyName;
+	
+	if (bIsArray)
+	{
+		Pin->PinType.ContainerType = EPinContainerType::Array;
+	}
+
+	return Pin;
+}
+
+
 UEdGraphPin* UCustomizableObjectNode::CustomCreatePin(EEdGraphPinDirection Direction, const FName& Type, const FName& Name, bool bIsArray)
 {
 	UEdGraphPin* Pin = CreatePin(Direction, Type, Name);
diff --git a/Engine/Plugins/Mutable/Source/CustomizableObjectEditor/Private/MuCOE/Nodes/CustomizableObjectNodeAnimationPose.cpp b/Engine/Plugins/Mutable/Source/CustomizableObjectEditor/Private/MuCOE/Nodes/CustomizableObjectNodeAnimationPose.cpp
--- a/Engine/Plugins/Mutable/Source/CustomizableObjectEditor/Private/MuCOE/Nodes/CustomizableObjectNodeAnimationPose.cpp
+++ b/Engine/Plugins/Mutable/Source/CustomizableObjectEditor/Private/MuCOE/Nodes/CustomizableObjectNodeAnimationPose.cpp
@@ -6,6 +6,7 @@
 #include "Animation/DebugSkelMeshComponent.h"
 #include "Animation/PoseAsset.h"
 #include "BonePose.h"
+#include "MuCO/CustomizableObjectCustomVersion.h"
 #include "MuCOE/EdGraphSchema_CustomizableObject.h"
 
 class UCustomizableObjectNodeRemapPins;
@@ -20,46 +21,57 @@ UCustomizableObjectNodeAnimationPose::UCustomizableObjectNodeAnimationPose() : S
 }
 
 
-void UCustomizableObjectNodeAnimationPose::PostEditChangeProperty(FPropertyChangedEvent& PropertyChangedEvent)
-{
-	Super::PostEditChangeProperty(PropertyChangedEvent);
-}
-
-
 void UCustomizableObjectNodeAnimationPose::AllocateDefaultPins(UCustomizableObjectNodeRemapPins* RemapPins)
 {
-	const UEdGraphSchema_CustomizableObject* Schema = GetDefault<UEdGraphSchema_CustomizableObject>();
-
-	FString PinName = TEXT("Input Mesh");
-	UEdGraphPin* InMeshPin = CustomCreatePin(EGPD_Input, Schema->PC_Mesh, FName(*PinName));
+	UEdGraphPin* InMeshPin = CustomCreatePinSimple(EGPD_Input, UEdGraphSchema_CustomizableObject::PC_Mesh);
 	InMeshPin->bDefaultValueIsIgnored = true;
 
-	PinName = TEXT("Table Pose");
-	UEdGraphPin* TablePosePin = CustomCreatePin(EGPD_Input, Schema->PC_PoseAsset, FName(*PinName));
+	UEdGraphPin* TablePosePin = CustomCreatePinSimple(EGPD_Input, UEdGraphSchema_CustomizableObject::PC_PoseAsset);
 	TablePosePin->bDefaultValueIsIgnored = true;
 
-	PinName = TEXT("Output Mesh");
-	UEdGraphPin* OutMeshPin = CustomCreatePin(EGPD_Output, Schema->PC_Mesh, FName(*PinName));
+	UEdGraphPin* OutMeshPin = CustomCreatePinSimple(EGPD_Output,  UEdGraphSchema_CustomizableObject::PC_Mesh);
 	OutMeshPin->bDefaultValueIsIgnored = true;
 }
 
-
-UEdGraphPin* UCustomizableObjectNodeAnimationPose::GetInputMeshPin() const
+void UCustomizableObjectNodeAnimationPose::BackwardsCompatibleFixup(int32 CustomizableObjectCustomVersion)
 {
-	FString PinName = FString(TEXT("Input Mesh"));
+	Super::BackwardsCompatibleFixup(CustomizableObjectCustomVersion);
 
-	UEdGraphPin* Pin = FindPin(PinName);
+	if (CustomizableObjectCustomVersion == FCustomizableObjectCustomVersion::UpdatedNodesPinName3)
+	{
+		if (UEdGraphPin* InputMesh = FindPin(TEXT("Input Mesh"), EGPD_Input))
+		{
+			InputMesh->PinName = TEXT("Mesh");
+			InputMesh->PinFriendlyName = LOCTEXT("Mesh_Pin_Category", "Mesh");
+		}
 
-	return Pin;
+		if (UEdGraphPin* OutputMesh = FindPin(TEXT("Table Pose"), EGPD_Input))
+		{
+			OutputMesh->PinName = TEXT("PoseAsset");
+			OutputMesh->PinFriendlyName = LOCTEXT("Pose_Pin_Category", "PoseAsset");
+		}
+		
+		if (UEdGraphPin* OutputMesh = FindPin(TEXT("Output Mesh"), EGPD_Output))
+		{
+			OutputMesh->PinName = TEXT("Mesh");
+			OutputMesh->PinFriendlyName = LOCTEXT("Mesh_Pin_Category", "Mesh");
+		}
+	}
 }
 
 
-UEdGraphPin* UCustomizableObjectNodeAnimationPose::GetTablePosePin() const
+UEdGraphPin* UCustomizableObjectNodeAnimationPose::GetInputMeshPin() const
 {
-	FString PinName = FString(TEXT("Table Pose"));
+	const FName PinName = UEdGraphSchema_CustomizableObject::GetPinCategoryName(UEdGraphSchema_CustomizableObject::PC_Mesh);
+	UEdGraphPin* Pin = FindPin(PinName, EGPD_Input);
+	return Pin;
+}
 
-	UEdGraphPin* Pin = FindPin(PinName);
 
+UEdGraphPin* UCustomizableObjectNodeAnimationPose::GetTablePosePin() const
+{
+	const FName PinName = UEdGraphSchema_CustomizableObject::GetPinCategoryName(UEdGraphSchema_CustomizableObject::PC_PoseAsset);
+	UEdGraphPin* Pin = FindPin(PinName, EGPD_Input);
 	return Pin;
 }
 
@@ -69,9 +81,9 @@ FText UCustomizableObjectNodeAnimationPose::GetNodeTitle(ENodeTitleType::Type Ti
 	if (PoseAsset)
 	{
 		FFormatNamedArguments Args;
-		Args.Add(TEXT("SkeletamMeshName"), FText::FromString(PoseAsset->GetName()));
+		Args.Add(TEXT("SkeletalMeshName"), FText::FromString(PoseAsset->GetName()));
 
-		return FText::Format(LOCTEXT("AnimationPose_Title", "{SkeletamMeshName}\nAnimation Pose"), Args);
+		return FText::Format(LOCTEXT("AnimationPose_Title", "{SkeletalMeshName}\nAnimation Pose"), Args);
 	}
 	else
 	{
@@ -82,8 +94,7 @@ FText UCustomizableObjectNodeAnimationPose::GetNodeTitle(ENodeTitleType::Type Ti
 
 FLinearColor UCustomizableObjectNodeAnimationPose::GetNodeTitleColor() const
 {
-	const UEdGraphSchema_CustomizableObject* Schema = GetDefault<UEdGraphSchema_CustomizableObject>();
-	return Schema->GetPinTypeColor(Schema->PC_Mesh);
+	return UEdGraphSchema_CustomizableObject::GetPinTypeColor(UEdGraphSchema_CustomizableObject::PC_Mesh);
 }
 
 
@@ -131,7 +142,7 @@ void UCustomizableObjectNodeAnimationPose::StaticRetrievePoseInformation(UPoseAs
 
 		OutCurve.CopyFrom(SkeletalMeshComponent->AnimCurves);
 
-		// Assuming one single pose, with a weigth set to 1.0
+		// Assuming one single pose, with a weight set to 1.0
 		FAnimExtractContext ExtractionContext;
 		ExtractionContext.bExtractRootMotion = false;
 		ExtractionContext.CurrentTime = 0.0f;
@@ -149,7 +160,7 @@ void UCustomizableObjectNodeAnimationPose::StaticRetrievePoseInformation(UPoseAs
 		const FReferenceSkeleton& RefSkeleton = RefSkeletalMesh->GetRefSkeleton();
 		for (int32 BoneIndex = 0; BoneIndex < RefSkeleton.GetRawBoneNum(); ++BoneIndex)
 		{
-			FName BoneName = RefSkeleton.GetBoneName(BoneIndex);
+			const FName BoneName = RefSkeleton.GetBoneName(BoneIndex);
 
 			FTransform CumulativePoseTransform;
 
@@ -170,8 +181,8 @@ void UCustomizableObjectNodeAnimationPose::StaticRetrievePoseInformation(UPoseAs
 				ParentIndex = RefSkeleton.GetParentIndex(ParentIndex);
 			}
 
-			FTransform BoneToComponentTranform = SkeletalMeshComponent->GetEditableComponentSpaceTransforms()[BoneIndex];
-			FTransform TransformToAdd = BoneToComponentTranform.Inverse() * CumulativePoseTransform;
+			const FTransform BoneToComponentTransform = SkeletalMeshComponent->GetEditableComponentSpaceTransforms()[BoneIndex];
+			FTransform TransformToAdd = BoneToComponentTransform.Inverse() * CumulativePoseTransform;
 
 			OutArrayBoneName.Add(BoneName);
 			OutArrayTransform.Add(TransformToAdd);
diff --git a/Engine/Plugins/Mutable/Source/CustomizableObjectEditor/Private/MuCOE/Nodes/CustomizableObjectNodeAnimationPose.h b/Engine/Plugins/Mutable/Source/CustomizableObjectEditor/Private/MuCOE/Nodes/CustomizableObjectNodeAnimationPose.h
--- a/Engine/Plugins/Mutable/Source/CustomizableObjectEditor/Private/MuCOE/Nodes/CustomizableObjectNodeAnimationPose.h
+++ b/Engine/Plugins/Mutable/Source/CustomizableObjectEditor/Private/MuCOE/Nodes/CustomizableObjectNodeAnimationPose.h
@@ -26,10 +26,7 @@ public:
 
 	UPROPERTY(EditAnywhere, Category=NoCategory)
 	TObjectPtr<UPoseAsset> PoseAsset;
-
-	// UObject interface.
-	virtual void PostEditChangeProperty(FPropertyChangedEvent& PropertyChangedEvent) override;
-
+	
 	// Begin EdGraphNode interface
 	UEdGraphPin* GetInputMeshPin() const;
 	UEdGraphPin* GetTablePosePin() const;
@@ -38,8 +35,9 @@ public:
 	// End EdGraphNode interface
 
 	// UCustomizableObjectNode interface
-	void AllocateDefaultPins(UCustomizableObjectNodeRemapPins* RemapPins) override;
-
+	virtual void AllocateDefaultPins(UCustomizableObjectNodeRemapPins* RemapPins) override;
+	virtual void BackwardsCompatibleFixup(int32 CustomizableObjectCustomVersion) override;
+	
 	// Retrieve the pose information from the PoseAsset
 	static void StaticRetrievePoseInformation(UPoseAsset* PoseAsset, class USkeletalMesh* RefSkeletalMesh, TArray<FName>& OutArrayBoneName, TArray<FTransform>& OutArrayTransform);
 };
diff --git a/Engine/Plugins/Mutable/Source/CustomizableObjectEditor/Private/MuCOE/Nodes/CustomizableObjectNodeComponentMeshBase.cpp b/Engine/Plugins/Mutable/Source/CustomizableObjectEditor/Private/MuCOE/Nodes/CustomizableObjectNodeComponentMeshBase.cpp
--- a/Engine/Plugins/Mutable/Source/CustomizableObjectEditor/Private/MuCOE/Nodes/CustomizableObjectNodeComponentMeshBase.cpp
+++ b/Engine/Plugins/Mutable/Source/CustomizableObjectEditor/Private/MuCOE/Nodes/CustomizableObjectNodeComponentMeshBase.cpp
@@ -20,7 +20,7 @@ void UCustomizableObjectNodeComponentMeshBase::AllocateDefaultPins(UCustomizable
 		LODPins.Add(Pin);
 	}
 	
-	OutputPin = CustomCreatePin(EGPD_Output, Schema->PC_Component, TEXT("Component"));
+	OutputPin = CustomCreatePin(EGPD_Output, Schema->PC_Component, FName(TEXT("Component")));
 }
 
 
diff --git a/Engine/Plugins/Mutable/Source/CustomizableObjectEditor/Private/MuCOE/Nodes/CustomizableObjectNodeFloatArithmeticOp.h b/Engine/Plugins/Mutable/Source/CustomizableObjectEditor/Private/MuCOE/Nodes/CustomizableObjectNodeFloatArithmeticOp.h
--- a/Engine/Plugins/Mutable/Source/CustomizableObjectEditor/Private/MuCOE/Nodes/CustomizableObjectNodeFloatArithmeticOp.h
+++ b/Engine/Plugins/Mutable/Source/CustomizableObjectEditor/Private/MuCOE/Nodes/CustomizableObjectNodeFloatArithmeticOp.h
@@ -32,13 +32,13 @@ public:
 	UCustomizableObjectNodeFloatArithmeticOp();
 
 	// Begin EdGraphNode interface
-	FText GetNodeTitle(ENodeTitleType::Type TitleType) const override;
-	FLinearColor GetNodeTitleColor() const override;
-	FText GetTooltipText() const override;
-	bool ShouldOverridePinNames() const override { return true; }
+	virtual FText GetNodeTitle(ENodeTitleType::Type TitleType) const override;
+	virtual FLinearColor GetNodeTitleColor() const override;
+	virtual FText GetTooltipText() const override;
+	virtual bool ShouldOverridePinNames() const override { return true; }
 
 	// UCustomizableObjectNode interface
-	void AllocateDefaultPins(UCustomizableObjectNodeRemapPins* RemapPins) override;
+	virtual void AllocateDefaultPins(UCustomizableObjectNodeRemapPins* RemapPins) override;
 
 	// First Operand
 	UEdGraphPin* XPin() const
@@ -56,7 +56,7 @@ public:
 	EFloatArithmeticOperation Operation;
 
 	// UCustomizableObjectNode interface
-	virtual bool IsAffectedByLOD() const { return false; }
+	virtual bool IsAffectedByLOD() const override { return false; }
 
 };
 
diff --git a/Engine/Plugins/Mutable/Source/CustomizableObjectEditor/Private/MuCOE/Nodes/CustomizableObjectNodeMeshMorphStackApplication.cpp b/Engine/Plugins/Mutable/Source/CustomizableObjectEditor/Private/MuCOE/Nodes/CustomizableObjectNodeMeshMorphStackApplication.cpp
--- a/Engine/Plugins/Mutable/Source/CustomizableObjectEditor/Private/MuCOE/Nodes/CustomizableObjectNodeMeshMorphStackApplication.cpp
+++ b/Engine/Plugins/Mutable/Source/CustomizableObjectEditor/Private/MuCOE/Nodes/CustomizableObjectNodeMeshMorphStackApplication.cpp
@@ -3,6 +3,7 @@
 #include "MuCOE/Nodes/CustomizableObjectNodeMeshMorphStackApplication.h"
 
 #include "Engine/SkeletalMesh.h"
+#include "MuCO/CustomizableObjectCustomVersion.h"
 #include "MuCOE/EdGraphSchema_CustomizableObject.h"
 #include "MuCOE/GraphTraversal.h"
 #include "MuCOE/ICustomizableObjectEditor.h"
@@ -17,14 +18,40 @@ class UCustomizableObjectNodeRemapPins;
 
 void UCustomizableObjectNodeMeshMorphStackApplication::AllocateDefaultPins(UCustomizableObjectNodeRemapPins* RemapPins)
 {
-	const UEdGraphSchema_CustomizableObject* Schema = GetDefault<UEdGraphSchema_CustomizableObject>();
-
 	// Input pins
-	CustomCreatePin(EGPD_Input, Schema->PC_Mesh, TEXT("InMesh"));
-	CustomCreatePin(EGPD_Input, Schema->PC_Stack, TEXT("Stack"));
-
+	
+	CustomCreatePinSimple(EGPD_Input, UEdGraphSchema_CustomizableObject::PC_Mesh);
+	CustomCreatePinSimple(EGPD_Input, UEdGraphSchema_CustomizableObject::PC_Stack);
+	
 	// Output pins
-	CustomCreatePin(EGPD_Output, Schema->PC_Mesh, TEXT("Result Mesh"));
+	
+	CustomCreatePinSimple(EGPD_Output, UEdGraphSchema_CustomizableObject::PC_Mesh);
+}
+
+
+void UCustomizableObjectNodeMeshMorphStackApplication::BackwardsCompatibleFixup(int32 CustomizableObjectCustomVersion)
+{
+	Super::BackwardsCompatibleFixup(CustomizableObjectCustomVersion);
+
+	if (CustomizableObjectCustomVersion == FCustomizableObjectCustomVersion::UpdatedNodesPinName3)
+	{
+		if (UEdGraphPin* InMeshPin = FindPin(TEXT("InMesh"), EGPD_Input))
+		{
+			InMeshPin->PinName = TEXT("Mesh");
+			InMeshPin->PinFriendlyName = LOCTEXT("Mesh_Pin_Category", "Mesh");
+		}
+
+		if (UEdGraphPin* InStackPin = FindPin(TEXT("Stack"), EGPD_Input))
+		{
+			InStackPin->PinFriendlyName = LOCTEXT("Stack_Pin_Category", "Stack");
+		}
+		
+		if (UEdGraphPin* OutMeshPin = FindPin(TEXT("Result Mesh"), EGPD_Output))
+		{
+			OutMeshPin->PinName = TEXT("Mesh");
+			OutMeshPin->PinFriendlyName = LOCTEXT("Mesh_Pin_Category", "Mesh");
+		}
+	}
 }
 
 
@@ -36,8 +63,7 @@ FText UCustomizableObjectNodeMeshMorphStackApplication::GetNodeTitle(ENodeTitleT
 
 FLinearColor UCustomizableObjectNodeMeshMorphStackApplication::GetNodeTitleColor() const
 {
-	const UEdGraphSchema_CustomizableObject* Schema = GetDefault<UEdGraphSchema_CustomizableObject>();
-	return Schema->GetPinTypeColor(Schema->PC_Mesh);
+	return UEdGraphSchema_CustomizableObject::GetPinTypeColor(UEdGraphSchema_CustomizableObject::PC_Mesh);
 }
 
 
@@ -49,13 +75,15 @@ FText UCustomizableObjectNodeMeshMorphStackApplication::GetTooltipText() const
 
 UEdGraphPin* UCustomizableObjectNodeMeshMorphStackApplication::GetMeshPin() const
 {
-	return FindPin(TEXT("InMesh"), EGPD_Input);
+	const FName PinName = UEdGraphSchema_CustomizableObject::GetPinCategoryName(UEdGraphSchema_CustomizableObject::PC_Mesh);
+	return FindPin(PinName, EGPD_Input);
 }
 
 
 UEdGraphPin* UCustomizableObjectNodeMeshMorphStackApplication::GetStackPin() const
 {
-	return FindPin(TEXT("Stack"), EGPD_Input);
+	const FName PinName = UEdGraphSchema_CustomizableObject::GetPinCategoryName(UEdGraphSchema_CustomizableObject::PC_Stack);
+	return FindPin(PinName, EGPD_Input);
 }
 
 TArray<FString> UCustomizableObjectNodeMeshMorphStackApplication::GetMorphList() const
diff --git a/Engine/Plugins/Mutable/Source/CustomizableObjectEditor/Private/MuCOE/Nodes/CustomizableObjectNodeMeshMorphStackApplication.h b/Engine/Plugins/Mutable/Source/CustomizableObjectEditor/Private/MuCOE/Nodes/CustomizableObjectNodeMeshMorphStackApplication.h
--- a/Engine/Plugins/Mutable/Source/CustomizableObjectEditor/Private/MuCOE/Nodes/CustomizableObjectNodeMeshMorphStackApplication.h
+++ b/Engine/Plugins/Mutable/Source/CustomizableObjectEditor/Private/MuCOE/Nodes/CustomizableObjectNodeMeshMorphStackApplication.h
@@ -24,7 +24,8 @@ public:
 
 	// UCustomizableObjectNode interface
 	virtual void AllocateDefaultPins(UCustomizableObjectNodeRemapPins* RemapPins) override;
-
+	virtual void BackwardsCompatibleFixup(int32 CustomizableObjectCustomVersion) override;
+	
 	// Own interface
 	/** Fills the list with all the morphs */
 	TArray<FString> GetMorphList() const;
diff --git a/Engine/Plugins/Mutable/Source/CustomizableObjectEditor/Private/MuCOE/Nodes/CustomizableObjectNodeObjectGroup.cpp b/Engine/Plugins/Mutable/Source/CustomizableObjectEditor/Private/MuCOE/Nodes/CustomizableObjectNodeObjectGroup.cpp
--- a/Engine/Plugins/Mutable/Source/CustomizableObjectEditor/Private/MuCOE/Nodes/CustomizableObjectNodeObjectGroup.cpp
+++ b/Engine/Plugins/Mutable/Source/CustomizableObjectEditor/Private/MuCOE/Nodes/CustomizableObjectNodeObjectGroup.cpp
@@ -109,8 +109,7 @@ FText UCustomizableObjectNodeObjectGroup::GetNodeTitle(ENodeTitleType::Type Titl
 
 FLinearColor UCustomizableObjectNodeObjectGroup::GetNodeTitleColor() const
 {
-	const UEdGraphSchema_CustomizableObject* Schema = GetDefault<UEdGraphSchema_CustomizableObject>();
-	return Schema->GetPinTypeColor(Schema->PC_Object);
+	return UEdGraphSchema_CustomizableObject::GetPinTypeColor(UEdGraphSchema_CustomizableObject::PC_Object);
 }
 
 
diff --git a/Engine/Plugins/Mutable/Source/CustomizableObjectEditor/Private/MuCOE/Nodes/CustomizableObjectNodeObjectGroup.h b/Engine/Plugins/Mutable/Source/CustomizableObjectEditor/Private/MuCOE/Nodes/CustomizableObjectNodeObjectGroup.h
--- a/Engine/Plugins/Mutable/Source/CustomizableObjectEditor/Private/MuCOE/Nodes/CustomizableObjectNodeObjectGroup.h
+++ b/Engine/Plugins/Mutable/Source/CustomizableObjectEditor/Private/MuCOE/Nodes/CustomizableObjectNodeObjectGroup.h
@@ -42,19 +42,19 @@ public:
 	int32 SocketPriority = 0;
 
 	// UObject interface.
-	void Serialize(FArchive& Ar) override;
-	void PostEditChangeProperty(FPropertyChangedEvent& PropertyChangedEvent) override;
+	virtual void Serialize(FArchive& Ar) override;
+	virtual void PostEditChangeProperty(FPropertyChangedEvent& PropertyChangedEvent) override;
 
 	// EdGraphNode interface
-	FText GetNodeTitle(ENodeTitleType::Type TitleType) const override;
-	FLinearColor GetNodeTitleColor() const override;
-	FText GetTooltipText() const override;
+	virtual FText GetNodeTitle(ENodeTitleType::Type TitleType) const override;
+	virtual FLinearColor GetNodeTitleColor() const override;
+	virtual FText GetTooltipText() const override;
 	virtual bool GetCanRenameNode() const override { return true; }
 	virtual void OnRenameNode(const FString& NewName) override;
 
 	// UCustomizableObjectNode interface
-	void AllocateDefaultPins(UCustomizableObjectNodeRemapPins* RemapPins) override;
-	bool IsSingleOutputNode() const override;
+	virtual void AllocateDefaultPins(UCustomizableObjectNodeRemapPins* RemapPins) override;
+	virtual bool IsSingleOutputNode() const override;
 
 	// Own interface
 	UEdGraphPin* ObjectsPin() const;
diff --git a/Engine/Plugins/Mutable/Source/CustomizableObjectEditor/Private/MuCOE/Nodes/CustomizableObjectNodePassThroughTexture.cpp b/Engine/Plugins/Mutable/Source/CustomizableObjectEditor/Private/MuCOE/Nodes/CustomizableObjectNodePassThroughTexture.cpp
--- a/Engine/Plugins/Mutable/Source/CustomizableObjectEditor/Private/MuCOE/Nodes/CustomizableObjectNodePassThroughTexture.cpp
+++ b/Engine/Plugins/Mutable/Source/CustomizableObjectEditor/Private/MuCOE/Nodes/CustomizableObjectNodePassThroughTexture.cpp
@@ -24,15 +24,20 @@ void UCustomizableObjectNodePassThroughTexture::BackwardsCompatibleFixup(int32 C
 			Texture_DEPRECATED = nullptr;
 		}
 	}
+	else if (CustomizableObjectCustomVersion == FCustomizableObjectCustomVersion::UpdatedNodesPinName3)
+	{
+		if (UEdGraphPin* FoundPin = FindPin(TEXT("Texture")))
+		{
+			FoundPin->PinName = TEXT("PassThrough Texture");
+			FoundPin->PinFriendlyName = LOCTEXT("PassThrough_Image_Pin_Category", "PassThrough Texture");
+		}
+	}
 }
 
 
 void UCustomizableObjectNodePassThroughTexture::AllocateDefaultPins(UCustomizableObjectNodeRemapPins* RemapPins)
 {
-	const UEdGraphSchema_CustomizableObject* Schema = GetDefault<UEdGraphSchema_CustomizableObject>();
-
-	FString PinName = TEXT("Texture");
-	UEdGraphPin* PinImagePin = CustomCreatePin(EGPD_Output, Schema->PC_PassThroughImage, FName(*PinName));
+	UEdGraphPin* PinImagePin = CustomCreatePinSimple(EGPD_Output, UEdGraphSchema_CustomizableObject::PC_PassThroughImage);
 	PinImagePin->bDefaultValueIsIgnored = true;
 }
 
@@ -55,8 +60,7 @@ FText UCustomizableObjectNodePassThroughTexture::GetNodeTitle(ENodeTitleType::Ty
 
 FLinearColor UCustomizableObjectNodePassThroughTexture::GetNodeTitleColor() const
 {
-	const UEdGraphSchema_CustomizableObject* Schema = GetDefault<UEdGraphSchema_CustomizableObject>();
-	return Schema->GetPinTypeColor(Schema->PC_PassThroughImage);
+	return UEdGraphSchema_CustomizableObject::GetPinTypeColor(UEdGraphSchema_CustomizableObject::PC_PassThroughImage);
 }
 
 
diff --git a/Engine/Plugins/Mutable/Source/CustomizableObjectEditor/Private/MuCOE/Nodes/CustomizableObjectNodePassThroughTexture.h b/Engine/Plugins/Mutable/Source/CustomizableObjectEditor/Private/MuCOE/Nodes/CustomizableObjectNodePassThroughTexture.h
--- a/Engine/Plugins/Mutable/Source/CustomizableObjectEditor/Private/MuCOE/Nodes/CustomizableObjectNodePassThroughTexture.h
+++ b/Engine/Plugins/Mutable/Source/CustomizableObjectEditor/Private/MuCOE/Nodes/CustomizableObjectNodePassThroughTexture.h
@@ -18,15 +18,15 @@ public:
 
 	// UCustomizableObjectNode interface
 	virtual void BackwardsCompatibleFixup(int32 CustomizableObjectCustomVersion) override;
-	void AllocateDefaultPins(UCustomizableObjectNodeRemapPins* RemapPins) override;
+	virtual void AllocateDefaultPins(UCustomizableObjectNodeRemapPins* RemapPins) override;
 
 	// UCustomizableObjectNodeTextureBase interface
 	virtual TObjectPtr<UTexture> GetTexture() override;
 
 	// Begin EdGraphNode interface
-	FText GetNodeTitle(ENodeTitleType::Type TitleType) const override;
-	FLinearColor GetNodeTitleColor() const override;
-	FText GetTooltipText() const override;
+	virtual FText GetNodeTitle(ENodeTitleType::Type TitleType) const override;
+	virtual FLinearColor GetNodeTitleColor() const override;
+	virtual FText GetTooltipText() const override;
 
 private:
 
diff --git a/Engine/Plugins/Mutable/Source/CustomizableObjectEditor/Private/MuCOE/Nodes/CustomizableObjectNodePassThroughTextureVariation.cpp b/Engine/Plugins/Mutable/Source/CustomizableObjectEditor/Private/MuCOE/Nodes/CustomizableObjectNodePassThroughTextureVariation.cpp
--- a/Engine/Plugins/Mutable/Source/CustomizableObjectEditor/Private/MuCOE/Nodes/CustomizableObjectNodePassThroughTextureVariation.cpp
+++ b/Engine/Plugins/Mutable/Source/CustomizableObjectEditor/Private/MuCOE/Nodes/CustomizableObjectNodePassThroughTextureVariation.cpp
@@ -2,6 +2,7 @@
 
 #include "MuCOE/Nodes/CustomizableObjectNodePassThroughTextureVariation.h"
 
+#include "MuCO/CustomizableObjectCustomVersion.h"
 #include "MuCOE/EdGraphSchema_CustomizableObject.h"
 
 #define LOCTEXT_NAMESPACE "CustomizableObjectEditor"
@@ -11,7 +12,7 @@ void UCustomizableObjectNodePassThroughTextureVariation::AllocateDefaultPins(UCu
 {
 	const UEdGraphSchema_CustomizableObject* Schema = GetDefault<UEdGraphSchema_CustomizableObject>();
 
-	UEdGraphPin* OutputPin = CustomCreatePin(EGPD_Output, Schema->PC_PassThroughImage, FName(TEXT("Texture")));
+	UEdGraphPin* OutputPin = CustomCreatePinSimple(EGPD_Output, UEdGraphSchema_CustomizableObject::PC_PassThroughImage);
 	OutputPin->bDefaultValueIsIgnored = true;
 
 	for (int VariationIndex = Variations.Num() - 1; VariationIndex >= 0; --VariationIndex)
@@ -29,6 +30,21 @@ void UCustomizableObjectNodePassThroughTextureVariation::AllocateDefaultPins(UCu
 }
 
 
+void UCustomizableObjectNodePassThroughTextureVariation::BackwardsCompatibleFixup(int32 CustomizableObjectCustomVersion)
+{
+	Super::BackwardsCompatibleFixup(CustomizableObjectCustomVersion);
+
+	if (CustomizableObjectCustomVersion == FCustomizableObjectCustomVersion::UpdatedNodesPinName3)
+	{
+		if (UEdGraphPin* Pin = FindPin(TEXT("Texture")))
+		{
+			Pin->PinName = TEXT("PassThrough Texture");
+			Pin->PinFriendlyName = LOCTEXT("PassThrough_Image_Pin_Category", "PassThrough Texture");
+		}
+	}
+}
+
+
 FText UCustomizableObjectNodePassThroughTextureVariation::GetNodeTitle(ENodeTitleType::Type TitleType) const
 {
 	return LOCTEXT("PassThrough_Texture_Variation", "PassThrough Texture Variation");
diff --git a/Engine/Plugins/Mutable/Source/CustomizableObjectEditor/Private/MuCOE/Nodes/CustomizableObjectNodePassThroughTextureVariation.h b/Engine/Plugins/Mutable/Source/CustomizableObjectEditor/Private/MuCOE/Nodes/CustomizableObjectNodePassThroughTextureVariation.h
--- a/Engine/Plugins/Mutable/Source/CustomizableObjectEditor/Private/MuCOE/Nodes/CustomizableObjectNodePassThroughTextureVariation.h
+++ b/Engine/Plugins/Mutable/Source/CustomizableObjectEditor/Private/MuCOE/Nodes/CustomizableObjectNodePassThroughTextureVariation.h
@@ -18,4 +18,5 @@ public:
 
 	// UCustomizableObjectNode interface
 	virtual void AllocateDefaultPins(UCustomizableObjectNodeRemapPins* RemapPins) override;
+	virtual void BackwardsCompatibleFixup(int32 CustomizableObjectCustomVersion) override;
 };
diff --git a/Engine/Plugins/Mutable/Source/CustomizableObjectEditor/Private/MuCOE/Nodes/CustomizableObjectNodeSwitchBase.cpp b/Engine/Plugins/Mutable/Source/CustomizableObjectEditor/Private/MuCOE/Nodes/CustomizableObjectNodeSwitchBase.cpp
--- a/Engine/Plugins/Mutable/Source/CustomizableObjectEditor/Private/MuCOE/Nodes/CustomizableObjectNodeSwitchBase.cpp
+++ b/Engine/Plugins/Mutable/Source/CustomizableObjectEditor/Private/MuCOE/Nodes/CustomizableObjectNodeSwitchBase.cpp
@@ -110,8 +110,7 @@ FText UCustomizableObjectNodeSwitchBase::GetNodeTitle(ENodeTitleType::Type Title
 
 FLinearColor UCustomizableObjectNodeSwitchBase::GetNodeTitleColor() const
 {
-	const UEdGraphSchema_CustomizableObject* Schema = GetDefault<UEdGraphSchema_CustomizableObject>();
-	return Schema->GetPinTypeColor(GetCategory());
+	return UEdGraphSchema_CustomizableObject::GetPinTypeColor(GetCategory());
 }
 
 
diff --git a/Engine/Plugins/Mutable/Source/CustomizableObjectEditor/Private/MuCOE/Nodes/CustomizableObjectNodeTextureBinarise.cpp b/Engine/Plugins/Mutable/Source/CustomizableObjectEditor/Private/MuCOE/Nodes/CustomizableObjectNodeTextureBinarise.cpp
--- a/Engine/Plugins/Mutable/Source/CustomizableObjectEditor/Private/MuCOE/Nodes/CustomizableObjectNodeTextureBinarise.cpp
+++ b/Engine/Plugins/Mutable/Source/CustomizableObjectEditor/Private/MuCOE/Nodes/CustomizableObjectNodeTextureBinarise.cpp
@@ -12,19 +12,14 @@ class UCustomizableObjectNodeRemapPins;
 
 void UCustomizableObjectNodeTextureBinarise::AllocateDefaultPins(UCustomizableObjectNodeRemapPins* RemapPins)
 {
-	const UEdGraphSchema_CustomizableObject* Schema = GetDefault<UEdGraphSchema_CustomizableObject>();
-
-	FString PinName = TEXT("Texture");
-	UEdGraphPin* OutputPin = CustomCreatePin(EGPD_Output, Schema->PC_Image, FName(*PinName));
+	UEdGraphPin* OutputPin = CustomCreatePinSimple(EGPD_Output, UEdGraphSchema_CustomizableObject::PC_Image);
 	OutputPin->bDefaultValueIsIgnored = true;
-
-	PinName = TEXT("Base Texture");
-	UEdGraphPin* ImagePin = CustomCreatePin(EGPD_Input, Schema->PC_Image, FName(*PinName));
-
+	
+	UEdGraphPin* ImagePin = CustomCreatePinSimple(EGPD_Input, UEdGraphSchema_CustomizableObject::PC_Image);
 	BaseImagePinReference = FEdGraphPinReference(ImagePin);
 	
-	PinName = TEXT("Threshold");
-	UEdGraphPin* ThresholdPin = CustomCreatePin(EGPD_Input, Schema->PC_Float, FName(*PinName));
+	const FName PinName = TEXT("Threshold");
+	CustomCreatePin(EGPD_Input, UEdGraphSchema_CustomizableObject::PC_Float, PinName);
 }
 
 
@@ -39,14 +34,26 @@ void UCustomizableObjectNodeTextureBinarise::BackwardsCompatibleFixup(int32 Cust
 			UCustomizableObjectNode::ReconstructNode();
 		}
 	}
-
-	if (CustomizableObjectCustomVersion == FCustomizableObjectCustomVersion::FixPinsNamesImageToTexture2)
+	else if (CustomizableObjectCustomVersion == FCustomizableObjectCustomVersion::FixPinsNamesImageToTexture2)
 	{
 		if (UEdGraphPin* TexturePin = FindPin(TEXT("Base Image"))) {
 			TexturePin->PinName = TEXT("Base Texture");
 			UCustomizableObjectNode::ReconstructNode();
 		}
 	}
+	else if (CustomizableObjectCustomVersion == FCustomizableObjectCustomVersion::UpdatedNodesPinName3)
+	{
+		if (UEdGraphPin* InputTexturePin = FindPin(TEXT("Base Texture"), EGPD_Input))
+		{
+			InputTexturePin->PinName = TEXT("Texture");
+			InputTexturePin->PinFriendlyName = LOCTEXT("Image_Pin_Category", "Texture");
+		}
+
+		if (UEdGraphPin* OutputTexturePin = FindPin(TEXT("Texture"), EGPD_Output))
+		{
+			OutputTexturePin->PinFriendlyName = LOCTEXT("Image_Pin_Category", "Texture");
+		}
+	}
 }
 
 
@@ -64,8 +71,7 @@ FText UCustomizableObjectNodeTextureBinarise::GetNodeTitle(ENodeTitleType::Type
 
 FLinearColor UCustomizableObjectNodeTextureBinarise::GetNodeTitleColor() const
 {
-	const UEdGraphSchema_CustomizableObject* Schema = GetDefault<UEdGraphSchema_CustomizableObject>();
-	return Schema->GetPinTypeColor(Schema->PC_Image);
+	return UEdGraphSchema_CustomizableObject::GetPinTypeColor(UEdGraphSchema_CustomizableObject::PC_Image);
 }
 
 
@@ -83,7 +89,8 @@ void UCustomizableObjectNodeTextureBinarise::Serialize(FArchive& Ar)
 
 	if (Ar.CustomVer(FCustomizableObjectCustomVersion::GUID) < FCustomizableObjectCustomVersion::PinsNamesImageToTexture)
 	{
-		BaseImagePinReference = FEdGraphPinReference(FindPin(TEXT("Base Image")));
+		const FName TargetPinName = UEdGraphSchema_CustomizableObject::GetPinCategoryName(UEdGraphSchema_CustomizableObject::PC_Image);
+		BaseImagePinReference = FEdGraphPinReference(FindPin(TargetPinName, EGPD_Input));
 	}
 }
 
diff --git a/Engine/Plugins/Mutable/Source/CustomizableObjectEditor/Private/MuCOE/Nodes/CustomizableObjectNodeTextureBinarise.h b/Engine/Plugins/Mutable/Source/CustomizableObjectEditor/Private/MuCOE/Nodes/CustomizableObjectNodeTextureBinarise.h
--- a/Engine/Plugins/Mutable/Source/CustomizableObjectEditor/Private/MuCOE/Nodes/CustomizableObjectNodeTextureBinarise.h
+++ b/Engine/Plugins/Mutable/Source/CustomizableObjectEditor/Private/MuCOE/Nodes/CustomizableObjectNodeTextureBinarise.h
@@ -18,9 +18,9 @@ public:
 	GENERATED_BODY()
 
 	// Begin EdGraphNode interface
-	FText GetNodeTitle(ENodeTitleType::Type TittleType) const override;
-	FLinearColor GetNodeTitleColor() const override;
-	FText GetTooltipText() const override;
+	virtual FText GetNodeTitle(ENodeTitleType::Type TittleType) const override;
+	virtual FLinearColor GetNodeTitleColor() const override;
+	virtual FText GetTooltipText() const override;
 	virtual void Serialize(FArchive& Ar) override;
 	
 	// UCustomizableObjectNode interface
diff --git a/Engine/Plugins/Mutable/Source/CustomizableObjectEditor/Private/MuCOE/Nodes/CustomizableObjectNodeTextureInvert.cpp b/Engine/Plugins/Mutable/Source/CustomizableObjectEditor/Private/MuCOE/Nodes/CustomizableObjectNodeTextureInvert.cpp
--- a/Engine/Plugins/Mutable/Source/CustomizableObjectEditor/Private/MuCOE/Nodes/CustomizableObjectNodeTextureInvert.cpp
+++ b/Engine/Plugins/Mutable/Source/CustomizableObjectEditor/Private/MuCOE/Nodes/CustomizableObjectNodeTextureInvert.cpp
@@ -12,15 +12,10 @@ class UCustomizableObjectNodeRemapPins;
 
 void UCustomizableObjectNodeTextureInvert::AllocateDefaultPins(UCustomizableObjectNodeRemapPins* RemapPins)
 {
-	const UEdGraphSchema_CustomizableObject* Schema = GetDefault<UEdGraphSchema_CustomizableObject>();
-
-	FString PinName = TEXT("Texture");
-	UEdGraphPin* OutputPin = CustomCreatePin(EGPD_Output, Schema->PC_Image, FName(*PinName));
+	UEdGraphPin* OutputPin = CustomCreatePinSimple(EGPD_Output, UEdGraphSchema_CustomizableObject::PC_Image);
 	OutputPin->bDefaultValueIsIgnored = true;
-
-	PinName = TEXT("Base Texture");
-	UEdGraphPin* ImagePin = CustomCreatePin(EGPD_Input, Schema->PC_Image, FName(*PinName));
-
+	
+	UEdGraphPin* ImagePin = CustomCreatePinSimple(EGPD_Input, UEdGraphSchema_CustomizableObject::PC_Image);
 	BaseImagePinReference = FEdGraphPinReference(ImagePin);
 }
 
@@ -33,8 +28,7 @@ void UCustomizableObjectNodeTextureInvert::BackwardsCompatibleFixup(int32 Custom
 	{
 		BaseImagePinReference = FEdGraphPinReference(FindPin(TEXT("Base Image")));
 	}
-	
-	if (CustomizableObjectCustomVersion == FCustomizableObjectCustomVersion::FixPinsNamesImageToTexture2)
+	else if (CustomizableObjectCustomVersion == FCustomizableObjectCustomVersion::FixPinsNamesImageToTexture2)
 	{
 		bool Replaced = false;
 		if (UEdGraphPin* TexturePin = FindPin(TEXT("Image")))
@@ -54,6 +48,19 @@ void UCustomizableObjectNodeTextureInvert::BackwardsCompatibleFixup(int32 Custom
 			UCustomizableObjectNode::ReconstructNode();
 		}
 	}
+	else if (CustomizableObjectCustomVersion == FCustomizableObjectCustomVersion::UpdatedNodesPinName3)
+	{
+		if (UEdGraphPin* InputTexturePin = FindPin(TEXT("Base Texture"), EGPD_Input))
+		{
+			InputTexturePin->PinName = TEXT("Texture");
+			InputTexturePin->PinFriendlyName = LOCTEXT("Image_Pin_Category", "Texture");
+		}
+
+		if (UEdGraphPin* OutputTexturePin = FindPin(TEXT("Texture"), EGPD_Output))
+		{
+			OutputTexturePin->PinFriendlyName = LOCTEXT("Image_Pin_Category", "Texture");
+		}
+	}
 }
 
 
@@ -71,8 +78,7 @@ FText UCustomizableObjectNodeTextureInvert::GetNodeTitle(ENodeTitleType::Type Ti
 
 FLinearColor UCustomizableObjectNodeTextureInvert::GetNodeTitleColor() const
 {
-	const UEdGraphSchema_CustomizableObject* Schema = GetDefault<UEdGraphSchema_CustomizableObject>();
-	return Schema->GetPinTypeColor(Schema->PC_Image);
+	return UEdGraphSchema_CustomizableObject::GetPinTypeColor(UEdGraphSchema_CustomizableObject::PC_Image);
 }
 
 
diff --git a/Engine/Plugins/Mutable/Source/CustomizableObjectEditor/Private/MuCOE/Nodes/CustomizableObjectNodeTextureSample.cpp b/Engine/Plugins/Mutable/Source/CustomizableObjectEditor/Private/MuCOE/Nodes/CustomizableObjectNodeTextureSample.cpp
--- a/Engine/Plugins/Mutable/Source/CustomizableObjectEditor/Private/MuCOE/Nodes/CustomizableObjectNodeTextureSample.cpp
+++ b/Engine/Plugins/Mutable/Source/CustomizableObjectEditor/Private/MuCOE/Nodes/CustomizableObjectNodeTextureSample.cpp
@@ -2,6 +2,7 @@
 
 #include "MuCOE/Nodes/CustomizableObjectNodeTextureSample.h"
 
+#include "MuCO/CustomizableObjectCustomVersion.h"
 #include "MuCOE/EdGraphSchema_CustomizableObject.h"
 
 class UCustomizableObjectNodeRemapPins;
@@ -17,38 +18,43 @@ UCustomizableObjectNodeTextureSample::UCustomizableObjectNodeTextureSample()
 }
 
 
-void UCustomizableObjectNodeTextureSample::PostEditChangeProperty(FPropertyChangedEvent& PropertyChangedEvent)
-{
-	//UProperty* PropertyThatChanged = PropertyChangedEvent.Property;
-	//if ( PropertyThatChanged && PropertyThatChanged->GetName() == TEXT("NumLODs") )
-	{
-	}
-
-	Super::PostEditChangeProperty(PropertyChangedEvent);
-}
-
-
 void UCustomizableObjectNodeTextureSample::AllocateDefaultPins(UCustomizableObjectNodeRemapPins* RemapPins)
 {
-	const UEdGraphSchema_CustomizableObject* Schema = GetDefault<UEdGraphSchema_CustomizableObject>();
-
-	FString PinName = TEXT("Color");
-	UEdGraphPin* ColorPin = CustomCreatePin(EGPD_Output, Schema->PC_Color, FName(*PinName));
+	// Output
+	
+	UEdGraphPin* ColorPin = CustomCreatePinSimple(EGPD_Output, UEdGraphSchema_CustomizableObject::PC_Color);
 	ColorPin->bDefaultValueIsIgnored = true;
-
-	PinName = TEXT("Texture");
-	UEdGraphPin* TexturePin = CustomCreatePin(EGPD_Input, Schema->PC_Image, FName(*PinName));
+	
+	// Inputs
+	
+	UEdGraphPin* TexturePin = CustomCreatePinSimple(EGPD_Input, UEdGraphSchema_CustomizableObject::PC_Image);
 	TexturePin->bDefaultValueIsIgnored = true;
 
-	PinName = TEXT("X");
-	UEdGraphPin* XPin = CustomCreatePin(EGPD_Input, Schema->PC_Float, FName(*PinName));
+	UEdGraphPin* XPin = CustomCreatePin(EGPD_Input, UEdGraphSchema_CustomizableObject::PC_Float,  FName(TEXT("X")));
 	XPin->bDefaultValueIsIgnored = true;
 
-	PinName = TEXT("Y");
-	UEdGraphPin* YPin = CustomCreatePin(EGPD_Input, Schema->PC_Float, FName(*PinName));
+	UEdGraphPin* YPin = CustomCreatePin(EGPD_Input, UEdGraphSchema_CustomizableObject::PC_Float, FName(TEXT("Y")));
 	YPin->bDefaultValueIsIgnored = true;
 }
 
+void UCustomizableObjectNodeTextureSample::BackwardsCompatibleFixup(int32 CustomizableObjectCustomVersion)
+{
+	Super::BackwardsCompatibleFixup(CustomizableObjectCustomVersion);
+
+	if (CustomizableObjectCustomVersion == FCustomizableObjectCustomVersion::UpdatedNodesPinName3)
+	{
+		if (UEdGraphPin* ColorOutput = FindPin(TEXT("Color"), EGPD_Output))
+		{
+			ColorOutput->PinFriendlyName = LOCTEXT("Color_Pin_Category", "Color");;
+		}
+		
+		if (UEdGraphPin* ImageInput = FindPin(TEXT("Texture"), EGPD_Input))
+		{
+			ImageInput->PinFriendlyName = LOCTEXT("Image_Pin_Category", "Texture");
+		}
+	}
+}
+
 
 FText UCustomizableObjectNodeTextureSample::GetNodeTitle(ENodeTitleType::Type TitleType) const
 {
@@ -58,8 +64,7 @@ FText UCustomizableObjectNodeTextureSample::GetNodeTitle(ENodeTitleType::Type Ti
 
 FLinearColor UCustomizableObjectNodeTextureSample::GetNodeTitleColor() const
 {
-	const UEdGraphSchema_CustomizableObject* Schema = GetDefault<UEdGraphSchema_CustomizableObject>();
-	return Schema->GetPinTypeColor(Schema->PC_Color);
+	return UEdGraphSchema_CustomizableObject::GetPinTypeColor(UEdGraphSchema_CustomizableObject::PC_Color);
 }
 
 
diff --git a/Engine/Plugins/Mutable/Source/CustomizableObjectEditor/Private/MuCOE/Nodes/CustomizableObjectNodeTextureSample.h b/Engine/Plugins/Mutable/Source/CustomizableObjectEditor/Private/MuCOE/Nodes/CustomizableObjectNodeTextureSample.h
--- a/Engine/Plugins/Mutable/Source/CustomizableObjectEditor/Private/MuCOE/Nodes/CustomizableObjectNodeTextureSample.h
+++ b/Engine/Plugins/Mutable/Source/CustomizableObjectEditor/Private/MuCOE/Nodes/CustomizableObjectNodeTextureSample.h
@@ -2,6 +2,7 @@
 
 #pragma once
 
+#include "MuCOE/EdGraphSchema_CustomizableObject.h"
 #include "MuCOE/Nodes/CustomizableObjectNode.h"
 
 #include "CustomizableObjectNodeTextureSample.generated.h"
@@ -22,20 +23,19 @@ public:
 
 	UCustomizableObjectNodeTextureSample();
 
-	// UObject interface.
-	void PostEditChangeProperty(FPropertyChangedEvent& PropertyChangedEvent) override;
-
 	// Begin EdGraphNode interface
-	FText GetNodeTitle(ENodeTitleType::Type TitleType) const override;
-	FLinearColor GetNodeTitleColor() const override;
-	FText GetTooltipText() const override;
+	virtual FText GetNodeTitle(ENodeTitleType::Type TitleType) const override;
+	virtual FLinearColor GetNodeTitleColor() const override;
+	virtual FText GetTooltipText() const override;
 
 	// UCustomizableObjectNode interface
-	void AllocateDefaultPins(UCustomizableObjectNodeRemapPins* RemapPins) override;
-
+	virtual void AllocateDefaultPins(UCustomizableObjectNodeRemapPins* RemapPins) override;
+	virtual void BackwardsCompatibleFixup(int32 CustomizableObjectCustomVersion) override;
+	
 	const UEdGraphPin* TexturePin() const
 	{
-		return FindPin(TEXT("Texture"));
+		const FName PinName = UEdGraphSchema_CustomizableObject::GetPinCategoryName(UEdGraphSchema_CustomizableObject::PC_Image);
+		return FindPin(PinName);
 	}
 
 	const UEdGraphPin* XPin() const
diff --git a/Engine/Plugins/Mutable/Source/CustomizableObjectEditor/Private/MuCOE/Nodes/CustomizableObjectNodeTextureSaturate.cpp b/Engine/Plugins/Mutable/Source/CustomizableObjectEditor/Private/MuCOE/Nodes/CustomizableObjectNodeTextureSaturate.cpp
--- a/Engine/Plugins/Mutable/Source/CustomizableObjectEditor/Private/MuCOE/Nodes/CustomizableObjectNodeTextureSaturate.cpp
+++ b/Engine/Plugins/Mutable/Source/CustomizableObjectEditor/Private/MuCOE/Nodes/CustomizableObjectNodeTextureSaturate.cpp
@@ -2,56 +2,76 @@
 
 #include "MuCOE/Nodes/CustomizableObjectNodeTextureSaturate.h"
 
+#include "MuCO/CustomizableObjectCustomVersion.h"
 #include "MuCOE/EdGraphSchema_CustomizableObject.h"
 
 class UCustomizableObjectNodeRemapPins;
 
 #define LOCTEXT_NAMESPACE "CustomizableObjectEditor"
 
+
 void UCustomizableObjectNodeTextureSaturate::AllocateDefaultPins(UCustomizableObjectNodeRemapPins* RemapPins)
 {
-	const UEdGraphSchema_CustomizableObject* Schema = GetDefault<UEdGraphSchema_CustomizableObject>();
-
-	FString PinName = TEXT("Texture");
-	UEdGraphPin* OutputPin = CustomCreatePin(EGPD_Output, Schema->PC_Image, FName(*PinName));
+	UEdGraphPin* OutputPin = CustomCreatePinSimple(EGPD_Output, UEdGraphSchema_CustomizableObject::PC_Image);
 	OutputPin->bDefaultValueIsIgnored = true;
-
-	PinName = TEXT("Base Texture");
-	UEdGraphPin* ImagePin = CustomCreatePin(EGPD_Input, Schema->PC_Image, FName(*PinName));
-
+	
+	UEdGraphPin* ImagePin = CustomCreatePinSimple(EGPD_Input, UEdGraphSchema_CustomizableObject::PC_Image);
 	BaseImagePinReference = FEdGraphPinReference(ImagePin);
 	
-	PinName = TEXT("Factor");
-	UEdGraphPin* FactorPin = CustomCreatePin(EGPD_Input, Schema->PC_Float, FName(*PinName));
+	UEdGraphPin* FactorPin = CustomCreatePin(EGPD_Input, UEdGraphSchema_CustomizableObject::PC_Float, FName(TEXT("Factor")));
 	FactorPin->bDefaultValueIsIgnored = true;
-
 	FactorPinReference = FEdGraphPinReference(FactorPin);
 }
 
+
+void UCustomizableObjectNodeTextureSaturate::BackwardsCompatibleFixup(int32 CustomizableObjectCustomVersion)
+{
+	Super::BackwardsCompatibleFixup(CustomizableObjectCustomVersion);
+
+	if (CustomizableObjectCustomVersion == FCustomizableObjectCustomVersion::UpdatedNodesPinName3)
+	{
+		if (UEdGraphPin* InputTexturePin = FindPin(TEXT("Base Texture"), EGPD_Input))
+		{
+			InputTexturePin->PinName = TEXT("Texture");
+			InputTexturePin->PinFriendlyName = LOCTEXT("Image_Pin_Category", "Texture");
+		}
+
+		if (UEdGraphPin* OutputTexturePin = FindPin(TEXT("Texture"), EGPD_Output))
+		{
+			OutputTexturePin->PinFriendlyName = LOCTEXT("Image_Pin_Category", "Texture");
+		}
+	}
+}
+
+
 UEdGraphPin* UCustomizableObjectNodeTextureSaturate::GetBaseImagePin() const
 {
 	return BaseImagePinReference.Get();
 }
 
+
 UEdGraphPin* UCustomizableObjectNodeTextureSaturate::GetFactorPin() const
 {
 	return FactorPinReference.Get();
 }
 
+
 FText UCustomizableObjectNodeTextureSaturate::GetNodeTitle(ENodeTitleType::Type TitleType) const
 {
 	return LOCTEXT("Texture_Saturate", "Texture Saturate");
 }
 
+
 FLinearColor UCustomizableObjectNodeTextureSaturate::GetNodeTitleColor() const
 {
-	const UEdGraphSchema_CustomizableObject* Schema = GetDefault<UEdGraphSchema_CustomizableObject>();
-	return Schema->GetPinTypeColor(Schema->PC_Image);
+	return UEdGraphSchema_CustomizableObject::GetPinTypeColor(UEdGraphSchema_CustomizableObject::PC_Image);
 }
 
+
 FText UCustomizableObjectNodeTextureSaturate::GetTooltipText() const
 {
 	return LOCTEXT("Texture_Saturate_Tooltip", "Get the provided texture with its saturation adjusted based on the numerical input provided where 1 equals to full saturation and 0 to no saturation.");
 }
 
+
 #undef LOCTEXT_NAMESPACE
diff --git a/Engine/Plugins/Mutable/Source/CustomizableObjectEditor/Private/MuCOE/Nodes/CustomizableObjectNodeTextureSaturate.h b/Engine/Plugins/Mutable/Source/CustomizableObjectEditor/Private/MuCOE/Nodes/CustomizableObjectNodeTextureSaturate.h
--- a/Engine/Plugins/Mutable/Source/CustomizableObjectEditor/Private/MuCOE/Nodes/CustomizableObjectNodeTextureSaturate.h
+++ b/Engine/Plugins/Mutable/Source/CustomizableObjectEditor/Private/MuCOE/Nodes/CustomizableObjectNodeTextureSaturate.h
@@ -18,21 +18,21 @@ public:
 	GENERATED_BODY()
 
 	// Begin EdGraphNode interface
-	FText GetNodeTitle(ENodeTitleType::Type TittleType) const override;
-	FLinearColor GetNodeTitleColor() const override;
-	FText GetTooltipText() const override;
+	virtual FText GetNodeTitle(ENodeTitleType::Type TittleType) const override;
+	virtual FLinearColor GetNodeTitleColor() const override;
+	virtual FText GetTooltipText() const override;
 	
 	// UCustomizableObjectNode interface
-	void AllocateDefaultPins(UCustomizableObjectNodeRemapPins* RemapPins) override;
-
+	virtual void AllocateDefaultPins(UCustomizableObjectNodeRemapPins* RemapPins) override;
+	virtual void BackwardsCompatibleFixup(int32 CustomizableObjectCustomVersion) override;
+	
 	UEdGraphPin* GetBaseImagePin() const;
 	UEdGraphPin* GetFactorPin() const;
 
 private:
 	UPROPERTY()
 	FEdGraphPinReference BaseImagePinReference;
-
-
+	
 	UPROPERTY()
 	FEdGraphPinReference FactorPinReference;
 };
diff --git a/Engine/Plugins/Mutable/Source/CustomizableObjectEditor/Private/MuCOE/Nodes/CustomizableObjectNodeTextureTransform.cpp b/Engine/Plugins/Mutable/Source/CustomizableObjectEditor/Private/MuCOE/Nodes/CustomizableObjectNodeTextureTransform.cpp
--- a/Engine/Plugins/Mutable/Source/CustomizableObjectEditor/Private/MuCOE/Nodes/CustomizableObjectNodeTextureTransform.cpp
+++ b/Engine/Plugins/Mutable/Source/CustomizableObjectEditor/Private/MuCOE/Nodes/CustomizableObjectNodeTextureTransform.cpp
@@ -11,38 +11,48 @@ class UCustomizableObjectNodeRemapPins;
 
 void UCustomizableObjectNodeTextureTransform::AllocateDefaultPins(UCustomizableObjectNodeRemapPins* RemapPins)
 {
-	const UEdGraphSchema_CustomizableObject* Schema = GetDefault<UEdGraphSchema_CustomizableObject>();
-
-	FString PinName = TEXT("Texture");
-	UEdGraphPin* OutputPin = CustomCreatePin(EGPD_Output, Schema->PC_Image, FName(*PinName));
+	UEdGraphPin* OutputPin = CustomCreatePinSimple(EGPD_Output, UEdGraphSchema_CustomizableObject::PC_Image);
 	OutputPin->bDefaultValueIsIgnored = true;
-
-	PinName = TEXT("Base Texture");
-	UEdGraphPin* ImagePin = CustomCreatePin(EGPD_Input, Schema->PC_Image, FName(*PinName));
-
+	
+	
+	UEdGraphPin* ImagePin = CustomCreatePinSimple(EGPD_Input, UEdGraphSchema_CustomizableObject::PC_Image);
 	BaseImagePinReference = FEdGraphPinReference(ImagePin);
 	
-	PinName = TEXT("Offset X");
-	UEdGraphPin* OffsetXPin = CustomCreatePin(EGPD_Input, Schema->PC_Float, FName(*PinName));
+	UEdGraphPin* OffsetXPin = CustomCreatePin(EGPD_Input, UEdGraphSchema_CustomizableObject::PC_Float, FName(TEXT("Offset X")));
 	OffsetXPin->bDefaultValueIsIgnored = true;
 	
-	PinName = TEXT("Offset Y");
-	UEdGraphPin* OffsetYPin = CustomCreatePin(EGPD_Input, Schema->PC_Float, FName(*PinName));
+	UEdGraphPin* OffsetYPin = CustomCreatePin(EGPD_Input, UEdGraphSchema_CustomizableObject::PC_Float, FName(TEXT("Offset Y")));
 	OffsetYPin->bDefaultValueIsIgnored = true;
 
-	PinName = TEXT("Scale X");
-	UEdGraphPin* ScaleXPin = CustomCreatePin(EGPD_Input, Schema->PC_Float, FName(*PinName));
+	UEdGraphPin* ScaleXPin = CustomCreatePin(EGPD_Input, UEdGraphSchema_CustomizableObject::PC_Float, FName(TEXT("Scale X")));
 	ScaleXPin->bDefaultValueIsIgnored = true;
 
-	PinName = TEXT("Scale Y");
-	UEdGraphPin* ScaleYPin = CustomCreatePin(EGPD_Input, Schema->PC_Float, FName(*PinName));
+	UEdGraphPin* ScaleYPin = CustomCreatePin(EGPD_Input, UEdGraphSchema_CustomizableObject::PC_Float, FName(TEXT("Scale Y")));
 	ScaleYPin->bDefaultValueIsIgnored = true;
 	
-	PinName = TEXT("Rotation");
-	UEdGraphPin* RotationPin = CustomCreatePin(EGPD_Input, Schema->PC_Float, FName(*PinName));
+	UEdGraphPin* RotationPin = CustomCreatePin(EGPD_Input, UEdGraphSchema_CustomizableObject::PC_Float, FName(TEXT("Rotation")));
 	RotationPin->bDefaultValueIsIgnored = true;
 }
 
+void UCustomizableObjectNodeTextureTransform::BackwardsCompatibleFixup(int32 CustomizableObjectCustomVersion)
+{
+	Super::BackwardsCompatibleFixup(CustomizableObjectCustomVersion);
+
+	if (CustomizableObjectCustomVersion == FCustomizableObjectCustomVersion::UpdatedNodesPinName3)
+	{
+		if (UEdGraphPin* InputTexturePin = FindPin(TEXT("Base Texture"), EGPD_Input))
+		{
+			InputTexturePin->PinName = TEXT("Texture");
+			InputTexturePin->PinFriendlyName = LOCTEXT("Image_Pin_Category", "Texture");
+		}
+
+		if (UEdGraphPin* OutputTexturePin = FindPin(TEXT("Texture"), EGPD_Output))
+		{
+			OutputTexturePin->PinFriendlyName = LOCTEXT("Image_Pin_Category", "Texture");
+		}
+	}
+}
+
 UEdGraphPin* UCustomizableObjectNodeTextureTransform::GetBaseImagePin() const
 {
 	return BaseImagePinReference.Get();
@@ -55,8 +65,7 @@ FText UCustomizableObjectNodeTextureTransform::GetNodeTitle(ENodeTitleType::Type
 
 FLinearColor UCustomizableObjectNodeTextureTransform::GetNodeTitleColor() const
 {
-	const UEdGraphSchema_CustomizableObject* Schema = GetDefault<UEdGraphSchema_CustomizableObject>();
-	return Schema->GetPinTypeColor(Schema->PC_Image);
+	return UEdGraphSchema_CustomizableObject::GetPinTypeColor(UEdGraphSchema_CustomizableObject::PC_Image);
 }
 
 FText UCustomizableObjectNodeTextureTransform::GetTooltipText() const
@@ -68,12 +77,6 @@ FText UCustomizableObjectNodeTextureTransform::GetTooltipText() const
 			"If Keep Aspect Ratio is set, an scaling factor preserving aspect ratio will be used as identity.");
 }
 
-void UCustomizableObjectNodeTextureTransform::Serialize(FArchive& Ar)
-{
-	Super::Serialize(Ar);
-
-	Ar.UsingCustomVersion(FCustomizableObjectCustomVersion::GUID);
-}
 
 
 #undef LOCTEXT_NAMESPACE
diff --git a/Engine/Plugins/Mutable/Source/CustomizableObjectEditor/Private/MuCOE/Nodes/CustomizableObjectNodeTextureTransform.h b/Engine/Plugins/Mutable/Source/CustomizableObjectEditor/Private/MuCOE/Nodes/CustomizableObjectNodeTextureTransform.h
--- a/Engine/Plugins/Mutable/Source/CustomizableObjectEditor/Private/MuCOE/Nodes/CustomizableObjectNodeTextureTransform.h
+++ b/Engine/Plugins/Mutable/Source/CustomizableObjectEditor/Private/MuCOE/Nodes/CustomizableObjectNodeTextureTransform.h
@@ -48,13 +48,14 @@ public:
 
 
 	// Begin EdGraphNode interface
-	FText GetNodeTitle(ENodeTitleType::Type TittleType) const override;
-	FLinearColor GetNodeTitleColor() const override;
-	FText GetTooltipText() const override;
-	virtual void Serialize(FArchive& Ar) override;
+	virtual FText GetNodeTitle(ENodeTitleType::Type TittleType) const override;
+	virtual FLinearColor GetNodeTitleColor() const override;
+	virtual FText GetTooltipText() const override;
+	//virtual void Serialize(FArchive& Ar) override;
 	
 	// UCustomizableObjectNode interface
-	void AllocateDefaultPins(UCustomizableObjectNodeRemapPins* RemapPins) override;
+	virtual void AllocateDefaultPins(UCustomizableObjectNodeRemapPins* RemapPins) override;
+	virtual void BackwardsCompatibleFixup(int32 CustomizableObjectCustomVersion) override;
 
 	UEdGraphPin* GetBaseImagePin() const;
 
diff --git a/Engine/Plugins/Mutable/Source/CustomizableObjectEditor/Private/MuCOE/Nodes/CustomizableObjectNodeTextureVariation.cpp b/Engine/Plugins/Mutable/Source/CustomizableObjectEditor/Private/MuCOE/Nodes/CustomizableObjectNodeTextureVariation.cpp
--- a/Engine/Plugins/Mutable/Source/CustomizableObjectEditor/Private/MuCOE/Nodes/CustomizableObjectNodeTextureVariation.cpp
+++ b/Engine/Plugins/Mutable/Source/CustomizableObjectEditor/Private/MuCOE/Nodes/CustomizableObjectNodeTextureVariation.cpp
@@ -2,6 +2,7 @@
 
 #include "MuCOE/Nodes/CustomizableObjectNodeTextureVariation.h"
 
+#include "MuCO/CustomizableObjectCustomVersion.h"
 #include "MuCOE/EdGraphSchema_CustomizableObject.h"
 
 class UCustomizableObjectNodeRemapPins;
@@ -18,7 +19,7 @@ UCustomizableObjectNodeTextureVariation::UCustomizableObjectNodeTextureVariation
 
 void UCustomizableObjectNodeTextureVariation::PostEditChangeProperty(FPropertyChangedEvent& PropertyChangedEvent)
 {
-	FProperty* PropertyThatChanged = PropertyChangedEvent.Property;
+	const FProperty* PropertyThatChanged = PropertyChangedEvent.Property;
 	if ( PropertyThatChanged  )
 	{
 		ReconstructNode();
@@ -30,25 +31,36 @@ void UCustomizableObjectNodeTextureVariation::PostEditChangeProperty(FPropertyCh
 
 void UCustomizableObjectNodeTextureVariation::AllocateDefaultPins(UCustomizableObjectNodeRemapPins* RemapPins)
 {
-	const UEdGraphSchema_CustomizableObject* Schema = GetDefault<UEdGraphSchema_CustomizableObject>();
-
-	UEdGraphPin* OutputPin = CustomCreatePin(EGPD_Output, Schema->PC_Image, FName(TEXT("Texture")));
+	UEdGraphPin* OutputPin = CustomCreatePinSimple(EGPD_Output, UEdGraphSchema_CustomizableObject::PC_Image);
 	OutputPin->bDefaultValueIsIgnored = true;
 
 	for ( int VariationIndex = Variations.Num() - 1; VariationIndex >= 0; --VariationIndex)
 	{
 		FString PinName = FString::Printf( TEXT("Variation %d"), VariationIndex);
-		UEdGraphPin* VariationPin = CustomCreatePin(EGPD_Input, Schema->PC_Image, FName(*PinName), false);
+		UEdGraphPin* VariationPin = CustomCreatePin(EGPD_Input, UEdGraphSchema_CustomizableObject::PC_Image, FName(*PinName), false);
 		VariationPin->bDefaultValueIsIgnored = true;
 
 		FString TagName = GetTagDisplayName(Variations[VariationIndex].Tag);
 		VariationPin->PinFriendlyName = FText::Format(LOCTEXT("Variation_Pin_FriendlyName", "Variation {0} [{1}]"), VariationIndex, FText::FromString(TagName));
 	}
 
-	UEdGraphPin* DefaultVariation = CustomCreatePin(EGPD_Input, Schema->PC_Image, FName(TEXT("Default")), false);
+	UEdGraphPin* DefaultVariation = CustomCreatePin(EGPD_Input, UEdGraphSchema_CustomizableObject::PC_Image, FName(TEXT("Default")), false);
 	DefaultVariation->bDefaultValueIsIgnored = true;
 }
 
+void UCustomizableObjectNodeTextureVariation::BackwardsCompatibleFixup(int32 CustomizableObjectCustomVersion)
+{
+	Super::BackwardsCompatibleFixup(CustomizableObjectCustomVersion);
+
+	if (CustomizableObjectCustomVersion == FCustomizableObjectCustomVersion::UpdatedNodesPinName3)
+	{
+		if (UEdGraphPin* TexturePin = FindPin(TEXT("Texture")))
+		{
+			TexturePin->PinFriendlyName = LOCTEXT("Image_Pin_Category", "Texture");
+		}
+	}
+}
+
 
 FText UCustomizableObjectNodeTextureVariation::GetNodeTitle(ENodeTitleType::Type TitleType) const
 {
@@ -58,8 +70,7 @@ FText UCustomizableObjectNodeTextureVariation::GetNodeTitle(ENodeTitleType::Type
 
 FLinearColor UCustomizableObjectNodeTextureVariation::GetNodeTitleColor() const
 {
-	const UEdGraphSchema_CustomizableObject* Schema = GetDefault<UEdGraphSchema_CustomizableObject>();
-	return Schema->GetPinTypeColor(Schema->PC_Image);
+	return UEdGraphSchema_CustomizableObject::GetPinTypeColor(UEdGraphSchema_CustomizableObject::PC_Image);
 }
 
 
diff --git a/Engine/Plugins/Mutable/Source/CustomizableObjectEditor/Private/MuCOE/Nodes/CustomizableObjectNodeTextureVariation.h b/Engine/Plugins/Mutable/Source/CustomizableObjectEditor/Private/MuCOE/Nodes/CustomizableObjectNodeTextureVariation.h
--- a/Engine/Plugins/Mutable/Source/CustomizableObjectEditor/Private/MuCOE/Nodes/CustomizableObjectNodeTextureVariation.h
+++ b/Engine/Plugins/Mutable/Source/CustomizableObjectEditor/Private/MuCOE/Nodes/CustomizableObjectNodeTextureVariation.h
@@ -36,21 +36,17 @@ public:
 	TArray<FCustomizableObjectTextureVariation> Variations;
 
 	// UObject interface.
-	void PostEditChangeProperty(FPropertyChangedEvent& PropertyChangedEvent) override;
+	virtual void PostEditChangeProperty(FPropertyChangedEvent& PropertyChangedEvent) override;
 
 	// EdGraphNode interface
 	virtual FText GetNodeTitle(ENodeTitleType::Type TitleType) const override;
-	FLinearColor GetNodeTitleColor() const override;
-	FText GetTooltipText() const override;
+	virtual FLinearColor GetNodeTitleColor() const override;
+	virtual FText GetTooltipText() const override;
 	
 	// UCustomizableObjectNode interface
 	virtual void AllocateDefaultPins(UCustomizableObjectNodeRemapPins* RemapPins) override;
-
-	UEdGraphPin* OutputPin() const
-	{
-		return FindPin(TEXT("Texture"));
-	}
-
+	virtual void BackwardsCompatibleFixup(int32 CustomizableObjectCustomVersion) override;
+	
 	UEdGraphPin* DefaultPin() const
 	{
 		return FindPin(TEXT("Default"));
diff --git a/Engine/Plugins/Mutable/Source/CustomizableObjectEditor/Private/MuCOE/SMutableCodeViewer.cpp b/Engine/Plugins/Mutable/Source/CustomizableObjectEditor/Private/MuCOE/SMutableCodeViewer.cpp
--- a/Engine/Plugins/Mutable/Source/CustomizableObjectEditor/Private/MuCOE/SMutableCodeViewer.cpp
+++ b/Engine/Plugins/Mutable/Source/CustomizableObjectEditor/Private/MuCOE/SMutableCodeViewer.cpp
@@ -2532,6 +2532,28 @@ bool SMutableCodeViewer::IsConstantResourceUsedByOperation(const int32 IndexOnCo
 						Data += sizeof(int32);
 					}
 				}
+				else if (OperationType == mu::EOpType::IN_CONDITIONAL)
+				{
+					mu::OP::ConditionalArgs Arguments = InProgram.GetOpArgs<mu::OP::ConditionalArgs>(OperationAddress);
+					
+					if (IndexOnConstantsArray == Arguments.condition)
+					{
+						bResourceLocated = true;
+						break;
+					}
+
+					if (IndexOnConstantsArray == Arguments.yes)
+					{
+						bResourceLocated = true;
+						break;
+					}
+
+					if (IndexOnConstantsArray == Arguments.no)
+					{
+						bResourceLocated = true;
+						break;
+					}
+				}
 				
 				break;
 			}
diff --git a/Engine/Plugins/Mutable/Source/CustomizableObjectEditor/Public/MuCOE/Nodes/CustomizableObjectNode.h b/Engine/Plugins/Mutable/Source/CustomizableObjectEditor/Public/MuCOE/Nodes/CustomizableObjectNode.h
--- a/Engine/Plugins/Mutable/Source/CustomizableObjectEditor/Public/MuCOE/Nodes/CustomizableObjectNode.h
+++ b/Engine/Plugins/Mutable/Source/CustomizableObjectEditor/Public/MuCOE/Nodes/CustomizableObjectNode.h
@@ -205,6 +205,9 @@ public:
 	/** If returns true, only a single link can go out of every output pin. False by default. */
 	virtual bool IsSingleOutputNode() const;
 
+	/** Wrapper of the pin creation method. Will use the PinName and PinFriendlyname returned by the "UEdGraphSchema_CustomizableObject" */
+	UEdGraphPin* CustomCreatePinSimple( EEdGraphPinDirection Direction, const FName& Category, bool bIsArray=false );
+	
 	// Wrapper of the pin creation method, to cope with multiple Unreal Engine versions
 	UEdGraphPin* CustomCreatePin( EEdGraphPinDirection Direction, const FName& Type, const FName& Name, bool bIsArray=false );
 
