[mutable] Fixed unselection on paste for the Projector nodes

#rnx

diff --git a/Engine/Plugins/Mutable/Source/CustomizableObjectEditor/Private/MuCOE/CustomizableObjectEditor.cpp b/Engine/Plugins/Mutable/Source/CustomizableObjectEditor/Private/MuCOE/CustomizableObjectEditor.cpp
--- a/Engine/Plugins/Mutable/Source/CustomizableObjectEditor/Private/MuCOE/CustomizableObjectEditor.cpp
+++ b/Engine/Plugins/Mutable/Source/CustomizableObjectEditor/Private/MuCOE/CustomizableObjectEditor.cpp
@@ -2,12 +2,9 @@
 
 #include "MuCOE/CustomizableObjectEditor.h"
 
-#include "AssetViewerSettings.h"
-#include "ContentBrowserModule.h"
 #include "CustomizableObjectCompiler.h"
 #include "CustomizableObjectEditorModule.h"
 #include "CustomizableObjectEditorPerformanceAnalyzer.h"
-#include "Animation/DebugSkelMeshComponent.h"
 #include "AssetRegistry/ARFilter.h"
 #include "AssetRegistry/AssetRegistryModule.h"
 #include "DetailsViewArgs.h"
@@ -18,14 +15,12 @@
 #include "Framework/MultiBox/MultiBoxBuilder.h"
 #include "Framework/Notifications/NotificationManager.h"
 #include "GraphEditorActions.h"
-#include "IContentBrowserSingleton.h"
 #include "HAL/PlatformApplicationMisc.h"
 #include "IDetailsView.h"
 #include "MaterialGraph/MaterialGraphNode.h"
 #include "MaterialGraph/MaterialGraphNode_Comment.h"
 #include "MuCO/CustomizableObjectInstance.h"
 #include "MuCO/CustomizableObjectSystem.h"
-#include "MuCO/UnrealPortabilityHelpers.h"
 #include "MuCO/CustomizableObjectCompilerTypes.h"
 #include "MuCOE/CustomizableObjectCustomSettings.h"
 #include "MuCOE/CustomizableObjectEditorActions.h"
@@ -34,16 +29,11 @@
 #include "MuCOE/CustomizableObjectEditorNodeContextCommands.h"
 #include "MuCOE/CustomizableObjectEditorStyle.h"
 #include "MuCOE/CustomizableObjectEditorViewportClient.h"
-#include "MuCOE/CustomizableObjectEditor_Deprecated.h"
 #include "MuCOE/CustomizableObjectGraph.h"
 #include "MuCOE/CustomizableObjectPreviewScene.h"
-#include "MuCOE/EdGraphSchema_CustomizableObject.h"
 #include "MuCOE/GraphTraversal.h"
 #include "MuCOE/Nodes/CustomizableObjectNodeEnumParameter.h"
-#include "MuCOE/Nodes/CustomizableObjectNodeFloatParameter.h"
 #include "MuCOE/Nodes/CustomizableObjectNodeGroupProjectorParameter.h"
-#include "MuCOE/Nodes/CustomizableObjectNodeMaterial.h"
-#include "MuCOE/Nodes/CustomizableObjectNodeMaterialVariation.h"
 #include "MuCOE/Nodes/CustomizableObjectNodeModifierClipMorph.h"
 #include "MuCOE/Nodes/CustomizableObjectNodeModifierClipWithMesh.h"
 #include "MuCOE/Nodes/CustomizableObjectNodeModifierTransformInMesh.h"
@@ -2180,8 +2170,11 @@ void FCustomizableObjectEditor::OnNodeTitleCommitted(const FText& NewText, EText
 
 void FCustomizableObjectEditor::UpdateGraphNodeProperties()
 {
+	// Cache a copy of the selected nodes so we can later restore them
+	const FGraphPanelSelectionSet PreClearingSelectedNodes = GraphEditor->GetSelectedNodes();
+	
 	OnSelectedGraphNodesChanged(FGraphPanelSelectionSet());
-	OnSelectedGraphNodesChanged(GraphEditor->GetSelectedNodes());
+	OnSelectedGraphNodesChanged(PreClearingSelectedNodes);
 }
 
 
diff --git a/Engine/Plugins/Mutable/Source/CustomizableObjectEditor/Private/MuCOE/CustomizableObjectEditor.h b/Engine/Plugins/Mutable/Source/CustomizableObjectEditor/Private/MuCOE/CustomizableObjectEditor.h
--- a/Engine/Plugins/Mutable/Source/CustomizableObjectEditor/Private/MuCOE/CustomizableObjectEditor.h
+++ b/Engine/Plugins/Mutable/Source/CustomizableObjectEditor/Private/MuCOE/CustomizableObjectEditor.h
@@ -183,7 +183,7 @@ public:
 	 */
 	void OnNodeTitleCommitted(const FText& NewText, ETextCommit::Type CommitInfo, UEdGraphNode* NodeBeingChanged);
 
-	virtual void UpdateGraphNodeProperties();
+	virtual void UpdateGraphNodeProperties() override;
 
 	/** Callback to notify the editor when the PreviewInstance has been updated */
 	void OnUpdatePreviewInstance(UCustomizableObjectInstance* Instance);
diff --git a/Engine/Plugins/Mutable/Source/CustomizableObjectEditor/Private/MuCOE/Nodes/CustomizableObjectNodeProjectorParameterDetails.cpp b/Engine/Plugins/Mutable/Source/CustomizableObjectEditor/Private/MuCOE/Nodes/CustomizableObjectNodeProjectorParameterDetails.cpp
--- a/Engine/Plugins/Mutable/Source/CustomizableObjectEditor/Private/MuCOE/Nodes/CustomizableObjectNodeProjectorParameterDetails.cpp
+++ b/Engine/Plugins/Mutable/Source/CustomizableObjectEditor/Private/MuCOE/Nodes/CustomizableObjectNodeProjectorParameterDetails.cpp
@@ -12,7 +12,6 @@
 #include "MuCOE/ICustomizableObjectEditor.h"
 #include "MuCOE/Nodes/CustomizableObjectNodeProjectorConstant.h"
 #include "MuCOE/Nodes/CustomizableObjectNodeProjectorParameter.h"
-#include "MuCOE/UnrealEditorPortabilityHelpers.h"
 #include "PropertyCustomizationHelpers.h"
 #include "Widgets/Input/SButton.h"
 #include "Widgets/Input/STextComboBox.h"
