[mutable] Fixed backwards compatibility error.

#rnx

diff --git a/Engine/Plugins/Mutable/Source/CustomizableObjectEditor/Private/MuCOE/Nodes/CustomizableObjectNodeTextureBinarise.cpp b/Engine/Plugins/Mutable/Source/CustomizableObjectEditor/Private/MuCOE/Nodes/CustomizableObjectNodeTextureBinarise.cpp
--- a/Engine/Plugins/Mutable/Source/CustomizableObjectEditor/Private/MuCOE/Nodes/CustomizableObjectNodeTextureBinarise.cpp
+++ b/Engine/Plugins/Mutable/Source/CustomizableObjectEditor/Private/MuCOE/Nodes/CustomizableObjectNodeTextureBinarise.cpp
@@ -29,14 +29,14 @@ void UCustomizableObjectNodeTextureBinarise::BackwardsCompatibleFixup(int32 Cust
 	
 	if (CustomizableObjectCustomVersion == FCustomizableObjectCustomVersion::FixPinsNamesImageToTexture2)
 	{
-		if (UEdGraphPin* TexturePin = FindPin(TEXT("Image"))) {
+		if (UEdGraphPin* TexturePin = FindPin(TEXT("Image")))
+		{
 			TexturePin->PinName = TEXT("Texture");
 			UCustomizableObjectNode::ReconstructNode();
 		}
-	}
-	else if (CustomizableObjectCustomVersion == FCustomizableObjectCustomVersion::FixPinsNamesImageToTexture2)
-	{
-		if (UEdGraphPin* TexturePin = FindPin(TEXT("Base Image"))) {
+
+		if (UEdGraphPin* TexturePin = FindPin(TEXT("Base Image")))
+		{
 			TexturePin->PinName = TEXT("Base Texture");
 			UCustomizableObjectNode::ReconstructNode();
 		}
