[Mutable] Fix deadlock when compiling a Customizable Object and shutting down the editor.

- Replaced Compiler Runnable GT Task Graph tasks, with a custom ticker tasks (as before).

diff --git a/Engine/Plugins/Mutable/Source/CustomizableObjectEditor/Private/MuCOE/CustomizableObjectCompileRunnable.cpp b/Engine/Plugins/Mutable/Source/CustomizableObjectEditor/Private/MuCOE/CustomizableObjectCompileRunnable.cpp
--- a/Engine/Plugins/Mutable/Source/CustomizableObjectEditor/Private/MuCOE/CustomizableObjectCompileRunnable.cpp
+++ b/Engine/Plugins/Mutable/Source/CustomizableObjectEditor/Private/MuCOE/CustomizableObjectCompileRunnable.cpp
@@ -2,6 +2,7 @@
 
 #include "MuCOE/CustomizableObjectCompileRunnable.h"
 
+#include "CustomizableObjectCompiler.h"
 #include "MuCOE/GenerateMutableSource/GenerateMutableSource.h"
 #include "HAL/FileManager.h"
 #include "MuCO/UnrealMutableModelDiskStreamer.h"
@@ -47,9 +48,9 @@ TAutoConsoleVariable<bool> CVarMutableCompilerFastCompression(
 	ECVF_Default);
 
 
-FCustomizableObjectCompileRunnable::FCustomizableObjectCompileRunnable(mu::Ptr<mu::Node> Root, TSharedPtr<FMutableCompilationContext> InCompilationContext)
+FCustomizableObjectCompileRunnable::FCustomizableObjectCompileRunnable(mu::Ptr<mu::Node> Root, const TSharedRef<FCustomizableObjectCompiler>& InCompiler)
 	: MutableRoot(Root)
-	, CompilationContext(InCompilationContext)
+	, WeakCompiler(InCompiler)
 	, bThreadCompleted(false)
 {
 	PrepareUnrealCompression();
@@ -187,7 +188,7 @@ uint32 FCustomizableObjectCompileRunnable::Run()
 			FString MorphName = MorphNameRef;
 
 			// It would be great to be able to do this conversion in a worker thread, but the engine doesn't support it yet.
-			auto LoadMeshFunc = [MeshData, MorphName, ResolvedMesh, CompletionEvent, CompilationContext=CompilationContext]() mutable
+			auto LoadMeshFunc = [MeshData, MorphName, ResolvedMesh, CompletionEvent, WeakCompiler = WeakCompiler]() mutable
 				{
 					MUTABLE_CPUPROFILER_SCOPE(LoadMeshFunc);
 					check(IsInGameThread());
@@ -200,12 +201,20 @@ uint32 FCustomizableObjectCompileRunnable::Run()
 						return;
 					}
 
+					TSharedPtr<FCustomizableObjectCompiler> Compiler = WeakCompiler.Pin();
+					if (!Compiler)
+					{
+						*ResolvedMesh = MakeShared<mu::FMesh>();
+						CompletionEvent.Trigger();
+						return;
+					}
+
 					// Ensure we don't pull unwanted data into the package when cooking.
 					FCookLoadScope CookLoadScope(ECookLoadType::EditorOnly);
 
 					if (MorphName.IsEmpty())
 					{
-						*ResolvedMesh = ConvertSkeletalMeshToMutable(MeshData, *CompilationContext);
+						*ResolvedMesh = ConvertSkeletalMeshToMutable(MeshData, *Compiler->CompilationContext);
 					}
 					else
 					{
@@ -213,7 +222,7 @@ uint32 FCustomizableObjectCompileRunnable::Run()
 						if (SkeletalMesh->FindMorphTarget(*MorphName))
 						{
 							TSharedPtr<mu::FMesh> Mesh;
-							Mesh = ConvertSkeletalMeshToMutable(MeshData, *CompilationContext);
+							Mesh = ConvertSkeletalMeshToMutable(MeshData, *Compiler->CompilationContext);
 							Mesh = BuildMorphedMutableMeshFromMesh(Mesh, SkeletalMesh, *MorphName, MeshData.LODIndex, MeshData.SectionIndex);
 							*ResolvedMesh = Mesh;
 						}
@@ -229,15 +238,15 @@ uint32 FCustomizableObjectCompileRunnable::Run()
 			else
 			{
 				check(!bRunImmediatlyIfPossible);
-				
-				// Could store event here
-				//LastTask = 
-				FFunctionGraphTask::CreateAndDispatchWhenReady(LoadMeshFunc,
-					TStatId{}, 
-					nullptr, // Prerequisites
-					ENamedThreads::GameThread
-					);
-
+				if (TSharedPtr<FCustomizableObjectCompiler> Compiler = WeakCompiler.Pin())
+				{
+					Compiler->AddGameThreadCompileTask(MoveTemp(LoadMeshFunc));	
+				}
+				else
+				{
+					*ResolvedMesh = MakeShared<mu::FMesh>();
+					CompletionEvent.Trigger();
+				}
 			}
 
 			return CompletionEvent;
diff --git a/Engine/Plugins/Mutable/Source/CustomizableObjectEditor/Private/MuCOE/CustomizableObjectCompileRunnable.h b/Engine/Plugins/Mutable/Source/CustomizableObjectEditor/Private/MuCOE/CustomizableObjectCompileRunnable.h
--- a/Engine/Plugins/Mutable/Source/CustomizableObjectEditor/Private/MuCOE/CustomizableObjectCompileRunnable.h
+++ b/Engine/Plugins/Mutable/Source/CustomizableObjectEditor/Private/MuCOE/CustomizableObjectCompileRunnable.h
@@ -54,7 +54,7 @@ private:
 
 public:
 
-	FCustomizableObjectCompileRunnable(mu::Ptr<mu::Node> Root, TSharedPtr<FMutableCompilationContext> CompilationContext);
+	FCustomizableObjectCompileRunnable(mu::Ptr<mu::Node> Root, const TSharedRef<FCustomizableObjectCompiler>& InCompiler);
 
 	// FRunnable interface
 	uint32 Run() override;
@@ -70,7 +70,8 @@ public:
 	TSharedPtr<mu::FModel, ESPMode::ThreadSafe> Model;
 
 	FCompilationOptions Options;
-	TSharedPtr<FMutableCompilationContext> CompilationContext;
+
+	TWeakPtr<FCustomizableObjectCompiler> WeakCompiler;
 
 	TArray<FMutableSourceTextureData> ReferencedTextures;
 	TArray<FMutableSourceMeshData> ReferencedMeshes;
diff --git a/Engine/Plugins/Mutable/Source/CustomizableObjectEditor/Private/MuCOE/CustomizableObjectCompiler.cpp b/Engine/Plugins/Mutable/Source/CustomizableObjectEditor/Private/MuCOE/CustomizableObjectCompiler.cpp
--- a/Engine/Plugins/Mutable/Source/CustomizableObjectEditor/Private/MuCOE/CustomizableObjectCompiler.cpp
+++ b/Engine/Plugins/Mutable/Source/CustomizableObjectEditor/Private/MuCOE/CustomizableObjectCompiler.cpp
@@ -44,6 +44,8 @@ bool FCustomizableObjectCompiler::Tick(bool bBlocking)
 {
 	MUTABLE_CPUPROFILER_SCOPE(FCustomizableObjectCompiler::Tick);
 
+	ProcessCompileTasks();
+	
 	bool bFinished = true;
 
 	if (TryPopCompileRequest())
@@ -357,8 +359,8 @@ void FCustomizableObjectCompiler::EnqueueCompileRequest(const TSharedRef<FCompil
 		
 	if (!CompileRequest->bAsync)
 	{
-		FCustomizableObjectCompiler SyncCompiler;
-		SyncCompiler.Compile(CompileRequest);
+		TSharedRef<FCustomizableObjectCompiler> SyncCompiler = MakeShared<FCustomizableObjectCompiler>();
+		SyncCompiler->Compile(CompileRequest);
 	}
 	else
 	{
@@ -432,6 +434,12 @@ void FCustomizableObjectCompiler::AddReferencedObjects(FReferenceCollector& Coll
 }
 
 
+void FCustomizableObjectCompiler::AddGameThreadCompileTask(TFunction<void()>&& Task)
+{
+	PendingGameThreadCompileTasks.Enqueue(Task);
+}
+
+
 void ProcessChildObjectsRecursively(const UCustomizableObject* ParentObject, FMutableGraphGenerationContext& GenerationContext)
 {
 	TArray<FName> ReferencedObjectNames;
@@ -775,6 +783,31 @@ ECompilationResultPrivate FCustomizableObjectCompiler::GetCompilationResult() co
 }
 
 
+void FCustomizableObjectCompiler::ProcessCompileTasks()
+{
+	MUTABLE_CPUPROFILER_SCOPE(CompileTasks);
+
+	check(IsInGameThread());
+
+	// See if there are compile-time tasks to run
+	constexpr double MaxSecondsPerFrame = 0.4;
+	double MaxTime = FPlatformTime::Seconds() + MaxSecondsPerFrame;
+
+	TFunction<void()> Task;
+	while (PendingGameThreadCompileTasks.Dequeue(Task))
+	{
+		Task();
+
+		// Simple time limit enforcement to avoid blocking the game thread if there are many requests.
+		double CurrentTime = FPlatformTime::Seconds();
+		if (CurrentTime >= MaxTime)
+		{
+			break;
+		}
+	}
+}
+
+
 void FCustomizableObjectCompiler::SetCompilationState(ECompilationStatePrivate State, ECompilationResultPrivate Result) const
 {
 	check(CurrentRequest);
@@ -800,7 +833,7 @@ void FCustomizableObjectCompiler::CompileInternal(bool bAsync)
 		return;
 	}
 
-	CompilationContext = MakeShared<FMutableCompilationContext>(CurrentObject, this, CurrentOptions);
+	CompilationContext = MakeShared<FMutableCompilationContext>(CurrentObject, SharedThis(this), CurrentOptions);
 	FMutableGraphGenerationContext GenerationContext(*CompilationContext);
 
 	// Perform a first participating objects pass
@@ -970,7 +1003,7 @@ void FCustomizableObjectCompiler::CompileInternal(bool bAsync)
 
 		ModelResources->StreamedResourceDataEditor = MoveTemp(GenerationContext.StreamedResourceData);
 		
-		CompileTask = MakeShareable(new FCustomizableObjectCompileRunnable(MutableRoot, CompilationContext));
+		CompileTask = MakeShareable(new FCustomizableObjectCompileRunnable(MutableRoot, SharedThis(this)));
 		CompileTask->Options = CurrentOptions;
 		CompileTask->ReferencedTextures = NewCompileTimeReferencedTextures;
 		CompileTask->ReferencedMeshes = NewCompileTimeReferencedMeshes;
@@ -1228,7 +1261,7 @@ mu::Ptr<mu::Node> FCustomizableObjectCompiler::Export(UCustomizableObject* Objec
 	CompilerOptions.bPhysicsAssetMergeEnabled = Object->GetPrivate()->IsPhysicsAssetMergeEnabled();
 	CompilerOptions.bAnimBpPhysicsManipulationEnabled = Object->GetPrivate()->IsEnabledAnimBpPhysicsAssetsManipulation();
 	
-	CompilationContext = MakeShared<FMutableCompilationContext>(Object, this, CompilerOptions);
+	CompilationContext = MakeShared<FMutableCompilationContext>(Object, SharedThis(this), CompilerOptions);
 	FMutableGraphGenerationContext GenerationContext(*CompilationContext);
 	
 	// Generate the mutable node expression
@@ -1495,12 +1528,11 @@ void FCustomizableObjectCompiler::ForceFinishCompilation()
 
 	else if (CompileTask.IsValid())
 	{
-		// Compilation needs game thread tasks every now and then. Wait for compilation to finish while
-		// giving execution time for these tasks.
-		// TODO: interruptible compilations?
+		// Compilation needs game thread tasks every now and then.
+		// Wait for compilation to finish while giving execution time for these tasks.
 		while (!CompileTask->IsCompleted())
 		{
-			FTaskGraphInterface::Get().ProcessThreadUntilIdle(ENamedThreads::GameThread);
+			ProcessCompileTasks();
 		}
 
 		// Order matters
diff --git a/Engine/Plugins/Mutable/Source/CustomizableObjectEditor/Private/MuCOE/CustomizableObjectCompiler.h b/Engine/Plugins/Mutable/Source/CustomizableObjectEditor/Private/MuCOE/CustomizableObjectCompiler.h
--- a/Engine/Plugins/Mutable/Source/CustomizableObjectEditor/Private/MuCOE/CustomizableObjectCompiler.h
+++ b/Engine/Plugins/Mutable/Source/CustomizableObjectEditor/Private/MuCOE/CustomizableObjectCompiler.h
@@ -36,11 +36,9 @@ namespace mu
 mu::Ptr<mu::NodeObject> GenerateMutableRoot(const UCustomizableObject* Object, FMutableGraphGenerationContext& GenerationContext);
 
 
-class FCustomizableObjectCompiler : public FTickableEditorObject, public FTickableCookObject, public FGCObject
+class FCustomizableObjectCompiler : public TSharedFromThis<FCustomizableObjectCompiler>, public FTickableEditorObject, public FTickableCookObject, public FGCObject
 {
 public:
-	virtual ~FCustomizableObjectCompiler() override {}
-
 	/** Check for pending compilation process. Returns true if an object has been updated. */
 	bool Tick(bool bBlocking = false);
 	int32 GetNumRemainingWork() const;
@@ -95,7 +93,9 @@ public:
 	{
 		return TEXT("FCustomizableObjectCompiler");
 	}
-	
+
+	void AddGameThreadCompileTask(TFunction<void()>&& Task);
+
 	/** */
 	TSharedPtr<FMutableCompilationContext> CompilationContext;
 
@@ -127,6 +127,8 @@ private:
 	
 	ECompilationResultPrivate GetCompilationResult() const;
 
+	void ProcessCompileTasks();
+	
 	/** Pointer to the Asynchronous Preloading process call back */
 	TSharedPtr<FStreamableHandle> AsynchronousStreamableHandlePtr;
 
@@ -162,5 +164,7 @@ private:
 
 	/** Compilation start time in seconds. */
 	double CompilationStartTime = 0;
+
+	TQueue< TFunction<void()>, EQueueMode::Mpsc > PendingGameThreadCompileTasks;
 };
 
diff --git a/Engine/Plugins/Mutable/Source/CustomizableObjectEditor/Private/MuCOE/CustomizableObjectEditorModule.cpp b/Engine/Plugins/Mutable/Source/CustomizableObjectEditor/Private/MuCOE/CustomizableObjectEditorModule.cpp
--- a/Engine/Plugins/Mutable/Source/CustomizableObjectEditor/Private/MuCOE/CustomizableObjectEditorModule.cpp
+++ b/Engine/Plugins/Mutable/Source/CustomizableObjectEditor/Private/MuCOE/CustomizableObjectEditorModule.cpp
@@ -270,9 +270,7 @@ void FCustomizableObjectEditorModule::StartupModule()
 void FCustomizableObjectEditorModule::ShutdownModule()
 {
 	FEditorDelegates::PreBeginPIE.RemoveAll(this);
-
-	check(Compiler.GetNumRemainingWork() == 0);
-
+	
 	if( FModuleManager::Get().IsModuleLoaded( "PropertyEditor" ) )
 	{
 		FPropertyEditorModule& PropertyModule = FModuleManager::GetModuleChecked<FPropertyEditorModule>("PropertyEditor");
@@ -921,7 +919,7 @@ void FCustomizableObjectEditorModule::PostBackwardsCompatibleFixup(UEdGraph& Gra
 
 void FCustomizableObjectEditorModule::EnqueueCompileRequest(const TSharedRef<FCompilationRequest>& InCompilationRequest, bool bForceRequest)
 {
-	Compiler.EnqueueCompileRequest(InCompilationRequest, bForceRequest);
+	Compiler->EnqueueCompileRequest(InCompilationRequest, bForceRequest);
 }
 
 
@@ -970,8 +968,8 @@ void FCustomizableObjectEditorModule::CompileCustomizableObject(UCustomizableObj
 
 int32 FCustomizableObjectEditorModule::Tick(bool bBlocking)
 {
-	Compiler.Tick(bBlocking);
-	return Compiler.GetNumRemainingWork();
+	Compiler->Tick(bBlocking);
+	return Compiler->GetNumRemainingWork();
 }
 
 
@@ -1057,20 +1055,20 @@ bool FCustomizableObjectEditorModule::IsCachedCookedPlatformDataLoaded(UCustomiz
 
 void FCustomizableObjectEditorModule::CancelCompileRequests()
 {
-	Compiler.ForceFinishCompilation();
-	Compiler.ClearCompileRequests();
+	Compiler->ForceFinishCompilation();
+	Compiler->ClearCompileRequests();
 }
 
 
 int32 FCustomizableObjectEditorModule::GetNumCompileRequests()
 {
-	return Compiler.GetNumRemainingWork();
+	return Compiler->GetNumRemainingWork();
 }
 
 
 bool FCustomizableObjectEditorModule::IsCompiling(const UCustomizableObject& Object) const
 {
-	return Object.GetPrivate()->IsLocked() || Compiler.IsRequestQueued(Object);
+	return Object.GetPrivate()->IsLocked() || Compiler->IsRequestQueued(Object);
 }
 
 
diff --git a/Engine/Plugins/Mutable/Source/CustomizableObjectEditor/Private/MuCOE/CustomizableObjectEditorModule.h b/Engine/Plugins/Mutable/Source/CustomizableObjectEditor/Private/MuCOE/CustomizableObjectEditorModule.h
--- a/Engine/Plugins/Mutable/Source/CustomizableObjectEditor/Private/MuCOE/CustomizableObjectEditorModule.h
+++ b/Engine/Plugins/Mutable/Source/CustomizableObjectEditor/Private/MuCOE/CustomizableObjectEditorModule.h
@@ -79,7 +79,7 @@ private:
 
 	FCustomizableObjectEditorLogger Logger;
 	
-	FCustomizableObjectCompiler Compiler;
+	TSharedRef<FCustomizableObjectCompiler> Compiler = MakeShared<FCustomizableObjectCompiler>();
 
 	// Command to look for Customizable Object Instance in the player pawn of the current world and open its Customizable Object Instance Editor
 	IConsoleCommand* LaunchCOIECommand = nullptr;
diff --git a/Engine/Plugins/Mutable/Source/CustomizableObjectEditor/Private/MuCOE/CustomizableObjectLayout.cpp b/Engine/Plugins/Mutable/Source/CustomizableObjectEditor/Private/MuCOE/CustomizableObjectLayout.cpp
--- a/Engine/Plugins/Mutable/Source/CustomizableObjectEditor/Private/MuCOE/CustomizableObjectLayout.cpp
+++ b/Engine/Plugins/Mutable/Source/CustomizableObjectEditor/Private/MuCOE/CustomizableObjectLayout.cpp
@@ -97,11 +97,11 @@ void UCustomizableObjectLayout::GenerateAutomaticBlocksFromUVs()
 	}
 
 	// Create a GenerationContext
-	FCustomizableObjectCompiler Compiler;
+	TSharedRef<FCustomizableObjectCompiler> Compiler = MakeShared<FCustomizableObjectCompiler>();
 	UCustomizableObject* Object = Node->GetGraphEditor()->GetCustomizableObject();
 	FCompilationOptions Options = Object->GetPrivate()->GetCompileOptions();
 	
-	FMutableCompilationContext CompilationContext(Object, &Compiler, Options);
+	FMutableCompilationContext CompilationContext(Object, Compiler, Options);
 	FMutableGraphGenerationContext GenerationContext(CompilationContext);
 
 	bool bOutWasEmpty = false;
diff --git a/Engine/Plugins/Mutable/Source/CustomizableObjectEditor/Private/MuCOE/GenerateMutableSource/GenerateMutableSource.cpp b/Engine/Plugins/Mutable/Source/CustomizableObjectEditor/Private/MuCOE/GenerateMutableSource/GenerateMutableSource.cpp
--- a/Engine/Plugins/Mutable/Source/CustomizableObjectEditor/Private/MuCOE/GenerateMutableSource/GenerateMutableSource.cpp
+++ b/Engine/Plugins/Mutable/Source/CustomizableObjectEditor/Private/MuCOE/GenerateMutableSource/GenerateMutableSource.cpp
@@ -129,7 +129,7 @@ uint32 GetTypeHash(const FGeneratedGroupProjectorsKey& Key)
 
 void FMutableCompilationContext::Log(const FText& Message, const TArray<const UObject*>& Context, const EMessageSeverity::Type MessageSeverity, const bool bAddBaseObjectInfo, const ELoggerSpamBin SpamBin) const
 {
-	if (Compiler)
+	if (TSharedPtr<FCustomizableObjectCompiler> Compiler = WeakCompiler.Pin())
 	{
 		Compiler->CompilerLog(Message, Context, MessageSeverity, bAddBaseObjectInfo, SpamBin);
 	}
@@ -138,7 +138,7 @@ void FMutableCompilationContext::Log(const FText& Message, const TArray<const UO
 
 void FMutableCompilationContext::Log(const FText& Message, const UObject* Context, const EMessageSeverity::Type MessageSeverity, const bool bAddBaseObjectInfo, const ELoggerSpamBin SpamBin) const
 {
-	if (Compiler)
+	if (TSharedPtr<FCustomizableObjectCompiler> Compiler = WeakCompiler.Pin())
 	{
 		Compiler->CompilerLog(Message, Context, MessageSeverity, bAddBaseObjectInfo, SpamBin);
 	}
@@ -212,10 +212,10 @@ uint32 FMutableCompilationContext::GetSkinWeightProfileIdUnique(const FName InPr
 }
 
 
-FMutableCompilationContext::FMutableCompilationContext(const UCustomizableObject* InObject, class FCustomizableObjectCompiler* InCompiler, const FCompilationOptions& InOptions)
+FMutableCompilationContext::FMutableCompilationContext(const UCustomizableObject* InObject, const TSharedPtr<FCustomizableObjectCompiler>& InCompiler, const FCompilationOptions& InOptions)
 	: Object(InObject)
 	, Options(InOptions)
-	, Compiler(InCompiler)
+	, WeakCompiler(InCompiler)
 {
 }
 
diff --git a/Engine/Plugins/Mutable/Source/CustomizableObjectEditor/Private/MuCOE/GenerateMutableSource/GenerateMutableSource.h b/Engine/Plugins/Mutable/Source/CustomizableObjectEditor/Private/MuCOE/GenerateMutableSource/GenerateMutableSource.h
--- a/Engine/Plugins/Mutable/Source/CustomizableObjectEditor/Private/MuCOE/GenerateMutableSource/GenerateMutableSource.h
+++ b/Engine/Plugins/Mutable/Source/CustomizableObjectEditor/Private/MuCOE/GenerateMutableSource/GenerateMutableSource.h
@@ -510,7 +510,7 @@ struct FMutableCompilationContext
 private:
 
 	/** Non-owned reference to the compiler object */
-	FCustomizableObjectCompiler* Compiler = nullptr;
+	TWeakPtr<FCustomizableObjectCompiler> WeakCompiler;
 
 	/** */
 	TMap<uint32, FName> UniqueSkinWeightProfileIds;
@@ -518,7 +518,7 @@ private:
 
 public:
 
-	FMutableCompilationContext(const UCustomizableObject* InObject, FCustomizableObjectCompiler* InCompiler, const FCompilationOptions& InOptions);
+	FMutableCompilationContext(const UCustomizableObject* InObject, const TSharedPtr<FCustomizableObjectCompiler>& InCompiler, const FCompilationOptions& InOptions);
 
 	/** Return the name of the current customizable object, for logging purposes. */
 	FString GetObjectName() const;
diff --git a/Engine/Plugins/Mutable/Source/CustomizableObjectEditor/Private/MuCOE/SMutableGraphViewer.h b/Engine/Plugins/Mutable/Source/CustomizableObjectEditor/Private/MuCOE/SMutableGraphViewer.h
--- a/Engine/Plugins/Mutable/Source/CustomizableObjectEditor/Private/MuCOE/SMutableGraphViewer.h
+++ b/Engine/Plugins/Mutable/Source/CustomizableObjectEditor/Private/MuCOE/SMutableGraphViewer.h
@@ -61,7 +61,7 @@ private:
 	TArray<FMutableSourceTextureData> ReferencedCompileTextures;
 
 	/** Object compiler. */
-	FCustomizableObjectCompiler Compiler;
+	TSharedRef<FCustomizableObjectCompiler> Compiler = MakeShared<FCustomizableObjectCompiler>();
 
 	/** Root nodes of the tree widget. */
 	TArray<TSharedPtr<FMutableGraphTreeElement>> RootNodes;
diff --git a/Engine/Plugins/Mutable/Source/CustomizableObjectEditor/Private/MuCOE/SMutableObjectViewer.cpp b/Engine/Plugins/Mutable/Source/CustomizableObjectEditor/Private/MuCOE/SMutableObjectViewer.cpp
--- a/Engine/Plugins/Mutable/Source/CustomizableObjectEditor/Private/MuCOE/SMutableObjectViewer.cpp
+++ b/Engine/Plugins/Mutable/Source/CustomizableObjectEditor/Private/MuCOE/SMutableObjectViewer.cpp
@@ -107,12 +107,12 @@ void SMutableObjectViewer::Construct(const FArguments& InArgs, UCustomizableObje
 void SMutableObjectViewer::GenerateMutableGraphPressed()
 {
 	// Convert from Unreal graph to Mutable graph.
-	FCustomizableObjectCompiler Compiler;
+	TSharedRef<FCustomizableObjectCompiler> Compiler = MakeShared<FCustomizableObjectCompiler>();
 	TArray<TSoftObjectPtr<const UTexture>> RuntimeTextures;
 	TArray<FMutableSourceTextureData> CompilerTextures;
 	TArray<TSoftObjectPtr<const UStreamableRenderAsset>> RuntimeMeshes;
 	TArray<FMutableSourceMeshData> CompilerMeshes;
-	mu::Ptr<mu::Node> RootNode = Compiler.Export(CustomizableObject, CompileOptions, RuntimeTextures, CompilerTextures, RuntimeMeshes, CompilerMeshes);
+	mu::Ptr<mu::Node> RootNode = Compiler->Export(CustomizableObject, CompileOptions, RuntimeTextures, CompilerTextures, RuntimeMeshes, CompilerMeshes);
 	if (!RootNode)
 	{
 		// TODO: Show errors
@@ -147,12 +147,12 @@ void SMutableObjectViewer::CompileMutableCodePressed()
 		{
 			CompileOptions.DebugBias = Bias;
 
-			FCustomizableObjectCompiler Compiler;
+			TSharedRef<FCustomizableObjectCompiler> Compiler = MakeShared<FCustomizableObjectCompiler>();
 			RuntimeTextures.Empty();
 			CompilerTextures.Empty();
 			RuntimeMeshes.Empty();
 			CompilerMeshes.Empty();
-			mu::NodePtr RootNode = Compiler.Export(CustomizableObject, CompileOptions, RuntimeTextures, CompilerTextures, RuntimeMeshes, CompilerMeshes);
+			mu::NodePtr RootNode = Compiler->Export(CustomizableObject, CompileOptions, RuntimeTextures, CompilerTextures, RuntimeMeshes, CompilerMeshes);
 			if (!RootNode)
 			{
 				// TODO: Show errors
@@ -161,7 +161,7 @@ void SMutableObjectViewer::CompileMutableCodePressed()
 			}
 
 			// Do the compilation to Mutable Code synchronously.
-			TSharedPtr<FCustomizableObjectCompileRunnable> CompileTask = MakeShareable(new FCustomizableObjectCompileRunnable(RootNode, Compiler.CompilationContext));
+			TSharedPtr<FCustomizableObjectCompileRunnable> CompileTask = MakeShareable(new FCustomizableObjectCompileRunnable(RootNode, Compiler));
 			CompileTask->Options = CompileOptions;
 			CompileTask->ReferencedTextures = CompilerTextures;
 			CompileTask->ReferencedMeshes = CompilerMeshes;
@@ -171,12 +171,12 @@ void SMutableObjectViewer::CompileMutableCodePressed()
 	}
 
 	// Convert from Unreal graph to Mutable graph.
-	FCustomizableObjectCompiler Compiler;
+	TSharedRef<FCustomizableObjectCompiler> Compiler = MakeShared<FCustomizableObjectCompiler>();
 	RuntimeTextures.Empty();
 	CompilerTextures.Empty();
 	RuntimeMeshes.Empty();
 	CompilerMeshes.Empty();
-	mu::NodePtr RootNode = Compiler.Export(CustomizableObject, CompileOptions, RuntimeTextures, CompilerTextures, RuntimeMeshes, CompilerMeshes);
+	mu::NodePtr RootNode = Compiler->Export(CustomizableObject, CompileOptions, RuntimeTextures, CompilerTextures, RuntimeMeshes, CompilerMeshes);
 	if (!RootNode)
 	{
 		// TODO: Show errors
@@ -185,7 +185,7 @@ void SMutableObjectViewer::CompileMutableCodePressed()
 	}
 
 	// Do the compilation to Mutable Code synchronously.
-	TSharedPtr<FCustomizableObjectCompileRunnable> CompileTask = MakeShareable(new FCustomizableObjectCompileRunnable(RootNode, Compiler.CompilationContext));
+	TSharedPtr<FCustomizableObjectCompileRunnable> CompileTask = MakeShareable(new FCustomizableObjectCompileRunnable(RootNode, Compiler));
 	CompileTask->Options = CompileOptions;
 	CompileTask->ReferencedTextures = CompilerTextures;
 	CompileTask->ReferencedMeshes = CompilerMeshes;
diff --git a/Engine/Plugins/Mutable/Source/CustomizableObjectEditor/Private/MuCOE/Thumbnails/CustomizableObjectInstanceThumbnailRenderer.cpp b/Engine/Plugins/Mutable/Source/CustomizableObjectEditor/Private/MuCOE/Thumbnails/CustomizableObjectInstanceThumbnailRenderer.cpp
--- a/Engine/Plugins/Mutable/Source/CustomizableObjectEditor/Private/MuCOE/Thumbnails/CustomizableObjectInstanceThumbnailRenderer.cpp
+++ b/Engine/Plugins/Mutable/Source/CustomizableObjectEditor/Private/MuCOE/Thumbnails/CustomizableObjectInstanceThumbnailRenderer.cpp
@@ -67,7 +67,7 @@ void UCustomizableObjectInstanceThumbnailRenderer::Draw(UObject* Object, int32 X
 			// Use default compilation options?
 			FCompilationOptions Options;
 
-			FCustomizableObjectCompiler Compiler;
+			TSharedRef<FCustomizableObjectCompiler> Compiler = MakeShared<FCustomizableObjectCompiler>();
 			// Compiler.Compile(CustomizableObjectInstance->CustomizableObject, Options, false); // TODO Can not be sync!
 		}
 
diff --git a/Engine/Plugins/Mutable/Source/CustomizableObjectEditor/Private/MuCOE/Thumbnails/CustomizableObjectThumbnailRenderer.cpp b/Engine/Plugins/Mutable/Source/CustomizableObjectEditor/Private/MuCOE/Thumbnails/CustomizableObjectThumbnailRenderer.cpp
--- a/Engine/Plugins/Mutable/Source/CustomizableObjectEditor/Private/MuCOE/Thumbnails/CustomizableObjectThumbnailRenderer.cpp
+++ b/Engine/Plugins/Mutable/Source/CustomizableObjectEditor/Private/MuCOE/Thumbnails/CustomizableObjectThumbnailRenderer.cpp
@@ -56,7 +56,7 @@ void UCustomizableObjectThumbnailRenderer::Draw(UObject* Object, int32 X, int32
 			// Use default compilation options?
 			FCompilationOptions Options = CustomizableObject->GetPrivate()->GetCompileOptions();
 
-			FCustomizableObjectCompiler Compiler;
+			TSharedRef<FCustomizableObjectCompiler> Compiler = MakeShared<FCustomizableObjectCompiler>();
 			// Compiler.Compile(CustomizableObject, Options, false); // TODO Can not be sync!
 		}
 
