[mutable] Protected against rare crash from editor CrashReporter.
#rnx
#tests Editor

diff --git a/Engine/Plugins/Mutable/Source/CustomizableObjectEditor/Private/MuCOE/GraphTraversal.cpp b/Engine/Plugins/Mutable/Source/CustomizableObjectEditor/Private/MuCOE/GraphTraversal.cpp
--- a/Engine/Plugins/Mutable/Source/CustomizableObjectEditor/Private/MuCOE/GraphTraversal.cpp
+++ b/Engine/Plugins/Mutable/Source/CustomizableObjectEditor/Private/MuCOE/GraphTraversal.cpp
@@ -249,7 +249,17 @@ UCustomizableObjectNodeObject* GetRootNode(const UCustomizableObject* Object, bo
 	// Look for the base object node
 	UCustomizableObjectNodeObject* Root = nullptr;
 	TArray<UCustomizableObjectNodeObject*> ObjectNodes;
-	Object->GetPrivate()->GetSource()->GetNodesOfClass<UCustomizableObjectNodeObject>(ObjectNodes);
+
+	if (Object)
+	{
+		if (const UCustomizableObjectPrivate* Private = Object->GetPrivate())
+		{
+			if (TObjectPtr<UEdGraph> Source = Private->GetSource())
+			{
+				Source->GetNodesOfClass<UCustomizableObjectNodeObject>(ObjectNodes);
+			}
+		}
+	}
 
 	bOutMultipleBaseObjectsFound = false;
 
