[Mutable] Bump current supported version for OverlayMaterials.

#rnx

diff --git a/Engine/Plugins/Mutable/Source/CustomizableObject/Internal/MuCO/CustomizableObjectPrivate.h b/Engine/Plugins/Mutable/Source/CustomizableObject/Internal/MuCO/CustomizableObjectPrivate.h
--- a/Engine/Plugins/Mutable/Source/CustomizableObject/Internal/MuCO/CustomizableObjectPrivate.h
+++ b/Engine/Plugins/Mutable/Source/CustomizableObject/Internal/MuCO/CustomizableObjectPrivate.h
@@ -1830,6 +1830,8 @@ public:
 
 		RemoveObsoleteBoolOps,
 
+		AddOverlayMaterials,
+
 		// -----<new versions can be added above this line>--------
 		LastCustomizableObjectVersion
 	};
