[mutable] Move more core code generation to tasks.
This is done to avoid frequent synchronizations with the game thread making some asset compilations very slow and dependent on the framerate.
#rnx

diff --git a/Engine/Plugins/Mutable/Source/MutableTools/Private/MuT/CodeGenerator.cpp b/Engine/Plugins/Mutable/Source/MutableTools/Private/MuT/CodeGenerator.cpp
--- a/Engine/Plugins/Mutable/Source/MutableTools/Private/MuT/CodeGenerator.cpp
+++ b/Engine/Plugins/Mutable/Source/MutableTools/Private/MuT/CodeGenerator.cpp
@@ -107,7 +107,8 @@ namespace mu
 	//---------------------------------------------------------------------------------------------
 	//---------------------------------------------------------------------------------------------
 	CodeGenerator::CodeGenerator(CompilerOptions::Private* Options)
-		: GenerateMeshConstantPipe(TEXT("GenerateMeshConstantPipe"))
+		: LocalPipe(TEXT("CodeGeneratorPipe")),
+		GenerateMeshConstantPipe(TEXT("GenerateMeshConstantPipe"))
 	{
 		CompilerOptions = Options;
 
@@ -205,13 +206,8 @@ namespace mu
 
 		else if (InNode->GetType()->IsA(NodeSurfaceNew::GetStaticType()))
 		{
-			const NodeSurfaceNew* surfNode = static_cast<const NodeSurfaceNew*>(InNode.get());
-
-			// This happens only if we generate a node graph that has a NodeSurfaceNew at the root.
-			FSurfaceGenerationResult SurfResult;
-			FSurfaceGenerationOptions SurfaceOptions(Options);
-			GenerateSurface(SurfResult, SurfaceOptions, surfNode);
-			OutResult.Op = SurfResult.surfaceOp;
+			// This no longer happens with the current tools.
+			check(false);
 			return;
 		}
 
@@ -411,11 +407,11 @@ namespace mu
 		// This can run in any thread.
 		UE::TUniqueLock Lock(GenerateLayoutConstantState.Mutex);
 		
-		TSharedPtr<const FLayout>* it = GenerateLayoutConstantState.GeneratedLayouts.Find({ SourceLayout,MeshIDPrefix });
+		TSharedPtr<const FLayout>* CachedLayout = GenerateLayoutConstantState.GeneratedLayouts.Find({ SourceLayout,MeshIDPrefix });
 
-		if (it)
+		if (CachedLayout)
 		{
-			return *it;
+			return *CachedLayout;
 		}
 
 		TSharedPtr<FLayout> GeneratedLayout = MakeShared<FLayout>();
@@ -594,70 +590,98 @@ namespace mu
 
 	void CodeGenerator::GenerateComponent_New(const FComponentGenerationOptions& Options, FGenericGenerationResult& Result, const NodeComponentNew* InNode)
 	{
-		const NodeComponentNew& node = *InNode;
+		TArray<FLODTask> LODTasks;
+		LODTasks.Reserve(InNode->LODs.Num());
 
-		// Create the expression for each component in this object
-		Ptr<ASTOpAddLOD> LODsOp = new ASTOpAddLOD();
-
-		for (int32 LODIndex = 0; LODIndex < node.LODs.Num(); ++LODIndex)
+		// Launch tasks for each lod, making sure each LOD depends on the previous one.
+		FLODTask LastLODTask;
+		for (int32 LODIndex = 0; LODIndex < InNode->LODs.Num(); ++LODIndex)
 		{
-			if (const NodeLOD* LODNode = node.LODs[LODIndex].get())
+			if (const NodeLOD* LODNode = InNode->LODs[LODIndex].get())
 			{
 				FLODGenerationOptions LODOptions(Options, LODIndex, InNode );
-				FGenericGenerationResult LODResult;
-				Generate_LOD(LODOptions, LODResult, LODNode);
-
-				LODsOp->lods.Emplace(LODsOp, LODResult.Op);
+				
+				bool bWasEmpty = false;
+				LastLODTask = GenerateLOD(LODOptions, LODNode, LastLODTask);
+				LODTasks.Add(LastLODTask);
 			}
 		}
 
-		Ptr<ASTOp> LastInstOp = LODsOp;
+		// Launch the task that generates the component
+		// There could be more concurrency here, but it doesn't look like it is relevant yet.
+		FComponentTask ComponentTask = LocalPipe.Launch(TEXT("MutableComponentNew"),
+			[
+				LODTasks, InNode, Options,
+				this
+			]
+			() mutable
+			{
+				FGenericGenerationResult Result;
 
-		if (NodeScalarPtr pScalarNode = InNode->OverlayMaterial)
-		{
-			Ptr<ASTOpAddOverlayMaterial> AddOverlayMaterialOp = new ASTOpAddOverlayMaterial();
-			AddOverlayMaterialOp->Instance = LastInstOp;
+				// Create the expression for each component in this object
+				Ptr<ASTOpAddLOD> LODsOp = new ASTOpAddLOD();
+				for (FLODTask& LODTask : LODTasks)
+				{
+					FGenericGenerationResult LODResult = LODTask.GetResult();
+					LODsOp->lods.Emplace(LODsOp, LODResult.Op);
+				}
 
-			// Scalar
-			FScalarGenerationResult ScalarResult;
-			GenerateScalar(ScalarResult, Options, pScalarNode);
-			AddOverlayMaterialOp->OverlayMaterialId = ScalarResult.op;
+				Ptr<ASTOp> LastInstOp = LODsOp;
 
-			LastInstOp = AddOverlayMaterialOp;
-		}
+				if (Ptr<NodeScalar> pScalarNode = InNode->OverlayMaterial)
+				{
+					Ptr<ASTOpAddOverlayMaterial> AddOverlayMaterialOp = new ASTOpAddOverlayMaterial();
+					AddOverlayMaterialOp->Instance = LastInstOp;
 
-		Ptr<ASTOpInstanceAdd> InstanceOp = new ASTOpInstanceAdd();
-		InstanceOp->type = EOpType::IN_ADDCOMPONENT;
-		InstanceOp->instance = Options.BaseInstance;
-		InstanceOp->value = LastInstOp;
-		InstanceOp->ExternalId = node.Id;
+					// Scalar
+					FScalarGenerationResult ScalarResult;
+					GenerateScalar(ScalarResult, Options, pScalarNode);
+					AddOverlayMaterialOp->OverlayMaterialId = ScalarResult.op;
 
-		Result.Op = InstanceOp;
+					LastInstOp = AddOverlayMaterialOp;
+				}
 
-		// Add a conditional if this component has conditions
-		for (const FirstPassGenerator::FComponent& Component: FirstPass.Components)
-		{
-			if (Component.Component != InNode)
-			{
-				continue;
-			}
+				Ptr<ASTOpInstanceAdd> InstanceOp = new ASTOpInstanceAdd();
+				InstanceOp->type = EOpType::IN_ADDCOMPONENT;
+				InstanceOp->instance = Options.BaseInstance;
+				InstanceOp->value = LastInstOp;
+				InstanceOp->ExternalId = InNode->Id;
 
-			if (Component.ComponentCondition || Component.ObjectCondition)
-			{
-				// TODO: This could be done earlier?
-				Ptr<ASTOpBoolAnd> ConditionOp = new ASTOpBoolAnd();
-				ConditionOp->A = Component.ObjectCondition;
-				ConditionOp->B = Component.ComponentCondition;
-
-				Ptr<ASTOpConditional> IfOp = new ASTOpConditional();
-				IfOp->type = EOpType::IN_CONDITIONAL;
-				IfOp->no = Options.BaseInstance;
-				IfOp->yes = Result.Op;
-				IfOp->condition = ConditionOp;
-
-				Result.Op = IfOp;
-			}
-		}
+				Result.Op = InstanceOp;
+
+				// Add a conditional if this component has conditions
+				for (const FirstPassGenerator::FComponent& Component : FirstPass.Components)
+				{
+					if (Component.Component != InNode)
+					{
+						continue;
+					}
+
+					if (Component.ComponentCondition || Component.ObjectCondition)
+					{
+						// TODO: This could be done earlier?
+						Ptr<ASTOpBoolAnd> ConditionOp = new ASTOpBoolAnd();
+						ConditionOp->A = Component.ObjectCondition;
+						ConditionOp->B = Component.ComponentCondition;
+
+						Ptr<ASTOpConditional> IfOp = new ASTOpConditional();
+						IfOp->type = EOpType::IN_CONDITIONAL;
+						IfOp->no = Options.BaseInstance;
+						IfOp->yes = Result.Op;
+						IfOp->condition = ConditionOp;
+
+						Result.Op = IfOp;
+					}
+				}
+
+				return Result;
+			},
+			LODTasks
+			);
+
+		// Sync point: we currently don't support task-based generation beyond components so we wait here.
+		// Otherwise we could return the component task to be chained with higher level tasks.
+		Result = ComponentTask.GetResult();
 	}
 
 
@@ -870,19 +894,12 @@ namespace mu
 	}
 
 
-    //---------------------------------------------------------------------------------------------
-    void CodeGenerator::GenerateSurface( FSurfaceGenerationResult& SurfaceResult, 
-										 const FSurfaceGenerationOptions& Options,
-                                         Ptr<const NodeSurfaceNew> SurfaceNode )
+	FSurfaceTask CodeGenerator::GenerateSurface( const FSurfaceGenerationOptions& Options, Ptr<const NodeSurfaceNew> SurfaceNode, FLODTask PreviousLODTask )
     {
         MUTABLE_CPUPROFILER_SCOPE(GenerateSurface);
 
-        // Build a series of operations to assemble the surface
-        Ptr<ASTOp> LastSurfOp;
-
         // Generate the mesh
         //------------------------------------------------------------------------
-        FMeshGenerationResult MeshResults;
 
 		// We don't add the mesh here, since it will be added directly at the top of the
 		// component expression in the NodeComponentNew generator with the right merges
@@ -911,11 +928,10 @@ namespace mu
 				continue;
 			}
 
-			if (Surface.ResultSurfaceOp)
+			if (Surface.ResultSurfaceTask.IsValid())
 			{
 				// Reuse the entire surface
-				SurfaceResult.surfaceOp = Surface.ResultSurfaceOp;
-				return;
+				return Surface.ResultSurfaceTask;
 			}
 			else
 			{
@@ -926,17 +942,9 @@ namespace mu
 
         if (TargetSurfaces.IsEmpty())
         {
-            return;
+            return UE::Tasks::MakeCompletedTask<FSurfaceGenerationResult>();
         }
 
-		// This assumes that the lods are processed in order. It checks it this way because some platforms may have empty lods at the top.
-		const bool bIsBaseForSharedSurface = 
-			SurfaceNode->SharedSurfaceId != INDEX_NONE && 
-			!SharedMeshOptionsMap.Contains(SurfaceNode->SharedSurfaceId);
-
-		// If this is true, we will reuse the surface properties from a higher LOD, se we can skip the generation of material properties and images.
-		const bool bShareSurface = SurfaceNode->SharedSurfaceId != INDEX_NONE && !bIsBaseForSharedSurface;
-
 		// Gather all modifiers that apply to this surface
 		TArray<FirstPassGenerator::FModifier> Modifiers;
 		constexpr bool bModifiersForBeforeOperations = false;
@@ -949,789 +957,889 @@ namespace mu
 		// This pass on the modifiers is only to detect errors that cannot be detected at the point they are applied.
 		CheckModifiersForSurface(*SurfaceNode, Modifiers, Options.LODIndex);
 		
-		TBitArray<> LayoutFromExtension;
-        //if (SurfaceNode->Mesh)
-        {
-            MUTABLE_CPUPROFILER_SCOPE(SurfaceMesh);
-
-            // Generate the mesh
-			FMeshGenerationStaticOptions MeshStaticOptions(ComponentId, Options.LODIndex);
-			MeshStaticOptions.ActiveTags = SurfaceNode->Tags;
-			MeshStaticOptions.State = Options.State;
-			FMeshGenerationDynamicOptions MeshDynamicOptions;
-			MeshDynamicOptions.bLayouts = true;
-
-			const FMeshGenerationResult* SharedMeshResults = nullptr;
-			if (bShareSurface)
-			{
-				// Do we have the surface we need to share it with?
-				SharedMeshResults = SharedMeshOptionsMap.Find(SurfaceNode->SharedSurfaceId);
-				check(SharedMeshResults);
-
-				// Override the layouts with the ones from the surface we share
-				MeshDynamicOptions.OverrideLayouts = SharedMeshResults->GeneratedLayouts;
-			}
-
-			// Normalize UVs if we're going to work with images and layouts.
-			// TODO: This should come from per-layout settings!
-			const bool bNormalizeUVs = false; // !SurfaceNode->Images.IsEmpty();
-			MeshDynamicOptions.bNormalizeUVs = bNormalizeUVs;
-
-			// Ensure UV islands remain within their main layout block on lower LODs to avoid unexpected reordering 
-			// of the layout blocks when reusing a surface between LODs. Used to fix small displacements on vertices
-			// that may cause them to fall on a different block.
-			MeshDynamicOptions.bClampUVIslands = bShareSurface;
-
-            FMeshTask MeshTask = GenerateMesh(MeshStaticOptions, UE::Tasks::MakeCompletedTask<FMeshGenerationDynamicOptions>(MeshDynamicOptions), SurfaceNode->Mesh);
+        // Generate the mesh
+		FMeshGenerationStaticOptions MeshStaticOptions(ComponentId, Options.LODIndex);
+		MeshStaticOptions.ActiveTags = SurfaceNode->Tags;
+		MeshStaticOptions.State = Options.State;
+		FMeshGenerationDynamicOptions MeshDynamicOptions;
+		MeshDynamicOptions.bLayouts = true;
+
+		// Normalize UVs if we're going to work with images and layouts.
+		// TODO: This should come from per-layout settings!
+		const bool bNormalizeUVs = false; // !SurfaceNode->Images.IsEmpty();
+		MeshDynamicOptions.bNormalizeUVs = bNormalizeUVs;
+
+		// The options depenend on the shared surface being generated, so we need to add the previous lod dependency
+		FMeshOptionsTask MeshOptionsTask = LocalPipe.Launch(TEXT("MutableSurfaceMeshOptions"),
+			[MeshDynamicOptions, SharedSurfaceId=SurfaceNode->SharedSurfaceId, this]() mutable
+			{
+				// This assumes that the lods are processed in order. It checks it this way because some platforms may have empty lods at the top.
+				bool bIsBaseForSharedSurface = SharedSurfaceId != INDEX_NONE;
+				if (bIsBaseForSharedSurface)
+				{
+					UE::TUniqueLock Lock(SharedMeshOptionsMapMutex);
+					bIsBaseForSharedSurface = !SharedMeshOptionsMap.Contains(SharedSurfaceId);
+				}
 
-			// Apply the modifier for the post-normal operations stage.
-			MeshTask = ApplyMeshModifiers(Modifiers,
-				MeshStaticOptions, UE::Tasks::MakeCompletedTask<FMeshGenerationDynamicOptions>(MeshDynamicOptions),
-				MeshTask, SharedMeshResults, SurfaceNode->GetMessageContext(), nullptr);
+				// If this is true, we will reuse the surface properties from a higher LOD, se we can skip the generation of material properties and images.
+				const bool bShareSurface = SharedSurfaceId != INDEX_NONE && !bIsBaseForSharedSurface;
 
-			// TODO
-			MeshResults = MeshTask.GetResult();
+				const FMeshGenerationResult* SharedMeshResults = nullptr;
+				if (bShareSurface)
+				{
+					UE::TUniqueLock Lock(SharedMeshOptionsMapMutex);
 
-			// Base mesh is allowed to be missing, aggregate all layouts and operations per layout indices in the
-			// generated mesh, base and extends.
-			TArray<CodeGenerator::FGeneratedLayout> SurfaceReferenceLayouts;
-			TArray<Ptr<ASTOp>> SurfaceLayoutOps;
+					// Do we have the surface we need to share it with?
+					SharedMeshResults = SharedMeshOptionsMap.Find(SharedSurfaceId);
+					check(SharedMeshResults);
 
-			int32 MaxLayoutNum = MeshResults.GeneratedLayouts.Num();
-			for (const FMeshGenerationResult::FExtraLayouts& ExtraLayoutData : MeshResults.ExtraMeshLayouts)
-			{
-				MaxLayoutNum = FMath::Max(MaxLayoutNum, ExtraLayoutData.GeneratedLayouts.Num());
-			}
-
-			SurfaceReferenceLayouts.SetNum(MaxLayoutNum);
-			SurfaceLayoutOps.SetNum(MaxLayoutNum);
-			LayoutFromExtension.Init(false, MaxLayoutNum);
+					// Override the layouts with the ones from the surface we share
+					MeshDynamicOptions.OverrideLayouts = SharedMeshResults->GeneratedLayouts;
+				}
 
-			// Add layouts form the base mesh.	
-            for (int32 LayoutIndex = 0; LayoutIndex < MeshResults.GeneratedLayouts.Num(); ++LayoutIndex)
-			{
-				if (!MeshResults.GeneratedLayouts[LayoutIndex].Layout)
+				// Ensure UV islands remain within their main layout block on lower LODs to avoid unexpected reordering 
+				// of the layout blocks when reusing a surface between LODs. Used to fix small displacements on vertices
+				// that may cause them to fall on a different block.
+				MeshDynamicOptions.bClampUVIslands = bShareSurface;
+
+				return MeshDynamicOptions;
+			},
+			PreviousLODTask
+		);
+
+        FMeshTask MeshTask = GenerateMesh(MeshStaticOptions, MeshOptionsTask, SurfaceNode->Mesh);
+
+		// Apply the modifier for the post-normal operations stage.
+		MeshTask = ApplyMeshModifiers(Modifiers,
+			MeshStaticOptions, UE::Tasks::MakeCompletedTask<FMeshGenerationDynamicOptions>(MeshDynamicOptions),
+			MeshTask, SurfaceNode->SharedSurfaceId, SurfaceNode->GetMessageContext(), nullptr);
+
+		FSurfaceTask SurfaceTask = LocalPipe.Launch(TEXT("MutableSurface"),
+			[
+				MeshTask,
+					SurfaceNode,
+					ComponentId, Options, 
+					Modifiers,
+					TargetSurfaces,
+					// We need to call some local methods which should be fine sinc4e. 
+					this
+			] 
+			() mutable
+			{
+				FMeshGenerationResult MeshResults = MeshTask.GetResult();
+
+				// Base mesh is allowed to be missing, aggregate all layouts and operations per layout indices in the
+				// generated mesh, base and extends.
+				TArray<FGeneratedLayout> SurfaceReferenceLayouts;
+				TArray<Ptr<ASTOp>> SurfaceLayoutOps;
+
+				int32 MaxLayoutNum = MeshResults.GeneratedLayouts.Num();
+				for (const FMeshGenerationResult::FExtraLayouts& ExtraLayoutData : MeshResults.ExtraMeshLayouts)
 				{
-					continue;
+					MaxLayoutNum = FMath::Max(MaxLayoutNum, ExtraLayoutData.GeneratedLayouts.Num());
 				}
 
-				SurfaceReferenceLayouts[LayoutIndex] = MeshResults.GeneratedLayouts[LayoutIndex];
+				SurfaceReferenceLayouts.SetNum(MaxLayoutNum);
+				SurfaceLayoutOps.SetNum(MaxLayoutNum);
 
-				bool bSharedHasThisLayout = SharedMeshResults 
-					&& 
-					SharedMeshResults->LayoutOps.IsValidIndex(LayoutIndex)
-					&&
-					SharedMeshResults->LayoutOps[LayoutIndex];
+				TBitArray<> LayoutFromExtension;
+				LayoutFromExtension.Init(false, MaxLayoutNum);
 
-				if (bSharedHasThisLayout)
+				// Scope for access control to shared data
+				bool bIsBaseForSharedSurface = false;
+				bool bShareSurface = false;
+				TArray<Ptr<ASTOp>> SharedResultLayoutOps;
 				{
-					SurfaceLayoutOps[LayoutIndex] = SharedMeshResults->LayoutOps[LayoutIndex];
-				}
-				else
-				{
-					Ptr<ASTOpConstantResource> ConstantLayoutOp = new ASTOpConstantResource();
-					ConstantLayoutOp->Type = EOpType::LA_CONSTANT;
+					UE::TUniqueLock Lock(SharedMeshOptionsMapMutex);
 
-					ConstantLayoutOp->SetValue(
-							SurfaceReferenceLayouts[LayoutIndex].Layout, 
-							CompilerOptions->OptimisationOptions.DiskCacheContext);
-					SurfaceLayoutOps[LayoutIndex] = ConstantLayoutOp;
-				}
-			}
+					bIsBaseForSharedSurface = SurfaceNode->SharedSurfaceId!=INDEX_NONE 
+						&& 
+						!SharedMeshOptionsMap.Contains(SurfaceNode->SharedSurfaceId);
 
-			// Add extra layouts. In case there is a missing reference layout, the first visited will
-			// take the role.
-			for (const FMeshGenerationResult::FExtraLayouts& ExtraLayoutsData : MeshResults.ExtraMeshLayouts)
-			{
-				if (!ExtraLayoutsData.MeshFragment)
-				{
-					// No mesh to add, we assume there are no layouts to add either.
-					check(ExtraLayoutsData.GeneratedLayouts.IsEmpty());
-					continue;
-				}
+					// If this is true, we will reuse the surface properties from a higher LOD, se we can skip the generation of material properties and images.
+					bShareSurface = SurfaceNode->SharedSurfaceId != INDEX_NONE && !bIsBaseForSharedSurface;
 
-				const TArray<CodeGenerator::FGeneratedLayout>& ExtraGeneratedLayouts = ExtraLayoutsData.GeneratedLayouts;
-            	for (int32 LayoutIndex = 0; LayoutIndex < ExtraGeneratedLayouts.Num(); ++LayoutIndex)
-				{
-					if (!ExtraGeneratedLayouts[LayoutIndex].Layout)
+					const FMeshGenerationResult* SharedMeshResults = nullptr;
+					if (bShareSurface)
 					{
-						continue;
-					}
+						// Do we have the surface we need to share it with?
+						SharedMeshResults = SharedMeshOptionsMap.Find(SurfaceNode->SharedSurfaceId);
+						check(SharedMeshResults);
 
-					bool bLayoutSetByThisExtension = false;
-					if (!SurfaceReferenceLayouts[LayoutIndex].Layout)
-					{
-						// This Layout slot is not set by the base surface, set it as reference.
-						SurfaceReferenceLayouts[LayoutIndex] = ExtraGeneratedLayouts[LayoutIndex];
-						bLayoutSetByThisExtension = true;
+						SharedResultLayoutOps = SharedMeshResults->LayoutOps;
 
-						LayoutFromExtension[LayoutIndex] = bLayoutSetByThisExtension;
 					}
-					
-					if (SharedMeshResults)
+
+					// Add layouts form the base mesh.	
+					for (int32 LayoutIndex = 0; LayoutIndex < MeshResults.GeneratedLayouts.Num(); ++LayoutIndex)
 					{
-						if (!SurfaceLayoutOps[LayoutIndex] && bLayoutSetByThisExtension)
+						if (!MeshResults.GeneratedLayouts[LayoutIndex].Layout)
 						{
-							check(SharedMeshResults->LayoutOps.IsValidIndex(LayoutIndex));
-							SurfaceLayoutOps[LayoutIndex] = SharedMeshResults->LayoutOps[LayoutIndex];
+							continue;
 						}
-					}
-					else
-					{
-						Ptr<ASTOpConstantResource> LayoutFragmentConstantOp = new ASTOpConstantResource();
-						LayoutFragmentConstantOp->Type = EOpType::LA_CONSTANT;
 
-						LayoutFragmentConstantOp->SetValue(
-								ExtraLayoutsData.GeneratedLayouts[LayoutIndex].Layout,
-								CompilerOptions->OptimisationOptions.DiskCacheContext);
+						SurfaceReferenceLayouts[LayoutIndex] = MeshResults.GeneratedLayouts[LayoutIndex];
 
-						Ptr<ASTOpLayoutMerge> LayoutMergeOp = new ASTOpLayoutMerge();
-						// Base may be null if the base does not have  a mesh with a layout at LayoutIndex.
-						// In that case, when applying the condition this can generate null layouts.
-						LayoutMergeOp->Base = SurfaceLayoutOps[LayoutIndex];
-						LayoutMergeOp->Added = LayoutFragmentConstantOp;
+						bool bSharedHasThisLayout = SharedMeshResults
+							&&
+							SharedMeshResults->LayoutOps.IsValidIndex(LayoutIndex)
+							&&
+							SharedMeshResults->LayoutOps[LayoutIndex];
 
-						if (ExtraLayoutsData.Condition)
+						if (bSharedHasThisLayout)
 						{
-							Ptr<ASTOpConditional> ConditionalOp = new ASTOpConditional();
-							ConditionalOp->type = EOpType::LA_CONDITIONAL;
-							ConditionalOp->no = SurfaceLayoutOps[LayoutIndex];
-							ConditionalOp->yes = LayoutMergeOp;
-							ConditionalOp->condition = ExtraLayoutsData.Condition;
-
-							SurfaceLayoutOps[LayoutIndex] = ConditionalOp;
+							SurfaceLayoutOps[LayoutIndex] = SharedMeshResults->LayoutOps[LayoutIndex];
 						}
 						else
 						{
-							SurfaceLayoutOps[LayoutIndex] = LayoutMergeOp;
+							Ptr<ASTOpConstantResource> ConstantLayoutOp = new ASTOpConstantResource();
+							ConstantLayoutOp->Type = EOpType::LA_CONSTANT;
+
+							ConstantLayoutOp->SetValue(
+								SurfaceReferenceLayouts[LayoutIndex].Layout,
+								CompilerOptions->OptimisationOptions.DiskCacheContext);
+							SurfaceLayoutOps[LayoutIndex] = ConstantLayoutOp;
 						}
 					}
 				}
-			}
 
-			Ptr<ASTOp> LastMeshOp = MeshResults.MeshOp;
-
-			check(SurfaceReferenceLayouts.Num() == SurfaceLayoutOps.Num());
-            for (int32 LayoutIndex = 0; LayoutIndex < SurfaceReferenceLayouts.Num(); ++LayoutIndex)
-			{
-				if (!SurfaceReferenceLayouts[LayoutIndex].Layout)
+				// Add extra layouts. In case there is a missing reference layout, the first visited will
+				// take the role.
+				for (const FMeshGenerationResult::FExtraLayouts& ExtraLayoutsData : MeshResults.ExtraMeshLayouts)
 				{
-					continue;
-				}
+					if (!ExtraLayoutsData.MeshFragment)
+					{
+						// No mesh to add, we assume there are no layouts to add either.
+						check(ExtraLayoutsData.GeneratedLayouts.IsEmpty());
+						continue;
+					}
 
-				if (SurfaceReferenceLayouts[LayoutIndex].Layout->GetLayoutPackingStrategy() == mu::EPackStrategy::Overlay)
-				{
-					continue;
-				}
+					const TArray<FGeneratedLayout>& ExtraGeneratedLayouts = ExtraLayoutsData.GeneratedLayouts;
+					for (int32 LayoutIndex = 0; LayoutIndex < ExtraGeneratedLayouts.Num(); ++LayoutIndex)
+					{
+						if (!ExtraGeneratedLayouts[LayoutIndex].Layout)
+						{
+							continue;
+						}
 
-				// Add layout packing instructions
-				if (!SharedMeshResults)
-				{
-					// Make sure we removed unnecessary blocks
-					Ptr<ASTOpLayoutFromMesh> ExtractOp = new ASTOpLayoutFromMesh();
-					ExtractOp->Mesh = LastMeshOp;
-					check(LayoutIndex < 256);
-					ExtractOp->LayoutIndex = uint8(LayoutIndex);
-
-					Ptr<ASTOpLayoutRemoveBlocks> RemoveOp = new ASTOpLayoutRemoveBlocks();
-					RemoveOp->Source = SurfaceLayoutOps[LayoutIndex];
-					RemoveOp->ReferenceLayout = ExtractOp;
-					SurfaceLayoutOps[LayoutIndex] = RemoveOp;
-
-					// Pack uv blocks
-					Ptr<ASTOpLayoutPack> LayoutPackOp = new ASTOpLayoutPack();
-					LayoutPackOp->Source = SurfaceLayoutOps[LayoutIndex];
-					SurfaceLayoutOps[LayoutIndex] = LayoutPackOp;
-				}
+						bool bLayoutSetByThisExtension = false;
+						if (!SurfaceReferenceLayouts[LayoutIndex].Layout)
+						{
+							// This Layout slot is not set by the base surface, set it as reference.
+							SurfaceReferenceLayouts[LayoutIndex] = ExtraGeneratedLayouts[LayoutIndex];
+							bLayoutSetByThisExtension = true;
 
-				// Create the expression to apply the layout to the mesh
-				{
-					Ptr<ASTOpMeshApplyLayout> ApplyLayoutOp = new ASTOpMeshApplyLayout();
-					ApplyLayoutOp->Mesh = LastMeshOp;
-					ApplyLayoutOp->Layout = SurfaceLayoutOps[LayoutIndex];
-					ApplyLayoutOp->Channel = (uint16)LayoutIndex;
-					
-					LastMeshOp = ApplyLayoutOp;
-				}
-			}
+							LayoutFromExtension[LayoutIndex] = bLayoutSetByThisExtension;
+						}
 
-			MeshResults.GeneratedLayouts = MoveTemp(SurfaceReferenceLayouts);
-			MeshResults.LayoutOps = MoveTemp(SurfaceLayoutOps); 
+						if (bShareSurface)
+						{
+							if (!SurfaceLayoutOps[LayoutIndex] && bLayoutSetByThisExtension)
+							{
+								check(SharedResultLayoutOps.IsValidIndex(LayoutIndex));
+								SurfaceLayoutOps[LayoutIndex] = SharedResultLayoutOps[LayoutIndex];
+							}
+						}
+						else
+						{
+							Ptr<ASTOpConstantResource> LayoutFragmentConstantOp = new ASTOpConstantResource();
+							LayoutFragmentConstantOp->Type = EOpType::LA_CONSTANT;
 
-            // Store in the surface for later use.
-			for (FirstPassGenerator::FSurface* TargetSurface : TargetSurfaces)
-			{
-				TargetSurface->ResultMeshOp = LastMeshOp;
-			}
-        }
+							LayoutFragmentConstantOp->SetValue(
+								ExtraLayoutsData.GeneratedLayouts[LayoutIndex].Layout,
+								CompilerOptions->OptimisationOptions.DiskCacheContext);
 
-        // Create the expression for each texture, if we are not reusing the surface from another LOD.
-        //------------------------------------------------------------------------
-		if (!bShareSurface)
-		{
-			for (int32 ImageIndex = 0; ImageIndex < SurfaceNode->Images.Num(); ++ImageIndex)
-			{
-				MUTABLE_CPUPROFILER_SCOPE(SurfaceTexture);
+							Ptr<ASTOpLayoutMerge> LayoutMergeOp = new ASTOpLayoutMerge();
+							// Base may be null if the base does not have  a mesh with a layout at LayoutIndex.
+							// In that case, when applying the condition this can generate null layouts.
+							LayoutMergeOp->Base = SurfaceLayoutOps[LayoutIndex];
+							LayoutMergeOp->Added = LayoutFragmentConstantOp;
 
-				// Any image-specific format or mipmapping needs to be applied at the end
-				Ptr<NodeImageMipmap> mipmapNode;
-				Ptr<NodeImageFormat> formatNode;
-				Ptr<NodeImageSwizzle> swizzleNode;
+							if (ExtraLayoutsData.Condition)
+							{
+								Ptr<ASTOpConditional> ConditionalOp = new ASTOpConditional();
+								ConditionalOp->type = EOpType::LA_CONDITIONAL;
+								ConditionalOp->no = SurfaceLayoutOps[LayoutIndex];
+								ConditionalOp->yes = LayoutMergeOp;
+								ConditionalOp->condition = ExtraLayoutsData.Condition;
 
-				bool bFound = false;
-				Ptr<NodeImage> pImageNode = SurfaceNode->Images[ImageIndex].Image;
+								SurfaceLayoutOps[LayoutIndex] = ConditionalOp;
+							}
+							else
+							{
+								SurfaceLayoutOps[LayoutIndex] = LayoutMergeOp;
+							}
+						}
+					}
+				}
 
-				while (!bFound && pImageNode)
+				Ptr<ASTOp> LastMeshOp = MeshResults.MeshOp;
+
+				check(SurfaceReferenceLayouts.Num() == SurfaceLayoutOps.Num());
+				for (int32 LayoutIndex = 0; LayoutIndex < SurfaceReferenceLayouts.Num(); ++LayoutIndex)
 				{
-					if (pImageNode->GetType()==NodeImageMipmap::GetStaticType())
+					if (!SurfaceReferenceLayouts[LayoutIndex].Layout)
 					{
-						NodeImageMipmap* tm = static_cast<NodeImageMipmap*>(pImageNode.get());
-						if (!mipmapNode) mipmapNode = tm;
-						pImageNode = tm->GetSource();
+						continue;
 					}
-					else if (pImageNode->GetType() == NodeImageFormat::GetStaticType())
+
+					if (SurfaceReferenceLayouts[LayoutIndex].Layout->GetLayoutPackingStrategy() == mu::EPackStrategy::Overlay)
 					{
-						NodeImageFormat* tf = static_cast<NodeImageFormat*>(pImageNode.get());
-						if (!formatNode) formatNode = tf;
-						pImageNode = tf->GetSource();
+						continue;
 					}
-					else if (pImageNode->GetType() == NodeImageSwizzle::GetStaticType())
-					{
-						NodeImageSwizzle* ts = static_cast<NodeImageSwizzle*>(pImageNode.get());
-
-						if (!ts->GetPrivate()->m_sources.IsEmpty())
-						{
-							NodeImage* Source = ts->GetSource(0).get();
-							
-							bool bAllSourcesAreTheSame = true;
-							for (int32 SourceIndex=1; SourceIndex<ts->GetPrivate()->m_sources.Num(); ++SourceIndex)
-							{
-								bAllSourcesAreTheSame = bAllSourcesAreTheSame && (Source == ts->GetSource(SourceIndex));
-							}
 
-							if (!swizzleNode && bAllSourcesAreTheSame)
-							{
-								swizzleNode = ts;
-								pImageNode = Source;
-							}
-							else
-							{
-								bFound = true;
-							}
-						}
-						else
-						{
-							// break loop if swizzle has no sources.
-							bFound = true;
-						}
+					// Add layout packing instructions
+					if (!bShareSurface)
+					{
+						// Make sure we removed unnecessary blocks
+						Ptr<ASTOpLayoutFromMesh> ExtractOp = new ASTOpLayoutFromMesh();
+						ExtractOp->Mesh = LastMeshOp;
+						check(LayoutIndex < 256);
+						ExtractOp->LayoutIndex = uint8(LayoutIndex);
+
+						Ptr<ASTOpLayoutRemoveBlocks> RemoveOp = new ASTOpLayoutRemoveBlocks();
+						RemoveOp->Source = SurfaceLayoutOps[LayoutIndex];
+						RemoveOp->ReferenceLayout = ExtractOp;
+						SurfaceLayoutOps[LayoutIndex] = RemoveOp;
+
+						// Pack uv blocks
+						Ptr<ASTOpLayoutPack> LayoutPackOp = new ASTOpLayoutPack();
+						LayoutPackOp->Source = SurfaceLayoutOps[LayoutIndex];
+						SurfaceLayoutOps[LayoutIndex] = LayoutPackOp;
 					}
-					else
+
+					// Create the expression to apply the layout to the mesh
 					{
-						bFound = true;
+						Ptr<ASTOpMeshApplyLayout> ApplyLayoutOp = new ASTOpMeshApplyLayout();
+						ApplyLayoutOp->Mesh = LastMeshOp;
+						ApplyLayoutOp->Layout = SurfaceLayoutOps[LayoutIndex];
+						ApplyLayoutOp->Channel = (uint16)LayoutIndex;
+
+						LastMeshOp = ApplyLayoutOp;
 					}
 				}
 
-				if (bFound)
-				{
-					const NodeSurfaceNew::FImageData& ImageData = SurfaceNode->Images[ImageIndex];
+				MeshResults.GeneratedLayouts = MoveTemp(SurfaceReferenceLayouts);
+				MeshResults.LayoutOps = MoveTemp(SurfaceLayoutOps);
 
-					const int32 LayoutIndex = ImageData.LayoutIndex;
+				// Store in the surface for later use.
+				for (FirstPassGenerator::FSurface* TargetSurface : TargetSurfaces)
+				{
+					TargetSurface->ResultMeshOp = LastMeshOp;
+				}
 
-					// If the layout index has been set to negative, it means we should ignore the layout for this image.
-					CompilerOptions::TextureLayoutStrategy ImageLayoutStrategy = (LayoutIndex < 0)
-						? CompilerOptions::TextureLayoutStrategy::None
-						: CompilerOptions::TextureLayoutStrategy::Pack
-						;
+				// Build a series of operations to assemble the surface
+				Ptr<ASTOp> LastSurfOp;
 
-					if (ImageLayoutStrategy == CompilerOptions::TextureLayoutStrategy::None)
+				// Create the expression for each texture, if we are not reusing the surface from another LOD.
+				//------------------------------------------------------------------------
+				if (!bShareSurface)
+				{
+					for (int32 ImageIndex = 0; ImageIndex < SurfaceNode->Images.Num(); ++ImageIndex)
 					{
-						// Generate the image
-						FImageGenerationOptions ImageOptions(ComponentId, Options.LODIndex);
-						ImageOptions.State = Options.State;
-						ImageOptions.ImageLayoutStrategy = ImageLayoutStrategy;
-						ImageOptions.ActiveTags = SurfaceNode->Tags;
-						ImageOptions.RectSize = { 0, 0 };
-						FImageGenerationResult Result;
-						GenerateImage(ImageOptions, Result, pImageNode);
-						Ptr<ASTOp> imageAd = Result.op;
-
-						// Placeholder block. Ideally this should be the actual image size
-						constexpr int32 FakeLayoutSize = 256;
-						FIntPoint GridSize(FakeLayoutSize, FakeLayoutSize);
-						FLayoutBlockDesc LayoutBlockDesc;
-						LayoutBlockDesc.BlockPixelsX = 1;
-						LayoutBlockDesc.BlockPixelsY = 1;
-						box< FIntVector2 > RectInCells;
-						RectInCells.min = { 0,0 };
-						RectInCells.size = { FakeLayoutSize ,FakeLayoutSize };
-						
-						imageAd = ApplyImageBlockModifiers(Modifiers, ImageOptions, imageAd, ImageData, GridSize, LayoutBlockDesc, RectInCells, SurfaceNode->GetMessageContext());
-						
-						check(imageAd);
-
-						if (swizzleNode)
-						{
-							Ptr<ASTOpImageSwizzle> fop = new ASTOpImageSwizzle();
-							fop->Format = swizzleNode->GetPrivate()->m_format;
-							fop->Sources[0] = imageAd;
-							fop->Sources[1] = imageAd;
-							fop->Sources[2] = imageAd;
-							fop->Sources[3] = imageAd;
-							fop->SourceChannels[0] = swizzleNode->GetPrivate()->m_sourceChannels[0];
-							fop->SourceChannels[1] = swizzleNode->GetPrivate()->m_sourceChannels[1];
-							fop->SourceChannels[2] = swizzleNode->GetPrivate()->m_sourceChannels[2];
-							fop->SourceChannels[3] = swizzleNode->GetPrivate()->m_sourceChannels[3];
-							check(fop->Format != EImageFormat::None);
-							imageAd = fop;
-						}
+						MUTABLE_CPUPROFILER_SCOPE(SurfaceTexture);
 
-						if (mipmapNode)
-						{
-							Ptr<ASTOpImageMipmap> op = new ASTOpImageMipmap();
-							op->Levels = 0;
-							op->Source = imageAd;
-							op->BlockLevels = 0;
-
-							op->AddressMode = mipmapNode->GetPrivate()->m_settings.AddressMode;
-							op->FilterType = mipmapNode->GetPrivate()->m_settings.FilterType;
-							imageAd = op;
-						}
+						// Any image-specific format or mipmapping needs to be applied at the end
+						Ptr<NodeImageMipmap> mipmapNode;
+						Ptr<NodeImageFormat> formatNode;
+						Ptr<NodeImageSwizzle> swizzleNode;
 
-						if (formatNode)
+						bool bFound = false;
+						Ptr<NodeImage> pImageNode = SurfaceNode->Images[ImageIndex].Image;
+
+						while (!bFound && pImageNode)
 						{
-							Ptr<ASTOpImagePixelFormat> fop = new ASTOpImagePixelFormat();
-							fop->Format = formatNode->GetPrivate()->m_format;
-							fop->FormatIfAlpha = formatNode->GetPrivate()->m_formatIfAlpha;
-							fop->Source = imageAd;
-							check(fop->Format != EImageFormat::None);
-							imageAd = fop;
-						}
+							if (pImageNode->GetType() == NodeImageMipmap::GetStaticType())
+							{
+								NodeImageMipmap* tm = static_cast<NodeImageMipmap*>(pImageNode.get());
+								if (!mipmapNode) mipmapNode = tm;
+								pImageNode = tm->GetSource();
+							}
+							else if (pImageNode->GetType() == NodeImageFormat::GetStaticType())
+							{
+								NodeImageFormat* tf = static_cast<NodeImageFormat*>(pImageNode.get());
+								if (!formatNode) formatNode = tf;
+								pImageNode = tf->GetSource();
+							}
+							else if (pImageNode->GetType() == NodeImageSwizzle::GetStaticType())
+							{
+								NodeImageSwizzle* ts = static_cast<NodeImageSwizzle*>(pImageNode.get());
 
-						Ptr<ASTOpInstanceAdd> op = new ASTOpInstanceAdd();
-						op->type = EOpType::IN_ADDIMAGE;
-						op->instance = LastSurfOp;
-						op->value = imageAd;
-						op->name = SurfaceNode->Images[ImageIndex].Name;
+								if (!ts->GetPrivate()->m_sources.IsEmpty())
+								{
+									NodeImage* Source = ts->GetSource(0).get();
 
-						LastSurfOp = op;
-					}
+									bool bAllSourcesAreTheSame = true;
+									for (int32 SourceIndex = 1; SourceIndex < ts->GetPrivate()->m_sources.Num(); ++SourceIndex)
+									{
+										bAllSourcesAreTheSame = bAllSourcesAreTheSame && (Source == ts->GetSource(SourceIndex));
+									}
 
-					else if (ImageLayoutStrategy == CompilerOptions::TextureLayoutStrategy::Pack) //-V547
-					{						
-						if (LayoutIndex >= MeshResults.GeneratedLayouts.Num() ||
-							LayoutIndex >= MeshResults.LayoutOps.Num())
-						{
-							ErrorLog->Add("Missing layout in object, or its parent.", ELMT_ERROR, SurfaceNode->GetMessageContext());
+									if (!swizzleNode && bAllSourcesAreTheSame)
+									{
+										swizzleNode = ts;
+										pImageNode = Source;
+									}
+									else
+									{
+										bFound = true;
+									}
+								}
+								else
+								{
+									// break loop if swizzle has no sources.
+									bFound = true;
+								}
+							}
+							else
+							{
+								bFound = true;
+							}
 						}
-						else
+
+						if (bFound)
 						{
-							TSharedPtr<const FLayout> pLayout = MeshResults.GeneratedLayouts[LayoutIndex].Layout;
-							check(pLayout);
+							const NodeSurfaceNew::FImageData& ImageData = SurfaceNode->Images[ImageIndex];
 
-							Ptr<ASTOpInstanceAdd> op = new ASTOpInstanceAdd();
-							op->type = EOpType::IN_ADDIMAGE;
-							op->instance = LastSurfOp;
+							const int32 LayoutIndex = ImageData.LayoutIndex;
 
-							// Image
-							//-------------------------------------
+							// If the layout index has been set to negative, it means we should ignore the layout for this image.
+							CompilerOptions::TextureLayoutStrategy ImageLayoutStrategy = (LayoutIndex < 0)
+								? CompilerOptions::TextureLayoutStrategy::None
+								: CompilerOptions::TextureLayoutStrategy::Pack
+								;
 
-							// Size of a layout block in pixels
-							FIntPoint GridSize = pLayout->GetGridSize();
-	
-							// Try to guess the layout block description from the first valid block that is generated.
-							FLayoutBlockDesc LayoutBlockDesc;
-							if (formatNode)
+							if (ImageLayoutStrategy == CompilerOptions::TextureLayoutStrategy::None)
 							{
-								LayoutBlockDesc.FinalFormat = formatNode->GetPrivate()->m_formatIfAlpha;
-								if (LayoutBlockDesc.FinalFormat == EImageFormat::None)
+								// Generate the image
+								FImageGenerationOptions ImageOptions(ComponentId, Options.LODIndex);
+								ImageOptions.State = Options.State;
+								ImageOptions.ImageLayoutStrategy = ImageLayoutStrategy;
+								ImageOptions.ActiveTags = SurfaceNode->Tags;
+								ImageOptions.RectSize = { 0, 0 };
+
+								// TODO: To tasks
+								FImageGenerationResult Result;
+								GenerateImage(ImageOptions, Result, pImageNode);
+								Ptr<ASTOp> imageAd = Result.op;
+
+								// Placeholder block. Ideally this should be the actual image size
+								constexpr int32 FakeLayoutSize = 256;
+								FIntPoint GridSize(FakeLayoutSize, FakeLayoutSize);
+								FLayoutBlockDesc LayoutBlockDesc;
+								LayoutBlockDesc.BlockPixelsX = 1;
+								LayoutBlockDesc.BlockPixelsY = 1;
+								box< FIntVector2 > RectInCells;
+								RectInCells.min = { 0,0 };
+								RectInCells.size = { FakeLayoutSize ,FakeLayoutSize };
+
+								imageAd = ApplyImageBlockModifiers(Modifiers, ImageOptions, imageAd, ImageData, GridSize, LayoutBlockDesc, RectInCells, SurfaceNode->GetMessageContext());
+
+								check(imageAd);
+
+								if (swizzleNode)
 								{
-									LayoutBlockDesc.FinalFormat = formatNode->GetPrivate()->m_format;
+									Ptr<ASTOpImageSwizzle> fop = new ASTOpImageSwizzle();
+									fop->Format = swizzleNode->GetPrivate()->m_format;
+									fop->Sources[0] = imageAd;
+									fop->Sources[1] = imageAd;
+									fop->Sources[2] = imageAd;
+									fop->Sources[3] = imageAd;
+									fop->SourceChannels[0] = swizzleNode->GetPrivate()->m_sourceChannels[0];
+									fop->SourceChannels[1] = swizzleNode->GetPrivate()->m_sourceChannels[1];
+									fop->SourceChannels[2] = swizzleNode->GetPrivate()->m_sourceChannels[2];
+									fop->SourceChannels[3] = swizzleNode->GetPrivate()->m_sourceChannels[3];
+									check(fop->Format != EImageFormat::None);
+									imageAd = fop;
 								}
-							}
 
-							bool bImageSizeWarning = false;
+								if (mipmapNode)
+								{
+									Ptr<ASTOpImageMipmap> op = new ASTOpImageMipmap();
+									op->Levels = 0;
+									op->Source = imageAd;
+									op->BlockLevels = 0;
+
+									op->AddressMode = mipmapNode->GetPrivate()->m_settings.AddressMode;
+									op->FilterType = mipmapNode->GetPrivate()->m_settings.FilterType;
+									imageAd = op;
+								}
 
-							// Start with a blank image. It will be completed later with the blockSize, format and mips information
-							Ptr<ASTOpImageBlankLayout> BlankImageOp;
-							Ptr<ASTOp> imageAd;
-							{
-								BlankImageOp = new ASTOpImageBlankLayout();
-								BlankImageOp->Layout = MeshResults.LayoutOps[LayoutIndex];
-								// The rest ok the op will be completed below
-								BlankImageOp->MipmapCount = 0;
-								imageAd = BlankImageOp;
+								if (formatNode)
+								{
+									Ptr<ASTOpImagePixelFormat> fop = new ASTOpImagePixelFormat();
+									fop->Format = formatNode->GetPrivate()->m_format;
+									fop->FormatIfAlpha = formatNode->GetPrivate()->m_formatIfAlpha;
+									fop->Source = imageAd;
+									check(fop->Format != EImageFormat::None);
+									imageAd = fop;
+								}
+
+								Ptr<ASTOpInstanceAdd> op = new ASTOpInstanceAdd();
+								op->type = EOpType::IN_ADDIMAGE;
+								op->instance = LastSurfOp;
+								op->value = imageAd;
+								op->name = SurfaceNode->Images[ImageIndex].Name;
+
+								LastSurfOp = op;
 							}
 
-							// Skip the block addition for this image if the layout was from a extension.
-							if (!LayoutFromExtension[LayoutIndex])
+							else if (ImageLayoutStrategy == CompilerOptions::TextureLayoutStrategy::Pack) //-V547
 							{
-								for (int32 BlockIndex = 0; BlockIndex < pLayout->GetBlockCount(); ++BlockIndex)
+								if (LayoutIndex >= MeshResults.GeneratedLayouts.Num() ||
+									LayoutIndex >= MeshResults.LayoutOps.Num())
 								{
-									// Generate the image
-									FImageGenerationOptions ImageOptions(ComponentId, Options.LODIndex);
-									ImageOptions.State = Options.State;
-									ImageOptions.ImageLayoutStrategy = ImageLayoutStrategy;
-									ImageOptions.RectSize = { 0,0 };
-									ImageOptions.ActiveTags = SurfaceNode->Tags;
-									ImageOptions.LayoutToApply = pLayout;
-									ImageOptions.LayoutBlockId = pLayout->Blocks[BlockIndex].Id;
-									FImageGenerationResult ImageResult;
-									GenerateImage(ImageOptions, ImageResult, pImageNode);
-									Ptr<ASTOp> blockAd = ImageResult.op;
-
-									if (!blockAd)
-									{
-										// The GenerateImage(...) above has failed, skip this block
-										SurfaceResult.surfaceOp = nullptr;
-										continue;
-									}
+									ErrorLog->Add("Missing layout in object, or its parent.", ELMT_ERROR, SurfaceNode->GetMessageContext());
+								}
+								else
+								{
+									TSharedPtr<const FLayout> pLayout = MeshResults.GeneratedLayouts[LayoutIndex].Layout;
+									check(pLayout);
 
-									// Calculate the desc of the generated block.
-									constexpr bool bReturnBestOption = true;
-									FImageDesc BlockDesc = blockAd->GetImageDesc(bReturnBestOption, nullptr);
+									Ptr<ASTOpInstanceAdd> op = new ASTOpInstanceAdd();
+									op->type = EOpType::IN_ADDIMAGE;
+									op->instance = LastSurfOp;
 
-									// Block in layout grid units (cells)
-									box< FIntVector2 > RectInCells;
-									RectInCells.min = pLayout->Blocks[BlockIndex].Min;
-									RectInCells.size = pLayout->Blocks[BlockIndex].Size;
+									// Image
+									//-------------------------------------
 
-									// Try to update the layout block desc if we don't know it yet.
-									UpdateLayoutBlockDesc(LayoutBlockDesc, BlockDesc, RectInCells.size);
+									// Size of a layout block in pixels
+									FIntPoint GridSize = pLayout->GetGridSize();
 
-									// Even if we force the size afterwards, we need some size hint in some cases, like image projections.
-									ImageOptions.RectSize = UE::Math::TIntVector2<int32>(BlockDesc.m_size);
+									// Try to guess the layout block description from the first valid block that is generated.
+									FLayoutBlockDesc LayoutBlockDesc;
+									if (formatNode)
+									{
+										LayoutBlockDesc.FinalFormat = formatNode->GetPrivate()->m_formatIfAlpha;
+										if (LayoutBlockDesc.FinalFormat == EImageFormat::None)
+										{
+											LayoutBlockDesc.FinalFormat = formatNode->GetPrivate()->m_format;
+										}
+									}
 
-									blockAd = ApplyImageBlockModifiers(Modifiers, ImageOptions, blockAd, ImageData, GridSize, LayoutBlockDesc, RectInCells, SurfaceNode->GetMessageContext());
+									bool bImageSizeWarning = false;
 
-									// Enforce block size and optimizations
-									blockAd = GenerateImageSize(blockAd, FIntVector2(BlockDesc.m_size));
+									// Start with a blank image. It will be completed later with the blockSize, format and mips information
+									Ptr<ASTOpImageBlankLayout> BlankImageOp;
+									Ptr<ASTOp> imageAd;
+									{
+										BlankImageOp = new ASTOpImageBlankLayout();
+										BlankImageOp->Layout = MeshResults.LayoutOps[LayoutIndex];
+										// The rest ok the op will be completed below
+										BlankImageOp->MipmapCount = 0;
+										imageAd = BlankImageOp;
+									}
+
+									// Skip the block addition for this image if the layout was from a extension.
+									if (!LayoutFromExtension[LayoutIndex])
+									{
+										for (int32 BlockIndex = 0; BlockIndex < pLayout->GetBlockCount(); ++BlockIndex)
+										{
+											// Generate the image
+											FImageGenerationOptions ImageOptions(ComponentId, Options.LODIndex);
+											ImageOptions.State = Options.State;
+											ImageOptions.ImageLayoutStrategy = ImageLayoutStrategy;
+											ImageOptions.RectSize = { 0,0 };
+											ImageOptions.ActiveTags = SurfaceNode->Tags;
+											ImageOptions.LayoutToApply = pLayout;
+											ImageOptions.LayoutBlockId = pLayout->Blocks[BlockIndex].Id;
+											FImageGenerationResult ImageResult;
+											GenerateImage(ImageOptions, ImageResult, pImageNode);
+											Ptr<ASTOp> blockAd = ImageResult.op;
+
+											if (!blockAd)
+											{
+												// The GenerateImage(...) above has failed, skip this block
+												continue;
+											}
+
+											// Calculate the desc of the generated block.
+											constexpr bool bReturnBestOption = true;
+											FImageDesc BlockDesc = blockAd->GetImageDesc(bReturnBestOption, nullptr);
+
+											// Block in layout grid units (cells)
+											box< FIntVector2 > RectInCells;
+											RectInCells.min = pLayout->Blocks[BlockIndex].Min;
+											RectInCells.size = pLayout->Blocks[BlockIndex].Size;
+
+											// Try to update the layout block desc if we don't know it yet.
+											UpdateLayoutBlockDesc(LayoutBlockDesc, BlockDesc, RectInCells.size);
+
+											// Even if we force the size afterwards, we need some size hint in some cases, like image projections.
+											ImageOptions.RectSize = UE::Math::TIntVector2<int32>(BlockDesc.m_size);
+
+											blockAd = ApplyImageBlockModifiers(Modifiers, ImageOptions, blockAd, ImageData, GridSize, LayoutBlockDesc, RectInCells, SurfaceNode->GetMessageContext());
+
+											// Enforce block size and optimizations
+											blockAd = GenerateImageSize(blockAd, FIntVector2(BlockDesc.m_size));
+
+											EImageFormat baseFormat = imageAd->GetImageDesc().m_format;
+											// Actually don't do it, it will be propagated from the top format operation.
+											//Ptr<ASTOp> blockAd = GenerateImageFormat(blockAd, baseFormat);
+
+											// Apply tiling to avoid generating chunks of image that are too big.
+											blockAd = ApplyTiling(blockAd, ImageOptions.RectSize, LayoutBlockDesc.FinalFormat);
+
+											// Compose layout operation
+											Ptr<ASTOpImageCompose> composeOp = new ASTOpImageCompose();
+											composeOp->Layout = MeshResults.LayoutOps[LayoutIndex];
+											composeOp->Base = imageAd;
+											composeOp->BlockImage = blockAd;
+
+											// Set the absolute block index.
+											check(pLayout->Blocks[BlockIndex].Id != FLayoutBlock::InvalidBlockId);
+											composeOp->BlockId = pLayout->Blocks[BlockIndex].Id;
+
+											imageAd = composeOp;
+										}
+									}
+									check(imageAd);
+
+									FMeshGenerationStaticOptions ModifierOptions(ComponentId, Options.LODIndex);
+									ModifierOptions.State = Options.State;
+									ModifierOptions.ActiveTags = SurfaceNode->Tags;
+									imageAd = ApplyImageExtendModifiers(Modifiers, ModifierOptions, MeshResults, imageAd, ImageLayoutStrategy,
+										LayoutIndex, ImageData, GridSize, LayoutBlockDesc,
+										SurfaceNode->GetMessageContext());
+
+									// Complete the base op
+									BlankImageOp->BlockSize[0] = uint16(LayoutBlockDesc.BlockPixelsX);
+									BlankImageOp->BlockSize[1] = uint16(LayoutBlockDesc.BlockPixelsY);
+									BlankImageOp->Format = GetUncompressedFormat(LayoutBlockDesc.FinalFormat);
+									BlankImageOp->GenerateMipmaps = LayoutBlockDesc.bBlocksHaveMips;
+									BlankImageOp->MipmapCount = 0;
+
+									if (swizzleNode)
+									{
+										Ptr<ASTOpImageSwizzle> fop = new ASTOpImageSwizzle();
+										fop->Format = swizzleNode->GetPrivate()->m_format;
+
+										for (int32 ChannelIndex = 0; ChannelIndex < swizzleNode->GetPrivate()->m_sourceChannels.Num(); ++ChannelIndex)
+										{
+											fop->Sources[ChannelIndex] = imageAd;
+											fop->SourceChannels[ChannelIndex] = swizzleNode->GetPrivate()->m_sourceChannels[ChannelIndex];
+										}
+										check(fop->Format != EImageFormat::None);
+										imageAd = fop;
+									}
 
-									EImageFormat baseFormat = imageAd->GetImageDesc().m_format;
-									// Actually don't do it, it will be propagated from the top format operation.
-									//Ptr<ASTOp> blockAd = GenerateImageFormat(blockAd, baseFormat);
+									// Apply mipmap and format if necessary, skip if format is None (possibly because a block was skipped above)
+									bool bNeedsMips =
+										(mipmapNode && LayoutBlockDesc.FinalFormat != EImageFormat::None)
+										||
+										LayoutBlockDesc.bBlocksHaveMips;
 
-									// Apply tiling to avoid generating chunks of image that are too big.
-									blockAd = ApplyTiling(blockAd, ImageOptions.RectSize, LayoutBlockDesc.FinalFormat);
+									if (bNeedsMips)
+									{
+										Ptr<ASTOpImageMipmap> mop = new ASTOpImageMipmap();
+
+										// At the end of the day, we want all the mipmaps. Maybe the code
+										// optimiser will split the process later.
+										mop->Levels = 0;
+										mop->bOnlyTail = false;
+										mop->Source = imageAd;
+
+										// We have to avoid mips smaller than the image format block size, so
+										// we will devide the layout block by the format block
+										const FImageFormatData& PixelFormatInfo = GetImageFormatData(LayoutBlockDesc.FinalFormat);
+
+										int32 mipsX = FMath::CeilLogTwo(LayoutBlockDesc.BlockPixelsX / PixelFormatInfo.PixelsPerBlockX);
+										int32 mipsY = FMath::CeilLogTwo(LayoutBlockDesc.BlockPixelsY / PixelFormatInfo.PixelsPerBlockY);
+										mop->BlockLevels = (uint8)FMath::Max(mipsX, mipsY);
+
+										if (LayoutBlockDesc.BlockPixelsX < PixelFormatInfo.PixelsPerBlockX || LayoutBlockDesc.BlockPixelsY < PixelFormatInfo.PixelsPerBlockY)
+										{
+											// In this case, the mipmap will never be useful for blocks, so we indicate that
+											// it should make the mips at the root of the expression.
+											mop->bOnlyTail = true;
+										}
+
+										mop->AddressMode = EAddressMode::ClampToEdge;
+										mop->FilterType = EMipmapFilterType::SimpleAverage;
+
+										if (mipmapNode)
+										{
+											mop->AddressMode = mipmapNode->GetPrivate()->m_settings.AddressMode;
+											mop->FilterType = mipmapNode->GetPrivate()->m_settings.FilterType;
+										}
+
+										imageAd = mop;
+									}
+
+									if (formatNode)
+									{
+										Ptr<ASTOpImagePixelFormat> fop = new ASTOpImagePixelFormat();
+										fop->Format = formatNode->GetPrivate()->m_format;
+										fop->FormatIfAlpha = formatNode->GetPrivate()->m_formatIfAlpha;
+										fop->Source = imageAd;
+										check(fop->Format != EImageFormat::None);
+										imageAd = fop;
+									}
 
-									// Compose layout operation
-									Ptr<ASTOpImageCompose> composeOp = new ASTOpImageCompose();
-									composeOp->Layout = MeshResults.LayoutOps[LayoutIndex];
-									composeOp->Base = imageAd;
-									composeOp->BlockImage = blockAd;
+									op->value = imageAd;
 
-									// Set the absolute block index.
-									check(pLayout->Blocks[BlockIndex].Id != FLayoutBlock::InvalidBlockId);
-									composeOp->BlockId = pLayout->Blocks[BlockIndex].Id;
+									// Name
+									op->name = SurfaceNode->Images[ImageIndex].Name;
 
-									imageAd = composeOp;
+									LastSurfOp = op;
 								}
 							}
-							check(imageAd);
-
-							FMeshGenerationStaticOptions ModifierOptions(ComponentId, Options.LODIndex);
-							ModifierOptions.State = Options.State;
-							ModifierOptions.ActiveTags = SurfaceNode->Tags;
-							imageAd = ApplyImageExtendModifiers( Modifiers, ModifierOptions, MeshResults, imageAd, ImageLayoutStrategy, 
-								LayoutIndex, ImageData, GridSize, LayoutBlockDesc, 
-								SurfaceNode->GetMessageContext());
-
-							// Complete the base op
-							BlankImageOp->BlockSize[0] = uint16(LayoutBlockDesc.BlockPixelsX);
-							BlankImageOp->BlockSize[1] = uint16(LayoutBlockDesc.BlockPixelsY);
-							BlankImageOp->Format = GetUncompressedFormat(LayoutBlockDesc.FinalFormat);
-							BlankImageOp->GenerateMipmaps = LayoutBlockDesc.bBlocksHaveMips;
-							BlankImageOp->MipmapCount = 0;
-
-							if (swizzleNode)
-							{
-								Ptr<ASTOpImageSwizzle> fop = new ASTOpImageSwizzle();
-								fop->Format = swizzleNode->GetPrivate()->m_format;
 
-								for (int32 ChannelIndex = 0; ChannelIndex < swizzleNode->GetPrivate()->m_sourceChannels.Num(); ++ChannelIndex)
-								{
-									fop->Sources[ChannelIndex] = imageAd;
-									fop->SourceChannels[ChannelIndex] = swizzleNode->GetPrivate()->m_sourceChannels[ChannelIndex];
-								}
-								check(fop->Format != EImageFormat::None);
-								imageAd = fop;
+							else
+							{
+								// Unimplemented texture layout strategy
+								check(false);
 							}
+						}
+					}
 
-							// Apply mipmap and format if necessary, skip if format is None (possibly because a block was skipped above)
-							bool bNeedsMips =
-								(mipmapNode && LayoutBlockDesc.FinalFormat != EImageFormat::None)
-								||
-								LayoutBlockDesc.bBlocksHaveMips;
+					// Create the expression for each vector
+					//------------------------------------------------------------------------
+					for (int32 t = 0; t < SurfaceNode->Vectors.Num(); ++t)
+					{
+						//MUTABLE_CPUPROFILER_SCOPE(SurfaceVector);
 
-							if (bNeedsMips)
-							{
-								Ptr<ASTOpImageMipmap> mop = new ASTOpImageMipmap();
-
-								// At the end of the day, we want all the mipmaps. Maybe the code
-								// optimiser will split the process later.
-								mop->Levels = 0;
-								mop->bOnlyTail = false;
-								mop->Source = imageAd;
-
-								// We have to avoid mips smaller than the image format block size, so
-								// we will devide the layout block by the format block
-								const FImageFormatData& PixelFormatInfo = GetImageFormatData(LayoutBlockDesc.FinalFormat);
-
-								int32 mipsX = FMath::CeilLogTwo(LayoutBlockDesc.BlockPixelsX / PixelFormatInfo.PixelsPerBlockX);
-								int32 mipsY = FMath::CeilLogTwo(LayoutBlockDesc.BlockPixelsY / PixelFormatInfo.PixelsPerBlockY);
-								mop->BlockLevels = (uint8)FMath::Max(mipsX, mipsY);
-								
-								if (LayoutBlockDesc.BlockPixelsX < PixelFormatInfo.PixelsPerBlockX || LayoutBlockDesc.BlockPixelsY < PixelFormatInfo.PixelsPerBlockY)
-								{
-									// In this case, the mipmap will never be useful for blocks, so we indicate that
-									// it should make the mips at the root of the expression.
-									mop->bOnlyTail = true;
-								}
+						if (Ptr<NodeColour> pVectorNode = SurfaceNode->Vectors[t].Vector)
+						{
+							Ptr<ASTOpInstanceAdd> op = new ASTOpInstanceAdd();
+							op->type = EOpType::IN_ADDVECTOR;
+							op->instance = LastSurfOp;
 
-								mop->AddressMode = EAddressMode::ClampToEdge;
-								mop->FilterType = EMipmapFilterType::SimpleAverage;
+							// Vector
+							FColorGenerationResult VectorResult;
+							GenerateColor(VectorResult, Options, pVectorNode);
+							op->value = VectorResult.op;
 
-								if (mipmapNode)
-								{
-									mop->AddressMode = mipmapNode->GetPrivate()->m_settings.AddressMode;
-									mop->FilterType = mipmapNode->GetPrivate()->m_settings.FilterType;
-								}
+							// Name
+							op->name = SurfaceNode->Vectors[t].Name;
 
-								imageAd = mop;
-							}
+							LastSurfOp = op;
+						}
+					}
 
-							if (formatNode)
-							{
-								Ptr<ASTOpImagePixelFormat> fop = new ASTOpImagePixelFormat();
-								fop->Format = formatNode->GetPrivate()->m_format;
-								fop->FormatIfAlpha = formatNode->GetPrivate()->m_formatIfAlpha;
-								fop->Source = imageAd;
-								check(fop->Format != EImageFormat::None);
-								imageAd = fop;
-							}
+					// Create the expression for each scalar
+					//------------------------------------------------------------------------
+					for (int32 t = 0; t < SurfaceNode->Scalars.Num(); ++t)
+					{
+						// MUTABLE_CPUPROFILER_SCOPE(SurfaceScalar);
 
-							op->value = imageAd;
+						if (NodeScalarPtr pScalarNode = SurfaceNode->Scalars[t].Scalar)
+						{
+							Ptr<ASTOpInstanceAdd> op = new ASTOpInstanceAdd();
+							op->type = EOpType::IN_ADDSCALAR;
+							op->instance = LastSurfOp;
+
+							// Scalar
+							FScalarGenerationResult ScalarResult;
+							GenerateScalar(ScalarResult, Options, pScalarNode);
+							op->value = ScalarResult.op;
 
 							// Name
-							op->name = SurfaceNode->Images[ImageIndex].Name;
+							op->name = SurfaceNode->Scalars[t].Name;
 
 							LastSurfOp = op;
 						}
 					}
 
-					else
+					// Create the expression for each string
+					//------------------------------------------------------------------------
+					for (int32 t = 0; t < SurfaceNode->Strings.Num(); ++t)
 					{
-						// Unimplemented texture layout strategy
-						check(false);
+						if (NodeStringPtr pStringNode = SurfaceNode->Strings[t].String)
+						{
+							Ptr<ASTOpInstanceAdd> op = new ASTOpInstanceAdd();
+							op->type = EOpType::IN_ADDSTRING;
+							op->instance = LastSurfOp;
+
+							FStringGenerationResult StringResult;
+							GenerateString(StringResult, Options, pStringNode);
+							op->value = StringResult.op;
+
+							// Name
+							op->name = SurfaceNode->Strings[t].Name;
+
+							LastSurfOp = op;
+						}
 					}
 				}
-			}
 
-			// Create the expression for each vector
-			//------------------------------------------------------------------------
-			for (int32 t = 0; t < SurfaceNode->Vectors.Num(); ++t)
-			{
-				//MUTABLE_CPUPROFILER_SCOPE(SurfaceVector);
+				FSurfaceGenerationResult SurfaceResult;
+				SurfaceResult.SurfaceOp = LastSurfOp;
 
-				if (Ptr<NodeColour> pVectorNode = SurfaceNode->Vectors[t].Vector)
+				// If we are going to share this surface properties, remember it.
+				if (bIsBaseForSharedSurface)
 				{
-					Ptr<ASTOpInstanceAdd> op = new ASTOpInstanceAdd();
-					op->type = EOpType::IN_ADDVECTOR;
-					op->instance = LastSurfOp;
+					UE::TUniqueLock Lock(SharedMeshOptionsMapMutex);
 
-					// Vector
-					FColorGenerationResult VectorResult;
-					GenerateColor(VectorResult, Options, pVectorNode);
-					op->value = VectorResult.op;
-
-					// Name
-					op->name = SurfaceNode->Vectors[t].Name;
-
-					LastSurfOp = op;
+					check(!SharedMeshOptionsMap.Contains(SurfaceNode->SharedSurfaceId));
+					SharedMeshOptionsMap.Add(SurfaceNode->SharedSurfaceId, MeshResults);
 				}
-			}
 
-			// Create the expression for each scalar
-			//------------------------------------------------------------------------
-			for (int32 t = 0; t < SurfaceNode->Scalars.Num(); ++t)
-			{
-				// MUTABLE_CPUPROFILER_SCOPE(SurfaceScalar);
+				return SurfaceResult;
+			}
+			, 
+			UE::Tasks::Prerequisites(MeshTask, PreviousLODTask)
+		);
 
-				if (NodeScalarPtr pScalarNode = SurfaceNode->Scalars[t].Scalar)
-				{
-					Ptr<ASTOpInstanceAdd> op = new ASTOpInstanceAdd();
-					op->type = EOpType::IN_ADDSCALAR;
-					op->instance = LastSurfOp;
+		for (FirstPassGenerator::FSurface* TargetSurface : TargetSurfaces)
+		{
+			TargetSurface->ResultSurfaceTask = SurfaceTask;
+		}
 
-					// Scalar
-					FScalarGenerationResult ScalarResult;
-					GenerateScalar(ScalarResult, Options, pScalarNode);
-					op->value = ScalarResult.op;
+		return SurfaceTask;
+	}
 
-					// Name
-					op->name = SurfaceNode->Scalars[t].Name;
 
-					LastSurfOp = op;
-				}
-			}
+	FLODTask CodeGenerator::GenerateLOD(const FLODGenerationOptions& Options, const NodeLOD* InNode, FLODTask PreviousLODTask )
+	{
+		MUTABLE_CPUPROFILER_SCOPE(GenerateLOD);
 
-			// Create the expression for each string
-			//------------------------------------------------------------------------
-			for (int32 t = 0; t < SurfaceNode->Strings.Num(); ++t)
+        // Look for all surfaces that belong to this component
+		TArray<int32> SurfaceIndices;
+		SurfaceIndices.Reserve(FirstPass.Surfaces.Num());
+		TArray<FSurfaceTask> SurfaceTasks;
+		SurfaceTasks.Reserve(FirstPass.Surfaces.Num());
+		for (int32 SurfaceIndex = 0; SurfaceIndex < FirstPass.Surfaces.Num(); ++SurfaceIndex)
+		{
+			const FirstPassGenerator::FSurface& SurfaceData = FirstPass.Surfaces[SurfaceIndex];
+			if (SurfaceData.Component == Options.Component
+				&&
+				SurfaceData.LOD == Options.LODIndex)
 			{
-				if (NodeStringPtr pStringNode = SurfaceNode->Strings[t].String)
+				// Apply state conditions: only generate it if it enabled in this state
 				{
-					Ptr<ASTOpInstanceAdd> op = new ASTOpInstanceAdd();
-					op->type = EOpType::IN_ADDSTRING;
-					op->instance = LastSurfOp;
+					bool bEnabledInThisState = true;
+					if (SurfaceData.StateCondition.Num() && Options.State >= 0)
+					{
+						bEnabledInThisState =
+							(Options.State < SurfaceData.StateCondition.Num())
+							&&
+							(SurfaceData.StateCondition[Options.State]);
+					}
+					if (!bEnabledInThisState)
+					{
+						continue;
+					}
+				}
 
-					FStringGenerationResult StringResult;
-					GenerateString(StringResult, Options, pStringNode);
-					op->value = StringResult.op;
 
-					// Name
-					op->name = SurfaceNode->Strings[t].Name;
+				FSurfaceGenerationOptions SurfaceOptions(Options);
+				FSurfaceTask SurfaceTask = GenerateSurface(SurfaceOptions, SurfaceData.Node, PreviousLODTask);
 
-					LastSurfOp = op;
-				}
+				SurfaceTasks.Add(SurfaceTask);
+				SurfaceIndices.Add(SurfaceIndex);
 			}
 		}
 
-        SurfaceResult.surfaceOp = LastSurfOp;
-
-		for (FirstPassGenerator::FSurface* TargetSurface : TargetSurfaces)
+		TArray<UE::Tasks::FTask> Requisites;
+		Requisites.Reserve( SurfaceTasks.Num() + 1);
+		Requisites.Append(SurfaceTasks);
+		if (PreviousLODTask.IsValid())
 		{
-			TargetSurface->ResultSurfaceOp = LastSurfOp;
+			Requisites.Add(PreviousLODTask);
 		}
 
-		// If we are going to share this surface properties, remember it.
-		if (bIsBaseForSharedSurface)
-		{
-			check(!SharedMeshOptionsMap.Contains(SurfaceNode->SharedSurfaceId));
-			SharedMeshOptionsMap.Add(SurfaceNode->SharedSurfaceId, MeshResults);
-		}
-    }
+		FLODTask LODTask = LocalPipe.Launch(TEXT("MutableLOD"),
+			[
+				SurfaceTasks, SurfaceIndices, 
+					this
+			]
+			() mutable
+			{
+				// Build a series of operations to assemble the component
+				Ptr<ASTOp> LastCompOp;
+				Ptr<ASTOp> LastMeshOp;
 
+				// This generates a different ID for each surface in the LOD and the component. It can be used to match it to the
+				// mesh surface. It cannot be 0 because it is a special case for the merge operation.
+				int32 surfaceID = 1;
 
-    //---------------------------------------------------------------------------------------------
-	void CodeGenerator::Generate_LOD(const FLODGenerationOptions& Options, FGenericGenerationResult& Result, const NodeLOD* InNode)
-	{
-		const NodeLOD& node = *InNode;
+				for (int32 SelectedSurfaceIndex=0; SelectedSurfaceIndex <SurfaceTasks.Num(); ++SelectedSurfaceIndex)
+				{
+					FSurfaceTask& SurfaceTask = SurfaceTasks[SelectedSurfaceIndex];
+					FSurfaceGenerationResult SurfaceGenerationResult = SurfaceTask.GetResult();
 
-		MUTABLE_CPUPROFILER_SCOPE(Generate_LOD);
+					const FirstPassGenerator::FSurface& SurfaceData = FirstPass.Surfaces[SurfaceIndices[SelectedSurfaceIndex]];
 
-		// Build a series of operations to assemble the component
-        Ptr<ASTOp> lastCompOp;
-        Ptr<ASTOp> lastMeshOp;
-        FString lastMeshName;
+					Ptr<ASTOpInstanceAdd> SurfaceOp = new ASTOpInstanceAdd();
+					SurfaceOp->type = EOpType::IN_ADDSURFACE;
+					SurfaceOp->name = SurfaceData.Node->Name;
+					SurfaceOp->instance = LastCompOp;
+					SurfaceOp->value = SurfaceGenerationResult.SurfaceOp;
+					SurfaceOp->id = surfaceID;
+					SurfaceOp->ExternalId = SurfaceData.Node->ExternalId;
+					SurfaceOp->SharedSurfaceId = SurfaceData.Node->SharedSurfaceId;
 
-        // This generates a different ID for each surface in the LOD and the component. It can be used to match it to the
-        // mesh surface. It cannot be 0 because it is a special case for the merge operation.
-        int32 surfaceID=1;
+					Ptr<ASTOp> SurfaceConditionOp = SurfaceData.FinalCondition;
 
-        // Look for all surfaces that belong to this component
-		for (int32 i = 0; i<FirstPass.Surfaces.Num(); ++i)
-		{
-			const FirstPassGenerator::FSurface& its = FirstPass.Surfaces[i];
-			if (its.Component==Options.Component
-				&&
-				its.LOD==Options.LODIndex)
-			{
-                // Apply state conditions: only generate it if it enabled in this state
-                {
-                    bool enabledInThisState = true;
-                    if (its.StateCondition.Num() && Options.State >= 0)
-                    {
-                        enabledInThisState =
-                                (Options.State < its.StateCondition.Num())
-                                &&
-                                ( its.StateCondition[Options.State] );
-                    }
-                    if (!enabledInThisState)
-                    {
-                        continue;
-                    }
-                }
-
-                Ptr<ASTOpInstanceAdd> SurfaceOp = new ASTOpInstanceAdd();
-				SurfaceOp->type = EOpType::IN_ADDSURFACE;
-				SurfaceOp->name = its.Node->Name;
-				SurfaceOp->instance = lastCompOp;
+					{
+						Ptr<ASTOpConditional> op = new ASTOpConditional();
+						op->type = EOpType::IN_CONDITIONAL;
+						op->no = LastCompOp;
+						op->yes = SurfaceOp;
+						op->condition = SurfaceConditionOp;
+						LastCompOp = op;
+					}
 
-				FSurfaceGenerationOptions SurfaceOptions(Options);
-				FSurfaceGenerationResult surfaceGenerationResult;
-                GenerateSurface( surfaceGenerationResult, SurfaceOptions, its.Node );
-				SurfaceOp->value = surfaceGenerationResult.surfaceOp;
-
-				SurfaceOp->id = surfaceID;
-				SurfaceOp->ExternalId = its.Node->ExternalId;
-				SurfaceOp->SharedSurfaceId = its.Node->SharedSurfaceId;
-
-                Ptr<ASTOp> SurfaceConditionOp = its.FinalCondition;
-
-                {
-                    Ptr<ASTOpConditional> op = new ASTOpConditional();
-                    op->type = EOpType::IN_CONDITIONAL;
-                    op->no = lastCompOp;
-                    op->yes = SurfaceOp;
-                    op->condition = SurfaceConditionOp;
-                    lastCompOp = op;
-                }
-
-                // Add the mesh with its condition
- 
-                // We add the merge op even for the first mesh, so that we set the surface id.
-				Ptr<ASTOp> mergeAd;
+					// Add the mesh with its condition
+
+					// We add the merge op even for the first mesh, so that we set the surface id.
+					Ptr<ASTOp> mergeAd;
+					{
+						Ptr<ASTOp> Added = SurfaceData.ResultMeshOp;
+
+						Ptr<ASTOpMeshMerge> MergeOp = new ASTOpMeshMerge();
+						MergeOp->Base = LastMeshOp;
+						MergeOp->Added = Added;
+						MergeOp->NewSurfaceID = surfaceID;
+						mergeAd = MergeOp;
+					}
+
+					if (SurfaceConditionOp)
+					{
+						Ptr<ASTOpConditional> op = new ASTOpConditional();
+						op->type = EOpType::ME_CONDITIONAL;
+						op->no = LastMeshOp;
+						op->yes = mergeAd;
+						op->condition = SurfaceConditionOp;
+						LastMeshOp = op;
+					}
+					else
+					{
+						LastMeshOp = mergeAd;
+					}
+
+					++surfaceID;
+				}
+
+				// Add op to optimize the skinning of the resulting mesh
 				{
-                    Ptr<ASTOp> Added = its.ResultMeshOp;
- 
-                    Ptr<ASTOpMeshMerge> MergeOp = new ASTOpMeshMerge();
-					MergeOp->Base = lastMeshOp;
-					MergeOp->Added = Added;
-					MergeOp->NewSurfaceID = surfaceID;
-                    mergeAd = MergeOp;
-                }
-
-                if (SurfaceConditionOp)
-                {
-                    Ptr<ASTOpConditional> op = new ASTOpConditional();
-                    op->type = EOpType::ME_CONDITIONAL;
-                    op->no = lastMeshOp;
-                    op->yes = mergeAd;
-                    op->condition = SurfaceConditionOp;
-                    lastMeshOp = op;
-                }
-                else
-                {
-                    lastMeshOp = mergeAd;
-                }
-
-				++surfaceID;
-			}
-		}
+					Ptr<ASTOpMeshOptimizeSkinning> mop = new ASTOpMeshOptimizeSkinning();
+					mop->Source = LastMeshOp;
+					LastMeshOp = mop;
+				}
 
-		// Add op to optimize the skinning of the resulting mesh
-		{
-			Ptr<ASTOpMeshOptimizeSkinning> mop = new ASTOpMeshOptimizeSkinning();
-			mop->Source = lastMeshOp;
-			lastMeshOp = mop;
-		}
+				// Add the component mesh
+				{
+					Ptr<ASTOpInstanceAdd> iop = new ASTOpInstanceAdd();
+					iop->type = EOpType::IN_ADDMESH;
+					iop->instance = LastCompOp;
+					iop->value = LastMeshOp;
 
-        // Add the component mesh
-        {
-            Ptr<ASTOpInstanceAdd> iop = new ASTOpInstanceAdd();
-            iop->type = EOpType::IN_ADDMESH;
-            iop->instance = lastCompOp;
-            iop->value = lastMeshOp;
+					LastCompOp = iop;
+				}
 
-            lastCompOp = iop;
-        }
+				FGenericGenerationResult Result;
+				Result.Op = LastCompOp;
+				return Result;
+			},
+			Requisites
+		);
 
-        Result.Op = lastCompOp;
+		return LODTask;
     }
 
 
-    //---------------------------------------------------------------------------------------------
-   // Ptr<ASTOp> CodeGenerator::Visit( const NodeObjectNew::Private& node )
 	void CodeGenerator::GenerateObject_New(const FObjectGenerationOptions& Options, FObjectGenerationResult& OutResult, const NodeObjectNew* InNode)
 	{
 		MUTABLE_CPUPROFILER_SCOPE(NodeObjectNew);
@@ -2028,12 +2136,12 @@ namespace mu
 	}
 
 	//---------------------------------------------------------------------------------------------
-	CodeGenerator::FMeshTask CodeGenerator::ApplyMeshModifiers(
+	FMeshTask CodeGenerator::ApplyMeshModifiers(
 		const TArray<FirstPassGenerator::FModifier>& Modifiers,
 		const FMeshGenerationStaticOptions& StaticOptions, 
 		FMeshOptionsTask Options,
 		FMeshTask BaseTask,
-		const FMeshGenerationResult* SharedMeshResults,
+		int32 SharedSurfaceId,
 		const void* ErrorContext,
 		const NodeMeshConstant* OriginalMeshNode)
 	{
@@ -2068,15 +2176,31 @@ namespace mu
 					MergedMeshStaticOptions.ModifiersToIgnore.Add(Modifier);
 
 					FMeshOptionsTask MergedMeshOptionsTask = UE::Tasks::Launch(TEXT("MutableMergedMeshOptions"),
-						[Options, SharedMeshResults, EditIndex]() mutable
+						[Options, SharedSurfaceId, EditIndex, this]() mutable
 						{
 							FMeshGenerationDynamicOptions Result = Options.GetResult();
 							Result.bEnsureAllVerticesHaveLayoutBlock = false;
 
-							if (SharedMeshResults)
+							UE::TUniqueLock Lock(SharedMeshOptionsMapMutex);
+
+							// This assumes that the lods are processed in order. It checks it this way because some platforms may have empty lods at the top.
+							const bool bIsBaseForSharedSurface = SharedSurfaceId != INDEX_NONE && !SharedMeshOptionsMap.Contains(SharedSurfaceId);
+
+							// If this is true, we will reuse the surface properties from a higher LOD, se we can skip the generation of material properties and images.
+							const bool bShareSurface = SharedSurfaceId != INDEX_NONE && !bIsBaseForSharedSurface;
+
+							const FMeshGenerationResult* SharedMeshResults = nullptr;
+							if (bShareSurface)
 							{
-								check(SharedMeshResults->ExtraMeshLayouts.IsValidIndex(EditIndex));
-								Result.OverrideLayouts = SharedMeshResults->ExtraMeshLayouts[EditIndex].GeneratedLayouts;
+								// Do we have the surface we need to share it with?
+								SharedMeshResults = SharedMeshOptionsMap.Find(SharedSurfaceId);
+								check(SharedMeshResults);
+
+								if (SharedMeshResults)
+								{
+									check(SharedMeshResults->ExtraMeshLayouts.IsValidIndex(EditIndex));
+									Result.OverrideLayouts = SharedMeshResults->ExtraMeshLayouts[EditIndex].GeneratedLayouts;
+								}
 							}
 
 							return Result;
@@ -2150,7 +2274,7 @@ namespace mu
 					constexpr bool bModifiersForBeforeOperations = false;
 					GetModifiersFor(StaticOptions.ComponentId, ModifierOptions.ActiveTags, bModifiersForBeforeOperations, ChildModifiers);
 
-					FMeshTask AddWithModifiersTask = ApplyMeshModifiers(ChildModifiers, ModifierOptions, Options, AddTask, SharedMeshResults, ErrorContext, nullptr);
+					FMeshTask AddWithModifiersTask = ApplyMeshModifiers(ChildModifiers, ModifierOptions, Options, AddTask, SharedSurfaceId, ErrorContext, nullptr);
 
 					LastMeshTask = UE::Tasks::Launch(TEXT("MutableMeshMergeModifier"),
 						[AddBaseTask, AddWithModifiersTask, LastMeshTask, Modifier, EditIndex]() mutable
diff --git a/Engine/Plugins/Mutable/Source/MutableTools/Private/MuT/CodeGenerator.h b/Engine/Plugins/Mutable/Source/MutableTools/Private/MuT/CodeGenerator.h
--- a/Engine/Plugins/Mutable/Source/MutableTools/Private/MuT/CodeGenerator.h
+++ b/Engine/Plugins/Mutable/Source/MutableTools/Private/MuT/CodeGenerator.h
@@ -195,11 +195,6 @@ namespace mu
 			int32 LODIndex = -1;
 		};
 
-		struct FGenericGenerationResult
-		{
-			Ptr<ASTOp> Op;
-		};
-
 		struct FGeneratedCacheKey
 		{
 			Ptr<const Node> Node;
@@ -218,7 +213,15 @@ namespace mu
 
 
 		void GenerateGeneric(const FGenericGenerationOptions& Options, FGenericGenerationResult& OutResult, const Ptr<const Node> InNode);
-		void Generate_LOD(const FLODGenerationOptions&, FGenericGenerationResult&, const NodeLOD*);
+
+		/** Launch the tasks to generate a LOD subgraph. 
+		* Since LODs need to be generated in order a dependency task representing the previous LOD generation can be specified. 
+		*/
+		FLODTask GenerateLOD(const FLODGenerationOptions&, const NodeLOD*, FLODTask PreviousLODTask);
+
+		/** Launch the tasks to generate a Surface subgraph.
+		*/
+		FSurfaceTask GenerateSurface(const FSurfaceGenerationOptions&, Ptr<const NodeSurfaceNew>, FLODTask PreviousLODTask);
 
 
 		struct FGeneratedObjectCacheKey
@@ -294,6 +297,9 @@ namespace mu
 
     private:
 
+		/** Any task accessing generic CodeGenerator members or "serial" data should run in this pipe. */
+		UE::Tasks::FPipe LocalPipe;
+
 		/** This struct contains the generation state that can only be accessed when generating NodeMeshConstants.
 		* It should only be accessed from tasks running in the specific FPipe below.
 		*/
@@ -521,19 +527,6 @@ namespace mu
         //-----------------------------------------------------------------------------------------
         // Meshes
 
-		/** */
-		struct FGeneratedLayout
-		{
-			TSharedPtr<const FLayout> Layout;
-			Ptr<const NodeLayout> Source;
-
-			FORCEINLINE bool operator==(const FGeneratedLayout& Other) const
-			{
-				return Layout == Other.Layout
-					&& Source == Other.Source;
-			}
-		};
-
 		/** Options that affect the generation of meshes. It is like list of what required data we want
 		* while parsing down the mesh node graph.
 		*/
@@ -589,32 +582,6 @@ namespace mu
 			TOptional<const void*> OverrideContext;
 		};
 
-		//! Store the results of the code generation of a mesh.
-		struct FMeshGenerationResult
-		{
-			//! Mesh after all code tree is applied
-			Ptr<ASTOp> MeshOp;
-
-			//! Original base mesh before removes, morphs, etc.
-			Ptr<ASTOp> BaseMeshOp;
-
-			/** Generated node layouts with their own block ids. */
-			TArray<FGeneratedLayout> GeneratedLayouts;
-
-			TArray<Ptr<ASTOp>> LayoutOps;
-
-			struct FExtraLayouts
-			{
-				/** Source node layouts to use with these extra mesh. They don't have block ids. */
-				TArray<FGeneratedLayout> GeneratedLayouts;
-				Ptr<ASTOp> Condition;
-				Ptr<ASTOp> MeshFragment;
-			};
-			TArray< FExtraLayouts > ExtraMeshLayouts;
-		};
-		
-
-		using FMeshTask = UE::Tasks::TTask<CodeGenerator::FMeshGenerationResult>;
 		using FMeshOptionsTask = UE::Tasks::TTask<CodeGenerator::FMeshGenerationDynamicOptions>;
 
 		/** Store the mesh generation data for surfaces that we intend to share across LODs. 
@@ -622,6 +589,12 @@ namespace mu
 		*/
 		TMap<int32, FMeshGenerationResult> SharedMeshOptionsMap;
 
+		/** Use this mutex to control acces to the SharedMeshOptionsMap attribute from any thread. 
+		* The task dependency will make sure that the actual order for data insert and retrieval is correct, but unrelated data can be added while 
+		* unrelated data is being retrieved so we need the mutex.
+		*/
+		UE::FMutex SharedMeshOptionsMapMutex;
+
 		/** This struct contains the generation state that can only be accessed when generating constant layouts.
 		* It should only be accessed after locking the included mutex.
 		*/
@@ -812,15 +785,6 @@ namespace mu
 
         void GenerateRange(FRangeGenerationResult&, const FGenericGenerationOptions&, Ptr<const NodeRange>);
 
-
-        //-----------------------------------------------------------------------------------------
-        struct FSurfaceGenerationResult
-        {
-            Ptr<ASTOp> surfaceOp;
-        };
-
-        void GenerateSurface( FSurfaceGenerationResult&, const FSurfaceGenerationOptions&, Ptr<const NodeSurfaceNew> );
-
 		//-----------------------------------------------------------------------------------------
 
 		struct FLayoutBlockDesc
@@ -842,7 +806,7 @@ namespace mu
 			const FMeshGenerationStaticOptions&,
 			FMeshOptionsTask DynamicOptionsTask,
 			FMeshTask BaseTask,
-			const FMeshGenerationResult* SharedSurfaceResults, 
+			int32 SharedSurfaceId,
 			const void* ErrorContext,
 			const NodeMeshConstant* OriginalMeshNode);
 
diff --git a/Engine/Plugins/Mutable/Source/MutableTools/Private/MuT/CodeGenerator_FirstPass.h b/Engine/Plugins/Mutable/Source/MutableTools/Private/MuT/CodeGenerator_FirstPass.h
--- a/Engine/Plugins/Mutable/Source/MutableTools/Private/MuT/CodeGenerator_FirstPass.h
+++ b/Engine/Plugins/Mutable/Source/MutableTools/Private/MuT/CodeGenerator_FirstPass.h
@@ -31,6 +31,61 @@ namespace mu
 	class FLayout;
 	struct FObjectState;
 
+	/** */
+	struct FGeneratedLayout
+	{
+		TSharedPtr<const FLayout> Layout;
+		Ptr<const NodeLayout> Source;
+
+		FORCEINLINE bool operator==(const FGeneratedLayout& Other) const
+		{
+			return Layout == Other.Layout
+				&& Source == Other.Source;
+		}
+	};
+
+	/** Store the results of the code generation of a mesh. */
+	struct FMeshGenerationResult
+	{
+		//! Mesh after all code tree is applied
+		Ptr<ASTOp> MeshOp;
+
+		//! Original base mesh before removes, morphs, etc.
+		Ptr<ASTOp> BaseMeshOp;
+
+		/** Generated node layouts with their own block ids. */
+		TArray<FGeneratedLayout> GeneratedLayouts;
+
+		TArray<Ptr<ASTOp>> LayoutOps;
+
+		struct FExtraLayouts
+		{
+			/** Source node layouts to use with these extra mesh. They don't have block ids. */
+			TArray<FGeneratedLayout> GeneratedLayouts;
+			Ptr<ASTOp> Condition;
+			Ptr<ASTOp> MeshFragment;
+		};
+		TArray< FExtraLayouts > ExtraMeshLayouts;
+	};
+
+	using FMeshTask = UE::Tasks::TTask<FMeshGenerationResult>;
+
+	struct FSurfaceGenerationResult
+	{
+		Ptr<ASTOp> SurfaceOp;
+	};
+
+	using FSurfaceTask = UE::Tasks::TTask<FSurfaceGenerationResult>;
+
+
+	struct FGenericGenerationResult
+	{
+		Ptr<ASTOp> Op;
+	};
+
+	using FLODTask = UE::Tasks::TTask<FGenericGenerationResult>;
+	using FComponentTask = UE::Tasks::TTask<FGenericGenerationResult>;
+
 	/** First pass of the code generation process.
 	 * It collects data about the object hierarchy, the conditions for each object and the global modifiers.
 	 */
@@ -133,7 +188,7 @@ namespace mu
             Ptr<ASTOp> FinalCondition;
 
             // This is filled in the final code generation pass
-            Ptr<ASTOp> ResultSurfaceOp;
+			FSurfaceTask ResultSurfaceTask;
             Ptr<ASTOp> ResultMeshOp;
 
 			bool operator==(const FSurface&) const = default;
diff --git a/Engine/Plugins/Mutable/Source/MutableTools/Private/MuT/CodeGenerator_Mesh.cpp b/Engine/Plugins/Mutable/Source/MutableTools/Private/MuT/CodeGenerator_Mesh.cpp
--- a/Engine/Plugins/Mutable/Source/MutableTools/Private/MuT/CodeGenerator_Mesh.cpp
+++ b/Engine/Plugins/Mutable/Source/MutableTools/Private/MuT/CodeGenerator_Mesh.cpp
@@ -575,7 +575,7 @@ namespace mu
 	
 
     //---------------------------------------------------------------------------------------------
-	CodeGenerator::FMeshTask CodeGenerator::GenerateMesh(const FMeshGenerationStaticOptions& StaticOptions, FMeshOptionsTask InOptions, const Ptr<const NodeMesh>& InUntypedNode)
+	FMeshTask CodeGenerator::GenerateMesh(const FMeshGenerationStaticOptions& StaticOptions, FMeshOptionsTask InOptions, const Ptr<const NodeMesh>& InUntypedNode)
     {
         if (!InUntypedNode)
         {
@@ -609,7 +609,7 @@ namespace mu
     }
 
 
-	CodeGenerator::FMeshTask CodeGenerator::GenerateMesh_Morph( const FMeshGenerationStaticOptions& StaticOptions, FMeshOptionsTask InOptions, const NodeMeshMorph* InMorphNode )
+	FMeshTask CodeGenerator::GenerateMesh_Morph( const FMeshGenerationStaticOptions& StaticOptions, FMeshOptionsTask InOptions, const NodeMeshMorph* InMorphNode )
     {
         NodeMeshMorph::Private& node = *InMorphNode->GetPrivate();
 
@@ -738,7 +738,7 @@ namespace mu
 	}
 
 
-	CodeGenerator::FMeshTask CodeGenerator::GenerateMesh_MakeMorph(const FMeshGenerationStaticOptions& StaticOptions, FMeshOptionsTask InOptions, const NodeMeshMakeMorph* InMakeMorphNode )
+	FMeshTask CodeGenerator::GenerateMesh_MakeMorph(const FMeshGenerationStaticOptions& StaticOptions, FMeshOptionsTask InOptions, const NodeMeshMakeMorph* InMakeMorphNode )
     {
         NodeMeshMakeMorph::Private& node = *InMakeMorphNode->GetPrivate();
 
@@ -823,7 +823,7 @@ namespace mu
 	}
 
 
-	CodeGenerator::FMeshTask CodeGenerator::GenerateMesh_Fragment(const FMeshGenerationStaticOptions& StaticOptions, FMeshOptionsTask InOptions, const NodeMeshFragment* Node )
+	FMeshTask CodeGenerator::GenerateMesh_Fragment(const FMeshGenerationStaticOptions& StaticOptions, FMeshOptionsTask InOptions, const NodeMeshFragment* Node )
     {
 		if (!Node->SourceMesh)
 		{
@@ -880,7 +880,7 @@ namespace mu
 	}
 
 
-	CodeGenerator::FMeshTask CodeGenerator::GenerateMesh_Switch(const FMeshGenerationStaticOptions& StaticOptions, FMeshOptionsTask InOptions, const NodeMeshSwitch* SwitchNode )
+	FMeshTask CodeGenerator::GenerateMesh_Switch(const FMeshGenerationStaticOptions& StaticOptions, FMeshOptionsTask InOptions, const NodeMeshSwitch* SwitchNode )
     {
         NodeMeshSwitch::Private& node = *SwitchNode->GetPrivate();
 
@@ -969,7 +969,7 @@ namespace mu
     }
 
 
-	CodeGenerator::FMeshTask CodeGenerator::GenerateMesh_Table(const FMeshGenerationStaticOptions& StaticOptions, FMeshOptionsTask InOptions, const NodeMeshTable* TableNode)
+	FMeshTask CodeGenerator::GenerateMesh_Table(const FMeshGenerationStaticOptions& StaticOptions, FMeshOptionsTask InOptions, const NodeMeshTable* TableNode)
 	{
 		Ptr<NodeMeshSwitch> SwitchNode;
 		GenerateTableSwitchNode<NodeMeshTable, ETableColumnType::Mesh, NodeMesh, NodeMeshSwitch>(*TableNode, SwitchNode,
@@ -1009,7 +1009,7 @@ namespace mu
 	}
 
 
-	CodeGenerator::FMeshTask CodeGenerator::GenerateMesh_Variation(const FMeshGenerationStaticOptions& StaticOptions, FMeshOptionsTask InOptions, const NodeMeshVariation* VariationNode )
+	FMeshTask CodeGenerator::GenerateMesh_Variation(const FMeshGenerationStaticOptions& StaticOptions, FMeshOptionsTask InOptions, const NodeMeshVariation* VariationNode )
     {
         NodeMeshVariation::Private& node = *VariationNode->GetPrivate();
 
@@ -1140,7 +1140,7 @@ namespace mu
     }
 
 
-	CodeGenerator::FMeshTask CodeGenerator::GenerateMesh_Constant(const FMeshGenerationStaticOptions& StaticOptions, FMeshOptionsTask InOptions, const NodeMeshConstant* InNode )
+	FMeshTask CodeGenerator::GenerateMesh_Constant(const FMeshGenerationStaticOptions& StaticOptions, FMeshOptionsTask InOptions, const NodeMeshConstant* InNode )
     {
 		MUTABLE_CPUPROFILER_SCOPE(GenerateMesh_Constant);
 
@@ -1410,13 +1410,13 @@ namespace mu
 		this->GetModifiersFor(StaticOptions.ComponentId, StaticOptions.ActiveTags, bModifiersForBeforeOperations, Modifiers);
 
 		// This task does the necessary processing of the mesh when it is available
-		FMeshTask ApplyModifiersTask = ApplyMeshModifiers(Modifiers, StaticOptions, InOptions, MeshProcessingTask, nullptr, InNode->GetMessageContext(), InNode);
+		FMeshTask ApplyModifiersTask = ApplyMeshModifiers(Modifiers, StaticOptions, InOptions, MeshProcessingTask, INDEX_NONE, InNode->GetMessageContext(), InNode);
 
 		return ApplyModifiersTask;
     }
 
 
-	CodeGenerator::FMeshTask CodeGenerator::GenerateMesh_Format(const FMeshGenerationStaticOptions& StaticOptions, FMeshOptionsTask InOptions, const NodeMeshFormat* format )
+	FMeshTask CodeGenerator::GenerateMesh_Format(const FMeshGenerationStaticOptions& StaticOptions, FMeshOptionsTask InOptions, const NodeMeshFormat* format )
     {
         NodeMeshFormat::Private& node = *format->GetPrivate();
 
@@ -1477,7 +1477,7 @@ namespace mu
     }
 
 
-	CodeGenerator::FMeshTask CodeGenerator::GenerateMesh_Transform(const FMeshGenerationStaticOptions& StaticOptions, FMeshOptionsTask InOptions, const NodeMeshTransform* TransformNode )
+	FMeshTask CodeGenerator::GenerateMesh_Transform(const FMeshGenerationStaticOptions& StaticOptions, FMeshOptionsTask InOptions, const NodeMeshTransform* TransformNode )
     {
         const NodeMeshTransform::Private& node = *TransformNode->GetPrivate();
 
@@ -1509,7 +1509,7 @@ namespace mu
      }
 
 
-	CodeGenerator::FMeshTask CodeGenerator::GenerateMesh_ClipMorphPlane(const FMeshGenerationStaticOptions& StaticOptions, FMeshOptionsTask InOptions, const NodeMeshClipMorphPlane* ClipNode )
+	FMeshTask CodeGenerator::GenerateMesh_ClipMorphPlane(const FMeshGenerationStaticOptions& StaticOptions, FMeshOptionsTask InOptions, const NodeMeshClipMorphPlane* ClipNode )
     {
 		if (!ClipNode->Source)
 		{
@@ -1581,7 +1581,7 @@ namespace mu
     }
 
 
-	CodeGenerator::FMeshTask CodeGenerator::GenerateMesh_ClipWithMesh(const FMeshGenerationStaticOptions& StaticOptions, FMeshOptionsTask InOptions, const NodeMeshClipWithMesh* ClipNode)
+	FMeshTask CodeGenerator::GenerateMesh_ClipWithMesh(const FMeshGenerationStaticOptions& StaticOptions, FMeshOptionsTask InOptions, const NodeMeshClipWithMesh* ClipNode)
     {
 		if (!ClipNode->Source)
 		{
@@ -1635,7 +1635,7 @@ namespace mu
     }
 
 
-	CodeGenerator::FMeshTask CodeGenerator::GenerateMesh_ClipDeform(const FMeshGenerationStaticOptions& StaticOptions, FMeshOptionsTask InOptions, const NodeMeshClipDeform* ClipDeform)
+	FMeshTask CodeGenerator::GenerateMesh_ClipDeform(const FMeshGenerationStaticOptions& StaticOptions, FMeshOptionsTask InOptions, const NodeMeshClipDeform* ClipDeform)
 	{
 		if (!ClipDeform->BaseMesh)
 		{
@@ -1688,7 +1688,8 @@ namespace mu
 		);
 	}
 
-	CodeGenerator::FMeshTask CodeGenerator::GenerateMesh_ApplyPose(const FMeshGenerationStaticOptions& StaticOptions, FMeshOptionsTask InOptions, const NodeMeshApplyPose* PoseNode )
+	
+	FMeshTask CodeGenerator::GenerateMesh_ApplyPose(const FMeshGenerationStaticOptions& StaticOptions, FMeshOptionsTask InOptions, const NodeMeshApplyPose* PoseNode )
     {
         const NodeMeshApplyPose::Private& node = *PoseNode->GetPrivate();
 
@@ -1742,7 +1743,7 @@ namespace mu
     }
 
 
-	CodeGenerator::FMeshTask CodeGenerator::GenerateMesh_Reshape(const FMeshGenerationStaticOptions& StaticOptions, FMeshOptionsTask InOptions, const NodeMeshReshape* Reshape)
+	FMeshTask CodeGenerator::GenerateMesh_Reshape(const FMeshGenerationStaticOptions& StaticOptions, FMeshOptionsTask InOptions, const NodeMeshReshape* Reshape)
 	{
 		const NodeMeshReshape::Private& Node = *Reshape->GetPrivate();
 
