[mutable] Fixed crash when trying to optimize an "ASTOpImageLayer" op.
- The crash was caused due to the destruction of the loaded data as no reference was being keept while working with the resource.
- Removed some unnecessary includes

#presubmit 67adbf15a8732e974e16fec1
#rnx

diff --git a/Engine/Plugins/Mutable/Source/MutableTools/Private/MuT/ASTOpConstantResource.cpp b/Engine/Plugins/Mutable/Source/MutableTools/Private/MuT/ASTOpConstantResource.cpp
--- a/Engine/Plugins/Mutable/Source/MutableTools/Private/MuT/ASTOpConstantResource.cpp
+++ b/Engine/Plugins/Mutable/Source/MutableTools/Private/MuT/ASTOpConstantResource.cpp
@@ -7,11 +7,9 @@
 #include "MuR/Mesh.h"
 #include "MuR/ModelPrivate.h"
 #include "MuR/ImagePrivate.h"
-#include "MuR/MutableMath.h"
 #include "MuR/PhysicsBody.h"
 #include "MuR/Serialisation.h"
 #include "MuR/Skeleton.h"
-#include "MuR/Types.h"
 
 #include "Containers/Array.h"
 #include "HAL/PlatformMath.h"
@@ -55,7 +53,7 @@ namespace mu
 		/** Shared context with cache settings and stats. */
 		FProxyFileContext& Options;
 
-		/** Prevent concurrent access to a signel resource. */
+		/** Prevent concurrent access to a signal resource. */
 		FCriticalSection Mutex;
 
 	public:
@@ -611,7 +609,7 @@ namespace mu
 
 		case EOpType::IM_CONSTANT:
 		{
-			const FImage* pImage = static_cast<const FImage*>(GetValue().Get());
+			TSharedPtr<const FImage> pImage = StaticCastSharedPtr<const FImage>(GetValue());
 			if (pImage->GetSizeX() <= 0 || pImage->GetSizeY() <= 0)
 			{
 				res = true;
@@ -799,7 +797,7 @@ namespace mu
 		ASTOp::EClosedMeshTest Result = EClosedMeshTest::Unknown;
 		if (Type == EOpType::ME_CONSTANT)
 		{
-			const FMesh* Mesh = static_cast<const FMesh*>(GetValue().Get());
+			TSharedPtr<const FMesh> Mesh = StaticCastSharedPtr<const FMesh>(GetValue());
 			if (Mesh)
 			{
 				if (Mesh->IsClosed())
diff --git a/Engine/Plugins/Mutable/Source/MutableTools/Private/MuT/ASTOpConstantResource.h b/Engine/Plugins/Mutable/Source/MutableTools/Private/MuT/ASTOpConstantResource.h
--- a/Engine/Plugins/Mutable/Source/MutableTools/Private/MuT/ASTOpConstantResource.h
+++ b/Engine/Plugins/Mutable/Source/MutableTools/Private/MuT/ASTOpConstantResource.h
@@ -47,7 +47,7 @@ namespace mu
 		//! Get a hash of the stored value.
 		uint64 GetValueHash() const;
 
-		//! Get a copy of the stored value
+		//! Get a pointer to the stored value.
 		TSharedPtr<const FResource> GetValue() const;
 
 		/** Set the value to store in this op.
diff --git a/Engine/Plugins/Mutable/Source/MutableTools/Private/MuT/ASTOpImageLayer.cpp b/Engine/Plugins/Mutable/Source/MutableTools/Private/MuT/ASTOpImageLayer.cpp
--- a/Engine/Plugins/Mutable/Source/MutableTools/Private/MuT/ASTOpImageLayer.cpp
+++ b/Engine/Plugins/Mutable/Source/MutableTools/Private/MuT/ASTOpImageLayer.cpp
@@ -12,8 +12,6 @@
 #include "MuT/ASTOpImageDisplace.h"
 #include "MuT/ASTOpSwitch.h"
 #include "MuR/ModelPrivate.h"
-#include "MuR/RefCounted.h"
-#include "MuR/Types.h"
 
 #include "Containers/Map.h"
 #include "HAL/PlatformMath.h"
