[Mutable] Fix invalid access to streamed resources when the streaming requests fail.

#rnx

diff --git a/Engine/Plugins/Mutable/Source/CustomizableObject/Private/MuCO/CustomizableObjectInstance.cpp b/Engine/Plugins/Mutable/Source/CustomizableObject/Private/MuCO/CustomizableObjectInstance.cpp
--- a/Engine/Plugins/Mutable/Source/CustomizableObject/Private/MuCO/CustomizableObjectInstance.cpp
+++ b/Engine/Plugins/Mutable/Source/CustomizableObject/Private/MuCO/CustomizableObjectInstance.cpp
@@ -6208,7 +6208,7 @@ void UCustomizableInstancePrivate::AdditionalAssetsAsyncLoaded(const TSharedRef<
 				}
 			}
 #else
-			if (ModelResources.StreamedResourceData.IsValidIndex(ResourceIndex))
+			if (ModelResources.StreamedResourceData.IsValidIndex(ResourceIndex) && ModelResources.StreamedResourceData[ResourceIndex].IsLoaded())
 			{
 				const FCustomizableObjectResourceData& ResourceData = ModelResources.StreamedResourceData[ResourceIndex].GetLoadedData();
 
