[Mutable] Fix GPU crash due to an invalid number of bone influences when the mesh has less than MAX_INFLUENCES_PER_STREAM and mesh LOD streaming is enabled.
- Rename generated meshes to display the component name.
- Fix crash when copying data from Mutable to the mesh due to incorrect FirstLODAvailable.

#rnx

diff --git a/Engine/Plugins/Mutable/Source/CustomizableObject/Internal/MuCO/CustomizableObjectSystemPrivate.h b/Engine/Plugins/Mutable/Source/CustomizableObject/Internal/MuCO/CustomizableObjectSystemPrivate.h
--- a/Engine/Plugins/Mutable/Source/CustomizableObject/Internal/MuCO/CustomizableObjectSystemPrivate.h
+++ b/Engine/Plugins/Mutable/Source/CustomizableObject/Internal/MuCO/CustomizableObjectSystemPrivate.h
@@ -90,9 +90,10 @@ public:
 	// These are the LODs that would be applied if this candidate is chosen
 	TMap<FName, uint8> MinLOD;
 
-	TMap<FName, uint8> RequestedLODLevels;
+	// These are the LODs that would be copied to the descriptor to trigger mesh updates on quality setting changes.
+	TMap<FName, uint8> QualitySettingMinLODs;
 
-	TMap<FName, uint8> FirstLODAvailable;
+	TMap<FName, uint8> RequestedLODLevels;
 
 	FMutableUpdateCandidate(UCustomizableObjectInstance* InCustomizableObjectInstance);
 	
@@ -513,9 +514,7 @@ public:
 	
 	void SetRequestedLODs(const TMap<FName, uint8>& RequestedLODs);
 
-	void SetFirstLODAvailable(const TMap<FName, uint8>& FirstLODs);
-
-	const TMap<FName, uint8>& GetFirstLODAvailable() const;
+	void SetQualitySettingMinLODs(const TMap<FName, uint8>& FirstLODs);
 
 	const FCustomizableObjectInstanceDescriptor& GetCapturedDescriptor() const;
 
@@ -607,6 +606,7 @@ public:
 	
 	TMap<FName, uint8> NumLODsAvailable;
 
+	/** Copy of UModelResources::ComponentFirstLODAvailable. First compiled LOD per component for the running platform. Constant. */
 	TMap<FName, uint8> FirstLODAvailable;
 	
 	TMap<FName, uint8> FirstResidentLOD;
diff --git a/Engine/Plugins/Mutable/Source/CustomizableObject/Internal/MuCO/DescriptorHash.h b/Engine/Plugins/Mutable/Source/CustomizableObject/Internal/MuCO/DescriptorHash.h
--- a/Engine/Plugins/Mutable/Source/CustomizableObject/Internal/MuCO/DescriptorHash.h
+++ b/Engine/Plugins/Mutable/Source/CustomizableObject/Internal/MuCO/DescriptorHash.h
@@ -23,9 +23,11 @@ private:
 	uint32 Hash = 0;
 
 public:
+	// MinLOD based on user input and quality settings. First LOD to generate.
 	TMap<FName, uint8> MinLODs;
 
-	TMap<FName, uint8> FirstLODAvailable;
+	// MinLOD based on quality settings. Used to trigger updates after changing the active quality level.
+	TMap<FName, uint8> QualitySettingMinLODs;
 	
 	// Array of bitmasks that indicate which LODs of each component have been requested
 	TMap<FName, uint8> RequestedLODsPerComponent;
diff --git a/Engine/Plugins/Mutable/Source/CustomizableObject/Private/MuCO/CustomizableObjectInstance.cpp b/Engine/Plugins/Mutable/Source/CustomizableObject/Private/MuCO/CustomizableObjectInstance.cpp
--- a/Engine/Plugins/Mutable/Source/CustomizableObject/Private/MuCO/CustomizableObjectInstance.cpp
+++ b/Engine/Plugins/Mutable/Source/CustomizableObject/Private/MuCO/CustomizableObjectInstance.cpp
@@ -2393,11 +2393,11 @@ bool UCustomizableInstancePrivate::UpdateSkeletalMesh_PostBeginUpdate0(UCustomiz
 
 		if (OperationData->bStreamMeshLODs)
 		{
-			SkeletalMesh = UCustomizableObjectSkeletalMesh::CreateSkeletalMesh(*Public);
+			SkeletalMesh = UCustomizableObjectSkeletalMesh::CreateSkeletalMesh(*Public, ComponentName);
 		}
 		else
 		{
-			FName SkeletalMeshName = GenerateUniqueNameFromCOInstance(*Public);
+			FName SkeletalMeshName = GenerateUniqueNameFromCOInstance(*Public, ComponentName);
 			SkeletalMesh = NewObject<USkeletalMesh>(GetTransientPackage(), SkeletalMeshName, RF_Transient);
 		}
 		
@@ -3301,7 +3301,7 @@ FMutableUpdateCandidate::FMutableUpdateCandidate(UCustomizableObjectInstance* In
 {
 	const FCustomizableObjectInstanceDescriptor& Descriptor = InCustomizableObjectInstance->GetPrivate()->GetDescriptor();
 	MinLOD = Descriptor.MinLOD;
-	FirstLODAvailable = Descriptor.FirstLODAvailable;
+	QualitySettingMinLODs = Descriptor.QualitySettingMinLODs;
 	RequestedLODLevels = Descriptor.GetRequestedLODLevels();
 }
 
@@ -3321,11 +3321,11 @@ void FMutableUpdateCandidate::Issue()
 void FMutableUpdateCandidate::ApplyLODUpdateParamsToInstance(FUpdateContextPrivate& Context)
 {
 	CustomizableObjectInstance->Descriptor.MinLOD = MinLOD;
-	CustomizableObjectInstance->Descriptor.FirstLODAvailable = FirstLODAvailable;
+	CustomizableObjectInstance->Descriptor.QualitySettingMinLODs = QualitySettingMinLODs;
 	CustomizableObjectInstance->Descriptor.RequestedLODLevels = RequestedLODLevels;
 
 	Context.SetMinLOD(MinLOD);
-	Context.SetFirstLODAvailable(FirstLODAvailable);
+	Context.SetQualitySettingMinLODs(QualitySettingMinLODs);
 	Context.SetRequestedLODs(RequestedLODLevels);
 }
 
@@ -4016,7 +4016,7 @@ void UCustomizableInstancePrivate::BuildMeshSockets(const TSharedRef<FUpdateCont
 		{
 			const FMutableRefSocket& RefSocket = RefSkeletalMeshData.Sockets[SocketIndex];
 
-			USkeletalMeshSocket* Socket = NewObject<USkeletalMeshSocket>(SkeletalMesh);
+			USkeletalMeshSocket* Socket = NewObject<USkeletalMeshSocket>(SkeletalMesh, RefSocket.SocketName);
 
 			Socket->SocketName = RefSocket.SocketName;
 			Socket->BoneName = RefSocket.BoneName;
@@ -4063,7 +4063,7 @@ void UCustomizableInstancePrivate::BuildMeshSockets(const TSharedRef<FUpdateCont
 					else
 					{
 						// New Socket
-						USkeletalMeshSocket* Socket = NewObject<USkeletalMeshSocket>(SkeletalMesh);
+						USkeletalMeshSocket* Socket = NewObject<USkeletalMeshSocket>(SkeletalMesh, MutableSocket.SocketName);
 						IndexToWriteSocket = Sockets.Add(Socket);
 						SocketMap.Add(MutableSocket.SocketName, TTuple<int32, int32>(IndexToWriteSocket, MutableSocket.Priority));
 					}
@@ -7318,7 +7318,8 @@ void UCustomizableObjectInstance::SetRequestedLODs(const TMap<FName, uint8>& InM
 
 	// DOCUMENTATION 
 	// MinLod =					Hard limit set by the budgeter or the user (artifical limit).			Per Skeletal Mesh Component.
-	// First LOD Available =	First available lod per platform and quality (scalability) settings.	Per Skeletal Mesh.
+	// QualitySettingMinLODs =	MinLOD based on the active quality settings.							COI Descriptor.
+	// First LOD Available =	First available lod per platform.										Per Skeletal Mesh.
 	// Requested LODs =			What the renderer wants to render.										Per Skeletal Mesh Component.
 
 	if (!GetPrivate()->CanUpdateInstance())
@@ -7359,8 +7360,6 @@ void UCustomizableObjectInstance::SetRequestedLODs(const TMap<FName, uint8>& InM
 	// Save the new LODs
 	MutableUpdateCandidate.MinLOD = InMinLODs;
 	MutableUpdateCandidate.RequestedLODLevels = Descriptor.GetRequestedLODLevels();
-
-	MutableUpdateCandidate.FirstLODAvailable.Reset();
 	
 	const TMap<FName, uint8>& GeneratedLODsPerComponent = GetPrivate()->CommittedDescriptorHash.RequestedLODsPerComponent;
 	
@@ -7373,7 +7372,8 @@ void UCustomizableObjectInstance::SetRequestedLODs(const TMap<FName, uint8>& InM
 		}
 
 		const uint8 MinLODIdx = CustomizableObject->GetPrivate()->GetMinLODIndex(ComponentName);
-		MutableUpdateCandidate.FirstLODAvailable.Add(ComponentName, MinLODIdx);
+		MutableUpdateCandidate.QualitySettingMinLODs.Add(ComponentName, MinLODIdx);
+
 		const int32 MaxLODIdx = ModelResources->ComponentLODs[ComponentName] - 1;
 		InMinLOD = FMath::Clamp(InMinLOD, MinLODIdx, MaxLODIdx);
 
diff --git a/Engine/Plugins/Mutable/Source/CustomizableObject/Private/MuCO/CustomizableObjectSkeletalMesh.cpp b/Engine/Plugins/Mutable/Source/CustomizableObject/Private/MuCO/CustomizableObjectSkeletalMesh.cpp
--- a/Engine/Plugins/Mutable/Source/CustomizableObject/Private/MuCO/CustomizableObjectSkeletalMesh.cpp
+++ b/Engine/Plugins/Mutable/Source/CustomizableObject/Private/MuCO/CustomizableObjectSkeletalMesh.cpp
@@ -12,9 +12,9 @@
 
 #include "Animation/SkinWeightProfile.h"
 
-FName GenerateUniqueNameFromCOInstance(const UCustomizableObjectInstance& Instance)
+FName GenerateUniqueNameFromCOInstance(const UCustomizableObjectInstance& Instance, FName ComponentName)
 {
-	FString BaseName = FString("MutableSkMesh-") + Instance.GetCustomizableObject()->GetName();
+	FString BaseName = FString("SK_") + Instance.GetCustomizableObject()->GetName()+ FString("_") + ComponentName.ToString();
 	FName SkeletalMeshName = MakeUniqueObjectName(GetTransientPackage(), USkeletalMesh::StaticClass(), FName(BaseName), 
 							                      EUniqueObjectNameOptions::GloballyUnique);
 
@@ -22,9 +22,9 @@ FName GenerateUniqueNameFromCOInstance(const UCustomizableObjectInstance& Instan
 }
 
 
-UCustomizableObjectSkeletalMesh* UCustomizableObjectSkeletalMesh::CreateSkeletalMesh(const UCustomizableObjectInstance& Instance)
+UCustomizableObjectSkeletalMesh* UCustomizableObjectSkeletalMesh::CreateSkeletalMesh(const UCustomizableObjectInstance& Instance, FName ComponentName)
 {
-	FName SkeletalMeshName = GenerateUniqueNameFromCOInstance(Instance);
+	FName SkeletalMeshName = GenerateUniqueNameFromCOInstance(Instance, ComponentName);
 	UCustomizableObjectSkeletalMesh* OutSkeletalMesh = NewObject<UCustomizableObjectSkeletalMesh>(GetTransientPackage(), SkeletalMeshName, RF_Transient);
 	
 	return OutSkeletalMesh;
diff --git a/Engine/Plugins/Mutable/Source/CustomizableObject/Private/MuCO/CustomizableObjectSystem.cpp b/Engine/Plugins/Mutable/Source/CustomizableObject/Private/MuCO/CustomizableObjectSystem.cpp
--- a/Engine/Plugins/Mutable/Source/CustomizableObject/Private/MuCO/CustomizableObjectSystem.cpp
+++ b/Engine/Plugins/Mutable/Source/CustomizableObject/Private/MuCO/CustomizableObjectSystem.cpp
@@ -411,16 +411,10 @@ void FUpdateContextPrivate::SetMinLOD(const TMap<FName, uint8>& MinLOD)
 }
 
 
-void FUpdateContextPrivate::SetFirstLODAvailable(const TMap<FName, uint8>& FirstLODs)
+void FUpdateContextPrivate::SetQualitySettingMinLODs(const TMap<FName, uint8>& QualitySettingsMinLODs)
 {
-	CapturedDescriptor.FirstLODAvailable = FirstLODs;
-	CapturedDescriptorHash.FirstLODAvailable = FirstLODs;
-}
-
-
-const TMap<FName, uint8>& FUpdateContextPrivate::GetFirstLODAvailable() const
-{
-	return CapturedDescriptor.FirstLODAvailable;
+	CapturedDescriptor.QualitySettingMinLODs = QualitySettingsMinLODs;
+	CapturedDescriptorHash.QualitySettingMinLODs = QualitySettingsMinLODs;
 }
 
 
@@ -2040,7 +2034,6 @@ namespace impl
 
 		TMap<FName, uint8> MinLODs = Operation->bStreamMeshLODs ? TMap<FName, uint8>() : Operation->GetMinLOD();
 		TMap<FName, uint8> RequestedLODs = Operation->GetRequestedLODs();
-		TMap<FName, uint8> FirstLODsAvailable = Operation->GetFirstLODAvailable();
 		
 		for (int32 InstanceComponentIndex = 0; InstanceComponentIndex < Operation->NumInstanceComponents; ++InstanceComponentIndex)
 		{
@@ -2051,7 +2044,7 @@ namespace impl
 			uint8& RequestedLOD = RequestedLODs.FindOrAdd(ComponentName, 0);
 			uint8& NumLODsAvailable = Operation->NumLODsAvailable.FindOrAdd(ComponentName, 0);
 			uint8& FirstResidentLOD = Operation->FirstResidentLOD.FindOrAdd(ComponentName, 0);
-			uint8& FirstLODAvailable = FirstLODsAvailable.FindOrAdd(ComponentName, 0);
+			const uint8& FirstLODAvailable = Operation->FirstLODAvailable.FindOrAdd(ComponentName, 0);
 
 			NumLODsAvailable = Operation->MutableInstance->GetLODCount(InstanceComponentIndex);
 			MinLOD = FMath::Clamp(MinLOD, FirstLODAvailable, NumLODsAvailable - 1);
@@ -2070,7 +2063,6 @@ namespace impl
 		
 		Operation->SetMinLOD(MinLODs);
 		Operation->SetRequestedLODs(RequestedLODs);
-		Operation->SetFirstLODAvailable(FirstLODsAvailable);
 	}
 	
 
@@ -4137,8 +4129,8 @@ int32 UCustomizableObjectSystem::TickInternal(const bool bBlocking)
 
 						FDescriptorHash LODUpdateDescriptorHash = CustomizableInstancePrivate->CommittedDescriptorHash;
 						LODUpdateDescriptorHash.MinLODs = LODUpdateCandidate.MinLOD;
+						LODUpdateDescriptorHash.QualitySettingMinLODs = LODUpdateCandidate.QualitySettingMinLODs;
 						LODUpdateDescriptorHash.RequestedLODsPerComponent = LODUpdateCandidate.RequestedLODLevels;
-						LODUpdateDescriptorHash.FirstLODAvailable = LODUpdateCandidate.FirstLODAvailable;
 
 						if (CustomizableInstancePrivate->MinSquareDistFromComponentToPlayer < MaxSquareDistanceFound &&
 							!LODUpdateDescriptorHash.IsSubset(CustomizableInstancePrivate->CommittedDescriptorHash))
diff --git a/Engine/Plugins/Mutable/Source/CustomizableObject/Private/MuCO/DescriptorHash.cpp b/Engine/Plugins/Mutable/Source/CustomizableObject/Private/MuCO/DescriptorHash.cpp
--- a/Engine/Plugins/Mutable/Source/CustomizableObject/Private/MuCO/DescriptorHash.cpp
+++ b/Engine/Plugins/Mutable/Source/CustomizableObject/Private/MuCO/DescriptorHash.cpp
@@ -53,7 +53,7 @@ FDescriptorHash::FDescriptorHash(const FCustomizableObjectInstanceDescriptor& De
 	Hash = HashCombine(Hash, GetTypeHash(Descriptor.GetBuildParameterRelevancy()));
 	
 	MinLODs = Descriptor.MinLOD;
-	FirstLODAvailable = Descriptor.FirstLODAvailable;
+	QualitySettingMinLODs = Descriptor.QualitySettingMinLODs;
 	RequestedLODsPerComponent = Descriptor.RequestedLODLevels;
 }
 
@@ -80,7 +80,7 @@ bool FDescriptorHash::IsSubset(const FDescriptorHash& Other) const
 	}
 	
 	// Scalability quality change : Do not care if it is a subset, always update
-	if (FirstLODAvailable != Other.FirstLODAvailable)
+	if (QualitySettingMinLODs != Other.QualitySettingMinLODs)
 	{
 		return false;
 	}
@@ -122,9 +122,9 @@ FString FDescriptorHash::ToString() const
 	Builder.Appendf(TEXT("])"));
 
 	Builder.Appendf(TEXT("FirstLODAvailable=["));
-	for (const TTuple<FName, uint8>& FirstLOD : FirstLODAvailable)
+	for (const TTuple<FName, uint8>& QualitySetting : QualitySettingMinLODs)
 	{
-		Builder.Appendf(TEXT("(%s, %i),"), *FirstLOD.Key.ToString(), FirstLOD.Value);
+		Builder.Appendf(TEXT("(%s, %i),"), *QualitySetting.Key.ToString(), QualitySetting.Value);
 	}
 	Builder.Appendf(TEXT("])"));
 
diff --git a/Engine/Plugins/Mutable/Source/CustomizableObject/Private/MuCO/UnrealConversionUtils.cpp b/Engine/Plugins/Mutable/Source/CustomizableObject/Private/MuCO/UnrealConversionUtils.cpp
--- a/Engine/Plugins/Mutable/Source/CustomizableObject/Private/MuCO/UnrealConversionUtils.cpp
+++ b/Engine/Plugins/Mutable/Source/CustomizableObject/Private/MuCO/UnrealConversionUtils.cpp
@@ -291,6 +291,11 @@ namespace UnrealConversionUtils
 			LODResource.SkinWeightVertexBuffer.SetNeedsCPUAccess(bAllowCPUAccess);
 
 			FSkinWeightDataVertexBuffer* VertexBuffer = LODResource.SkinWeightVertexBuffer.GetDataVertexBuffer();
+
+			// NumBoneInfluences must be equal to MaxBoneInfluences. Set and then GetMaxBoneInfluences to know the number of bone influences to use (at least MAX_INFLUENCES_PER_STREAM).
+			VertexBuffer->SetMaxBoneInfluences(NumBoneInfluences);
+			NumBoneInfluences = VertexBuffer->GetMaxBoneInfluences();
+			
 			VertexBuffer->Init(NumBoneInfluences, DummyNumVertices);
 			VertexBuffer->SetMetaData(NumVertices, NumBoneInfluences, bUse16BitBoneIndex, bUse16BitBoneWeights);
 
diff --git a/Engine/Plugins/Mutable/Source/CustomizableObject/Public/MuCO/CustomizableObjectInstanceDescriptor.h b/Engine/Plugins/Mutable/Source/CustomizableObject/Public/MuCO/CustomizableObjectInstanceDescriptor.h
--- a/Engine/Plugins/Mutable/Source/CustomizableObject/Public/MuCO/CustomizableObjectInstanceDescriptor.h
+++ b/Engine/Plugins/Mutable/Source/CustomizableObject/Public/MuCO/CustomizableObjectInstanceDescriptor.h
@@ -330,14 +330,16 @@ struct CUSTOMIZABLEOBJECT_API FCustomizableObjectInstanceDescriptor
 	/** If this is set to true, when updating the instance an additional step will be performed to calculate the list of instance parameters that are relevant for the current parameter values. */
 	bool bBuildParameterRelevancy = false;
 
-	/** These are the LODs Mutable can generate, they MUST NOT be used in an update (Mutable thread). */
+	/** These are the LODs Mutable can generate based on user input and quality settings, they MUST NOT be used in an update (Mutable thread). */
 	TMap<FName, uint8> MinLOD;
 
+	/** These are the MinLODs based on the active quality setting */
+	TMap<FName, uint8> QualitySettingMinLODs;
+
 	/** RequestedLODs per component to generate.
 	 * Key is the Component name. Value represents the first LOD to generate. Zero means all LODs. */
 	TMap<FName, uint8> RequestedLODLevels;
 	
-	TMap<FName, uint8> FirstLODAvailable;
 	
 	// Friends
 	friend FDescriptorHash;
diff --git a/Engine/Plugins/Mutable/Source/CustomizableObject/Public/MuCO/CustomizableObjectSkeletalMesh.h b/Engine/Plugins/Mutable/Source/CustomizableObject/Public/MuCO/CustomizableObjectSkeletalMesh.h
--- a/Engine/Plugins/Mutable/Source/CustomizableObject/Public/MuCO/CustomizableObjectSkeletalMesh.h
+++ b/Engine/Plugins/Mutable/Source/CustomizableObject/Public/MuCO/CustomizableObjectSkeletalMesh.h
@@ -22,7 +22,7 @@ class FUpdateContextPrivate;
 
 class FCustomizableObjectMeshStreamIn;
 
-FName GenerateUniqueNameFromCOInstance(const UCustomizableObjectInstance& Instance);
+FName GenerateUniqueNameFromCOInstance(const UCustomizableObjectInstance& Instance, FName ComponentName);
 
 /**
  * CustomizableObjectSkeletalMesh is a class inheriting from USkeletalMesh with the sole purpose of allowing skeletal meshes generated by Mutable
@@ -38,7 +38,7 @@ class CUSTOMIZABLEOBJECT_API UCustomizableObjectSkeletalMesh : public USkeletalM
 
 private:
 
-	static UCustomizableObjectSkeletalMesh* CreateSkeletalMesh(const UCustomizableObjectInstance& InInstance);
+	static UCustomizableObjectSkeletalMesh* CreateSkeletalMesh(const UCustomizableObjectInstance& InInstance, FName ComponentName);
 
 	void InitMutableStreamingData(const TSharedRef<FUpdateContextPrivate>& InOperationData, const FName& ComponentName, const int32 FirstLOD, const int32 LODCount);
 
