[Mutable] Prefetch bulk data files for high-quality mips when streaming texture mipmaps.
+ Add MutableStreamRequest to gather and stream multiple binary blobs.
+ Add CVarEnablePrefetchHighQualityMips to turn on/off the prefetch of high-quality mips.

#rnx

diff --git a/Engine/Plugins/Mutable/Source/CustomizableObject/Internal/MuCO/CustomizableObjectMipDataProvider.h b/Engine/Plugins/Mutable/Source/CustomizableObject/Internal/MuCO/CustomizableObjectMipDataProvider.h
--- a/Engine/Plugins/Mutable/Source/CustomizableObject/Internal/MuCO/CustomizableObjectMipDataProvider.h
+++ b/Engine/Plugins/Mutable/Source/CustomizableObject/Internal/MuCO/CustomizableObjectMipDataProvider.h
@@ -22,6 +22,8 @@ class FThreadSafeCounter;
 class UCustomizableObjectInstance;
 class UObject;
 class UTexture;
+class FMutableStreamRequest;
+struct FModelStreamableBulkData;
 
 /** This struct stores the data relevant for the construction of a specific texture. 
 * This includes all the data required to rebuild the image (or any of its mips).
@@ -32,7 +34,7 @@ public:
 	FMutableUpdateContext() = default;
 	
 	FMutableUpdateContext(const FString& InCustomizableObjectPathName, const FString& InInstancePathName, TSharedPtr<mu::FSystem> InSystem,
-		TSharedPtr<mu::FModel, ESPMode::ThreadSafe> InModel, TSharedPtr<const mu::FParameters> InParameters, int32 InState);
+		TSharedPtr<mu::FModel, ESPMode::ThreadSafe> InModel, TSharedPtr<FModelStreamableBulkData> InModelStreamable, TSharedPtr<const mu::FParameters> InParameters, int32 InState);
 
 	~FMutableUpdateContext();
 
@@ -44,6 +46,8 @@ public:
 	
 	TSharedPtr<mu::FModel, ESPMode::ThreadSafe> GetModel() const;
 
+	TSharedPtr<FModelStreamableBulkData> GetModelStreamableBulkData();
+
 	TSharedPtr<const mu::FParameters> GetParameters() const;
 	
 	int32 GetState() const;
@@ -60,6 +64,8 @@ private:
 
 	TSharedPtr<mu::FSystem> System;
 	TSharedPtr<mu::FModel, ESPMode::ThreadSafe> Model;
+	TSharedPtr<FModelStreamableBulkData> ModelStreamableBulkData;
+
 	TSharedPtr<const mu::FParameters> Parameters;
 	int32 State = -1;
 
@@ -94,10 +100,21 @@ struct FMutableImageOperationData
 
 	TArray<FMutableMipUpdateLevel> Levels;
 
+	uint32 MutableTaskId = FMutableTaskGraph::INVALID_ID;
+
+	struct CUSTOMIZABLEOBJECT_API FPrefetchMemoryCounter
+	{
+		alignas(8) static inline std::atomic<SSIZE_T> Counter{ 0 };
+	};
+	using ArrayDataContainerType = TArray<uint8, mu::FDefaultMemoryTrackingAllocator<FPrefetchMemoryCounter>>;
+	ArrayDataContainerType AllocatedMemory;
+
 	// Used to sync with the FMutableTextureMipDataProvider and FRenderAssetUpdate::Tick
 	FThreadSafeCounter* Counter;
 	FTextureUpdateSyncOptions::FCallback RescheduleCallback;
 
+	bool bIsCancelled = false;
+
 	/** Access to the Counter must be protected with this because it may be accessed from another thread to null it. */
 	FCriticalSection CounterTaskLock;
 
@@ -120,7 +137,7 @@ public:
 	void AbortPollMips() override;
 
 private:
-	void CancelCounterSafely();
+	void CancelAsyncTasks();
 	void PrintWarningAndAdvanceToCleanup();
 
 public:
@@ -135,7 +152,7 @@ public:
 
 	TSharedPtr<FMutableImageOperationData> OperationData;
 
-	uint32 MutableTaskId = FMutableTaskGraph::INVALID_ID;
+	TUniquePtr<FMutableStreamRequest> PrefetchRequest;
 };
 
 
diff --git a/Engine/Plugins/Mutable/Source/CustomizableObject/Internal/MuCO/CustomizableObjectPrivate.h b/Engine/Plugins/Mutable/Source/CustomizableObject/Internal/MuCO/CustomizableObjectPrivate.h
--- a/Engine/Plugins/Mutable/Source/CustomizableObject/Internal/MuCO/CustomizableObjectPrivate.h
+++ b/Engine/Plugins/Mutable/Source/CustomizableObject/Internal/MuCO/CustomizableObjectPrivate.h
@@ -597,7 +597,9 @@ struct CUSTOMIZABLEOBJECT_API FMutableStreamableBlock
 
 	/** Used to store properties of the data, necessary for its recovery. For instance if it is high-res. */
 	UPROPERTY()
-	uint32 Flags = 0;
+	uint16 Flags = 0;
+	
+	uint16 IsPrefetched = 0;
 
 	UPROPERTY()
 	uint64 Offset = 0;
@@ -610,7 +612,6 @@ struct CUSTOMIZABLEOBJECT_API FMutableStreamableBlock
 		return Ar;
 	}
 };
-template<> struct TCanBulkSerialize<FMutableStreamableBlock> { enum { Value = true }; };
 static_assert(sizeof(FMutableStreamableBlock) == 8 * 2);
 
 
@@ -870,7 +871,6 @@ enum class EMutableFileFlags : uint8
 };
 
 
-#if WITH_EDITOR
 namespace MutablePrivate
 {
 	enum class EStreamableDataType : uint32 // uint32 for padding and DDC purposes
@@ -883,7 +883,7 @@ namespace MutablePrivate
 		DataTypeCount
 	};
 
-
+#if WITH_EDITOR
 	struct FBlock
 	{
 		/** Used on some data types as the index to the block stored in the CustomizableObject */
@@ -1083,9 +1083,8 @@ namespace MutablePrivate
 	UE::DerivedData::FValueId CUSTOMIZABLEOBJECT_API GetDerivedDataModelResourcesId();
 	UE::DerivedData::FValueId CUSTOMIZABLEOBJECT_API GetDerivedDataModelStreamableBulkDataId();
 	UE::DerivedData::FValueId CUSTOMIZABLEOBJECT_API GetDerivedDataBulkDataFilesId();
-}
 #endif
-
+}
 
 struct CUSTOMIZABLEOBJECT_API FModelStreamableBulkData
 {
@@ -1098,6 +1097,9 @@ struct CUSTOMIZABLEOBJECT_API FModelStreamableBulkData
 
 	TArray<FByteBulkData> StreamableBulkData;
 
+	// File path to stream resources from when not using FByteBulkData or DDC.
+	FString FullFilePath;
+
 	void Serialize(FArchive& Ar, UObject* Owner, bool bCooked);
 
 #if WITH_EDITORONLY_DATA
@@ -1108,6 +1110,8 @@ struct CUSTOMIZABLEOBJECT_API FModelStreamableBulkData
 		Ar << Data.RealTimeMorphStreamables;
 
 		// Don't serialize FByteBulkData manually, the data will be skipped.
+		
+		Ar << Data.FullFilePath;
 
 		return Ar;
 	}
@@ -1125,6 +1129,8 @@ class UModelStreamableData : public UObject
 public:
 	virtual void Serialize(FArchive& Ar) override;
 
+	virtual void PostLoad() override;
+
 	TSharedPtr<FModelStreamableBulkData> StreamingData;
 };
 
@@ -1832,6 +1838,8 @@ public:
 
 		AddOverlayMaterials,
 
+		PrefetchHighQualityMipMaps,
+
 		// -----<new versions can be added above this line>--------
 		LastCustomizableObjectVersion
 	};
diff --git a/Engine/Plugins/Mutable/Source/CustomizableObject/Internal/MuCO/CustomizableObjectSystemPrivate.h b/Engine/Plugins/Mutable/Source/CustomizableObject/Internal/MuCO/CustomizableObjectSystemPrivate.h
--- a/Engine/Plugins/Mutable/Source/CustomizableObject/Internal/MuCO/CustomizableObjectSystemPrivate.h
+++ b/Engine/Plugins/Mutable/Source/CustomizableObject/Internal/MuCO/CustomizableObjectSystemPrivate.h
@@ -44,6 +44,7 @@ struct FTexturePlatformData;
 // Split StreamedBulkData into chunks smaller than MUTABLE_STREAMED_DATA_MAXCHUNKSIZE
 #define MUTABLE_STREAMED_DATA_MAXCHUNKSIZE		(512 * 1024 * 1024)
 
+extern TAutoConsoleVariable<bool> CVarMutableHighPriorityLoading;
 
 /** End a Customizable Object Instance Update. All code paths of an update have to end here. */
 void FinishUpdateGlobal(const TSharedRef<FUpdateContextPrivate>& Context);
@@ -330,6 +331,8 @@ struct FInstanceUpdateData
 		TSharedPtr<const mu::FImage> Image;
 		TWeakObjectPtr<UTexture2D> Cached;
 
+		TArray<int32> ConstantImagesNeededToGenerate;
+
 		bool bIsPassThrough = false;
 		bool bIsNonProgressive = false;
 	};
diff --git a/Engine/Plugins/Mutable/Source/CustomizableObject/Private/MuCO/CustomizableObject.cpp b/Engine/Plugins/Mutable/Source/CustomizableObject/Private/MuCO/CustomizableObject.cpp
--- a/Engine/Plugins/Mutable/Source/CustomizableObject/Private/MuCO/CustomizableObject.cpp
+++ b/Engine/Plugins/Mutable/Source/CustomizableObject/Private/MuCO/CustomizableObject.cpp
@@ -2127,6 +2127,19 @@ void UModelStreamableData::Serialize(FArchive& Ar)
 	}
 }
 
+void UModelStreamableData::PostLoad()
+{
+	Super::PostLoad();
+
+	if (StreamingData)
+	{
+		const FString OutermostName = GetOutermost()->GetName();
+		FString PackageFilename = FPackageName::LongPackageNameToFilename(OutermostName);
+		FPaths::MakeStandardFilename(PackageFilename);
+		StreamingData->FullFilePath = PackageFilename;
+	}
+}
+
 
 void UCustomizableObjectPrivate::SetModel(const TSharedPtr<mu::FModel, ESPMode::ThreadSafe>& Model, const FGuid Id)
 {
diff --git a/Engine/Plugins/Mutable/Source/CustomizableObject/Private/MuCO/CustomizableObjectInstance.cpp b/Engine/Plugins/Mutable/Source/CustomizableObject/Private/MuCO/CustomizableObjectInstance.cpp
--- a/Engine/Plugins/Mutable/Source/CustomizableObject/Private/MuCO/CustomizableObjectInstance.cpp
+++ b/Engine/Plugins/Mutable/Source/CustomizableObject/Private/MuCO/CustomizableObjectInstance.cpp
@@ -5428,12 +5428,6 @@ bool UCustomizableInstancePrivate::BuildOrCopyRenderData(const TSharedRef<FUpdat
 	return true;
 }
 
-static bool bEnableHighPriorityLoading = true;
-FAutoConsoleVariableRef CVarMutableHighPriorityLoading(
-	TEXT("Mutable.EnableLoadingAssetsWithHighPriority"),
-	bEnableHighPriorityLoading,
-	TEXT("If enabled, the request to load additional assets will have high priority."));
-
 
 UE::Tasks::FTask UCustomizableInstancePrivate::LoadAdditionalAssetsAndData(const TSharedRef<FUpdateContextPrivate>& OperationData)
 {
@@ -5841,6 +5835,8 @@ UE::Tasks::FTask UCustomizableInstancePrivate::LoadAdditionalAssetsAndData(const
 #endif
 	}
 	
+	const bool bEnableHighPriorityLoading = CVarMutableHighPriorityLoading.GetValueOnAnyThread();
+
 	TArray<UE::Tasks::FTaskEvent> StreamingCompletionEvents;
 	if (AssetsToStream.Num() > 0)
 	{
@@ -6444,6 +6440,7 @@ void UCustomizableInstancePrivate::BuildMaterials(const TSharedRef<FUpdateContex
 		Public->GetPathName(),
 		UCustomizableObjectSystem::GetInstance()->GetPrivate()->MutableSystem,
 		OperationData->Model,
+		CustomizableObject->GetPrivate()->GetModelStreamableBulkData(),
 		OperationData->Parameters,
 	    OperationData->GetCapturedDescriptor().GetState());
 
@@ -6930,6 +6927,7 @@ void UCustomizableInstancePrivate::BuildMaterials(const TSharedRef<FUpdateContex
 														MutableMipDataProviderFactory->ImageRef.Component = uint8(InstanceComponentIndex);
 														MutableMipDataProviderFactory->ImageRef.Image = uint8(ImageIndex);
 														MutableMipDataProviderFactory->ImageRef.BaseMip = uint8(Image.BaseMip);
+														MutableMipDataProviderFactory->ImageRef.ConstantImagesNeededToGenerate = Image.ConstantImagesNeededToGenerate;
 														MutableMipDataProviderFactory->UpdateContext = UpdateContext;
 														MutableTexture->AddAssetUserData(MutableMipDataProviderFactory);
 													}
diff --git a/Engine/Plugins/Mutable/Source/CustomizableObject/Private/MuCO/CustomizableObjectMipDataProvider.cpp b/Engine/Plugins/Mutable/Source/CustomizableObject/Private/MuCO/CustomizableObjectMipDataProvider.cpp
--- a/Engine/Plugins/Mutable/Source/CustomizableObject/Private/MuCO/CustomizableObjectMipDataProvider.cpp
+++ b/Engine/Plugins/Mutable/Source/CustomizableObject/Private/MuCO/CustomizableObjectMipDataProvider.cpp
@@ -2,8 +2,10 @@
 
 #include "MuCO/CustomizableObjectMipDataProvider.h"
 
+#include "MuCO/CustomizableObjectPrivate.h"
 #include "MuCO/CustomizableObjectInstance.h"
 #include "MuCO/CustomizableObjectSystemPrivate.h"
+#include "MuCO/MutableStreamRequest.h"
 #include "MuR/Model.h"
 #include "TextureResource.h"
 #include "UnrealMutableImageProvider.h"
@@ -21,18 +23,26 @@ UMutableTextureMipDataProviderFactory::UMutableTextureMipDataProviderFactory(con
 
 }
 
+bool bPrefetchHighQualityMips = true;
+static FAutoConsoleVariableRef CVarEnablePrefetchHighQualityMips(
+	TEXT("mutable.EnablePrefetchHighQualityMips"),
+	bPrefetchHighQualityMips,
+	TEXT("If true, prefetch the data of high-quality mips to ensure it is available during the Mip generation task."),
+	ECVF_Default);
 
 FMutableUpdateContext::FMutableUpdateContext(
 	const FString& InCustomizableObjectPathName,
 	const FString& InInstancePathName,
 	TSharedPtr<mu::FSystem> InSystem,
 	TSharedPtr<mu::FModel, ESPMode::ThreadSafe> InModel,
+	TSharedPtr<FModelStreamableBulkData> InModelStreamable,
 	TSharedPtr<const mu::FParameters> InParameters,
 	int32 InState) :
 	CustomizableObjectPathName(InCustomizableObjectPathName),
 	InstancePathName(InInstancePathName),
 	System(InSystem),
 	Model(InModel),
+	ModelStreamableBulkData(InModelStreamable),
 	Parameters(InParameters),
 	State(InState)
 {
@@ -79,6 +89,12 @@ TSharedPtr<mu::FModel, ESPMode::ThreadSafe> FMutableUpdateContext::GetModel() co
 }
 
 
+TSharedPtr<FModelStreamableBulkData> FMutableUpdateContext::GetModelStreamableBulkData()
+{
+	return ModelStreamableBulkData;
+}
+
+
 TSharedPtr<const mu::FParameters> FMutableUpdateContext::GetParameters() const
 {
 	return Parameters;
@@ -165,19 +181,16 @@ namespace Impl
 
 		auto EndUpdateImage = [](TSharedPtr<FMutableImageOperationData>& OperationData)
 		{
-			// The request could be cancelled in parallel from CancelCounterSafely and its value be changed
+			// The request could be cancelled in parallel from CancelAsyncTasks and its value be changed
 			// between reading it and actually running Decrement() and RescheduleCallback(), so lock
 			FScopeLock Lock(&OperationData->CounterTaskLock);
-
-			if (OperationData->Counter) // If the request has been cancelled the counter will be null
+			if (!OperationData->bIsCancelled)
 			{
 				// Make the FMutableTextureMipDataProvider continue
 				OperationData->Counter->Decrement();
+				check(OperationData->Counter->GetValue() == 0);
 
-				if (OperationData->Counter->GetValue() == 0)
-				{
-					OperationData->RescheduleCallback();
-				}
+				OperationData->RescheduleCallback();
 			}
 			
 			TRACE_END_REGION(UE_MUTABLE_MIPDATA_PROVIDER_UPDATE_IMAGE_REGION);
@@ -346,6 +359,9 @@ int32 FMutableTextureMipDataProvider::GetMips(const FTextureUpdateContext& Conte
 	// Maximum value to skip, will be minimized by the first mip level requested
 	int32 MipsToSkip = 256;
 
+	uint32 AllocatedMemorySize = 0;
+	uint8* AllocatedMemory = nullptr;
+
 	for (int32 MipIndex = StartingMipIndex; MipIndex < CurrentFirstLODIdx; ++MipIndex)
 	{
 		check(Context.MipsView.IsValidIndex(MipIndex) && MipInfos.IsValidIndex(MipIndex));
@@ -374,36 +390,122 @@ int32 FMutableTextureMipDataProvider::GetMips(const FTextureUpdateContext& Conte
 			check(LODBias >= 0);
 			check(NumMips == Context.MipsView.Num() + LODBias);
 
-			OperationData->Levels.Add({ MipIndex + LODBias, Dest, (int32)MipInfo.SizeX, (int32)MipInfo.SizeY, (int32)MipInfo.DataSize, MipInfo.Format });
+			const uint64 MipDataSize = MipInfo.DataSize > 0 ? MipInfo.DataSize : GPixelFormats[MipInfo.Format].Get2DTextureMipSizeInBytes(MipInfo.SizeX, MipInfo.SizeY, 0);
+
+			// Allocated memory used to prefetch
+			if (MipDataSize > AllocatedMemorySize)
+			{
+				AllocatedMemorySize = MipDataSize;
+				AllocatedMemory = reinterpret_cast<uint8*>(MipInfo.DestData);
+			}
+
+			OperationData->Levels.Add({ MipIndex + LODBias, Dest, (int32)MipInfo.SizeX, (int32)MipInfo.SizeY, (int32)MipDataSize, MipInfo.Format });
 			MipsToSkip = FMath::Min(MipsToSkip, MipIndex + LODBias);
 		}
 	}
 
-	UCustomizableObjectSystemPrivate* CustomizableObjectSystem = UCustomizableObjectSystem::GetInstance()->GetPrivate();
-	if (CustomizableObjectSystem)
+	if (OperationData.IsValid())
 	{
-		if (OperationData.IsValid())
+		OperationData->RequestedImage = ImageRef;
+		OperationData->UpdateContext = UpdateContext;
+		OperationData->MipsToSkip = MipsToSkip;
+		OperationData->Counter = SyncOptions.Counter;
+		OperationData->RescheduleCallback = SyncOptions.RescheduleCallback;
+
+		// Increment to stop PollMips from running until the Mutable request task finishes. 
+		// If a request completes immediately, then it will call the callback
+		// but that won't do anything because the tick would not try to acquire the lock since it is already locked.
+		check(OperationData->Counter);
+		check(OperationData->Counter->GetValue() == 0);
+		OperationData->Counter->Increment(); // Prefetch task 
+		OperationData->Counter->Increment(); // MipUpdate task
+
+		MUTABLE_CPUPROFILER_SCOPE(ImagePrefetch);
+		const TSharedPtr<FModelStreamableBulkData>& ModelStreamableBulkData = OperationData->UpdateContext->GetModelStreamableBulkData();
+		
+		PrefetchRequest = MakeUnique<FMutableStreamRequest>(ModelStreamableBulkData);
+
+		if (bPrefetchHighQualityMips)
 		{
-			OperationData->RequestedImage = ImageRef;
-			OperationData->UpdateContext = UpdateContext;
-			OperationData->MipsToSkip = MipsToSkip;
-
-			check(SyncOptions.Counter);
-			// Increment to stop PollMips from running until the Mutable request task finishes. 
-			// If a request completes immediately, then it will call the callback
-			// but that won't do anything because the tick would not try to acquire the lock since it is already locked.
-			SyncOptions.Counter->Increment();
-			OperationData->Counter = SyncOptions.Counter;
-			OperationData->RescheduleCallback = SyncOptions.RescheduleCallback;
-
-			MutableTaskId = CustomizableObjectSystem->MutableTaskGraph.AddMutableThreadTaskLowPriority(
-				TEXT("Mutable_MipUpdate"),
-				[OperationData = this->OperationData]()
+			TArray<int32> RomsToPrefetch;
+
+			uint32 MaxBlockSize = 0;
+
+			TSharedPtr<mu::FModel, ESPMode::ThreadSafe> Model = UpdateContext->GetModel();
+			check(Model);
+
+			for (int32 ImageToLoadIndex : ImageRef.ConstantImagesNeededToGenerate)
+			{
+				for (const FMutableMipUpdateLevel& Level : OperationData->Levels)
 				{
-					Impl::Task_Mutable_UpdateImage(OperationData);
-				});
+					int32 RomId = Model->GetConstantImageRomId(ImageToLoadIndex, Level.MipLevel);
+
+					if (RomId >= 0 && Model->IsRomHighRes(RomId))
+					{
+						FMutableStreamableBlock& Block = ModelStreamableBulkData->ModelStreamables[RomId];
+						if (!Block.IsPrefetched)
+						{
+							Block.IsPrefetched = true;
+
+							MaxBlockSize = FMath::Max(MaxBlockSize, Model->GetRomSize(RomId));
+							RomsToPrefetch.Add(RomId);
+						}
+					}
+				}
+			}
+
+			if (AllocatedMemorySize < MaxBlockSize)
+			{
+				OperationData->AllocatedMemory.SetNum(MaxBlockSize);
+				AllocatedMemory = OperationData->AllocatedMemory.GetData();
+			}
+
+			for (int32 RomId : RomsToPrefetch)
+			{
+				const FMutableStreamableBlock& Block = ModelStreamableBulkData->ModelStreamables[RomId];
+				PrefetchRequest->AddBlock(
+					ModelStreamableBulkData->ModelStreamables[RomId],
+					MutablePrivate::EStreamableDataType::Model,
+					TArrayView<uint8>(AllocatedMemory, Model->GetRomSize(RomId)));
+			}
 		}
 
+		UE::Tasks::FTask StreamTask = PrefetchRequest->Stream();
+
+		UE::Tasks::TTask PrefetchTask = UE::Tasks::Launch(TEXT("MutableImagePrefetchTask"),
+			[OperationData = this->OperationData]()
+			{
+				MUTABLE_CPUPROFILER_SCOPE(ImagePrefetchTask);
+
+				OperationData->AllocatedMemory.Empty();
+
+				FScopeLock Lock(&OperationData->CounterTaskLock);
+
+				// Decrement counter after completing the prefetch task.
+				OperationData->Counter->Decrement();
+
+				if (OperationData->bIsCancelled)
+				{
+					check(OperationData->Counter == 0);
+					OperationData->RescheduleCallback();
+					return;
+				}
+
+				UCustomizableObjectSystemPrivate* CustomizableObjectSystem = UCustomizableObjectSystem::GetInstance()->GetPrivate();
+				if (CustomizableObjectSystem)
+				{
+					OperationData->MutableTaskId = CustomizableObjectSystem->MutableTaskGraph.AddMutableThreadTaskLowPriority(
+						TEXT("Mutable_MipUpdate"),
+						[OperationData = OperationData]()
+						{
+							Impl::Task_Mutable_UpdateImage(OperationData);
+						});
+				}
+			},
+			StreamTask,
+			UE::Tasks::ETaskPriority::Inherit,
+			UE::Tasks::EExtendedTaskPriority::Inline);
+
 		AdvanceTo(ETickState::PollMips, ETickThread::Async);
 	}
 	else
@@ -517,12 +619,7 @@ void FMutableTextureMipDataProvider::Cancel(const FTextureUpdateSyncOptions& Syn
 {
 	bRequestAborted = true;
 
-	CancelCounterSafely();
-
-	if (UCustomizableObjectSystem* System = UCustomizableObjectSystem::GetInstance())
-	{
-		System->GetPrivate()->MutableTaskGraph.CancelMutableThreadTaskLowPriority(MutableTaskId);	
-	}
+	CancelAsyncTasks();
 }
 
 
@@ -537,26 +634,34 @@ void FMutableTextureMipDataProvider::AbortPollMips()
 {
 	bRequestAborted = true;
 
-	CancelCounterSafely();
-
-	if (UCustomizableObjectSystem* System = UCustomizableObjectSystem::GetInstance())
-	{
-		System->GetPrivate()->MutableTaskGraph.CancelMutableThreadTaskLowPriority(MutableTaskId);	
-	}
+	CancelAsyncTasks();
 }
 
 
-void FMutableTextureMipDataProvider::CancelCounterSafely()
+void FMutableTextureMipDataProvider::CancelAsyncTasks()
 {
+	if (PrefetchRequest)
+	{
+		PrefetchRequest->Cancel();
+	}
+
 	if (OperationData.IsValid())
 	{
 		// The Counter could be read in parallel from Task_Mutable_UpdateImage, so lock
 		FScopeLock Lock(&OperationData->CounterTaskLock);
+		if (!OperationData->bIsCancelled)
+		{
+			OperationData->bIsCancelled = true;
 
-		if (OperationData->Counter)
+			// Decrement counter. Do not set the value to zero since we must wait for the PrefetchTask to complete
+			if (OperationData->Counter->GetValue() > 0)
+			{
+				OperationData->Counter->Decrement();
+			}
+		}
+		if (UCustomizableObjectSystem* System = UCustomizableObjectSystem::GetInstance())
 		{
-			OperationData->Counter->Set(0);
-			OperationData->Counter = nullptr;
+			System->GetPrivate()->MutableTaskGraph.CancelMutableThreadTaskLowPriority(OperationData->MutableTaskId);
 		}
 	}
 }
diff --git a/Engine/Plugins/Mutable/Source/CustomizableObject/Private/MuCO/CustomizableObjectSystem.cpp b/Engine/Plugins/Mutable/Source/CustomizableObject/Private/MuCO/CustomizableObjectSystem.cpp
--- a/Engine/Plugins/Mutable/Source/CustomizableObject/Private/MuCO/CustomizableObjectSystem.cpp
+++ b/Engine/Plugins/Mutable/Source/CustomizableObject/Private/MuCO/CustomizableObjectSystem.cpp
@@ -155,6 +155,11 @@ TAutoConsoleVariable<bool> CVarFixLowPriorityTasksOverlap(
 	true,
 	TEXT("If true, use code that fixes the Low Priority Tasks overlap."));
 
+TAutoConsoleVariable<bool> CVarMutableHighPriorityLoading(
+	TEXT("Mutable.EnableLoadingAssetsWithHighPriority"),
+	true,
+	TEXT("If enabled, the request to load additional assets will have high priority."));
+
 int32 UCustomizableObjectSystemPrivate::SkeletalMeshMinLodQualityLevel = -1;
 
 
@@ -2987,9 +2992,10 @@ namespace impl
 	{
 		MUTABLE_CPUPROFILER_SCOPE(Task_Mutable_GetImages_GetImage)
 
-		const mu::FExtendedImageDesc& ImageDesc = GetImageDescTask.GetResult();
+		mu::FExtendedImageDesc& ImageDesc = GetImageDescTask.GetResult();
 
 		FInstanceUpdateData::FImage& Image = OperationData->InstanceUpdateData.Images[ImageIndex];
+		Image.ConstantImagesNeededToGenerate = MoveTemp(ImageDesc.ConstantImagesNeededToGenerate);
 
 		const UCustomizableObjectSystemPrivate* CustomizableObjectSystemPrivateData = UCustomizableObjectSystem::GetInstanceChecked()->GetPrivate();
 		
diff --git a/Engine/Plugins/Mutable/Source/CustomizableObject/Private/MuCO/MutableStreamRequest.cpp b/Engine/Plugins/Mutable/Source/CustomizableObject/Private/MuCO/MutableStreamRequest.cpp
new file mode 100644
--- /dev/null
+++ b/Engine/Plugins/Mutable/Source/CustomizableObject/Private/MuCO/MutableStreamRequest.cpp
@@ -0,0 +1,203 @@
+// Copyright Epic Games, Inc. All Rights Reserved.
+
+#include "MutableStreamRequest.h"
+
+#include "MuCO/CustomizableObjectPrivate.h"
+#include "MuCO/CustomizableObjectSystemPrivate.h"
+#include "HAL/PlatformFileManager.h"
+
+
+FMutableStreamRequest::FMutableStreamRequest(const TSharedPtr<FModelStreamableBulkData>& InModelStreamableBulkData) :
+	ModelStreamableBulkData(InModelStreamableBulkData)
+{
+}
+
+
+const TSharedPtr<FModelStreamableBulkData>& FMutableStreamRequest::GetModelStreamableBulkData() const
+{
+	return ModelStreamableBulkData;
+}
+
+
+void FMutableStreamRequest::AddBlock(const FMutableStreamableBlock& Block, MutablePrivate::EStreamableDataType DataType, TArrayView<uint8> AllocatedMemoryView)
+{
+	if (bIsStreaming)
+	{
+		check(false);
+		return;
+	}
+
+	IAsyncReadFileHandle* FileHandle = nullptr;
+
+#if WITH_EDITOR
+	int32 FileHandleIndex = OpenFilesIds.Find(Block.FileId);
+	if (FileHandleIndex == INDEX_NONE)
+	{
+		const FString& FullFileName = ModelStreamableBulkData->FullFilePath + UCustomizableObjectPrivate::GetDataTypeExtension(DataType);
+		IAsyncReadFileHandle* ReadFileHandle = FPlatformFileManager::Get().GetPlatformFile().OpenAsyncRead(*FullFileName);
+		OpenFileHandles.Emplace(ReadFileHandle);
+
+		FileHandleIndex = OpenFilesIds.Add(Block.FileId);
+		check(OpenFileHandles.Num() == OpenFilesIds.Num());
+	}
+
+	if (OpenFileHandles.IsValidIndex(FileHandleIndex))
+	{
+		FileHandle = OpenFileHandles[FileHandleIndex].Get();
+	}
+#endif
+		
+	BlockReadInfos.Emplace(Block.Offset, FileHandle, AllocatedMemoryView, Block.FileId);
+}
+
+
+UE::Tasks::FTask FMutableStreamRequest::Stream()
+{
+	if (bIsStreaming)
+	{
+		check(false);
+		return UE::Tasks::MakeCompletedTask<void>();
+	}
+
+	bIsStreaming = true;
+
+	for (const FBlockReadInfo& Block : BlockReadInfos)
+	{
+		HeapMemory->CompletionEvents.Emplace(TEXT("AsyncReadDataReadyEvent"));
+	}
+
+	UE::Tasks::Launch(TEXT("CustomizableObjectReadRequestTask"),
+		[
+			ModelStreamableBulkData = ModelStreamableBulkData,
+			BlockReadInfos = MoveTemp(BlockReadInfos),
+			HeapMemory = this->HeapMemory
+		]()
+		{
+			MUTABLE_CPUPROFILER_SCOPE(CustomizableInstanceLoadBlocksAsyncRead_Request);
+			FScopeLock Lock(&HeapMemory->ReadRequestLock);
+
+			// Task cancelled, early out.
+			if (HeapMemory->bIsCancelled)
+			{
+				// Trigger preallocated events to complete GatherStreamingRequestsCompletionTask prerequisites
+				for (UE::Tasks::FTaskEvent& CompletionEvent : HeapMemory->CompletionEvents)
+				{
+					CompletionEvent.Trigger();
+				}
+
+				return;
+			}
+
+			const bool bUseFBulkData = !ModelStreamableBulkData->StreamableBulkData.IsEmpty();
+			const EAsyncIOPriorityAndFlags Priority = CVarMutableHighPriorityLoading.GetValueOnAnyThread() ? AIOP_High : AIOP_Normal;
+
+			const int32 NumBlocks = BlockReadInfos.Num();
+			for (int32 Index = 0; Index < NumBlocks; ++Index)
+			{
+				const FBlockReadInfo& Block = BlockReadInfos[Index];
+				UE::Tasks::FTaskEvent& CompletionEvent = HeapMemory->CompletionEvents[Index];
+
+				if (bUseFBulkData)
+				{
+					FBulkDataIORequestCallBack IOCallback =
+						[CompletionEvent, FileId = Block.FileId](bool bWasCancelled, IBulkDataIORequest*) mutable
+						{
+							CompletionEvent.Trigger();
+						};
+
+					check(ModelStreamableBulkData->StreamableBulkData.IsValidIndex(Block.FileId));
+					FByteBulkData& ByteBulkData = ModelStreamableBulkData->StreamableBulkData[Block.FileId];
+
+					HeapMemory->BulkReadRequests.Add(TSharedPtr<IBulkDataIORequest>(ByteBulkData.CreateStreamingRequest(
+						Block.Offset,
+						(int64)Block.AllocatedMemoryView.Num(),
+						Priority,
+						&IOCallback,
+						Block.AllocatedMemoryView.GetData())));
+				}
+				else if (Block.FileHandle)
+				{
+					FAsyncFileCallBack ReadRequestCallBack =
+						[CompletionEvent, FileId = Block.FileId](bool bWasCancelled, IAsyncReadRequest*) mutable
+						{
+							CompletionEvent.Trigger();
+						};
+
+					HeapMemory->ReadRequests.Add(TSharedPtr<IAsyncReadRequest>(Block.FileHandle->ReadRequest(
+						Block.Offset,
+						(int64)Block.AllocatedMemoryView.Num(),
+						Priority,
+						&ReadRequestCallBack,
+						Block.AllocatedMemoryView.GetData())));
+				}
+				else
+				{
+					ensure(false);
+					CompletionEvent.Trigger();
+				}
+			}
+		},
+		UE::Tasks::ETaskPriority::High);
+	
+		
+	return UE::Tasks::Launch(TEXT("GatherStreamingRequestsCompletionTask"),
+        [
+            OpenFileHandles = OpenFileHandles,
+			HeapMemory = this->HeapMemory
+        ]() mutable 
+        {
+			{
+				FScopeLock Lock(&HeapMemory->ReadRequestLock);
+
+				for (TSharedPtr<IAsyncReadRequest>& ReadRequest : HeapMemory->ReadRequests)
+				{
+					if (ReadRequest)
+					{
+						ReadRequest->WaitCompletion();
+					}
+				}
+
+				for (TSharedPtr<IBulkDataIORequest>& BulkReadRequest : HeapMemory->BulkReadRequests)
+				{
+					if (BulkReadRequest)
+					{
+						BulkReadRequest->WaitCompletion();
+					}
+				}
+
+				HeapMemory->BulkReadRequests.Empty();
+				HeapMemory->ReadRequests.Empty();
+			}
+			
+			OpenFileHandles.Empty();
+        },
+		HeapMemory->CompletionEvents,
+        UE::Tasks::ETaskPriority::High);
+}
+
+
+void FMutableStreamRequest::Cancel()
+{
+	FScopeLock Lock(&HeapMemory->ReadRequestLock);
+
+	if (!HeapMemory->bIsCancelled)
+	{
+		HeapMemory->bIsCancelled = true;
+
+		for (TSharedPtr<IAsyncReadRequest>& ReadRequest : HeapMemory->ReadRequests)
+		{
+			if (ReadRequest)
+			{
+				ReadRequest->Cancel();
+			}
+		}
+
+		for (TSharedPtr<IBulkDataIORequest>& BulkReadRequest : HeapMemory->BulkReadRequests)
+		{
+			if (BulkReadRequest)
+			{
+				BulkReadRequest->Cancel();
+			}
+		}
+	}
+}
diff --git a/Engine/Plugins/Mutable/Source/CustomizableObject/Private/MuCO/MutableStreamRequest.h b/Engine/Plugins/Mutable/Source/CustomizableObject/Private/MuCO/MutableStreamRequest.h
new file mode 100644
--- /dev/null
+++ b/Engine/Plugins/Mutable/Source/CustomizableObject/Private/MuCO/MutableStreamRequest.h
@@ -0,0 +1,67 @@
+// Copyright Epic Games, Inc. All Rights Reserved.
+
+#pragma once
+#include "MuCO/CustomizableObjectPrivate.h"
+
+struct FBlockReadInfo;
+struct FMutableStreamableBlock;
+struct FModelStreamableBulkData;
+
+
+constexpr int32 BulkDataFileOffset =
+#if WITH_EDITOR
+	sizeof(MutableCompiledDataStreamHeader);
+#else
+		0;
+#endif
+
+
+/** Stream data from .mut or Bulk files.
+ * 
+ * Can be destroyed at any time, even when the stream request is in progress. */
+class FMutableStreamRequest
+{
+public:
+	FMutableStreamRequest(const TSharedPtr<FModelStreamableBulkData>& InModelStreamableBulkData);
+	
+	const TSharedPtr<FModelStreamableBulkData>& GetModelStreamableBulkData() const;
+
+	/** Add a block to stream. */
+	void AddBlock(const FMutableStreamableBlock& Block, MutablePrivate::EStreamableDataType DataType, TArrayView<uint8> AllocatedMemoryView);
+
+	/** Stream the previously requested data. */
+	UE::Tasks::FTask Stream();
+
+	/** Cancel pending requests */
+	void Cancel();
+
+private:
+	struct FBlockReadInfo
+	{
+		uint64 Offset;
+		IAsyncReadFileHandle* FileHandle;
+		TArrayView<uint8> AllocatedMemoryView;
+		uint32 FileId;
+	};
+	
+	const TSharedPtr<FModelStreamableBulkData> ModelStreamableBulkData;
+
+	TArray<uint32> OpenFilesIds;
+	TArray<TSharedPtr<IAsyncReadFileHandle>> OpenFileHandles;
+
+	struct FHeapMemory
+	{
+		TArray<UE::Tasks::FTaskEvent> CompletionEvents;
+		TArray<TSharedPtr<IAsyncReadRequest>> ReadRequests;
+		TArray<TSharedPtr<IBulkDataIORequest>> BulkReadRequests;
+
+		bool bIsCancelled = false;
+		FCriticalSection ReadRequestLock;
+	};
+
+	TSharedRef<FHeapMemory> HeapMemory = MakeShared<FHeapMemory>();
+
+	TArray<FBlockReadInfo> BlockReadInfos;
+
+	bool bIsStreaming = false;
+};
diff --git a/Engine/Plugins/Mutable/Source/CustomizableObject/Public/MuCO/CustomizableObjectSystem.h b/Engine/Plugins/Mutable/Source/CustomizableObject/Public/MuCO/CustomizableObjectSystem.h
--- a/Engine/Plugins/Mutable/Source/CustomizableObject/Public/MuCO/CustomizableObjectSystem.h
+++ b/Engine/Plugins/Mutable/Source/CustomizableObject/Public/MuCO/CustomizableObjectSystem.h
@@ -137,6 +137,8 @@ struct FMutableImageReference
 	uint8 Image = 0;
 
 	uint8 BaseMip = 0;
+
+	TArray<int32> ConstantImagesNeededToGenerate;
 };
 
 
diff --git a/Engine/Plugins/Mutable/Source/CustomizableObjectEditor/Private/MuCOE/CustomizableObjectCompileRunnable.cpp b/Engine/Plugins/Mutable/Source/CustomizableObjectEditor/Private/MuCOE/CustomizableObjectCompileRunnable.cpp
--- a/Engine/Plugins/Mutable/Source/CustomizableObjectEditor/Private/MuCOE/CustomizableObjectCompileRunnable.cpp
+++ b/Engine/Plugins/Mutable/Source/CustomizableObjectEditor/Private/MuCOE/CustomizableObjectCompileRunnable.cpp
@@ -367,6 +367,7 @@ FCustomizableObjectSaveDDRunnable::FCustomizableObjectSaveDDRunnable(const TShar
 	{
 		// We will be saving all compilation data in two separate files, write CO Data
 		FullFileName = CustomizableObject->GetPrivate()->GetCompiledDataFileName(Options.TargetPlatform);
+		PlatformData->ModelStreamableBulkData->FullFilePath = FullFileName;
 	}
 }
 
diff --git a/Engine/Plugins/Mutable/Source/CustomizableObjectEditor/Private/MuCOE/CustomizableObjectCompiler.cpp b/Engine/Plugins/Mutable/Source/CustomizableObjectEditor/Private/MuCOE/CustomizableObjectCompiler.cpp
--- a/Engine/Plugins/Mutable/Source/CustomizableObjectEditor/Private/MuCOE/CustomizableObjectCompiler.cpp
+++ b/Engine/Plugins/Mutable/Source/CustomizableObjectEditor/Private/MuCOE/CustomizableObjectCompiler.cpp
@@ -1339,7 +1339,7 @@ void FCustomizableObjectCompiler::FinishCompilationTask()
 			{
 				Flags = Model->IsRomHighRes(FileIndex) ? EMutableFileFlags::HighRes : EMutableFileFlags::None;
 			}
-			ModelStreamables.Add(FileIndex, FMutableStreamableBlock{ 0, uint32(Flags), Offset });
+			ModelStreamables.Add(FileIndex, FMutableStreamableBlock{ .FileId = 0, .Flags = uint16(Flags), .Offset = Offset });
 			Offset += ResourceSize;
 		}
 	}
@@ -1387,7 +1387,7 @@ void FCustomizableObjectCompiler::FinishCompilationTask()
 			ResourceMeshData.NameResolutionMap = MeshData.Value.NameResolutionMap;
 			ResourceMeshData.Size = DataSizeInBytes;
 			EMutableFileFlags Flags = EMutableFileFlags::None;
-			ResourceMeshData.Block = FMutableStreamableBlock{ uint32(0), uint32(Flags), RealTimeMorphDataOffsetInBytes };
+			ResourceMeshData.Block = FMutableStreamableBlock{ .FileId = uint32(0), .Flags = uint16(Flags), .Offset = RealTimeMorphDataOffsetInBytes };
 			ResourceMeshData.SourceId = MeshData.Value.SourceId;
 
 			RealTimeMorphDataOffsetInBytes += DataSizeInBytes;
@@ -1417,7 +1417,7 @@ void FCustomizableObjectCompiler::FinishCompilationTask()
 			ResourceMeshData.PhysicsAssetIndex = MeshData.Value.PhysicsAssetIndex;
 			ResourceMeshData.Size = DataSizeInBytes;
 			EMutableFileFlags Flags = EMutableFileFlags::None;
-			ResourceMeshData.Block = FMutableStreamableBlock{ uint32(0), uint32(Flags), ClothingDataOffsetInBytes };
+			ResourceMeshData.Block = FMutableStreamableBlock{ .FileId = uint32(0), .Flags = uint16(Flags), .Offset = ClothingDataOffsetInBytes };
 			ResourceMeshData.SourceId = MeshData.Value.SourceId;
 
 			ClothingDataOffsetInBytes += DataSizeInBytes;
diff --git a/Engine/Plugins/Mutable/Source/MutableRuntime/Private/MuR/CodeRunner.cpp b/Engine/Plugins/Mutable/Source/MutableRuntime/Private/MuR/CodeRunner.cpp
--- a/Engine/Plugins/Mutable/Source/MutableRuntime/Private/MuR/CodeRunner.cpp
+++ b/Engine/Plugins/Mutable/Source/MutableRuntime/Private/MuR/CodeRunner.cpp
@@ -137,6 +137,7 @@ namespace mu
 		if (Type == FScheduledOp::EType::ImageDesc)
 		{
 			ImageDescResults.Reserve(64);
+			ImageDescConstantImages.Reserve(32);
 		}
 	
 		// We will read this in the end, so make sure we keep it.
@@ -5998,30 +5999,33 @@ namespace mu
 			Result.m_lods = Program.ConstantImages[ImageIndex].LODCount;
 
 			int32 LODIndexIndex = Program.ConstantImages[ImageIndex].FirstIndex;
-			int32 LODIndex = 0;
-			for (; LODIndex < Result.m_lods; ++LODIndex)
 			{
-				const int32 CurrentIndexIndex = LODIndexIndex + LODIndex;
-				const FConstantResourceIndex CurrentIndex = Program.ConstantImageLODIndices[CurrentIndexIndex];
-
-				bool bIsLODAvailable = false;
-				if (!CurrentIndex.Streamable)
+				int32 LODIndex = 0;
+				for (; LODIndex < Result.m_lods; ++LODIndex)
 				{
-					bIsLODAvailable = true;
-				}
-				else
-				{
-					uint32 RomId = CurrentIndex.Index;
-					bIsLODAvailable = System->StreamInterface->DoesBlockExist(Model.Get(), RomId);
-				}
+					const int32 CurrentIndexIndex = LODIndexIndex + LODIndex;
+					const FConstantResourceIndex CurrentIndex = Program.ConstantImageLODIndices[CurrentIndexIndex];
 
-				if (bIsLODAvailable)
-				{
-					break;
+					bool bIsLODAvailable = false;
+					if (!CurrentIndex.Streamable)
+					{
+						bIsLODAvailable = true;
+					}
+					else
+					{
+						uint32 RomId = CurrentIndex.Index;
+						bIsLODAvailable = System->StreamInterface->DoesBlockExist(Model.Get(), RomId);
+					}
+
+					if (bIsLODAvailable)
+					{
+						break;
+					}
 				}
-			}
 
-			Result.FirstLODAvailable = LODIndex;
+				Result.FirstLODAvailable = LODIndex;
+			}
+			ImageDescConstantImages.Add(ImageIndex);
 
 			StoreValidDesc(Item);
 			break;
diff --git a/Engine/Plugins/Mutable/Source/MutableRuntime/Private/MuR/CodeRunner.h b/Engine/Plugins/Mutable/Source/MutableRuntime/Private/MuR/CodeRunner.h
--- a/Engine/Plugins/Mutable/Source/MutableRuntime/Private/MuR/CodeRunner.h
+++ b/Engine/Plugins/Mutable/Source/MutableRuntime/Private/MuR/CodeRunner.h
@@ -163,6 +163,7 @@ namespace  mu
 		
 		/** Image descriptor intermediate results. */
 		TMap<OP::ADDRESS, FExtendedImageDesc> ImageDescResults;
+		TArray<int32> ImageDescConstantImages; 
 
 		/** Only used for correct mip skipping with external images. It is the LOD for which the image is build. */
 		int32 ImageLOD;
diff --git a/Engine/Plugins/Mutable/Source/MutableRuntime/Private/MuR/CodeRunnerTasked.cpp b/Engine/Plugins/Mutable/Source/MutableRuntime/Private/MuR/CodeRunnerTasked.cpp
--- a/Engine/Plugins/Mutable/Source/MutableRuntime/Private/MuR/CodeRunnerTasked.cpp
+++ b/Engine/Plugins/Mutable/Source/MutableRuntime/Private/MuR/CodeRunnerTasked.cpp
@@ -154,7 +154,8 @@ namespace mu
 
 		HeapData.SetNum(0, EAllowShrinking::No);
 		ImageDescResults.Reset();
-		
+		ImageDescConstantImages.Reset();
+
 		RunnerCompletionEvent = UE::Tasks::FTaskEvent(TEXT("CodeRunnerCompletionEvent"));
 		
 		bool bProfile = false;
@@ -489,7 +490,10 @@ namespace mu
 	/** */
 	const FExtendedImageDesc& CodeRunner::GetImageDescResult(OP::ADDRESS ResultAddress)
 	{
-		return ImageDescResults.FindChecked(ResultAddress);
+		FExtendedImageDesc& Result = ImageDescResults.FindChecked(ResultAddress);
+		Result.ConstantImagesNeededToGenerate = ImageDescConstantImages;
+
+		return Result;
 	}
 
 	//---------------------------------------------------------------------------------------------
diff --git a/Engine/Plugins/Mutable/Source/MutableRuntime/Private/MuR/Model.cpp b/Engine/Plugins/Mutable/Source/MutableRuntime/Private/MuR/Model.cpp
--- a/Engine/Plugins/Mutable/Source/MutableRuntime/Private/MuR/Model.cpp
+++ b/Engine/Plugins/Mutable/Source/MutableRuntime/Private/MuR/Model.cpp
@@ -408,6 +408,27 @@ namespace mu
 		return m_pD->Program.Roms[Index].IsHighRes==1;
 	}
 
+    int32 FModel::GetConstantImageRomId(int32 ConstantImageIndex, int32 LODIndex) const
+    {
+        const FProgram& Program = m_pD->Program;
+        check(Program.ConstantImages.IsValidIndex(ConstantImageIndex));
+        
+        FImageLODRange LODRange = Program.ConstantImages[ConstantImageIndex];
+
+        if (LODIndex >= LODRange.LODCount) 
+        {
+            return -1; 
+        }
+
+        FConstantResourceIndex ResourceIndex = Program.ConstantImageLODIndices[LODRange.FirstIndex + LODIndex];
+        if (!ResourceIndex.Streamable)
+        {
+            return -1;
+        }
+
+        return ResourceIndex.Index;
+    }
+
 
     TSharedPtr<FParameters> FModel::NewParameters(TSharedPtr<const FModel> Model, const FParameters* OldParameters )
     {
diff --git a/Engine/Plugins/Mutable/Source/MutableRuntime/Private/MuR/System.cpp b/Engine/Plugins/Mutable/Source/MutableRuntime/Private/MuR/System.cpp
--- a/Engine/Plugins/Mutable/Source/MutableRuntime/Private/MuR/System.cpp
+++ b/Engine/Plugins/Mutable/Source/MutableRuntime/Private/MuR/System.cpp
@@ -416,12 +416,12 @@ namespace mu
 		TEXT("If different than 0, clear the image desc cache at every image."),
 		ECVF_Scalability);
 
-	void FSystem::GetImageDescInline(FInstance::FID instanceID, FResourceID ImageId, FExtendedImageDesc& OutDesc)
+	FExtendedImageDesc FSystem::GetImageDescInline(FInstance::FID instanceID, FResourceID ImageId)
 	{
 		LLM_SCOPE_BYNAME(TEXT("MutableRuntime"));
 		MUTABLE_CPUPROFILER_SCOPE(SystemGetImageDesc);
 
-		OutDesc = FExtendedImageDesc();
+		FExtendedImageDesc Result;
 
 		// Find the live instance
 		FLiveInstance* LiveInstance = m_pD->FindLiveInstance(instanceID);
@@ -454,12 +454,14 @@ namespace mu
 			UE::Tasks::FTask CompletionEvent = Runner->StartRun(bForceInlineExecution);
 			check(CompletionEvent.IsCompleted());
 
-			OutDesc = Runner->GetImageDescResult(RootAddress);
+			Result = Runner->GetImageDescResult(RootAddress);
 
 			m_pD->WorkingMemoryManager.EndRunnerThread();
 		}
 
 		m_pD->WorkingMemoryManager.CurrentInstanceCache = nullptr;
+
+		return Result;
 	}
 
 
diff --git a/Engine/Plugins/Mutable/Source/MutableRuntime/Public/MuR/ImageTypes.h b/Engine/Plugins/Mutable/Source/MutableRuntime/Public/MuR/ImageTypes.h
--- a/Engine/Plugins/Mutable/Source/MutableRuntime/Public/MuR/ImageTypes.h
+++ b/Engine/Plugins/Mutable/Source/MutableRuntime/Public/MuR/ImageTypes.h
@@ -100,7 +100,7 @@ namespace mu
 		}
 	};
 
-	/** An image descriptor with information about missing data LODs */
+	/** An image descriptor with information about missing data LODs and which images are needed to load to generate the image */
 	struct FExtendedImageDesc : public FImageDesc
 	{
 		/** 
@@ -108,6 +108,7 @@ namespace mu
 		 * It is assumed that non available LODs are consecutive starting from the mip with higher resolution.
 		 **/
 		uint8 FirstLODAvailable = 0;
+		TArray<int32> ConstantImagesNeededToGenerate;
 	};
 
 	/** List of supported modes in generic image layering operations. */
diff --git a/Engine/Plugins/Mutable/Source/MutableRuntime/Public/MuR/Model.h b/Engine/Plugins/Mutable/Source/MutableRuntime/Public/MuR/Model.h
--- a/Engine/Plugins/Mutable/Source/MutableRuntime/Public/MuR/Model.h
+++ b/Engine/Plugins/Mutable/Source/MutableRuntime/Public/MuR/Model.h
@@ -145,7 +145,13 @@ namespace mu
 		bool IsMeshData(int32 Index) const;
 
 		bool IsRomHighRes(int32 Index) const;
-    	
+   
+		/** 
+		 * Returns the ConstantImageIndex LODIndex rom id. In case ConstantImageIndex does not have LODIndex 
+		 * returns -1.
+		 */
+    	int32 GetConstantImageRomId(int32 ConstantImageIndex, int32 LODIndex) const;
+
 		//-----------------------------------------------------------------------------------------
 		// Interface pattern
 		//-----------------------------------------------------------------------------------------
diff --git a/Engine/Plugins/Mutable/Source/MutableRuntime/Public/MuR/System.h b/Engine/Plugins/Mutable/Source/MutableRuntime/Public/MuR/System.h
--- a/Engine/Plugins/Mutable/Source/MutableRuntime/Public/MuR/System.h
+++ b/Engine/Plugins/Mutable/Source/MutableRuntime/Public/MuR/System.h
@@ -165,7 +165,7 @@ namespace mu
 
 		//! Only valid between BeginUpdate and EndUpdate
 		//! Calculate the description of an image, without generating it.
-		void GetImageDescInline(FInstance::FID InstanceID, FResourceID ImageId, FExtendedImageDesc& OutDesc);
+		FExtendedImageDesc GetImageDescInline(FInstance::FID InstanceID, FResourceID ImageId);
 
 		//! Only valid between BeginUpdate and EndUpdate
 		//! \param MipsToSkip Number of mips to skip compared from the full image.
