[Mutable] Add missing FBlockReadInfo constructor to fix Android builds.

#rnx

diff --git a/Engine/Plugins/Mutable/Source/CustomizableObject/Private/MuCO/MutableStreamRequest.cpp b/Engine/Plugins/Mutable/Source/CustomizableObject/Private/MuCO/MutableStreamRequest.cpp
--- a/Engine/Plugins/Mutable/Source/CustomizableObject/Private/MuCO/MutableStreamRequest.cpp
+++ b/Engine/Plugins/Mutable/Source/CustomizableObject/Private/MuCO/MutableStreamRequest.cpp
@@ -201,3 +201,12 @@ void FMutableStreamRequest::Cancel()
 		}
 	}
 }
+
+
+FMutableStreamRequest::FBlockReadInfo::FBlockReadInfo(uint64 InOffset, IAsyncReadFileHandle* InFileHandle, TArrayView<uint8>& InAllocatedMemoryView, uint32 InFileId)
+	: Offset(InOffset)
+	, FileHandle(InFileHandle)
+	, AllocatedMemoryView(InAllocatedMemoryView)
+	, FileId(InFileId)
+{
+}
diff --git a/Engine/Plugins/Mutable/Source/CustomizableObject/Private/MuCO/MutableStreamRequest.h b/Engine/Plugins/Mutable/Source/CustomizableObject/Private/MuCO/MutableStreamRequest.h
--- a/Engine/Plugins/Mutable/Source/CustomizableObject/Private/MuCO/MutableStreamRequest.h
+++ b/Engine/Plugins/Mutable/Source/CustomizableObject/Private/MuCO/MutableStreamRequest.h
@@ -38,6 +38,9 @@ public:
 private:
 	struct FBlockReadInfo
 	{
+		FBlockReadInfo() = default;
+		FBlockReadInfo(uint64 InOffset, IAsyncReadFileHandle* InFileHandle, TArrayView<uint8>& InAllocatedMemoryView, uint32 InFileId);
+
 		uint64 Offset;
 		IAsyncReadFileHandle* FileHandle;
 		TArrayView<uint8> AllocatedMemoryView;
