[mutable] Fixed assignation of old value to projector layer for layers greater than the first one.
- No RangeIndex was being provided when deleting the opacity value for a given range, so, the value in the first range index was being destroyed, instead of the value in the range index provided.
- Fixed the addition of an extra range value that was never being used.

#rnx

diff --git a/Engine/Plugins/Mutable/Source/CustomizableObject/Private/MuCO/CustomizableObjectInstanceDescriptor.cpp b/Engine/Plugins/Mutable/Source/CustomizableObject/Private/MuCO/CustomizableObjectInstanceDescriptor.cpp
--- a/Engine/Plugins/Mutable/Source/CustomizableObject/Private/MuCO/CustomizableObjectInstanceDescriptor.cpp
+++ b/Engine/Plugins/Mutable/Source/CustomizableObject/Private/MuCO/CustomizableObjectInstanceDescriptor.cpp
@@ -1351,10 +1351,11 @@ void FCustomizableObjectInstanceDescriptor::SetIntParameterSelectedOption(const
 	}
 	else
 	{
+		// Contingency in case the RangeIndex can not be used as an index of the "ParameterRangeValueNames" array
 		if (!IntParameters[ParameterIndexInInstance].ParameterRangeValueNames.IsValidIndex(RangeIndex))
 		{
 			const int32 InsertionIndex = IntParameters[ParameterIndexInInstance].ParameterRangeValueNames.Num();
-			const int32 NumInsertedElements = RangeIndex + 1 - IntParameters[ParameterIndexInInstance].ParameterRangeValueNames.Num();
+			const int32 NumInsertedElements = RangeIndex + 1 - InsertionIndex;
 			IntParameters[ParameterIndexInInstance].ParameterRangeValueNames.InsertDefaulted(InsertionIndex, NumInsertedElements);
 		}
 
@@ -2676,7 +2677,7 @@ void FCustomizableObjectInstanceDescriptor::RemoveLayerAt(const FName& Projector
 	RemoveValueFromIntRange(ParamName + IMAGE_PARAMETER_POSTFIX, RangeIndex);
 
 	// Opacity Range.
-	RemoveValueFromFloatRange(ParamName + OPACITY_PARAMETER_POSTFIX);
+	RemoveValueFromFloatRange(ParamName + OPACITY_PARAMETER_POSTFIX, RangeIndex);
 }
 
 
diff --git a/Engine/Plugins/Mutable/Source/CustomizableObjectEditor/Private/MuCOE/CustomizableInstanceDetails.cpp b/Engine/Plugins/Mutable/Source/CustomizableObjectEditor/Private/MuCOE/CustomizableInstanceDetails.cpp
--- a/Engine/Plugins/Mutable/Source/CustomizableObjectEditor/Private/MuCOE/CustomizableInstanceDetails.cpp
+++ b/Engine/Plugins/Mutable/Source/CustomizableObjectEditor/Private/MuCOE/CustomizableInstanceDetails.cpp
@@ -2267,7 +2267,7 @@ FReply FCustomizableInstanceDetails::OnProjectorLayerAdded(FString ParamName)
 	FScopedTransaction LocalTransaction(LOCTEXT("AddProjectorLayer", "Add Projector Layer"));
 	CustomInstance->Modify();
 
-	const int32 NumLayers = CustomInstance->MultilayerProjectorNumLayers(*ParamName) + 1;
+	const int32 NumLayers = CustomInstance->MultilayerProjectorNumLayers(*ParamName);
 	CustomInstance->MultilayerProjectorCreateLayer(*ParamName, NumLayers);
 
 	UpdateInstance();
