[mutable] Fixed indeterminism in the randomization of int parameter
- The sorting of the found related COs was being performed too early.
- Added stage in all testing commandlets that require CO compilation that deletes any previous compiled data.

#rnx

diff --git a/Engine/Plugins/Mutable/Source/CustomizableObjectEditor/Private/MuCOE/CustomizableObjectCompiler.cpp b/Engine/Plugins/Mutable/Source/CustomizableObjectEditor/Private/MuCOE/CustomizableObjectCompiler.cpp
--- a/Engine/Plugins/Mutable/Source/CustomizableObjectEditor/Private/MuCOE/CustomizableObjectCompiler.cpp
+++ b/Engine/Plugins/Mutable/Source/CustomizableObjectEditor/Private/MuCOE/CustomizableObjectCompiler.cpp
@@ -452,12 +452,6 @@ void ProcessChildObjectsRecursively(const UCustomizableObject* ParentObject, FMu
 		return;
 	}
 	
-	// Required to be deterministic.
-	ReferencedObjectNames.Sort([](const FName& A, const FName& B)
-	{
-		return A.LexicalLess(B);
-	});
-	
 	bool bMultipleBaseObjectsFound = false;
 	
 	TArray<FAssetData> AssetDataArray;
@@ -467,6 +461,12 @@ void ProcessChildObjectsRecursively(const UCustomizableObject* ParentObject, FMu
 	Filter.ClassPaths = { UCustomizableObject::StaticClass()->GetClassPathName() };
 	AssetRegistryModule.Get().GetAssets(Filter, AssetDataArray);
 
+	// Required to be deterministic.
+	AssetDataArray.Sort([](const FAssetData& A, const FAssetData& B)
+	{
+		return A.PackageName.LexicalLess(B.PackageName);
+	});
+	
 	for (FAssetData AssetData : AssetDataArray)
 	{
 		FSoftObjectPath SoftObjectPath = AssetData.GetSoftObjectPath();
@@ -498,19 +498,13 @@ void ProcessChildObjectsRecursively(const UCustomizableObject* ParentObject, FMu
 			}
 
 			ICustomizableObjectVersionBridgeInterface* CustomizableObjectVersionBridgeInterface = Cast<ICustomizableObjectVersionBridgeInterface>(GenerationContext.RootVersionBridge);
-
-			if (CustomizableObjectVersionBridgeInterface)
+			if (ensure(CustomizableObjectVersionBridgeInterface))
 			{
 				if (!CustomizableObjectVersionBridgeInterface->IsVersionStructIncludedInCurrentRelease(ChildObject->VersionStruct))
 				{
 					continue;
 				}
 			}
-			else
-			{
-				// This should never happen as the ICustomizableObjectVersionBridgeInterface was already checked at the start of the compilation
-				ensure(false);
-			}
 		}
 		
 		if (!bMultipleBaseObjectsFound)
diff --git a/Engine/Plugins/Mutable/Source/MutableValidation/Private/MuV/CustomizableObjectCompilationUtility.cpp b/Engine/Plugins/Mutable/Source/MutableValidation/Private/MuV/CustomizableObjectCompilationUtility.cpp
--- a/Engine/Plugins/Mutable/Source/MutableValidation/Private/MuV/CustomizableObjectCompilationUtility.cpp
+++ b/Engine/Plugins/Mutable/Source/MutableValidation/Private/MuV/CustomizableObjectCompilationUtility.cpp
@@ -13,6 +13,7 @@
 #include "Interfaces/ITargetPlatform.h"
 #include "MuCOE/CompileRequest.h"
 #include "MuCOE/CustomizableObjectEditorFunctionLibrary.h"
+#include "HAL/FileManager.h"
 
 
 bool FCustomizableObjectCompilationUtility::CompileCustomizableObject(UCustomizableObject* InCustomizableObject, const bool bShouldLogMutableLogs /* = true */, const FCompilationOptions* InCompilationOptionsOverride  /* nullptr */)
@@ -60,6 +61,22 @@ bool FCustomizableObjectCompilationUtility::CompileCustomizableObject(UCustomiza
 		UE_LOG(LogMutable, Log, TEXT("(int) model_package_data_bytes_limit_bytes : %llu "), CompilationOptions.PackagedDataBytesLimit);
 	}
 
+	// Delete data from previous compilations just to be sure it does not affect the new compilation operation
+	IFileManager& FileManager = IFileManager::Get();
+	const FString CompiledDataDirectory = CustomizableObject->GetPrivate()->GetCompiledDataFolderPath();
+	if (CustomizableObject->IsCompiled() && FileManager.DirectoryExists(*CompiledDataDirectory))
+	{
+		UE_LOG(LogMutable,Display,TEXT("Deleting old mutable compiled data from '%s'..."), *CompiledDataDirectory);
+		if ( !FileManager.DeleteDirectory( *CompiledDataDirectory, false, true) )
+		{
+			UE_LOG(LogMutable,Warning,TEXT("Failed to delete old compiled data directory '%s'. This could affect the compilation of the CO."), *CompiledDataDirectory);
+		}
+		else
+		{
+			UE_LOG(LogMutable, Display, TEXT("Directory deleted successfully"));
+		}
+	}
+	
 	// Compilation completed successfully
 	bool bCompilationSuccess = false;
 	
