[mutable] Change strategy to decide if a texture format needs channel swizzle because of the internal represetnation.
#rnx

diff --git a/Engine/Plugins/Mutable/Source/CustomizableObjectEditor/Private/MuCOE/GenerateMutableSource/GenerateMutableSourceSurface.cpp b/Engine/Plugins/Mutable/Source/CustomizableObjectEditor/Private/MuCOE/GenerateMutableSource/GenerateMutableSourceSurface.cpp
--- a/Engine/Plugins/Mutable/Source/CustomizableObjectEditor/Private/MuCOE/GenerateMutableSource/GenerateMutableSourceSurface.cpp
+++ b/Engine/Plugins/Mutable/Source/CustomizableObjectEditor/Private/MuCOE/GenerateMutableSource/GenerateMutableSourceSurface.cpp
@@ -764,11 +764,16 @@ mu::Ptr<mu::NodeSurface> GenerateMutableSourceSurface(const UEdGraphPin * Pin, F
 										bool bUseLA = GenerationContext.CompilationContext.Options.TargetPlatform->SupportsFeature(ETargetPlatformFeatures::NormalmapLAEncodingMode);
 										if (bUseLA)
 										{
-											// See GetQualityFormat in TextureFormatASTC.cpp to understand why
-											if (UnrealTargetPlatformFormat == PF_ASTC_6x6 || UnrealTargetPlatformFormat == PF_ASTC_6x6_NORM_RG)
+											// We'll have to trust the reference texture because the actual internal settings are opaque.
+											// See GetQualityFormat in TextureFormatASTC.cpp to understand how it works, but it depoends on some inaccessible texture name constants like this:
+											//bool bIsNormalMapFormat = TextureFormatName == GTextureFormatNameASTC_NormalAG ||
+											//		TextureFormatName == GTextureFormatNameASTC_NormalRG ||
+											//		TextureFormatName == GTextureFormatNameASTC_NormalLA ||
+											//		TextureFormatName == GTextureFormatNameASTC_NormalRG_Precise;
+											bool bIsNormalMapFormat = ReferenceTexture->IsNormalMap();
+
+											if (bIsNormalMapFormat)
 											{
-												MutableFormat = mu::EImageFormat::ASTC_4x4_RGBA_LDR;
-												MutableFormatIfAlpha = mu::EImageFormat::ASTC_4x4_RGBA_LDR;
 
 												// Insert a channel swizzle
 												mu::Ptr<mu::NodeImageSwizzle> Swizzle = new mu::NodeImageSwizzle;
