[Mutable] Added more compile scopes.

diff --git a/Engine/Plugins/Mutable/Source/CustomizableObjectEditor/Private/MuCOE/CustomizableObjectCompiler.cpp b/Engine/Plugins/Mutable/Source/CustomizableObjectEditor/Private/MuCOE/CustomizableObjectCompiler.cpp
--- a/Engine/Plugins/Mutable/Source/CustomizableObjectEditor/Private/MuCOE/CustomizableObjectCompiler.cpp
+++ b/Engine/Plugins/Mutable/Source/CustomizableObjectEditor/Private/MuCOE/CustomizableObjectCompiler.cpp
@@ -333,6 +333,8 @@ void FCustomizableObjectCompiler::Compile(const TSharedRef<FCompilationRequest>&
 
 void FCustomizableObjectCompiler::EnqueueCompileRequest(const TSharedRef<FCompilationRequest>& CompileRequest, bool bForceRequests)
 {
+	MUTABLE_CPUPROFILER_SCOPE(FCustomizableObjectCompiler::EnqueueCompileRequest)
+
 	check(IsInGameThread());
 
 	const UCustomizableObject* CustomizableObject = CompileRequest->GetCustomizableObject();
@@ -359,6 +361,8 @@ void FCustomizableObjectCompiler::EnqueueCompileRequest(const TSharedRef<FCompil
 		
 	if (!CompileRequest->bAsync)
 	{
+		MUTABLE_CPUPROFILER_SCOPE(SyncCompile)
+		
 		TSharedRef<FCustomizableObjectCompiler> SyncCompiler = MakeShared<FCustomizableObjectCompiler>();
 		SyncCompiler->Compile(CompileRequest);
 	}
@@ -816,7 +820,7 @@ void FCustomizableObjectCompiler::SetCompilationState(ECompilationStatePrivate S
 
 void FCustomizableObjectCompiler::CompileInternal(bool bAsync)
 {
-	MUTABLE_CPUPROFILER_SCOPE(FCustomizableObjectCompiler::Compile)
+	MUTABLE_CPUPROFILER_SCOPE(FCustomizableObjectCompiler::CompileInternal)
 	
 	UE_LOG(LogMutable, Verbose, TEXT("PROFILE: [ %16.8f ] FCustomizableObjectCompiler::Compile start."), FPlatformTime::Seconds());
 	
