[mutable] Remove root mesh formatting operation in favour of operations on the leave meshes.
Small optimizations to mesh format code related to tangent space.
Encode morph targets in 3-components instead of 4.
Moved SetSurfaceFormat function to make it accessible from other parts of the code and simplified its parameters.
#rnx

diff --git a/Engine/Plugins/Mutable/Source/CustomizableObjectEditor/Private/MuCOE/CustomizableObjectCompileRunnable.cpp b/Engine/Plugins/Mutable/Source/CustomizableObjectEditor/Private/MuCOE/CustomizableObjectCompileRunnable.cpp
--- a/Engine/Plugins/Mutable/Source/CustomizableObjectEditor/Private/MuCOE/CustomizableObjectCompileRunnable.cpp
+++ b/Engine/Plugins/Mutable/Source/CustomizableObjectEditor/Private/MuCOE/CustomizableObjectCompileRunnable.cpp
@@ -212,21 +212,7 @@ uint32 FCustomizableObjectCompileRunnable::Run()
 					// Ensure we don't pull unwanted data into the package when cooking.
 					FCookLoadScope CookLoadScope(ECookLoadType::EditorOnly);
 
-					if (MorphName.IsEmpty())
-					{
-						*ResolvedMesh = ConvertSkeletalMeshToMutable(MeshData, *Compiler->CompilationContext);
-					}
-					else
-					{
-						const USkeletalMesh* SkeletalMesh = Cast<USkeletalMesh>(MeshData.Mesh.LoadSynchronous());
-						if (SkeletalMesh->FindMorphTarget(*MorphName))
-						{
-							TSharedPtr<mu::FMesh> Mesh;
-							Mesh = ConvertSkeletalMeshToMutable(MeshData, *Compiler->CompilationContext);
-							Mesh = BuildMorphedMutableMeshFromMesh(Mesh, SkeletalMesh, *MorphName, MeshData.LODIndex, MeshData.SectionIndex);
-							*ResolvedMesh = Mesh;
-						}
-					}
+					*ResolvedMesh = ConvertSkeletalMeshToMutable(MeshData, *Compiler->CompilationContext, MorphName);
 
 					CompletionEvent.Trigger();
 				};
diff --git a/Engine/Plugins/Mutable/Source/CustomizableObjectEditor/Private/MuCOE/GenerateMutableSource/GenerateMutableSource.h b/Engine/Plugins/Mutable/Source/CustomizableObjectEditor/Private/MuCOE/GenerateMutableSource/GenerateMutableSource.h
--- a/Engine/Plugins/Mutable/Source/CustomizableObjectEditor/Private/MuCOE/GenerateMutableSource/GenerateMutableSource.h
+++ b/Engine/Plugins/Mutable/Source/CustomizableObjectEditor/Private/MuCOE/GenerateMutableSource/GenerateMutableSource.h
@@ -984,5 +984,4 @@ uint32 AddUniqueSurfaceMetadata(const FSkeletalMaterial* Material, const FSkelMe
 
 /** Convert a mesh to Mutable format.
 */
-TSharedPtr<mu::FMesh> ConvertSkeletalMeshToMutable(const FMutableSourceMeshData& Source, FMutableCompilationContext& Context);
-TSharedPtr<mu::FMesh> BuildMorphedMutableMeshFromMesh(TSharedPtr<mu::FMesh> BaseSourceMesh, const USkeletalMesh* SkeletalMesh, const FString& MorphTargetName, int32 LODIndex, int32 SectionIndex);
+TSharedPtr<mu::FMesh> ConvertSkeletalMeshToMutable(const FMutableSourceMeshData& Source, FMutableCompilationContext& Context, const FString& MorphName);
diff --git a/Engine/Plugins/Mutable/Source/CustomizableObjectEditor/Private/MuCOE/GenerateMutableSource/GenerateMutableSourceMesh.cpp b/Engine/Plugins/Mutable/Source/CustomizableObjectEditor/Private/MuCOE/GenerateMutableSource/GenerateMutableSourceMesh.cpp
--- a/Engine/Plugins/Mutable/Source/CustomizableObjectEditor/Private/MuCOE/GenerateMutableSource/GenerateMutableSourceMesh.cpp
+++ b/Engine/Plugins/Mutable/Source/CustomizableObjectEditor/Private/MuCOE/GenerateMutableSource/GenerateMutableSourceMesh.cpp
@@ -15,6 +15,7 @@
 #include "MuCO/CustomizableObjectInstance.h"
 #include "MuCO/CustomizableObjectSystem.h"
 #include "MuCO/UnrealConversionUtils.h"
+#include "MuCO/MutableMeshBufferUtils.h"
 #include "MuCOE/CustomizableObjectCompiler.h"
 #include "MuCOE/CustomizableObjectLayout.h"
 #include "MuCOE/GenerateMutableSource/GenerateMutableSourceFloat.h"
@@ -45,6 +46,7 @@
 #include "MuT/NodeMeshTable.h"
 #include "MuT/NodeMeshVariation.h"
 #include "MuR/Types.h"
+#include "MuR/OpMeshFormat.h"
 #include "PhysicsEngine/PhysicsAsset.h"
 #include "PhysicsEngine/SkeletalBodySetup.h"
 #include "Engine/SkinnedAssetCommon.h"
@@ -53,6 +55,154 @@
 
 #define LOCTEXT_NAMESPACE "CustomizableObjectEditor"
 
+
+void SetSurfaceFormat(FMutableCompilationContext& Context,
+	mu::FMeshBufferSet& OutVertexBufferFormat,
+	mu::FMeshBufferSet& OutIndexBufferFormat,
+	const FMutableGraphMeshGenerationData& MeshData)
+{
+	// Limit skinning weights if necessary
+	// \todo: make it more flexible to support 3 or 5 or 1 weight, since there is support for this in 4.25
+	const int32 MutableBonesPerVertex = FGPUBaseSkinVertexFactory::UseUnlimitedBoneInfluences(MeshData.MaxNumBonesPerVertex, Context.Options.TargetPlatform) &&
+		MeshData.MaxNumBonesPerVertex < (int32)Context.Options.CustomizableObjectNumBoneInfluences ?
+		MeshData.MaxNumBonesPerVertex :
+		(int32)Context.Options.CustomizableObjectNumBoneInfluences;
+
+	ensure(MutableBonesPerVertex <= MAX_TOTAL_INFLUENCES);
+
+	if (MutableBonesPerVertex != MeshData.MaxNumBonesPerVertex)
+	{
+		UE_LOG(LogMutable, Verbose, TEXT("In object [%s] Mesh bone number adjusted from %d to %d."), *Context.Object->GetFullName(), MeshData.MaxNumBonesPerVertex, MutableBonesPerVertex);
+	}
+
+	int MutableBufferCount = MUTABLE_VERTEXBUFFER_TEXCOORDS + 1;
+	if (MeshData.bHasVertexColors)
+	{
+		++MutableBufferCount;
+	}
+
+	if (MeshData.MaxNumBonesPerVertex > 0 && MeshData.MaxBoneIndexTypeSizeBytes > 0)
+	{
+		++MutableBufferCount;
+	}
+
+	if (MeshData.bHasRealTimeMorphs)
+	{
+		MutableBufferCount += 2;
+	}
+
+	if (MeshData.bHasClothing)
+	{
+		MutableBufferCount += 2;
+	}
+
+	MutableBufferCount += MeshData.SkinWeightProfilesSemanticIndices.Num();
+
+	OutVertexBufferFormat.SetBufferCount(MutableBufferCount);
+
+	int32 CurrentVertexBuffer = 0;
+
+	// Vertex buffer
+	MutableMeshBufferUtils::SetupVertexPositionsBuffer(CurrentVertexBuffer, OutVertexBufferFormat);
+	++CurrentVertexBuffer;
+
+	// Tangent buffer
+	MutableMeshBufferUtils::SetupTangentBuffer(CurrentVertexBuffer, OutVertexBufferFormat);
+	++CurrentVertexBuffer;
+
+	// Texture coords buffer
+	MutableMeshBufferUtils::SetupTexCoordinatesBuffer(CurrentVertexBuffer, MeshData.NumTexCoordChannels, OutVertexBufferFormat);
+	++CurrentVertexBuffer;
+
+	// Skin buffer
+	if (MeshData.MaxNumBonesPerVertex > 0 && MeshData.MaxBoneIndexTypeSizeBytes > 0)
+	{
+		const int32 MaxBoneWeightTypeSizeBytes = Context.Options.b16BitBoneWeightsEnabled ? 2 : 1;
+		MutableMeshBufferUtils::SetupSkinBuffer(CurrentVertexBuffer, MeshData.MaxBoneIndexTypeSizeBytes, MaxBoneWeightTypeSizeBytes, MutableBonesPerVertex, OutVertexBufferFormat);
+		++CurrentVertexBuffer;
+	}
+
+	// Colour buffer
+	if (MeshData.bHasVertexColors)
+	{
+		MutableMeshBufferUtils::SetupVertexColorBuffer(CurrentVertexBuffer, OutVertexBufferFormat);
+		++CurrentVertexBuffer;
+	}
+
+	// MorphTarget vertex tracking info buffers
+	if (MeshData.bHasRealTimeMorphs)
+	{
+		using namespace mu;
+		{
+			const int32 ElementSize = sizeof(uint32);
+			constexpr int32 ChannelCount = 1;
+			const EMeshBufferSemantic Semantics[ChannelCount] = { EMeshBufferSemantic::Other };
+			const int32 SemanticIndices[ChannelCount] = { 0 };
+			const EMeshBufferFormat Formats[ChannelCount] = { EMeshBufferFormat::UInt32 };
+			const int32 Components[ChannelCount] = { 1 };
+			const int32 Offsets[ChannelCount] = { 0 };
+
+			OutVertexBufferFormat.SetBuffer(CurrentVertexBuffer, ElementSize, ChannelCount, Semantics, SemanticIndices, Formats, Components, Offsets);
+			++CurrentVertexBuffer;
+		}
+
+		{
+			const int32 ElementSize = sizeof(uint32);
+			constexpr int32 ChannelCount = 1;
+			const EMeshBufferSemantic Semantics[ChannelCount] = { EMeshBufferSemantic::Other };
+			const int32 SemanticIndices[ChannelCount] = { 1 };
+			const EMeshBufferFormat Formats[ChannelCount] = { EMeshBufferFormat::UInt32 };
+			const int32 Components[ChannelCount] = { 1 };
+			const int32 Offsets[ChannelCount] = { 0 };
+
+			OutVertexBufferFormat.SetBuffer(CurrentVertexBuffer, ElementSize, ChannelCount, Semantics, SemanticIndices, Formats, Components, Offsets);
+			++CurrentVertexBuffer;
+		}
+	}
+
+	//Clothing Data Buffer.
+	if (MeshData.bHasClothing)
+	{
+		{
+			using namespace mu;
+			const int32 ElementSize = sizeof(int32);
+			constexpr int32 ChannelCount = 1;
+			const EMeshBufferSemantic Semantics[ChannelCount] = { EMeshBufferSemantic::Other };
+			const int32 SemanticIndices[ChannelCount] = { 2 };
+			const EMeshBufferFormat Formats[ChannelCount] = { EMeshBufferFormat::Int32 };
+			const int32 Components[ChannelCount] = { 1 };
+			const int32 Offsets[ChannelCount] = { 0 };
+
+			OutVertexBufferFormat.SetBuffer(CurrentVertexBuffer, ElementSize, ChannelCount, Semantics, SemanticIndices, Formats, Components, Offsets);
+			++CurrentVertexBuffer;
+		}
+
+		{
+			using namespace mu;
+			const int32 ElementSize = sizeof(uint32);
+			constexpr int32 ChannelCount = 1;
+			const EMeshBufferSemantic Semantics[ChannelCount] = { EMeshBufferSemantic::Other };
+			const int32 SemanticIndices[ChannelCount] = { 3 };
+			const EMeshBufferFormat Formats[ChannelCount] = { EMeshBufferFormat::UInt32 };
+			const int32 Components[ChannelCount] = { 1 };
+			const int32 Offsets[ChannelCount] = { 0 };
+
+			OutVertexBufferFormat.SetBuffer(CurrentVertexBuffer, ElementSize, ChannelCount, Semantics, SemanticIndices, Formats, Components, Offsets);
+			++CurrentVertexBuffer;
+		}
+	}
+
+	for (int32 ProfileSemanticIndex : MeshData.SkinWeightProfilesSemanticIndices)
+	{
+		MutableMeshBufferUtils::SetupSkinWeightProfileBuffer(CurrentVertexBuffer, MeshData.MaxBoneIndexTypeSizeBytes, 1, MutableBonesPerVertex, ProfileSemanticIndex, OutVertexBufferFormat);
+		++CurrentVertexBuffer;
+	}
+
+	// Index buffer
+	MutableMeshBufferUtils::SetupIndexBuffer(OutIndexBufferFormat);
+}
+
+
 void GetLODAndSectionForAutomaticLODs(const FMutableGraphGenerationContext& Context, const UCustomizableObjectNode& Node, const USkeletalMesh& SkeletalMesh,
 	const int32 LODIndexConnected, const int32 SectionIndexConnected, int32& OutLODIndex, int32& OutSectionIndex, const bool bOnlyConnectedLOD)
 {
@@ -707,7 +857,8 @@ namespace MutablePrivate
 
 TSharedPtr<mu::FMesh> ConvertSkeletalMeshToMutable(
 	const FMutableSourceMeshData& Source,
-	FMutableCompilationContext& Context
+	FMutableCompilationContext& Context,
+	const FString& MorphName
 )
 {
 	const FString MeshName = Source.Mesh.GetLongPackageName().ToLower();
@@ -788,9 +939,6 @@ TSharedPtr<mu::FMesh> ConvertSkeletalMeshToMutable(
 			return nullptr;
 		}
 
-		// Done while in GenerateContext
-		//Context.AddParticipatingObject(*Skeleton);
-
 		FMutableComponentInfo* MutComponentInfo = Context.GetComponentInfo(Source.Component);
 		check(MutComponentInfo);
 
@@ -1170,9 +1318,43 @@ TSharedPtr<mu::FMesh> ConvertSkeletalMeshToMutable(
 
 	}
 
+	const FSoftSkinVertex* SourceVertexData = Vertices.GetData() + VertexStart;
+
+	// Apply the morph if necessary
+	TUniquePtr<FSoftSkinVertex[]> MorphedVertexData;
+	if (!MorphName.IsEmpty())
+	{
+		UMorphTarget* MorphTarget = SkeletalMesh->FindMorphTarget(*MorphName);
+		if (MorphTarget && MorphTarget->GetMorphLODModels().IsValidIndex(Source.LODIndex))
+		{
+			// Make a copy of the original vertex data since we will modify it.
+			MorphedVertexData = MakeUnique<FSoftSkinVertex[]>(VertexCount);
+			FMemory::Memcpy(MorphedVertexData.Get(), SourceVertexData, VertexCount*sizeof(FSoftSkinVertex));
+
+			// Bake the morph in the copy of the source vertex data
+			int32 MaterialVertexStart = SkeletalMesh->GetImportedModel()->LODModels[Source.LODIndex].Sections[Source.SectionIndex].GetVertexBufferIndex();
+
+			const FMorphTargetLODModel& MorphLODModel = MorphTarget->GetMorphLODModels()[Source.LODIndex];
+			for (const FMorphTargetDelta& MorphDelta : MorphLODModel.Vertices)
+			{
+				const int32 VertexIndex = MorphDelta.SourceIdx - MaterialVertexStart;
+				if (VertexIndex >= 0 && VertexIndex < VertexCount)
+				{
+					FVector3f& Position = MorphedVertexData[VertexIndex].Position;
+					Position += MorphDelta.PositionDelta;
+
+					FVector4f& Normal = MorphedVertexData[VertexIndex].TangentZ;
+					Normal += FVector4f(MorphDelta.TangentZDelta, 0.0f);
+				}
+			}
+
+			SourceVertexData = MorphedVertexData.Get();
+		}
+	}
+
 	UE::Tasks::FTaskEvent CompletionEvent(TEXT("MeshConversionCompletion"));
 
-	MutablePrivate::CopyBufferClearingPadding(CompletionEvent, MutableMesh->GetVertexBuffers(), UnrealSourceVertexFormat, Vertices.GetData() + VertexStart);
+	MutablePrivate::CopyBufferClearingPadding(CompletionEvent, MutableMesh->GetVertexBuffers(), UnrealSourceVertexFormat, SourceVertexData);
 
 	// Indices
 	CompletionEvent.AddPrerequisites(
@@ -2264,6 +2446,63 @@ TSharedPtr<mu::FMesh> ConvertSkeletalMeshToMutable(
 		MeshSurface.BoneMapCount = MutableMesh->BoneMap.Num();
 	}
 
+
+	// Convert to the expected mesh format
+	// TODO: generate directly in this format.
+	{
+		FMutableGraphMeshGenerationData MeshData;
+
+		if (Context.Options.bSkinWeightProfilesEnabled && Source.LODIndex >= 0)
+		{
+			if (ImportedModel && ImportedModel->LODModels.IsValidIndex(Source.LODIndex))
+			{
+				for (const auto& SkinWeightProfile : ImportedModel->LODModels[Source.LODIndex].SkinWeightProfiles)
+				{
+					const int32 ProfileId = static_cast<int32>(Context.GetSkinWeightProfileIdUnique(SkinWeightProfile.Key));
+					MeshData.SkinWeightProfilesSemanticIndices.AddUnique(ProfileId);
+				}
+			}
+		}
+
+		if (ImportedModel->LODModels.IsValidIndex(Source.LODIndex) &&
+			ImportedModel->LODModels[Source.LODIndex].Sections.IsValidIndex(Source.SectionIndex))
+		{
+			MeshData.bHasVertexColors = SkeletalMesh->GetHasVertexColors();
+			MeshData.NumTexCoordChannels = ImportedModel->LODModels[Source.LODIndex].NumTexCoords;
+			MeshData.MaxBoneIndexTypeSizeBytes = ImportedModel->LODModels[Source.LODIndex].RequiredBones.Num() > 256 ? 2 : 1;
+			MeshData.MaxNumBonesPerVertex = ImportedModel->LODModels[Source.LODIndex].GetMaxBoneInfluences();
+
+			// When mesh data is combined we will get an upper and lower bound of the number of triangles.
+			MeshData.MaxNumTriangles = ImportedModel->LODModels[Source.LODIndex].Sections[Source.SectionIndex].NumTriangles;
+			MeshData.MinNumTriangles = MeshData.MaxNumTriangles;
+
+			MeshData.bHasRealTimeMorphs = Context.Options.bRealTimeMorphTargetsEnabled && SkeletalMesh->GetMorphTargets().Num();
+			MeshData.bHasClothing = Context.Options.bClothingEnabled && ImportedModel->LODModels[Source.LODIndex].HasClothData();
+		}
+
+		TSharedPtr<mu::FMesh> ForcedFormat = MakeShared<mu::FMesh>();
+		SetSurfaceFormat(Context, ForcedFormat->GetVertexBuffers(), ForcedFormat->GetIndexBuffers(), MeshData);
+
+		TSharedPtr<mu::FMesh> FormattedMesh = MakeShared<mu::FMesh>();
+		bool bOutSuccess = false;
+		constexpr bool bKeepSystemBuffers = true;
+		constexpr bool bFormatVertices = true;
+		constexpr bool bFormatIndices = true;
+		constexpr bool bIgnoreMissingChannels = true;
+		mu::MeshFormat(FormattedMesh.Get(), MutableMesh.Get(), ForcedFormat.Get(),
+			bKeepSystemBuffers,
+			bFormatVertices,
+			bFormatIndices,
+			bIgnoreMissingChannels,
+			bOutSuccess);
+
+		mu::MeshOptimizeBuffers(FormattedMesh.Get());
+
+		check(bOutSuccess);
+
+		MutableMesh = FormattedMesh;
+	}
+
 	return MutableMesh;
 }
 
@@ -2386,7 +2625,8 @@ TSharedPtr<mu::FMesh> ConvertSkeletalMeshToMutable(
 	else
 	{
 		// At some point this shouldn't happen anymore because all mesh conversion could be moved to the core compilation stage.
-		MutableMesh = ConvertSkeletalMeshToMutable(Source, GenerationContext.CompilationContext);
+		FString NoMorph;
+		MutableMesh = ConvertSkeletalMeshToMutable(Source, GenerationContext.CompilationContext, NoMorph);
 	}
 
 	return MutableMesh;
@@ -2660,70 +2900,6 @@ TSharedPtr<mu::FMesh> GenerateMutableStaticMesh(TSoftObjectPtr<UStreamableRender
 }
 
 
-TSharedPtr<mu::FMesh> BuildMorphedMutableMeshFromMesh(TSharedPtr<mu::FMesh> BaseSourceMesh, const USkeletalMesh* SkeletalMesh, const FString& MorphTargetName, int32 LODIndex, int32 SectionIndex)
-{
-	MUTABLE_CPUPROFILER_SCOPE(BuildMorphedMutableMeshFromMesh);
-
-	// Clone it (it will probably be shared)
-	TSharedPtr<mu::FMesh> MorphedSourceMesh = BaseSourceMesh->Clone();
-
-	// Bake the morph in the new mutable mesh
-	UMorphTarget* MorphTarget = SkeletalMesh ? SkeletalMesh->FindMorphTarget(*MorphTargetName) : nullptr;
-
-	if (MorphTarget && MorphTarget->GetMorphLODModels().IsValidIndex(LODIndex))
-	{
-		int32 PosBuf = -1;
-		int32 PosChannel = -1;
-		MorphedSourceMesh->GetVertexBuffers().FindChannel(mu::EMeshBufferSemantic::Position, 0, &PosBuf, &PosChannel);
-		check(PosBuf >= 0 && PosChannel >= 0);
-
-		int32 PosElemSize = MorphedSourceMesh->GetVertexBuffers().GetElementSize(PosBuf);
-		int32 PosOffset = MorphedSourceMesh->GetVertexBuffers().GetChannelOffset(PosBuf, PosChannel);
-		uint8* PosBuffer = MorphedSourceMesh->GetVertexBuffers().GetBufferData(PosBuf) + PosOffset;
-
-		int32 NorBuf = -1;
-		int32 NorChannel = -1;
-		MorphedSourceMesh->GetVertexBuffers().FindChannel(mu::EMeshBufferSemantic::Normal, 0, &NorBuf, &NorChannel);
-
-		const bool bHasNormals = NorBuf >= 0 && NorChannel >= 0;
-
-		int32 NorElemSize = bHasNormals ? MorphedSourceMesh->GetVertexBuffers().GetElementSize(NorBuf) : 0;
-		int32 NorOffset = bHasNormals ? MorphedSourceMesh->GetVertexBuffers().GetChannelOffset(NorBuf, NorChannel) : 0;
-		uint8* NorBuffer = bHasNormals ? MorphedSourceMesh->GetVertexBuffers().GetBufferData(NorBuf) + NorOffset : nullptr;
-
-		int32 MaterialVertexStart = SkeletalMesh->GetImportedModel()->LODModels[LODIndex].Sections[SectionIndex].GetVertexBufferIndex();
-		int32 MeshVertexCount = MorphedSourceMesh->GetVertexBuffers().GetElementCount();
-
-		const FMorphTargetLODModel& MorphLODModel = MorphTarget->GetMorphLODModels()[LODIndex];
-		for (const FMorphTargetDelta& MorphDelta : MorphLODModel.Vertices)
-		{
-			const int32 VertexIndex = MorphDelta.SourceIdx - MaterialVertexStart;
-			if (VertexIndex >= 0 && VertexIndex < MeshVertexCount)
-			{
-				{
-					float* const PosData = reinterpret_cast<float*>(PosBuffer + PosElemSize * VertexIndex);
-					const FVector3f MorphedPosition = FVector3f(PosData[0], PosData[1], PosData[2]) + MorphDelta.PositionDelta;
-					PosData[0] = MorphedPosition.X;
-					PosData[1] = MorphedPosition.Y;
-					PosData[2] = MorphedPosition.Z;
-				}
-
-				if (bHasNormals)
-				{
-					float* const NorData = reinterpret_cast<float*>(NorBuffer + NorElemSize * VertexIndex);
-					const FVector3f MorphedNormal = FVector3f(NorData[0], NorData[1], NorData[2]) + MorphDelta.TangentZDelta;
-					NorData[0] = MorphedNormal.X;
-					NorData[1] = MorphedNormal.Y;
-					NorData[2] = MorphedNormal.Z;
-				}
-			}
-		}
-	}
-
-	return MorphedSourceMesh;
-}
-
-
 TSharedPtr<mu::FMesh> BuildMorphedMutableMesh(const UEdGraphPin* BaseSourcePin, const FString& MorphTargetName, FMutableGraphGenerationContext& GenerationContext, const bool bOnlyConnectedLOD, const FName& RowName)
 {
 	check(BaseSourcePin);
@@ -3270,9 +3446,8 @@ bool GetAndValidateReshapePhysicsToDeform(
 mu::NodeMeshPtr GenerateMorphMesh(const UEdGraphPin* Pin,
 	TArray<FMorphNodeData> TypedNodeMorphs,
 	int32 MorphIndex,
-	mu::NodeMeshPtr SourceNode,
+	mu::Ptr<mu::NodeMesh> SourceNode,
 	FMutableGraphGenerationContext & GenerationContext,
-	FMutableGraphMeshGenerationData & MeshData,
 	const bool bOnlyConnectedLOD,
 	const FString& TableColumnName = "")
 {
@@ -3298,7 +3473,7 @@ mu::NodeMeshPtr GenerateMorphMesh(const UEdGraphPin* Pin,
 	}
 	else
 	{
-		mu::NodeMeshPtr NextMorph = GenerateMorphMesh(Pin, TypedNodeMorphs, MorphIndex + 1, SourceNode, GenerationContext, MeshData, bOnlyConnectedLOD, TableColumnName); // TODO FutureGMT change to a for. This recursion can be problematic with the production cache
+		mu::NodeMeshPtr NextMorph = GenerateMorphMesh(Pin, TypedNodeMorphs, MorphIndex + 1, SourceNode, GenerationContext, bOnlyConnectedLOD, TableColumnName); // TODO FutureGMT change to a for. This recursion can be problematic with the production cache
 		Result->SetBase(NextMorph);
 	}
 	
@@ -3473,18 +3648,21 @@ mu::Ptr<mu::NodeMesh> GenerateMutableSourceMesh(const UEdGraphPin* Pin,
 	}
 
 	//SkeletalMesh Result
-	mu::NodeMeshPtr Result;
+	mu::Ptr<mu::NodeMesh> Result;
 
 	//SkeletalMesh + Morphs Result
-	mu::NodeMeshPtr MorphResult;
+	mu::Ptr<mu::NodeMesh> MorphResult;
 	
 	if (const UCustomizableObjectNodeSkeletalMesh* TypedNodeSkel = Cast<UCustomizableObjectNodeSkeletalMesh>(Node))
 	{
 		mu::NodeMeshConstantPtr MeshNode = new mu::NodeMeshConstant();
 		Result = MeshNode;
 
-		GenerationContext.LoadObject(TypedNodeSkel->SkeletalMesh);
-		if (TypedNodeSkel->SkeletalMesh)
+		if (TypedNodeSkel->SkeletalMesh.IsNull())
+		{
+			GenerationContext.Log(LOCTEXT("MissingSkeletalMesh", "No Skeletal Mesh set in the SkeletalMesh node."), Node);
+		}
+		else
 		{
 			int32 LODIndexConnected = -1; // LOD which the pin is connected to
 			int32 SectionIndexConnected = -1;
@@ -3497,7 +3675,9 @@ mu::Ptr<mu::NodeMesh> GenerateMutableSourceMesh(const UEdGraphPin* Pin,
 				TypedNodeSkel->GetPinSection(*Pin, LODIndexConnected, SectionIndexConnected, LayoutIndex);
 			}
 
-			GetLODAndSectionForAutomaticLODs(GenerationContext, *Node, *TypedNodeSkel->SkeletalMesh, LODIndexConnected, SectionIndexConnected, LODIndex, SectionIndex, bOnlyConnectedLOD);
+			// TODO: Avoid loading the mesh here.
+			USkeletalMesh* SkeletalMesh = GenerationContext.LoadObject(TypedNodeSkel->SkeletalMesh);
+			GetLODAndSectionForAutomaticLODs(GenerationContext, *Node, *SkeletalMesh, LODIndexConnected, SectionIndexConnected, LODIndex, SectionIndex, bOnlyConnectedLOD);
 			
 			// First process the mesh tags that are going to make the mesh unique and affect whether it's repeated in 
 			// the mesh cache or not
@@ -3526,7 +3706,7 @@ mu::Ptr<mu::NodeMesh> GenerateMutableSourceMesh(const UEdGraphPin* Pin,
 
 			if (GenerationContext.CompilationContext.Object->GetPrivate()->IsAssetUserDataMergeEnabled())
 			{
-				const TArray<UAssetUserData*>* AssetUserDataArray = TypedNodeSkel->SkeletalMesh->GetAssetUserDataArray();
+				const TArray<UAssetUserData*>* AssetUserDataArray = SkeletalMesh->GetAssetUserDataArray();
 
 				if (AssetUserDataArray)
 				{
@@ -3552,10 +3732,10 @@ mu::Ptr<mu::NodeMesh> GenerateMutableSourceMesh(const UEdGraphPin* Pin,
 				}
 			}
 
-			FSkeletalMeshModel* ImportedModel = TypedNodeSkel->SkeletalMesh->GetImportedModel();
+			FSkeletalMeshModel* ImportedModel = SkeletalMesh->GetImportedModel();
 			
 			constexpr bool bIsPassthrough = false;
-			TSharedPtr<mu::FMesh> MutableMesh = GenerateMutableSkeletalMesh(TypedNodeSkel->SkeletalMesh, TypedNodeSkel->AnimInstance, LODIndexConnected, SectionIndexConnected,
+			TSharedPtr<mu::FMesh> MutableMesh = GenerateMutableSkeletalMesh(SkeletalMesh, TypedNodeSkel->AnimInstance, LODIndexConnected, SectionIndexConnected,
 																LODIndex, SectionIndex, MeshUniqueTags, SurfaceMetadataId, GenerationContext, TypedNodeSkel, nullptr, bIsPassthrough);
 			if (MutableMesh)
 			{
@@ -3588,12 +3768,12 @@ mu::Ptr<mu::NodeMesh> GenerateMutableSourceMesh(const UEdGraphPin* Pin,
 					MutableMesh->AddStreamedResource(BitCast<uint64>(ResourceId));
 				}
 
-				AddSocketTagsToMesh(TypedNodeSkel->SkeletalMesh.Get(), MutableMesh, GenerationContext);
+				AddSocketTagsToMesh(SkeletalMesh, MutableMesh, GenerationContext);
 
 				if (UCustomizableObjectSystem::GetInstance()->IsMutableAnimInfoDebuggingEnabled())
 				{
 					FString MeshPath;
-					TypedNodeSkel->SkeletalMesh->GetOuter()->GetPathName(nullptr, MeshPath);
+					SkeletalMesh->GetOuter()->GetPathName(nullptr, MeshPath);
 					FString MeshTag = FString("__MeshPath:") + MeshPath;
 					AddTagToMutableMeshUnique(*MutableMesh, MeshTag);
 				}
@@ -3601,16 +3781,16 @@ mu::Ptr<mu::NodeMesh> GenerateMutableSourceMesh(const UEdGraphPin* Pin,
 				if (ImportedModel->LODModels.IsValidIndex(LODIndex) &&
 					ImportedModel->LODModels[LODIndex].Sections.IsValidIndex(SectionIndex))
 				{
-					MeshData.bHasVertexColors = TypedNodeSkel->SkeletalMesh->GetHasVertexColors();	
+					MeshData.bHasVertexColors = SkeletalMesh->GetHasVertexColors();	
 					MeshData.NumTexCoordChannels = ImportedModel->LODModels[LODIndex].NumTexCoords;
 					MeshData.MaxBoneIndexTypeSizeBytes = ImportedModel->LODModels[LODIndex].RequiredBones.Num() > 256 ? 2 : 1;
 					MeshData.MaxNumBonesPerVertex = ImportedModel->LODModels[LODIndex].GetMaxBoneInfluences();
 
 					// When mesh data is combined we will get an upper and lower bound of the number of triangles.
 					MeshData.MaxNumTriangles = ImportedModel->LODModels[LODIndex].Sections[SectionIndex].NumTriangles;
-					MeshData.MinNumTriangles = ImportedModel->LODModels[LODIndex].Sections[SectionIndex].NumTriangles;
+					MeshData.MinNumTriangles = MeshData.MaxNumTriangles;
 
-					MeshData.bHasRealTimeMorphs = GenerationContext.CompilationContext.Options.bRealTimeMorphTargetsEnabled && TypedNodeSkel->SkeletalMesh->GetMorphTargets().Num();
+					MeshData.bHasRealTimeMorphs = GenerationContext.CompilationContext.Options.bRealTimeMorphTargetsEnabled && SkeletalMesh->GetMorphTargets().Num();
 					MeshData.bHasClothing = GenerationContext.CompilationContext.Options.bClothingEnabled && ImportedModel->LODModels[LODIndex].HasClothData();
 				}
 			}
@@ -3672,13 +3852,9 @@ mu::Ptr<mu::NodeMesh> GenerateMutableSourceMesh(const UEdGraphPin* Pin,
 			// Applying Mesh Morph Nodes
 			if (GenerationContext.MeshMorphStack.Num())
 			{
-				MorphResult = GenerateMorphMesh(Pin, GenerationContext.MeshMorphStack, 0, Result, GenerationContext, MeshData, bOnlyConnectedLOD);
+				MorphResult = GenerateMorphMesh(Pin, GenerationContext.MeshMorphStack, 0, Result, GenerationContext, bOnlyConnectedLOD);
 			}
 		}
-		else
-		{
-			GenerationContext.Log(LOCTEXT("MissingSkeletalMesh", "No Skeletal Mesh set in the SkeletalMesh node."), Node);
-		}
 	}
 
 	else if (const UCustomizableObjectNodeStaticMesh* TypedNodeStatic = Cast<UCustomizableObjectNodeStaticMesh>(Node))
@@ -4378,7 +4554,7 @@ mu::Ptr<mu::NodeMesh> GenerateMutableSourceMesh(const UEdGraphPin* Pin,
 							// Applying Mesh Morph Nodes
 							if (DefaultSkeletalMesh && GenerationContext.MeshMorphStack.Num())
 							{
-								MorphResult = GenerateMorphMesh(Pin, GenerationContext.MeshMorphStack, 0, Result, GenerationContext, MeshData, bOnlyConnectedLOD, MutableColumnName);
+								MorphResult = GenerateMorphMesh(Pin, GenerationContext.MeshMorphStack, 0, Result, GenerationContext, bOnlyConnectedLOD, MutableColumnName);
 							}
 						}
 					}
diff --git a/Engine/Plugins/Mutable/Source/CustomizableObjectEditor/Private/MuCOE/GenerateMutableSource/GenerateMutableSourceMesh.h b/Engine/Plugins/Mutable/Source/CustomizableObjectEditor/Private/MuCOE/GenerateMutableSource/GenerateMutableSourceMesh.h
--- a/Engine/Plugins/Mutable/Source/CustomizableObjectEditor/Private/MuCOE/GenerateMutableSource/GenerateMutableSourceMesh.h
+++ b/Engine/Plugins/Mutable/Source/CustomizableObjectEditor/Private/MuCOE/GenerateMutableSource/GenerateMutableSourceMesh.h
@@ -17,6 +17,12 @@ class UStaticMesh;
 class UStreamableRenderAsset;
 struct FMutableGraphGenerationContext;
 struct FMutableGraphMeshGenerationData;
+struct FMutableCompilationContext;
+
+namespace mu
+{
+	class FMeshBufferSet;
+}
 
 
 /** Returns the corrected LOD and Section Index when using Automatic LOD From Mesh strategy.
@@ -104,3 +110,10 @@ TSharedPtr<mu::FMesh> BuildMorphedMutableMesh(const UEdGraphPin* BaseSourcePin,
  * @param bOnlyConnectedLOD Corrected LOD and Section will unconditionally always be the connected ones.
  * @return  Mutable Mesh Node. */
 mu::Ptr<mu::NodeMesh> GenerateMutableSourceMesh(const UEdGraphPin* Pin, FMutableGraphGenerationContext& GenerationContext, FMutableGraphMeshGenerationData& MeshData, uint32 SurfaceMetadataId, bool bLinkedToExtendMaterial, bool bOnlyConnectedLOD);
+
+
+/** */
+extern void SetSurfaceFormat(FMutableCompilationContext&,
+	mu::FMeshBufferSet& OutVertexBufferFormat, 
+	mu::FMeshBufferSet& OutIndexBufferFormat, 
+	const FMutableGraphMeshGenerationData& MeshData);
diff --git a/Engine/Plugins/Mutable/Source/CustomizableObjectEditor/Private/MuCOE/GenerateMutableSource/GenerateMutableSourceModifier.cpp b/Engine/Plugins/Mutable/Source/CustomizableObjectEditor/Private/MuCOE/GenerateMutableSource/GenerateMutableSourceModifier.cpp
--- a/Engine/Plugins/Mutable/Source/CustomizableObjectEditor/Private/MuCOE/GenerateMutableSource/GenerateMutableSourceModifier.cpp
+++ b/Engine/Plugins/Mutable/Source/CustomizableObjectEditor/Private/MuCOE/GenerateMutableSource/GenerateMutableSourceModifier.cpp
@@ -302,10 +302,7 @@ mu::Ptr<mu::NodeModifier> GenerateMutableSourceModifier(const UEdGraphPin * Pin,
 				mu::Ptr<mu::NodeMesh> MeshPtr = AddMeshNode;
 
 				mu::Ptr<mu::NodeMeshFormat> MeshFormat = new mu::NodeMeshFormat();
-				SetSurfaceFormat(GenerationContext,
-					MeshFormat->GetVertexBuffers(), MeshFormat->GetIndexBuffers(), MeshData,
-					GenerationContext.CompilationContext.Options.CustomizableObjectNumBoneInfluences,
-					GenerationContext.CompilationContext.Options.b16BitBoneWeightsEnabled);
+				SetSurfaceFormat(GenerationContext.CompilationContext, MeshFormat->GetVertexBuffers(), MeshFormat->GetIndexBuffers(), MeshData);
 
 				MeshFormat->SetSource(MeshPtr.get());
 
diff --git a/Engine/Plugins/Mutable/Source/CustomizableObjectEditor/Private/MuCOE/GenerateMutableSource/GenerateMutableSourceSurface.cpp b/Engine/Plugins/Mutable/Source/CustomizableObjectEditor/Private/MuCOE/GenerateMutableSource/GenerateMutableSourceSurface.cpp
--- a/Engine/Plugins/Mutable/Source/CustomizableObjectEditor/Private/MuCOE/GenerateMutableSource/GenerateMutableSourceSurface.cpp
+++ b/Engine/Plugins/Mutable/Source/CustomizableObjectEditor/Private/MuCOE/GenerateMutableSource/GenerateMutableSourceSurface.cpp
@@ -17,7 +17,6 @@
 #include "Rendering/SkeletalMeshLODModel.h"
 
 #include "MuCO/CustomizableObjectInstance.h"
-#include "MuCO/MutableMeshBufferUtils.h"
 #include "MuCO/UnrealConversionUtils.h"
 #include "MuCOE/CustomizableObjectCompiler.h"
 #include "MuCOE/CustomizableObjectLayout.h"
@@ -59,152 +58,6 @@
 #define LOCTEXT_NAMESPACE "CustomizableObjectEditor"
 
 
-void SetSurfaceFormat( FMutableGraphGenerationContext& GenerationContext,
-					   mu::FMeshBufferSet& OutVertexBufferFormat, mu::FMeshBufferSet& OutIndexBufferFormat, const FMutableGraphMeshGenerationData& MeshData, 
-					   ECustomizableObjectNumBoneInfluences ECustomizableObjectNumBoneInfluences, bool bWith16BitWeights)
-{
-	// Limit skinning weights if necessary
-	// \todo: make it more flexible to support 3 or 5 or 1 weight, since there is support for this in 4.25
-	const int32 MutableBonesPerVertex = FGPUBaseSkinVertexFactory::UseUnlimitedBoneInfluences(MeshData.MaxNumBonesPerVertex, GenerationContext.CompilationContext.Options.TargetPlatform) &&
-										MeshData.MaxNumBonesPerVertex < (int32)ECustomizableObjectNumBoneInfluences ?
-											MeshData.MaxNumBonesPerVertex :
-											(int32)ECustomizableObjectNumBoneInfluences;
-
-	ensure(MutableBonesPerVertex <= MAX_TOTAL_INFLUENCES);
-	
-	if (MutableBonesPerVertex != MeshData.MaxNumBonesPerVertex)
-	{
-		UE_LOG(LogMutable, Verbose, TEXT("In object [%s] Mesh bone number adjusted from %d to %d."), *GenerationContext.GetObjectName(), MeshData.MaxNumBonesPerVertex, MutableBonesPerVertex);
-	}
-
-	int MutableBufferCount = MUTABLE_VERTEXBUFFER_TEXCOORDS + 1;
-	if (MeshData.bHasVertexColors)
-	{
-		++MutableBufferCount;
-	}
-
-	if (MeshData.MaxNumBonesPerVertex > 0 && MeshData.MaxBoneIndexTypeSizeBytes > 0)
-	{
-		++MutableBufferCount;
-	}
-
-	if (MeshData.bHasRealTimeMorphs)
-	{
-		MutableBufferCount += 2;
-	}
-
-	if (MeshData.bHasClothing)
-	{
-		MutableBufferCount += 2;
-	}
-
-	MutableBufferCount += MeshData.SkinWeightProfilesSemanticIndices.Num();
-
-	OutVertexBufferFormat.SetBufferCount(MutableBufferCount);
-
-	int32 CurrentVertexBuffer = 0;
-
-	// Vertex buffer
-	MutableMeshBufferUtils::SetupVertexPositionsBuffer(CurrentVertexBuffer, OutVertexBufferFormat);
-	++CurrentVertexBuffer;
-
-	// Tangent buffer
-	MutableMeshBufferUtils::SetupTangentBuffer(CurrentVertexBuffer, OutVertexBufferFormat);
-	++CurrentVertexBuffer;
-
-	// Texture coords buffer
-	MutableMeshBufferUtils::SetupTexCoordinatesBuffer(CurrentVertexBuffer, MeshData.NumTexCoordChannels, OutVertexBufferFormat);
-	++CurrentVertexBuffer;
-
-	// Skin buffer
-	if (MeshData.MaxNumBonesPerVertex > 0 && MeshData.MaxBoneIndexTypeSizeBytes > 0)
-	{
-		const int32 MaxBoneWeightTypeSizeBytes = bWith16BitWeights ? 2 : 1;
-		MutableMeshBufferUtils::SetupSkinBuffer(CurrentVertexBuffer, MeshData.MaxBoneIndexTypeSizeBytes, MaxBoneWeightTypeSizeBytes, MutableBonesPerVertex, OutVertexBufferFormat);
-		++CurrentVertexBuffer;
-	}
-
-	// Colour buffer
-	if (MeshData.bHasVertexColors)
-	{
-		MutableMeshBufferUtils::SetupVertexColorBuffer(CurrentVertexBuffer, OutVertexBufferFormat);
-		++CurrentVertexBuffer;
-	}
-
-	// MorphTarget vertex tracking info buffers
-	if (MeshData.bHasRealTimeMorphs)
-	{
-		using namespace mu;
-		{
-			const int32 ElementSize = sizeof(uint32);
-			constexpr int32 ChannelCount = 1;
-			const EMeshBufferSemantic Semantics[ChannelCount] = { EMeshBufferSemantic::Other };
-			const int32 SemanticIndices[ChannelCount] = { 0 };
-			const EMeshBufferFormat Formats[ChannelCount] = { EMeshBufferFormat::UInt32 };
-			const int32 Components[ChannelCount] = { 1 };
-			const int32 Offsets[ChannelCount] = { 0 };
-
-			OutVertexBufferFormat.SetBuffer(CurrentVertexBuffer, ElementSize, ChannelCount, Semantics, SemanticIndices, Formats, Components, Offsets);
-			++CurrentVertexBuffer;
-		}
-
-		{
-			const int32 ElementSize = sizeof(uint32);
-			constexpr int32 ChannelCount = 1;
-			const EMeshBufferSemantic Semantics[ChannelCount] = { EMeshBufferSemantic::Other };
-			const int32 SemanticIndices[ChannelCount] = { 1 };
-			const EMeshBufferFormat Formats[ChannelCount] = { EMeshBufferFormat::UInt32 };
-			const int32 Components[ChannelCount] = { 1 };
-			const int32 Offsets[ChannelCount] = { 0 };
-
-			OutVertexBufferFormat.SetBuffer(CurrentVertexBuffer, ElementSize, ChannelCount, Semantics, SemanticIndices, Formats, Components, Offsets);
-			++CurrentVertexBuffer;
-		}
-	}
-
-	//Clothing Data Buffer.
-	if (MeshData.bHasClothing)
-	{
-		{
-			using namespace mu;
-			const int32 ElementSize = sizeof(int32);
-			constexpr int32 ChannelCount = 1;
-			const EMeshBufferSemantic Semantics[ChannelCount] = { EMeshBufferSemantic::Other };
-			const int32 SemanticIndices[ChannelCount] = { 2 };
-			const EMeshBufferFormat Formats[ChannelCount] = { EMeshBufferFormat::Int32 };
-			const int32 Components[ChannelCount] = { 1 };
-			const int32 Offsets[ChannelCount] = { 0 };
-
-			OutVertexBufferFormat.SetBuffer(CurrentVertexBuffer, ElementSize, ChannelCount, Semantics, SemanticIndices, Formats, Components, Offsets);
-			++CurrentVertexBuffer;
-		}
-
-		{
-			using namespace mu;
-			const int32 ElementSize = sizeof(uint32);
-			constexpr int32 ChannelCount = 1;
-			const EMeshBufferSemantic Semantics[ChannelCount] = { EMeshBufferSemantic::Other };
-			const int32 SemanticIndices[ChannelCount] = { 3 };
-			const EMeshBufferFormat Formats[ChannelCount] = { EMeshBufferFormat::UInt32 };
-			const int32 Components[ChannelCount] = { 1 };
-			const int32 Offsets[ChannelCount] = { 0 };
-
-			OutVertexBufferFormat.SetBuffer(CurrentVertexBuffer, ElementSize, ChannelCount, Semantics, SemanticIndices, Formats, Components, Offsets);
-			++CurrentVertexBuffer;
-		}
-	}
-
-	for (int32 ProfileSemanticIndex : MeshData.SkinWeightProfilesSemanticIndices)
-	{
-		MutableMeshBufferUtils::SetupSkinWeightProfileBuffer(CurrentVertexBuffer, MeshData.MaxBoneIndexTypeSizeBytes, 1, MutableBonesPerVertex, ProfileSemanticIndex, OutVertexBufferFormat);
-		++CurrentVertexBuffer;
-	}
-
-	// Index buffer
-	MutableMeshBufferUtils::SetupIndexBuffer(OutIndexBufferFormat);
-}
-
-
 mu::Ptr<mu::NodeSurface> GenerateMutableSourceSurface(const UEdGraphPin * Pin, FMutableGraphGenerationContext & GenerationContext)
 {
 	MUTABLE_CPUPROFILER_SCOPE(GenerateMutableSourceSurface);
@@ -338,26 +191,13 @@ mu::Ptr<mu::NodeSurface> GenerateMutableSourceSurface(const UEdGraphPin * Pin, F
 
 				GenerationContext.LayoutGenerationFlags.Pop();
 
-				if (MeshNode)
+				if (!MeshNode)
 				{
-					mu::Ptr<mu::NodeMeshFormat> MeshFormatNode = new mu::NodeMeshFormat();
-					MeshFormatNode->SetSource(MeshNode.get());
-					SetSurfaceFormat(GenerationContext,
-						MeshFormatNode->GetVertexBuffers(), MeshFormatNode->GetIndexBuffers(), MeshData,
-						GenerationContext.CompilationContext.Options.CustomizableObjectNumBoneInfluences,
-						GenerationContext.CompilationContext.Options.b16BitBoneWeightsEnabled);
-
-					// \TODO: Make it an option?
-					MeshFormatNode->SetOptimizeBuffers(true);
-
-					MeshFormatNode->SetMessageContext(Node);
-
-					SurfNode->ExternalId = SurfaceMetadataUniqueHash;
-					SurfNode->Mesh = MeshFormatNode;
+					GenerationContext.Log(LOCTEXT("MeshFailed", "Mesh generation failed."), Node);
 				}
 				else
 				{
-					GenerationContext.Log(LOCTEXT("MeshFailed", "Mesh generation failed."), Node);
+					SurfNode->Mesh = MeshNode;
 				}
 			}
 		}
@@ -765,7 +605,7 @@ mu::Ptr<mu::NodeSurface> GenerateMutableSourceSurface(const UEdGraphPin * Pin, F
 										if (bUseLA)
 										{
 											// We'll have to trust the reference texture because the actual internal settings are opaque.
-											// See GetQualityFormat in TextureFormatASTC.cpp to understand how it works, but it depoends on some inaccessible texture name constants like this:
+											// See GetQualityFormat in TextureFormatASTC.cpp to understand how it works, but it depends on some inaccessible texture name constants like this:
 											//bool bIsNormalMapFormat = TextureFormatName == GTextureFormatNameASTC_NormalAG ||
 											//		TextureFormatName == GTextureFormatNameASTC_NormalRG ||
 											//		TextureFormatName == GTextureFormatNameASTC_NormalLA ||
diff --git a/Engine/Plugins/Mutable/Source/CustomizableObjectEditor/Private/MuCOE/GenerateMutableSource/GenerateMutableSourceSurface.h b/Engine/Plugins/Mutable/Source/CustomizableObjectEditor/Private/MuCOE/GenerateMutableSource/GenerateMutableSourceSurface.h
--- a/Engine/Plugins/Mutable/Source/CustomizableObjectEditor/Private/MuCOE/GenerateMutableSource/GenerateMutableSourceSurface.h
+++ b/Engine/Plugins/Mutable/Source/CustomizableObjectEditor/Private/MuCOE/GenerateMutableSource/GenerateMutableSourceSurface.h
@@ -17,7 +17,3 @@ namespace mu
 
 extern mu::Ptr<mu::NodeSurface> GenerateMutableSourceSurface(const UEdGraphPin* Pin, FMutableGraphGenerationContext & GenerationContext);
 
-extern void SetSurfaceFormat(FMutableGraphGenerationContext& GenerationContext,
-	mu::FMeshBufferSet& OutVertexBufferFormat, mu::FMeshBufferSet& OutIndexBufferFormat, const FMutableGraphMeshGenerationData& MeshData,
-	ECustomizableObjectNumBoneInfluences ECustomizableObjectNumBoneInfluences, bool bWith16BitWeights);
-
diff --git a/Engine/Plugins/Mutable/Source/MutableRuntime/Internal/MuR/OpMeshFormat.h b/Engine/Plugins/Mutable/Source/MutableRuntime/Internal/MuR/OpMeshFormat.h
--- a/Engine/Plugins/Mutable/Source/MutableRuntime/Internal/MuR/OpMeshFormat.h
+++ b/Engine/Plugins/Mutable/Source/MutableRuntime/Internal/MuR/OpMeshFormat.h
@@ -32,6 +32,6 @@ namespace mu
     void MeshFormatBuffer( const FMeshBufferSet& Source, FMeshBuffer& Result, int32 ResultElementOffset, bool bHasSpecialSemantics, uint32 IDPrefix);
 
 	/** Try to reduce the mesh size by reducing the component count, and data type of some buffers. */
-	void MeshOptimizeBuffers( FMesh* );
+	MUTABLERUNTIME_API void MeshOptimizeBuffers( FMesh* );
 
 }
diff --git a/Engine/Plugins/Mutable/Source/MutableRuntime/Private/MuR/OpMeshDifference.h b/Engine/Plugins/Mutable/Source/MutableRuntime/Private/MuR/OpMeshDifference.h
--- a/Engine/Plugins/Mutable/Source/MutableRuntime/Private/MuR/OpMeshDifference.h
+++ b/Engine/Plugins/Mutable/Source/MutableRuntime/Private/MuR/OpMeshDifference.h
@@ -77,13 +77,12 @@ namespace mu
 
 
 		// Make a delta of every vertex
-		// TODO: Not always 4 components
 		// TODO: Not always floats
 		int32 differentVertexCount = 0;
 		TArray<bool> isVertexDifferent;
 		isVertexDifferent.SetNumZeroed(VertexCount);
-		TArray< FVector4f > deltas;
-		deltas.SetNum(VertexCount * numChannels);
+		TArray< FVector3f > deltas;
+		deltas.SetNumZeroed(VertexCount * numChannels);
 
 		for ( int32 c=0; c<numChannels; ++c )
 		{
@@ -95,12 +94,12 @@ namespace mu
 
 			for ( int32 v=0; v<VertexCount; ++v )
 			{
-				FVector4f base = baseIt.GetAsVec4f();
-				FVector4f target = targetIt.GetAsVec4f();
-				FVector4f delta = target-base;
+				FVector3f base = baseIt.GetAsVec3f();
+				FVector3f target = targetIt.GetAsVec3f();
+				FVector3f delta = target-base;
 				deltas[ v*numChannels+c ] = delta;		
 
-				if (!delta.Equals(FVector4f(0,0,0,0),UE_SMALL_NUMBER))
+				if (!delta.Equals(FVector3f(0,0,0),UE_SMALL_NUMBER))
 				{
 					if ( !isVertexDifferent[v] )
 					{
@@ -113,7 +112,7 @@ namespace mu
 				++baseIt;		
 				++targetIt;		
 			}
-		}
+		} 
 
 
 		// Create the morph mesh
@@ -134,9 +133,9 @@ namespace mu
 				semantic.Add(semantics[c]);
 				semanticIndex.Add(semanticIndices[c]);
 				format.Add(EMeshBufferFormat::Float32);
-				components.Add(4);
+				components.Add(3);
 				offsets.Add(ElementSize);
-				ElementSize += 4 * sizeof(float);
+				ElementSize += 3 * sizeof(float);
 			}
 
 			int32 BufferIndex = 0;
@@ -192,7 +191,7 @@ namespace mu
 					// all channels
 					for ( int32 c=0; c<numChannels; ++c )
 					{
-						constexpr int32 ChannelSizeBytes = 16;
+						constexpr int32 ChannelSizeBytes = 3*sizeof(float);
 						FMemory::Memcpy(ResultData, &deltas[v * numChannels + c], ChannelSizeBytes);
 						ResultData += ChannelSizeBytes;
 					}
diff --git a/Engine/Plugins/Mutable/Source/MutableRuntime/Private/MuR/OpMeshFormat.cpp b/Engine/Plugins/Mutable/Source/MutableRuntime/Private/MuR/OpMeshFormat.cpp
--- a/Engine/Plugins/Mutable/Source/MutableRuntime/Private/MuR/OpMeshFormat.cpp
+++ b/Engine/Plugins/Mutable/Source/MutableRuntime/Private/MuR/OpMeshFormat.cpp
@@ -52,11 +52,11 @@ namespace mu
 			int32 SourceChannelIndex;
 			Source.FindChannel(ResultChannel.Semantic, ResultChannel.SemanticIndex, &SourceBufferIndex, &SourceChannelIndex);
 
-			int32 resultElemSize = ResultBuffer.ElementSize;
-			int32 resultComponents = ResultChannel.ComponentCount;
-			int32 resultChannelSize = GetMeshFormatData(ResultChannel.Format).SizeInBytes * resultComponents;
-			uint8* pResultBuf = ResultBuffer.Data.GetData() + ResultOffsetElements*ResultBuffer.ElementSize;
-			pResultBuf += ResultChannel.Offset;
+			int32 ResultElemSize = ResultBuffer.ElementSize;
+			int32 ResultComponents = ResultChannel.ComponentCount;
+			int32 ResultChannelSize = GetMeshFormatData(ResultChannel.Format).SizeInBytes * ResultComponents;
+			uint8* ResultBuf = ResultBuffer.Data.GetData() + ResultOffsetElements*ResultBuffer.ElementSize;
+			ResultBuf += ResultChannel.Offset;
 
 			// Case 1: Special semantics that may be implicit or relative
 			//---------------------------------------------------------------------------------
@@ -74,7 +74,7 @@ namespace mu
 					if (bHasSameFormat)
 					{
 						check(SourceChannel == ResultChannel);
-						FMemory::Memcpy(pResultBuf,SourceBuffer.Data.GetData(),SourceBuffer.Data.Num());
+						FMemory::Memcpy(ResultBuf,SourceBuffer.Data.GetData(),SourceBuffer.Data.Num());
 					}
 					else
 					{
@@ -84,7 +84,7 @@ namespace mu
 						const uint32* SourceData = reinterpret_cast<const uint32*>(SourceBuffer.Data.GetData());
 
 						check(ResultChannel.Format == EMeshBufferFormat::UInt64);
-						uint64* ResultData = reinterpret_cast<uint64*>(pResultBuf);
+						uint64* ResultData = reinterpret_cast<uint64*>(ResultBuf);
 
 						for (int32 i = 0; i < SourceElementCount; ++i)
 						{
@@ -101,7 +101,7 @@ namespace mu
 					check(IDPrefix);
 
 					check(ResultChannel.Format == EMeshBufferFormat::UInt64);
-					uint64* ResultData = reinterpret_cast<uint64*>(pResultBuf);
+					uint64* ResultData = reinterpret_cast<uint64*>(ResultBuf);
 
 					for (int32 VertexIndex = 0; VertexIndex < SourceElementCount; ++VertexIndex)
 					{
@@ -127,7 +127,7 @@ namespace mu
 					if (bHasSameFormat)
 					{
 						check(SourceChannel==ResultChannel);
-						FMemory::Memcpy(pResultBuf, SourceBuffer.Data.GetData(), SourceBuffer.Data.Num());
+						FMemory::Memcpy(ResultBuf, SourceBuffer.Data.GetData(), SourceBuffer.Data.Num());
 					}
 					else
 					{
@@ -136,7 +136,7 @@ namespace mu
 						const uint16* SourceData = reinterpret_cast<const uint16*>(SourceBuffer.Data.GetData());
 
 						check(ResultChannel.Format == EMeshBufferFormat::UInt64);
-						uint64* ResultData = reinterpret_cast<uint64*>(pResultBuf);
+						uint64* ResultData = reinterpret_cast<uint64*>(ResultBuf);
 
 						for (int32 i = 0; i < SourceElementCount; ++i)
 						{
@@ -178,12 +178,12 @@ namespace mu
 					{
 						for (int32 v = 0; v < SourceElementCount; ++v)
 						{
-							float* pTypedResultBuf = (float*)pResultBuf;
-							for (int32 i = 0; i < resultComponents; ++i)
+							float* pTypedResultBuf = (float*)ResultBuf;
+							for (int32 i = 0; i < ResultComponents; ++i)
 							{
 								pTypedResultBuf[i] = 1.0f;
 							}
-							pResultBuf += resultElemSize;
+							ResultBuf += ResultElemSize;
 						}
 						break;
 					}
@@ -192,12 +192,12 @@ namespace mu
 					{
 						for (int32 v = 0; v < SourceElementCount; ++v)
 						{
-							uint8* pTypedResultBuf = (uint8*)pResultBuf;
-							for (int32 i = 0; i < resultComponents; ++i)
+							uint8* pTypedResultBuf = (uint8*)ResultBuf;
+							for (int32 i = 0; i < ResultComponents; ++i)
 							{
 								pTypedResultBuf[i] = 255;
 							}
-							pResultBuf += resultElemSize;
+							ResultBuf += ResultElemSize;
 						}
 						break;
 					}
@@ -206,12 +206,12 @@ namespace mu
 					{
 						for (int32 v = 0; v < SourceElementCount; ++v)
 						{
-							uint16* pTypedResultBuf = (uint16*)pResultBuf;
-							for (int32 i = 0; i < resultComponents; ++i)
+							uint16* pTypedResultBuf = (uint16*)ResultBuf;
+							for (int32 i = 0; i < ResultComponents; ++i)
 							{
 								pTypedResultBuf[i] = 65535;
 							}
-							pResultBuf += resultElemSize;
+							ResultBuf += ResultElemSize;
 						}
 						break;
 					}
@@ -228,8 +228,8 @@ namespace mu
 					// TODO: and maybe raise a warning?
 					for (int32 v = 0; v < SourceElementCount; ++v)
 					{
-						FMemory::Memzero(pResultBuf, resultChannelSize);
-						pResultBuf += resultElemSize;
+						FMemory::Memzero(ResultBuf, ResultChannelSize);
+						ResultBuf += ResultElemSize;
 					}
 				}
 
@@ -241,105 +241,116 @@ namespace mu
 			//---------------------------------------------------------------------------------
 			{
 				// Get the data about the source format
-				EMeshBufferSemantic sourceSemantic;
-				int32 sourceSemanticIndex;
-				EMeshBufferFormat sourceFormat;
-				int32 sourceComponents;
-				int32 sourceOffset;
+				EMeshBufferSemantic SourceSemantic;
+				int32 SourceSemanticIndex;
+				EMeshBufferFormat SourceFormat;
+				int32 SourceComponents;
+				int32 SourceOffset;
 				Source.GetChannel
 				(
 					SourceBufferIndex, SourceChannelIndex,
-					&sourceSemantic, &sourceSemanticIndex,
-					&sourceFormat, &sourceComponents,
-					&sourceOffset
+					&SourceSemantic, &SourceSemanticIndex,
+					&SourceFormat, &SourceComponents,
+					&SourceOffset
 				);
-				check(sourceSemantic == ResultChannel.Semantic
+				check(SourceSemantic == ResultChannel.Semantic
 					&&
-					sourceSemanticIndex == ResultChannel.SemanticIndex);
+					SourceSemanticIndex == ResultChannel.SemanticIndex);
 
-				int32 sourceElemSize = Source.GetElementSize(SourceBufferIndex);
-				const uint8* pSourceBuf = Source.GetBufferData(SourceBufferIndex);
-				pSourceBuf += sourceOffset;
+				int32 SourceElemSize = Source.GetElementSize(SourceBufferIndex);
+				const uint8* SourceBuf = Source.GetBufferData(SourceBufferIndex);
+				SourceBuf += SourceOffset;
 
 				// Copy element by element
-				for (int32 v = 0; v < SourceElementCount; ++v)
+				if (ResultChannel.Format == SourceFormat && ResultComponents == SourceComponents)
 				{
-					if (ResultChannel.Format == sourceFormat && resultComponents == sourceComponents)
+					for (int32 v = 0; v < SourceElementCount; ++v)
 					{
-						FMemory::Memcpy(pResultBuf, pSourceBuf, resultChannelSize);
+						FMemory::Memcpy(ResultBuf, SourceBuf, ResultChannelSize);
+						ResultBuf += ResultElemSize;
+						SourceBuf += SourceElemSize;
 					}
-					else if (ResultChannel.Format == EMeshBufferFormat::PackedDir8_W_TangentSign
-						||
-						ResultChannel.Format == EMeshBufferFormat::PackedDirS8_W_TangentSign)
-					{
-						check(sourceComponents >= 3);
-						check(resultComponents == 4);
+				}
+				else if (ResultChannel.Format == EMeshBufferFormat::PackedDir8_W_TangentSign
+					||
+					ResultChannel.Format == EMeshBufferFormat::PackedDirS8_W_TangentSign)
+				{
+					check(SourceComponents >= 3);
+					check(ResultComponents == 4);
+
+					// Look for the full tangent space if not done already
+					int32 tanXBuf = -1, tanXChan = -1, tanYBuf = -1, tanYChan = -1, tanZBuf = -1, tanZChan = -1;
+					Source.FindChannel(EMeshBufferSemantic::Tangent, ResultChannel.SemanticIndex, &tanXBuf, &tanXChan);
+					Source.FindChannel(EMeshBufferSemantic::Binormal, ResultChannel.SemanticIndex, &tanYBuf, &tanYChan);
+					Source.FindChannel(EMeshBufferSemantic::Normal, ResultChannel.SemanticIndex, &tanZBuf, &tanZChan);
 
+					UntypedMeshBufferIteratorConst xIt(Source, EMeshBufferSemantic::Tangent, ResultChannel.SemanticIndex);
+					UntypedMeshBufferIteratorConst yIt(Source, EMeshBufferSemantic::Binormal, ResultChannel.SemanticIndex);
+					UntypedMeshBufferIteratorConst zIt(Source, EMeshBufferSemantic::Normal, ResultChannel.SemanticIndex);
+
+					for (int32 v = 0; v < SourceElementCount; ++v)
+					{
 						// convert the 3 first components
 						for (int32 i = 0; i < 3; ++i)
 						{
-							if (i < sourceComponents)
+							if (i < SourceComponents)
 							{
 								ConvertData
 								(
 									i,
-									pResultBuf, ResultChannel.Format,
-									pSourceBuf, sourceFormat
+									ResultBuf, ResultChannel.Format,
+									SourceBuf, SourceFormat
 								);
 							}
 						}
 
-
 						// Add the tangent sign
-						uint8* pData = reinterpret_cast<uint8*>(pResultBuf);
-
-						// Look for the full tangent space
-						int32 tanXBuf, tanXChan, tanYBuf, tanYChan, tanZBuf, tanZChan;
-						Source.FindChannel(EMeshBufferSemantic::Tangent, ResultChannel.SemanticIndex, &tanXBuf, &tanXChan);
-						Source.FindChannel(EMeshBufferSemantic::Binormal, ResultChannel.SemanticIndex, &tanYBuf, &tanYChan);
-						Source.FindChannel(EMeshBufferSemantic::Normal, ResultChannel.SemanticIndex, &tanZBuf, &tanZChan);
-
 						if (tanXBuf >= 0 && tanYBuf >= 0 && tanZBuf >= 0)
 						{
-							UntypedMeshBufferIteratorConst xIt(Source, EMeshBufferSemantic::Tangent, ResultChannel.SemanticIndex);
-							UntypedMeshBufferIteratorConst yIt(Source, EMeshBufferSemantic::Binormal, ResultChannel.SemanticIndex);
-							UntypedMeshBufferIteratorConst zIt(Source, EMeshBufferSemantic::Normal, ResultChannel.SemanticIndex);
-
-							xIt += v;
-							yIt += v;
-							zIt += v;
-
 							FMatrix44f Mat(xIt.GetAsVec3f(), yIt.GetAsVec3f(), zIt.GetAsVec3f(), FVector3f(0, 0, 0));
 
-							uint8 sign = 0;
 							if (ResultChannel.Format == EMeshBufferFormat::PackedDir8_W_TangentSign)
 							{
-								sign = Mat.RotDeterminant() < 0 ? 0 : 255;
+								uint8* pData = reinterpret_cast<uint8*>(ResultBuf);
+								uint8 Sign = Mat.RotDeterminant() < 0 ? 0 : 255;
+								pData[3] = Sign;
 							}
 							else if (ResultChannel.Format == EMeshBufferFormat::PackedDirS8_W_TangentSign)
 							{
-								sign = Mat.RotDeterminant() < 0 ? -128 : 127;
+								int8* pData = reinterpret_cast<int8*>(ResultBuf);
+								int8 Sign = Mat.RotDeterminant() < 0 ? -128 : 127;
+								pData[3] = Sign;
 							}
-							pData[3] = sign;
+
+							xIt++;
+							yIt++;
+							zIt++;
 						}
 						else
 						{
 							// At least initialize it to avoid garbage.
+							int8* pData = reinterpret_cast<int8*>(ResultBuf);
 							pData[3] = 0;
 						}
+						ResultBuf += ResultElemSize;
+						SourceBuf += SourceElemSize;
+
 					}
-					else
+				}
+				else
+				{
+					for (int32 v = 0; v < SourceElementCount; ++v)
 					{
 						// Convert formats
-						for (int32 i = 0; i < resultComponents; ++i)
+						for (int32 i = 0; i < ResultComponents; ++i)
 						{
-							if (i < sourceComponents)
+							if (i < SourceComponents)
 							{
 								ConvertData
 								(
 									i,
-									pResultBuf, ResultChannel.Format,
-									pSourceBuf, sourceFormat
+									ResultBuf, ResultChannel.Format,
+									SourceBuf, SourceFormat
 								);
 							}
 							else
@@ -347,7 +358,7 @@ namespace mu
 								// Add zeros. TODO: Warning?
 								FMemory::Memzero
 								(
-									pResultBuf + GetMeshFormatData(ResultChannel.Format).SizeInBytes * i,
+									ResultBuf + GetMeshFormatData(ResultChannel.Format).SizeInBytes * i,
 									GetMeshFormatData(ResultChannel.Format).SizeInBytes
 								);
 							}
@@ -357,34 +368,34 @@ namespace mu
 						// Extra step to normalise some semantics in some formats
 						// TODO: Make it optional, and add different normalisation types n, n^2
 						// TODO: Optimise
-						if (sourceSemantic == EMeshBufferSemantic::BoneWeights)
+						if (SourceSemantic == EMeshBufferSemantic::BoneWeights)
 						{
 							if (ResultChannel.Format == EMeshBufferFormat::NUInt8)
 							{
-								uint8* pData = (uint8*)pResultBuf;
-								uint8 accum = 0;
-								for (int32 i = 0; i < resultComponents; ++i)
+								uint8* Data = (uint8*)ResultBuf;
+								uint8 Accum = 0;
+								for (int32 i = 0; i < ResultComponents; ++i)
 								{
-									accum += pData[i];
+									Accum += Data[i];
 								}
-								pData[0] += 255 - accum;
+								Data[0] += 255 - Accum;
 							}
 
 							else if (ResultChannel.Format == EMeshBufferFormat::NUInt16)
 							{
-								uint16* pData = (uint16*)pResultBuf;
-								uint16 accum = 0;
-								for (int32 i = 0; i < resultComponents; ++i)
+								uint16* Data = (uint16*)ResultBuf;
+								uint16 Accum = 0;
+								for (int32 i = 0; i < ResultComponents; ++i)
 								{
-									accum += pData[i];
+									Accum += Data[i];
 								}
-								pData[0] += 65535 - accum;
+								Data[0] += 65535 - Accum;
 							}
 						}
-					}
 
-					pResultBuf += resultElemSize;
-					pSourceBuf += sourceElemSize;
+						ResultBuf += ResultElemSize;
+						SourceBuf += SourceElemSize;
+					}
 				}
 			}
 		}
@@ -672,6 +683,8 @@ namespace mu
 
 	void MeshOptimizeBuffers( FMesh* InMesh )
 	{
+		MUTABLE_CPUPROFILER_SCOPE(MeshOptimizeBuffers);
+
 		if (!InMesh)
 		{
 			return;
diff --git a/Engine/Plugins/Mutable/Source/MutableRuntime/Public/MuR/MeshBufferSet.h b/Engine/Plugins/Mutable/Source/MutableRuntime/Public/MuR/MeshBufferSet.h
--- a/Engine/Plugins/Mutable/Source/MutableRuntime/Public/MuR/MeshBufferSet.h
+++ b/Engine/Plugins/Mutable/Source/MutableRuntime/Public/MuR/MeshBufferSet.h
@@ -137,27 +137,18 @@ namespace mu
 	/** */
 	struct FMeshBufferChannel
 	{
-		FMeshBufferChannel()
-		{
-			Semantic = EMeshBufferSemantic::None;
-			Format = EMeshBufferFormat::None;
-			SemanticIndex = 0;
-			Offset = 0;
-			ComponentCount = 0;
-		}
-
-		EMeshBufferSemantic Semantic;
+		EMeshBufferSemantic Semantic = EMeshBufferSemantic::None;
 
-		EMeshBufferFormat Format;
+		EMeshBufferFormat Format = EMeshBufferFormat::None;
 
 		/** Index of the semantic, in case there are more than one of this type. */
-		int32 SemanticIndex;
+		int32 SemanticIndex = 0;
 
 		/** Offset in bytes from the begining of a buffer element */
-		uint16 Offset;
+		uint16 Offset = 0;
 
 		/** Number of components of the type in Format for every value in the channel */
-		uint16 ComponentCount;
+		uint16 ComponentCount = 0;
 
 		inline bool operator==(const FMeshBufferChannel& Other) const
 		{
diff --git a/Engine/Plugins/Mutable/Source/MutableTools/Private/MuT/CodeGenerator_Mesh.cpp b/Engine/Plugins/Mutable/Source/MutableTools/Private/MuT/CodeGenerator_Mesh.cpp
--- a/Engine/Plugins/Mutable/Source/MutableTools/Private/MuT/CodeGenerator_Mesh.cpp
+++ b/Engine/Plugins/Mutable/Source/MutableTools/Private/MuT/CodeGenerator_Mesh.cpp
@@ -196,6 +196,16 @@ namespace mu
 		int32 TexCoordsBufferIndex = -1;
 		int32 TexCoordsChannelIndex = -1;
 		Mesh->GetVertexBuffers().FindChannel(EMeshBufferSemantic::TexCoords, LayoutChannel, &TexCoordsBufferIndex, &TexCoordsChannelIndex);
+
+
+		if (TexCoordsBufferIndex<0 || TexCoordsChannelIndex<0)
+		{
+			// This is actually fine when using shared materials across LODs
+			FString Msg = FString::Printf(TEXT("LOD %d : Trying to generate layout for missing UV channel %d. No layout blocks generated."), StaticMeshOptions.LODIndex, LayoutChannel);
+			ErrorLog->Add(Msg, ELMT_INFO, ErrorContext);
+			return;
+		}
+
 		check(TexCoordsBufferIndex >= 0);
 		check(TexCoordsChannelIndex >= 0);
 
@@ -769,7 +779,8 @@ namespace mu
 				InOptions
 			);
 
-			BaseTask = GenerateMesh(StaticOptions, BaseOptionsTask, node.m_pBase );
+			FMeshGenerationStaticOptions BaseStaticOptions = StaticOptions;
+			BaseTask = GenerateMesh(BaseStaticOptions, BaseOptionsTask, node.m_pBase );
         }
 
         // Target
diff --git a/Engine/Plugins/Mutable/Source/MutableTools/Public/MuT/Compiler.h b/Engine/Plugins/Mutable/Source/MutableTools/Public/MuT/Compiler.h
--- a/Engine/Plugins/Mutable/Source/MutableTools/Public/MuT/Compiler.h
+++ b/Engine/Plugins/Mutable/Source/MutableTools/Public/MuT/Compiler.h
@@ -27,10 +27,15 @@ namespace mu
 	/** Arguments are Texture ID, output generated texture and a "run immediately" flag. */
 	typedef TFunction<UE::Tasks::FTask(int32, TSharedPtr<TSharedPtr<FImage>>, bool)> FReferencedImageResourceFunc;
 
-	/** Arguments are Mesh ID, optional morph name, output generated mesh and a "run immediately" flag. */
+	/** Arguments are 
+	* - Mesh ID, 
+	* - optional morph name, 
+	* - output generated mesh 
+	* - "run immediately" flag. 
+	*/
 	typedef TFunction<UE::Tasks::FTask(int32, const FString&, TSharedPtr<TSharedPtr<FMesh>>, bool)> FReferencedMeshResourceFunc;
 
-    //! \brief Options used to compile the models with a compiler.
+    /** Options used to compile the models with a compiler. */
     class MUTABLETOOLS_API CompilerOptions : public RefCounted
     {
     public:
