[Mutable] Added scope to LoadCompiledDataFromDisk.

#rnx

diff --git a/Engine/Plugins/Mutable/Source/CustomizableObject/Private/MuCO/CustomizableObject.cpp b/Engine/Plugins/Mutable/Source/CustomizableObject/Private/MuCO/CustomizableObject.cpp
--- a/Engine/Plugins/Mutable/Source/CustomizableObject/Private/MuCO/CustomizableObject.cpp
+++ b/Engine/Plugins/Mutable/Source/CustomizableObject/Private/MuCO/CustomizableObject.cpp
@@ -597,6 +597,8 @@ void UModelResources::InitCookData()
 
 void UCustomizableObjectPrivate::LoadCompiledDataFromDisk()
 {
+	MUTABLE_CPUPROFILER_SCOPE(UCustomizableObjectPrivate::LoadCompiledDataFromDisk)
+	
 	ITargetPlatformManagerModule& TargetPlatformManager = GetTargetPlatformManagerRef();
 	const ITargetPlatform* RunningPlatform = TargetPlatformManager.GetRunningTargetPlatform();
 	check(RunningPlatform);
