[Mutable] Move clothing physics clipping in a mutable external data manipulation plugin.

- Added IMutableClothingModule interface.
- Added new experimental MutabelClothing plugin implementing the MutabeClothingModule interface.

diff --git a/Engine/Plugins/Experimental/MutableClothing/MutableClothing.uplugin b/Engine/Plugins/Experimental/MutableClothing/MutableClothing.uplugin
new file mode 100644
--- /dev/null
+++ b/Engine/Plugins/Experimental/MutableClothing/MutableClothing.uplugin
@@ -0,0 +1,31 @@
+{
+	"FileVersion": 3,
+	"Version": 1,
+	"VersionName": "1.0",
+	"FriendlyName": "Mutable Clothing",
+	"Description": "Adds Mutable functionality to work with clothing",
+	"Category": "Mutable",
+	"CreatedBy": "Epic Games, Inc.",
+	"CreatedByURL": "https://epicgames.com",
+	"DocsURL": "",
+	"MarketplaceURL": "",
+	"SupportURL": "",
+	"CanContainContent": false,
+	"IsBetaVersion": false,
+	"Installed": false,
+	"EnabledByDefault": false,
+	"IsExperimentalVersion": true,
+	"Modules": [
+		{
+			"Name": "MutableClothing",
+			"Type": "Runtime",
+			"LoadingPhase": "Default"
+		}
+	],
+	"Plugins": [
+		{
+			"Name": "Mutable",
+			"Enabled": true
+		}
+	]
+}
diff --git a/Engine/Plugins/Experimental/MutableClothing/Source/MutableClothing/MutableClothing.Build.cs b/Engine/Plugins/Experimental/MutableClothing/Source/MutableClothing/MutableClothing.Build.cs
new file mode 100644
--- /dev/null
+++ b/Engine/Plugins/Experimental/MutableClothing/Source/MutableClothing/MutableClothing.Build.cs
@@ -0,0 +1,24 @@
+// Copyright Epic Games, Inc. All Rights Reserved.
+
+using UnrealBuildTool;
+
+public class MutableClothing : ModuleRules
+{
+	public MutableClothing(ReadOnlyTargetRules TargetRules) : base(TargetRules)
+	{
+		PrivateDependencyModuleNames.AddRange(
+			new string[]
+			{
+				"Core",
+				"CoreUObject",
+				"Engine",
+				"ClothingSystemRuntimeCommon",
+			});
+		
+        PublicDependencyModuleNames.AddRange(
+	        new string[]
+	        {
+		        "CustomizableObject",
+     		});
+    }
+}
diff --git a/Engine/Plugins/Experimental/MutableClothing/Source/MutableClothing/Private/MutableClothingModule.cpp b/Engine/Plugins/Experimental/MutableClothing/Source/MutableClothing/Private/MutableClothingModule.cpp
new file mode 100644
--- /dev/null
+++ b/Engine/Plugins/Experimental/MutableClothing/Source/MutableClothing/Private/MutableClothingModule.cpp
@@ -0,0 +1,429 @@
+// Copyright Epic Games, Inc. All Rights Reserved.
+
+
+#include "MutableClothingModule.h"
+
+#include "Modules/ModuleManager.h"
+#include "ClothingAsset.h"
+
+IMPLEMENT_MODULE(FMutableClothingModule, MutableClothing);
+
+
+void FMutableClothingModule::StartupModule()
+{
+}
+
+
+void FMutableClothingModule::ShutdownModule()
+{
+}
+
+bool FMutableClothingModule::UpdateClothSimulationLOD(
+		int32 InSimulationLODIndex,
+		UClothingAssetCommon& InOutClothingAsset,
+		TConstArrayView<TArrayView<FMeshToMeshVertData>> InOutAttachedLODsRenderData)
+{
+	TRACE_CPUPROFILER_EVENT_SCOPE(MutableClothing_UpdateClothSimulationLOD)
+
+	check(InOutClothingAsset.LodData.IsValidIndex(InSimulationLODIndex));
+
+	if (!InOutAttachedLODsRenderData.Num())
+	{
+		return false;
+	}
+
+	FClothLODDataCommon& ClothingSimData = InOutClothingAsset.LodData[InSimulationLODIndex];
+	const int32 NumSimulationVertices = ClothingSimData.PhysicalMeshData.Vertices.Num();
+	
+	TBitArray<> VertexRemoveMask;
+	VertexRemoveMask.Init(true, NumSimulationVertices);
+
+	for (int32 RenderDataIndex = 0; RenderDataIndex < InOutAttachedLODsRenderData.Num(); ++RenderDataIndex)
+	{
+		for (const FMeshToMeshVertData& VertClothData : InOutAttachedLODsRenderData[RenderDataIndex])
+		{
+			const uint16* Indices = VertClothData.SourceMeshVertIndices;
+
+			VertexRemoveMask[Indices[0]] = false;
+			VertexRemoveMask[Indices[1]] = false;
+			VertexRemoveMask[Indices[2]] = false;
+		}
+	}	
+
+	const int32 FirstRemoved = VertexRemoveMask.Find(true);
+	if (FirstRemoved == INDEX_NONE)
+	{
+		// No vertices removed.
+		return false;
+	}
+
+	TArray<int32> MeshToMeshIndexRemap;
+	MeshToMeshIndexRemap.SetNumUninitialized(NumSimulationVertices);
+	{
+		for (int32 Index = 0; Index < FirstRemoved; ++Index)
+		{
+			MeshToMeshIndexRemap[Index] = Index;
+		}
+		
+		int32 NewMeshVerticesNum = FirstRemoved;
+
+		for (int32 Index = FirstRemoved; Index < NumSimulationVertices; ++Index)
+		{
+			MeshToMeshIndexRemap[Index] = !VertexRemoveMask[Index] ? NewMeshVerticesNum++ : INDEX_NONE;
+		}
+
+		// Apply index remap to the render data. 
+		for (const TArrayView<FMeshToMeshVertData>& VertDataView : InOutAttachedLODsRenderData)
+		{
+			for (FMeshToMeshVertData& MeshToMeshVertData : VertDataView)
+			{
+				MeshToMeshVertData.SourceMeshVertIndices[0] = MeshToMeshIndexRemap[MeshToMeshVertData.SourceMeshVertIndices[0]];
+				MeshToMeshVertData.SourceMeshVertIndices[1] = MeshToMeshIndexRemap[MeshToMeshVertData.SourceMeshVertIndices[1]];
+				MeshToMeshVertData.SourceMeshVertIndices[2] = MeshToMeshIndexRemap[MeshToMeshVertData.SourceMeshVertIndices[2]];
+			}
+
+#if DO_CHECK
+			int32 NumUnmappedVertices = 0;
+			for (FMeshToMeshVertData& MeshToMeshVertData : VertDataView)
+			{
+				NumUnmappedVertices += static_cast<int32>( 
+						(MeshToMeshVertData.SourceMeshVertIndices[0] == INDEX_NONE) |
+						(MeshToMeshVertData.SourceMeshVertIndices[1] == INDEX_NONE) |
+						(MeshToMeshVertData.SourceMeshVertIndices[2] == INDEX_NONE) );
+			}
+
+			check(NumUnmappedVertices == 0);
+#endif
+		}
+
+		const auto CopyIfUsed = [&VertexRemoveMask](auto& Dst, const auto& Src)
+		{
+			const int32 SrcNumElems = Src.Num();
+			for (int32 Idx = 0, DstNumElems = 0; Idx < SrcNumElems; ++Idx)
+			{
+				if (!VertexRemoveMask[Idx])
+				{
+					Dst[DstNumElems++] = Src[Idx];
+				}
+			}
+		};
+
+		FClothLODDataCommon ResultClothingSimData;
+	
+		ResultClothingSimData.PhysicalMeshData.Vertices.SetNum(NewMeshVerticesNum);
+		ResultClothingSimData.PhysicalMeshData.Normals.SetNum(NewMeshVerticesNum);
+		ResultClothingSimData.PhysicalMeshData.BoneData.SetNum(NewMeshVerticesNum);
+		ResultClothingSimData.PhysicalMeshData.InverseMasses.SetNum(NewMeshVerticesNum);
+
+		CopyIfUsed(ResultClothingSimData.PhysicalMeshData.Vertices,      ClothingSimData.PhysicalMeshData.Vertices);
+		CopyIfUsed(ResultClothingSimData.PhysicalMeshData.Normals,       ClothingSimData.PhysicalMeshData.Normals);
+		CopyIfUsed(ResultClothingSimData.PhysicalMeshData.BoneData,      ClothingSimData.PhysicalMeshData.BoneData);
+		CopyIfUsed(ResultClothingSimData.PhysicalMeshData.InverseMasses, ClothingSimData.PhysicalMeshData.InverseMasses);
+
+		const TMap<uint32, FPointWeightMap>& SourceWeightMaps = ClothingSimData.PhysicalMeshData.WeightMaps;
+		TMap<uint32, FPointWeightMap>& DestWeightMaps = ResultClothingSimData.PhysicalMeshData.WeightMaps;
+
+		for (const TPair<uint32, FPointWeightMap>& WeightMap : SourceWeightMaps)
+		{
+			if (WeightMap.Value.Values.Num() > 0)
+			{
+				FPointWeightMap& NewWeightMap = ResultClothingSimData.PhysicalMeshData.AddWeightMap(WeightMap.Key);
+				NewWeightMap.Values.SetNum(NewMeshVerticesNum);
+
+				CopyIfUsed(NewWeightMap.Values, WeightMap.Value.Values);
+			}
+		}
+	
+		const TArray<int32>& IndexMap = MeshToMeshIndexRemap;
+		const auto TrimAndRemapTriangles = [&IndexMap](TArray<uint32>& Dst, const TArray<uint32>& Src) -> int32
+		{
+			check(Src.Num() % 3 == 0);
+
+			const int32 SrcNumElems = Src.Num();
+
+			int32 DstNumElems = 0;
+			for (int32 Idx = 0; Idx < SrcNumElems; Idx += 3)
+			{
+				const int32 Idx0 = IndexMap[Src[Idx + 0]];
+				const int32 Idx1 = IndexMap[Src[Idx + 1]];
+				const int32 Idx2 = IndexMap[Src[Idx + 2]];
+
+				// triangles are only copied if all vertices are used.
+				if (!((Idx0 < 0) | (Idx1 < 0) | (Idx2 < 0)))
+				{
+					Dst[DstNumElems + 0] = Idx0;
+					Dst[DstNumElems + 1] = Idx1;
+					Dst[DstNumElems + 2] = Idx2;
+
+					DstNumElems += 3;
+				}
+			}
+
+			return DstNumElems;
+		};
+
+		ResultClothingSimData.PhysicalMeshData.Indices.SetNum(ClothingSimData.PhysicalMeshData.Indices.Num());
+		const int32 NumRemainingIndices = TrimAndRemapTriangles(
+				ResultClothingSimData.PhysicalMeshData.Indices, 
+				ClothingSimData.PhysicalMeshData.Indices);
+		ResultClothingSimData.PhysicalMeshData.Indices.SetNum(NumRemainingIndices, EAllowShrinking::No);
+
+		const auto TrimAndRemapVertexSet = [&IndexMap](TSet<int32>& Dst, const TSet<int32>& Src)
+		{	
+			Dst.Reserve(Src.Num());
+			for(const int32 SrcIdx : Src)
+			{
+				const int32 MappedIdx = IndexMap[SrcIdx];
+
+				if (MappedIdx >= 0)
+				{
+					Dst.Add(MappedIdx);
+				}
+			}
+		};
+
+		TrimAndRemapVertexSet(
+				ResultClothingSimData.PhysicalMeshData.SelfCollisionVertexSet, 
+				ClothingSimData.PhysicalMeshData.SelfCollisionVertexSet);
+
+		{
+			TRACE_CPUPROFILER_EVENT_SCOPE(MutableClothing_BuildClothTetherData)
+
+			auto TrimAndRemapTethers = [&IndexMap](FClothTetherData& Dst, const FClothTetherData& Src)
+			{
+				Dst.Tethers.Reserve(Src.Tethers.Num());
+				for (const TArray<TTuple<int32, int32, float>>& SrcTetherCluster : Src.Tethers)
+				{
+					TArray<TTuple<int32, int32, float>>& DstTetherCluster = Dst.Tethers.Emplace_GetRef();
+					DstTetherCluster.Reserve(SrcTetherCluster.Num());
+					for (const TTuple<int32, int32, float>& Tether : SrcTetherCluster)
+					{
+						const int32 Index0 = IndexMap[Tether.Get<0>()];
+						const int32 Index1 = IndexMap[Tether.Get<1>()];
+						if ((Index0 >= 0) & (Index1 >= 0))
+						{
+							DstTetherCluster.Emplace(Index0, Index1, Tether.Get<2>());
+						}
+					}
+
+					if (!DstTetherCluster.Num())
+					{
+						Dst.Tethers.RemoveAt(Dst.Tethers.Num() - 1, 1, EAllowShrinking::No);
+					}
+				}
+			};
+
+			TrimAndRemapTethers(
+					ResultClothingSimData.PhysicalMeshData.GeodesicTethers, 
+					ClothingSimData.PhysicalMeshData.GeodesicTethers);
+
+			TrimAndRemapTethers(
+					ResultClothingSimData.PhysicalMeshData.EuclideanTethers, 
+					ClothingSimData.PhysicalMeshData.EuclideanTethers);
+		}
+
+		ClothingSimData.PhysicalMeshData = MoveTemp(ResultClothingSimData.PhysicalMeshData);
+
+		if (ClothingSimData.TransitionUpSkinData.Num())
+		{
+			ResultClothingSimData.TransitionUpSkinData.SetNum(NewMeshVerticesNum);	
+			CopyIfUsed(ResultClothingSimData.TransitionUpSkinData, ClothingSimData.TransitionUpSkinData);
+
+			ClothingSimData.TransitionUpSkinData = MoveTemp(ResultClothingSimData.TransitionUpSkinData);
+		}
+
+		if (ClothingSimData.TransitionDownSkinData.Num())
+		{
+			ResultClothingSimData.TransitionDownSkinData.SetNum(NewMeshVerticesNum);
+			CopyIfUsed(ResultClothingSimData.TransitionDownSkinData, ClothingSimData.TransitionDownSkinData);
+			
+			ClothingSimData.TransitionDownSkinData = MoveTemp(ResultClothingSimData.TransitionDownSkinData);
+		}
+
+		auto RemapTransitionMeshToMeshVertData = [](TArray<FMeshToMeshVertData>& InOutVertData, TArray<int32>& IndexMap)
+		{
+			for (FMeshToMeshVertData& VertData : InOutVertData)
+			{
+				uint16* Indices = VertData.SourceMeshVertIndices;
+				Indices[0] = (uint16)IndexMap[Indices[0]];
+				Indices[1] = (uint16)IndexMap[Indices[1]];
+				Indices[2] = (uint16)IndexMap[Indices[2]];
+			}
+		};
+
+		if (InOutClothingAsset.LodData.Num() > InSimulationLODIndex + 1)
+		{
+			TArray<FMeshToMeshVertData>& NextLodTransitionData = InOutClothingAsset.LodData[InSimulationLODIndex + 1].TransitionUpSkinData;
+
+			if (NextLodTransitionData.Num())
+			{
+				RemapTransitionMeshToMeshVertData(NextLodTransitionData, MeshToMeshIndexRemap);
+			}
+		}
+
+		if (InSimulationLODIndex - 1 >= 0)
+		{
+			TArray<FMeshToMeshVertData>& PrevLodTransitionData = InOutClothingAsset.LodData[InSimulationLODIndex - 1].TransitionDownSkinData;
+
+			if (PrevLodTransitionData.Num())
+			{
+				RemapTransitionMeshToMeshVertData(PrevLodTransitionData, MeshToMeshIndexRemap);
+			}
+		}
+	}
+
+	return true;
+}
+
+
+void FMutableClothingModule::FixLODTransitionMappings(
+		int32 SimulationLODIndex, 
+		UClothingAssetCommon& InOutClothingAsset)
+{
+	TRACE_CPUPROFILER_EVENT_SCOPE(MutableClothingModule_FixLODTransitionMappings)
+
+	check(InOutClothingAsset.LodData.IsValidIndex(SimulationLODIndex));
+
+	struct FMeshPhysicsDesc
+	{
+		const TArray<FVector3f>& Vertices;
+		const TArray<FVector3f>& Normals;
+		const TArray<uint32>& Indices;
+	};
+
+	auto RebindVertex = [](const FMeshPhysicsDesc& Mesh, const FVector3f& InPosition, const FVector3f& InNormal, FMeshToMeshVertData& Out)
+	{
+		const FVector3f Normal = InNormal;
+
+		// We don't have the mesh tangent, find something plausible.
+		FVector3f Tan0, Tan1;
+		Normal.FindBestAxisVectors(Tan0, Tan1);
+		const FVector3f Tangent = Tan0;
+		
+		// Some of the math functions take as argument FVector, we'd want to be FVector3f. 
+		// This should be changed once support for the single type in the FMath functions is added. 
+		const FVector Position = (FVector)InPosition;
+		int32 BestBaseTriangleIdx = INDEX_NONE;
+		FVector::FReal BestDistanceSq = TNumericLimits<FVector::FReal>::Max();
+		
+		const int32 NumIndices = Mesh.Indices.Num();
+		check(NumIndices % 3 == 0);
+
+		for (int32 I = 0; I < NumIndices; I += 3)
+		{
+			const FVector& A = (FVector)Mesh.Vertices[Mesh.Indices[I + 0]];
+			const FVector& B = (FVector)Mesh.Vertices[Mesh.Indices[I + 1]];
+			const FVector& C = (FVector)Mesh.Vertices[Mesh.Indices[I + 2]];
+
+			FVector ClosestTrianglePoint = FMath::ClosestPointOnTriangleToPoint(Position, (FVector)A, (FVector)B, (FVector)C);
+
+			const FVector::FReal CurrentDistSq = (ClosestTrianglePoint - Position).SizeSquared();
+			if (CurrentDistSq < BestDistanceSq)
+			{
+				BestDistanceSq = CurrentDistSq;
+				BestBaseTriangleIdx = I;
+			}
+		}
+
+		check(BestBaseTriangleIdx >= 0);
+
+		auto ComputeBaryCoordsAndDist = [](const FVector3f& A, const FVector3f& B, const FVector3f& C, const FVector3f& P) -> FVector4f
+		{
+			FPlane4f TrianglePlane(A, B, C);
+
+			const FVector3f PointOnTriPlane = FVector3f::PointPlaneProject(P, TrianglePlane);
+			const FVector3f BaryCoords = (FVector3f)FMath::ComputeBaryCentric2D((FVector)PointOnTriPlane, (FVector)A, (FVector)B, (FVector)C);
+
+			return FVector4f(BaryCoords, TrianglePlane.PlaneDot((FVector3f)P));
+		};
+
+		const FVector3f& A = Mesh.Vertices[Mesh.Indices[BestBaseTriangleIdx + 0]];
+		const FVector3f& B = Mesh.Vertices[Mesh.Indices[BestBaseTriangleIdx + 1]];
+		const FVector3f& C = Mesh.Vertices[Mesh.Indices[BestBaseTriangleIdx + 2]];
+
+		Out.PositionBaryCoordsAndDist = ComputeBaryCoordsAndDist(A, B, C, (FVector3f)Position );
+		Out.NormalBaryCoordsAndDist = ComputeBaryCoordsAndDist(A, B, C, (FVector3f)Position + Normal );
+		Out.TangentBaryCoordsAndDist = ComputeBaryCoordsAndDist(A, B, C, (FVector3f)Position + Tangent );
+		Out.SourceMeshVertIndices[0] = (uint16)Mesh.Indices[BestBaseTriangleIdx + 0];
+		Out.SourceMeshVertIndices[1] = (uint16)Mesh.Indices[BestBaseTriangleIdx + 1]; 
+		Out.SourceMeshVertIndices[2] = (uint16)Mesh.Indices[BestBaseTriangleIdx + 2];
+	};
+
+	auto RecreateTransitionData = [&RebindVertex]( 
+		const FMeshPhysicsDesc& ToMesh, const FMeshPhysicsDesc& FromMesh, TArray<FMeshToMeshVertData>& InOutTransitionData)
+	{
+		if (!InOutTransitionData.Num())
+		{
+			return;
+		}
+
+		const int32 TransitionDataNum = InOutTransitionData.Num();
+		
+		for (int32 I = 0; I < TransitionDataNum; ++I)
+		{
+			FMeshToMeshVertData& VertData = InOutTransitionData[I];
+			uint16* Indices = VertData.SourceMeshVertIndices;
+
+			// If any original indices are missing but the vertex is still alive rebind the vertex.
+			// In general, the number of rebinds should be small.
+
+			// Currently, if any index is missing we rebind to the closest triangle but it could be nice to use the remaining indices, 
+			// if any, to find the most appropriate triangle to bind to. 
+			const bool bNeedsRebind = (Indices[0] == 0xFFFF) | (Indices[1] == 0xFFFF) | (Indices[2] == 0xFFFF);
+
+			if (bNeedsRebind)
+			{
+				RebindVertex(ToMesh, FromMesh.Vertices[I], FromMesh.Normals[I], VertData);
+			}
+		}
+	};
+
+	const int32 NumLODs = InOutClothingAsset.LodData.Num();
+
+	const FMeshPhysicsDesc CurrentPhysicsMesh {
+			InOutClothingAsset.LodData[SimulationLODIndex].PhysicalMeshData.Vertices,
+			InOutClothingAsset.LodData[SimulationLODIndex].PhysicalMeshData.Normals,
+			InOutClothingAsset.LodData[SimulationLODIndex].PhysicalMeshData.Indices };
+
+	if (SimulationLODIndex + 1 < NumLODs && 
+		InOutClothingAsset.LodData[SimulationLODIndex].TransitionDownSkinData.Num())
+	{
+		const FMeshPhysicsDesc TransitionUpTarget {  
+				InOutClothingAsset.LodData[SimulationLODIndex + 1].PhysicalMeshData.Vertices,
+				InOutClothingAsset.LodData[SimulationLODIndex + 1].PhysicalMeshData.Normals,
+				InOutClothingAsset.LodData[SimulationLODIndex + 1].PhysicalMeshData.Indices };
+		
+		if (TransitionUpTarget.Vertices.Num())
+		{
+			RecreateTransitionData(
+					TransitionUpTarget, 
+					CurrentPhysicsMesh, 
+					InOutClothingAsset.LodData[SimulationLODIndex].TransitionDownSkinData);
+		}
+		else
+		{
+			InOutClothingAsset.LodData[SimulationLODIndex].TransitionDownSkinData.Empty();
+		}
+	}
+	
+	if (SimulationLODIndex - 1 >= 0 &&
+		InOutClothingAsset.LodData[SimulationLODIndex].TransitionUpSkinData.Num())
+	{	
+		FMeshPhysicsDesc TransitionDownTarget {  
+				InOutClothingAsset.LodData[SimulationLODIndex - 1].PhysicalMeshData.Vertices,
+				InOutClothingAsset.LodData[SimulationLODIndex - 1].PhysicalMeshData.Normals,
+				InOutClothingAsset.LodData[SimulationLODIndex - 1].PhysicalMeshData.Indices };
+
+		if (TransitionDownTarget.Vertices.Num())
+		{
+			RecreateTransitionData(
+					TransitionDownTarget, 
+					CurrentPhysicsMesh, 
+					InOutClothingAsset.LodData[SimulationLODIndex].TransitionUpSkinData);
+		}
+		else
+		{
+			InOutClothingAsset.LodData[SimulationLODIndex].TransitionUpSkinData.Empty();
+		}
+	}
+}
diff --git a/Engine/Plugins/Experimental/MutableClothing/Source/MutableClothing/Private/MutableClothingModule.h b/Engine/Plugins/Experimental/MutableClothing/Source/MutableClothing/Private/MutableClothingModule.h
new file mode 100644
--- /dev/null
+++ b/Engine/Plugins/Experimental/MutableClothing/Source/MutableClothing/Private/MutableClothingModule.h
@@ -0,0 +1,24 @@
+// Copyright Epic Games, Inc. All Rights Reserved.
+
+#pragma once
+
+#include "MuCO/Plugins/IMutableClothingModule.h"
+
+
+class MUTABLECLOTHING_API FMutableClothingModule : public IMutableClothingModule
+{
+public:
+	// IModuleInterface interface
+	virtual void StartupModule() override;
+	virtual void ShutdownModule() override;
+
+	// IMutableClothingModule interface
+	virtual bool UpdateClothSimulationLOD(
+			int32 InSimulationLODIndex,	
+			UClothingAssetCommon& InOutClothingAsset,
+			TConstArrayView<TArrayView<FMeshToMeshVertData>> InOutAttachedLODsRenderData) override;
+
+	virtual void FixLODTransitionMappings(
+			int32 InSimulationLODIndex, 
+			UClothingAssetCommon& InOutClothingAsset) override;
+};
diff --git a/Engine/Plugins/Mutable/Source/CustomizableObject/Private/MuCO/CustomizableObjectInstance.cpp b/Engine/Plugins/Mutable/Source/CustomizableObject/Private/MuCO/CustomizableObjectInstance.cpp
--- a/Engine/Plugins/Mutable/Source/CustomizableObject/Private/MuCO/CustomizableObjectInstance.cpp
+++ b/Engine/Plugins/Mutable/Source/CustomizableObject/Private/MuCO/CustomizableObjectInstance.cpp
@@ -55,6 +55,8 @@
 #include "Serialization/BulkData.h"
 #include "HAL/PlatformFileManager.h"
 
+#include "MuCO/Plugins/IMutableClothingModule.h"
+
 #include UE_INLINE_GENERATED_CPP_BY_NAME(CustomizableObjectInstance)
 
 #if WITH_EDITOR
@@ -4201,8 +4203,12 @@ namespace
 
 }
 
-void UCustomizableInstancePrivate::BuildOrCopyClothingData(const TSharedRef<FUpdateContextPrivate>& OperationData, USkeletalMesh * SkeletalMesh, const USkeletalMesh* LastUpdateSkeletalMesh, 
-	UCustomizableObjectInstance* CustomizableObjectInstance, FCustomizableObjectInstanceComponentIndex InstanceComponentIndex)
+void UCustomizableInstancePrivate::BuildOrCopyClothingData(
+		const TSharedRef<FUpdateContextPrivate>& OperationData, 
+		USkeletalMesh* SkeletalMesh, 
+		const USkeletalMesh* LastUpdateSkeletalMesh, 
+		UCustomizableObjectInstance* CustomizableObjectInstance, 
+		FCustomizableObjectInstanceComponentIndex InstanceComponentIndex)
 {
 	MUTABLE_CPUPROFILER_SCOPE(UCustomizableInstancePrivate::BuildOrCopyClothingData);
 
@@ -4220,8 +4226,8 @@ void UCustomizableInstancePrivate::BuildOrCopyClothingData(const TSharedRef<FUpd
 	}
 
 	const bool bAllowClothingPhysicsEdits = !bDisableClothingPhysicsEditsPropagation && ModelResources.bAllowClothingPhysicsEditsPropagation;
-	// First we need to discover if any clothing asset is used for the instance. 
 
+	// First we need to discover if any clothing asset is used for the instance. 
 	struct FSectionWithClothData
 	{
 		int32 ClothAssetIndex;
@@ -4355,15 +4361,13 @@ void UCustomizableInstancePrivate::BuildOrCopyClothingData(const TSharedRef<FUpd
 						continue;
 					}
 
-					SectionsWithCloth.Add(FSectionWithClothData
-							{
-							 	ClothAssetIndex, ClothAssetLodIndex, 
-								Section, LODIndex, FirstVertex, 
-								IndicesView16Bits, IndicesView32Bits, 
-								ClothingDataView, 
-								MakeArrayView(SectionClothingData->Data.GetData(), SectionClothingData->Data.Num()),
-								TArray<FMeshToMeshVertData>() 
-							});
+					SectionsWithCloth.Add(FSectionWithClothData {
+						ClothAssetIndex, ClothAssetLodIndex, 
+						Section, LODIndex, FirstVertex, 
+						IndicesView16Bits, IndicesView32Bits, 
+						ClothingDataView, 
+						MakeArrayView(SectionClothingData->Data.GetData(), SectionClothingData->Data.Num()),
+						TArray<FMeshToMeshVertData>() });
 				}
 			}
 		}
@@ -4379,39 +4383,6 @@ void UCustomizableInstancePrivate::BuildOrCopyClothingData(const TSharedRef<FUpd
 		return; // Nothing to do.
 	}
 		
-	TArray<FCustomizableObjectClothingAssetData> NewClothingAssetsData;
-	NewClothingAssetsData.SetNum(ClothingAssetsData.Num());
-
-	{
-		for (FSectionWithClothData& SectionWithCloth : SectionsWithCloth)
-		{
-			const FCustomizableObjectClothingAssetData& SrcAssetData = ClothingAssetsData[SectionWithCloth.ClothAssetIndex];
-			FCustomizableObjectClothingAssetData& DstAssetData = NewClothingAssetsData[SectionWithCloth.ClothAssetIndex];
-		
-			// Only initilialize once, multiple sections with cloth could point to the same cloth asset.
-			if (!DstAssetData.LodMap.Num())
-			{
-				DstAssetData.LodMap.Init(INDEX_NONE, OperationData->NumLODsAvailable[ComponentName]);
-
-				DstAssetData.UsedBoneNames = SrcAssetData.UsedBoneNames;
-				DstAssetData.UsedBoneIndices = SrcAssetData.UsedBoneIndices;
-				DstAssetData.ReferenceBoneIndex = SrcAssetData.ReferenceBoneIndex;
-				DstAssetData.Name = SrcAssetData.Name;
-
-				if (!bAllowClothingPhysicsEdits)
-				{
-					DstAssetData.LodData = SrcAssetData.LodData;
-				}
-				else
-				{
-					DstAssetData.LodData.SetNum(SrcAssetData.LodData.Num());
-				}
-			}
-
-			DstAssetData.LodMap[SectionWithCloth.Lod] = SectionWithCloth.ClothAssetLodIndex;
-		}
-	}
-
 	{
 		MUTABLE_CPUPROFILER_SCOPE(CopyMeshToMeshData)
 		
@@ -4430,9 +4401,9 @@ void UCustomizableInstancePrivate::BuildOrCopyClothingData(const TSharedRef<FUpd
 				// FMeshToMeshVertData and FCustomizableObjectMeshToMeshVertData have the same memory footprint and
 				// bytes in a FCustomizableObjectMeshToMeshVertData form a valid FMeshToMeshVertData (not the other way around).
 
-				static_assert(sizeof(FCustomizableObjectMeshToMeshVertData) == sizeof(FMeshToMeshVertData), "");
-				static_assert(TIsTrivial<FCustomizableObjectMeshToMeshVertData>::Value, "");
-				static_assert(TIsTrivial<FMeshToMeshVertData>::Value, "");
+				static_assert(sizeof(FCustomizableObjectMeshToMeshVertData) == sizeof(FMeshToMeshVertData));
+				static_assert(TIsTrivial<FCustomizableObjectMeshToMeshVertData>::Value);
+				static_assert(TIsTrivial<FMeshToMeshVertData>::Value);
 
 				const int32 VertexDataIndex = SectionWithCloth.ClothingDataIndicesView[VertexIdx];
 				check(VertexDataIndex >= 0);
@@ -4445,553 +4416,217 @@ void UCustomizableInstancePrivate::BuildOrCopyClothingData(const TSharedRef<FUpd
 		}
 	}
 
-	if (bAllowClothingPhysicsEdits)
+	struct FPerClothAssetData
 	{
-		// Indices remaps for {Section, AssetLod}, needed to recreate the lod transition data.
-		TMap<int32, TArray<TArray<int32>>> PhysicsSectionLodsIndicesRemaps;
-
-		check( SectionsWithCloth.Num() > 0 );
-		FSectionWithClothData* MaxSection = Algo::MaxElement(SectionsWithCloth, 
-				[](const FSectionWithClothData& A, const FSectionWithClothData& B) { return  A.ClothAssetIndex < B.ClothAssetIndex; } );
-		
-		PhysicsSectionLodsIndicesRemaps.Reserve(MaxSection->ClothAssetIndex + 1);
-			
-		for (FSectionWithClothData& SectionLods : SectionsWithCloth)
-		{
-			TArray<TArray<int32>>& Value = PhysicsSectionLodsIndicesRemaps.FindOrAdd(SectionLods.ClothAssetIndex, {});
-			Value.SetNum(FMath::Max(Value.Num(), SectionLods.ClothAssetLodIndex + 1));
-		}
+		int32 MinLOD = 0;
 
-		{
-			MUTABLE_CPUPROFILER_SCOPE(RemapPhysicsMesh)
-
-			for (FSectionWithClothData& SectionWithCloth : SectionsWithCloth)
-			{
-				const FCustomizableObjectClothingAssetData& SrcClothingAssetData = ClothingAssetsData[SectionWithCloth.ClothAssetIndex];
-				FCustomizableObjectClothingAssetData& NewClothingAssetData = NewClothingAssetsData[SectionWithCloth.ClothAssetIndex];
-						
-				const FClothLODDataCommon& SrcLodData = SrcClothingAssetData.LodData[SectionWithCloth.ClothAssetLodIndex];
-				FClothLODDataCommon& NewLodData = NewClothingAssetData.LodData[SectionWithCloth.ClothAssetLodIndex];
-				
-				const int32 PhysicalMeshVerticesNum = SrcLodData.PhysicalMeshData.Vertices.Num();
+		// Indices in SectionsWithCloth for render sections attached to this ClothAsset.
+		TArray<TArray<int32, TInlineAllocator<8>>, TInlineAllocator<8>> AttachedSections; 
 
-				if (!PhysicalMeshVerticesNum)
-				{
-					// Nothing to do.
-					continue;
-				}
-				
-				// Vertices not indexed in the mesh to mesh data generated for this section need to be removed.
-				TArray<uint8> VertexUtilizationBuffer;
-				VertexUtilizationBuffer.Init(0, PhysicalMeshVerticesNum);
-			
-				// Discover used vertices.
-				const int32 SectionVerticesNum = SectionWithCloth.ClothingDataIndicesView.Num();
-
-				TArray<uint8> RenderVertexUtilizationBuffer;
-				RenderVertexUtilizationBuffer.Init(0, SectionVerticesNum);
+		FName Name;
+		
+		UPhysicsAsset* PhysicsAsset = nullptr;
+		UClothingAssetCommon* ClothingAsset = nullptr;
+	};
 
-				// Sometimes, at least when a clip morph is applied, vertices are not removed from the section
-				// and only the triangles (indices) that form the mesh are modified.
+	TMap<int32, FPerClothAssetData> PerClothAssetData;
+	PerClothAssetData.Reserve(32);
 
-				auto GenerateRenderUtilizationBuffer = [&RenderVertexUtilizationBuffer](const auto& IndicesView, int32 SectionBaseVertex)
-				{
-					const int32 IndicesCount = IndicesView.Num();
-					check(IndicesCount % 3 == 0);
-					for ( int32 I = 0; I < IndicesCount; I += 3 )
-					{
-						RenderVertexUtilizationBuffer[int32(IndicesView[I + 0]) - SectionBaseVertex] = 1;
-						RenderVertexUtilizationBuffer[int32(IndicesView[I + 1]) - SectionBaseVertex] = 1;
-						RenderVertexUtilizationBuffer[int32(IndicesView[I + 2]) - SectionBaseVertex] = 1;
-					}
-				};
-				
-				if (SectionWithCloth.SectionIndex16View.Num())
-				{
-					GenerateRenderUtilizationBuffer(SectionWithCloth.SectionIndex16View, SectionWithCloth.BaseVertex);
-				}
-				else
-				{
-					check(SectionWithCloth.SectionIndex32View.Num());
-					GenerateRenderUtilizationBuffer(SectionWithCloth.SectionIndex32View, SectionWithCloth.BaseVertex);
-				}
+	// Gather attached sections clothing asset LOD.
+	const int32 NumSectionsWithCloth = SectionsWithCloth.Num();
+	for (int32 SectionWithClothIndex = 0; SectionWithClothIndex < NumSectionsWithCloth; ++SectionWithClothIndex)
+	{
+		const FSectionWithClothData& SectionWithCloth = SectionsWithCloth[SectionWithClothIndex];
+		FPerClothAssetData& AssetData = PerClothAssetData.FindOrAdd(SectionWithCloth.ClothAssetIndex); 
 
-				const TArray<FMeshToMeshVertData>& SectionClothMappingData = SectionWithCloth.MappingData;
-				for (int32 Idx = 0; Idx < SectionVerticesNum; ++Idx)
-				{
-					if (RenderVertexUtilizationBuffer[Idx])
-					{
-						const uint16* Indices = SectionClothMappingData[Idx].SourceMeshVertIndices;
+		AssetData.MinLOD = FMath::Min(AssetData.MinLOD, SectionWithCloth.ClothAssetLodIndex);
 
-						// It may happen with missing data files
-						if (VertexUtilizationBuffer.IsValidIndex(Indices[0]))
-						{
-							VertexUtilizationBuffer[Indices[0]] = 1;
-						}
+		int32 MaxLOD = FMath::Max<int32>(AssetData.AttachedSections.Num() - 1, SectionWithCloth.ClothAssetLodIndex);
 
-						if (VertexUtilizationBuffer.IsValidIndex(Indices[1]))
-						{
-							VertexUtilizationBuffer[Indices[1]] = 1;
-						}
+		AssetData.AttachedSections.SetNum(MaxLOD + 1);
+		AssetData.AttachedSections[SectionWithCloth.ClothAssetLodIndex].Add(SectionWithClothIndex);
+	}
 
-						if (VertexUtilizationBuffer.IsValidIndex(Indices[2]))
-						{
-							VertexUtilizationBuffer[Indices[2]] = 1;
-						}
-					}
-				}
+	for (TPair<int32, FPerClothAssetData>& Data : PerClothAssetData)
+	{
+		int32 ClothAssetIndex = Data.Key;
+		FPerClothAssetData& ClothAssetData = Data.Value; 
+	
+		ClothAssetData.Name = ClothingAssetsData[ClothAssetIndex].Name;
+		ClothAssetData.PhysicsAsset = ClothingPhysicsAssets[ClothAssetIndex];
+	}
 
-				TArray<int32>& IndexMap = PhysicsSectionLodsIndicesRemaps[SectionWithCloth.ClothAssetIndex][SectionWithCloth.ClothAssetLodIndex];
-				IndexMap.SetNumUninitialized(PhysicalMeshVerticesNum);
+	{ 
+		MUTABLE_CPUPROFILER_SCOPE(CreateClothingAssets)
 
-				// Compute index remap and number of remaining physics vertices.
-				// -1 indicates the vertex has been removed.
-				int32 NewPhysicalMeshVerticesNum = 0;
-				for (int32 Idx = 0; Idx < PhysicalMeshVerticesNum; ++Idx)
+		auto CreateNewClothConfigFromData = [](UObject* Outer, const FCustomizableObjectClothConfigData& ConfigData) -> UClothConfigCommon* 
+		{
+			UClass* ClothConfigClass = FindObject<UClass>(nullptr, *ConfigData.ClassPath);
+			if (ClothConfigClass)
+			{
+				UClothConfigCommon* ClothConfig = NewObject<UClothConfigCommon>(Outer, ClothConfigClass);
+				if (ClothConfig)
 				{
-					IndexMap[Idx] = VertexUtilizationBuffer[Idx] ? NewPhysicalMeshVerticesNum++ : -1;
-				}
+					FMemoryReaderView MemoryReader(ConfigData.ConfigBytes);
+					ClothConfig->Serialize(MemoryReader);
 
-				const bool bHasVerticesRemoved = NewPhysicalMeshVerticesNum < PhysicalMeshVerticesNum;
-				if (!bHasVerticesRemoved)
-				{
-					// If no vertices are removed the IndexMap is no longer needed. The lack of data in the map 
-					// can indicates that no vertex has been removed to subsequent operations.
-					IndexMap.Reset();
+					return ClothConfig;
 				}
-				
-				const auto CopyIfUsed = [&VertexUtilizationBuffer, bHasVerticesRemoved](auto& Dst, const auto& Src)
-				{	
-					const int32 SrcNumElems = Src.Num();
-
-					if (!bHasVerticesRemoved)
-					{
-						for (int32 Idx = 0; Idx < SrcNumElems; ++Idx)
-						{
-							Dst[Idx] = Src[Idx];
-						}
+			}
 
-						return;
-					}
+			return nullptr;
+		};
 
-					for (int32 Idx = 0, DstNumElems = 0; Idx < SrcNumElems; ++Idx)
-					{
-						if (VertexUtilizationBuffer[Idx])
-						{
-							Dst[DstNumElems++] = Src[Idx];
-						}
-					}
-				};
+		TArray<TTuple<FName, UClothConfigCommon*>> SharedConfigs;
+		SharedConfigs.Reserve(ClothSharedConfigsData.Num());
+ 
+		for (const FCustomizableObjectClothConfigData& ConfigData : ClothSharedConfigsData)
+		{
+			UClothConfigCommon* ClothConfig = CreateNewClothConfigFromData(SkeletalMesh, ConfigData);
+			if (ClothConfig)
+			{
+				SharedConfigs.Emplace(ConfigData.ConfigName, ClothConfig);
+			}
+		}
+	
+		bool bAllNamesUnique = true;
+		TArray<FName, TInlineAllocator<8>> UniqueAssetNames;
 
-				NewLodData.PhysicalMeshData.MaxBoneWeights = SrcLodData.PhysicalMeshData.MaxBoneWeights;
+		for (TPair<int32, FPerClothAssetData>& AssetData : PerClothAssetData)
+		{
+			const int32 PrevNumUniqueElems = UniqueAssetNames.Num();
+			const int32 ElemIndex = UniqueAssetNames.AddUnique(AssetData.Value.Name);
 
-				NewLodData.PhysicalMeshData.Vertices.SetNum(NewPhysicalMeshVerticesNum);
-				NewLodData.PhysicalMeshData.Normals.SetNum(NewPhysicalMeshVerticesNum);
-				NewLodData.PhysicalMeshData.BoneData.SetNum(NewPhysicalMeshVerticesNum);
-				NewLodData.PhysicalMeshData.InverseMasses.SetNum(NewPhysicalMeshVerticesNum);
+			if (ElemIndex < PrevNumUniqueElems)
+			{
+				bAllNamesUnique = false;
+				break;
+			}
+		}
 
-				CopyIfUsed(NewLodData.PhysicalMeshData.Vertices, SrcLodData.PhysicalMeshData.Vertices);
-				CopyIfUsed(NewLodData.PhysicalMeshData.Normals, SrcLodData.PhysicalMeshData.Normals);
-				CopyIfUsed(NewLodData.PhysicalMeshData.BoneData, SrcLodData.PhysicalMeshData.BoneData);
-				CopyIfUsed(NewLodData.PhysicalMeshData.InverseMasses, SrcLodData.PhysicalMeshData.InverseMasses);
+		for (TPair<int32, FPerClothAssetData>& AssetData : PerClothAssetData)
+		{
+			int32 AssetIndex = AssetData.Key;
+			FPerClothAssetData& ClothAssetData = AssetData.Value;  
 
-				const int32 PrevIndex = SectionWithCloth.Lod - 1;
-				const bool bNeedsTransitionUpData = NewClothingAssetData.LodMap.IsValidIndex(PrevIndex) && NewClothingAssetData.LodMap[PrevIndex] != INDEX_NONE;
-				if (bNeedsTransitionUpData)
-				{
-					NewLodData.TransitionUpSkinData.SetNum(SrcLodData.TransitionUpSkinData.Num() ? NewPhysicalMeshVerticesNum : 0);	
-					CopyIfUsed(NewLodData.TransitionUpSkinData, SrcLodData.TransitionUpSkinData);
-				}
+			FName ClothingAssetObjectName = bAllNamesUnique 
+					? ClothAssetData.Name
+					: FName(FString::Printf(TEXT("%s_%d"), *ClothAssetData.Name.ToString(), AssetIndex));
 
-				const int32 NextIndex = SectionWithCloth.Lod + 1;
-				const bool bNeedsTransitionDownData = NewClothingAssetData.LodMap.IsValidIndex(NextIndex) && NewClothingAssetData.LodMap[NextIndex] != INDEX_NONE;
-				if (bNeedsTransitionDownData)
-				{
-					NewLodData.TransitionDownSkinData.SetNum(SrcLodData.TransitionDownSkinData.Num() ? NewPhysicalMeshVerticesNum : 0);
-					CopyIfUsed(NewLodData.TransitionDownSkinData, SrcLodData.TransitionDownSkinData);
-				}
-				
-				const TMap<uint32, FPointWeightMap>& SrcPhysWeightMaps = SrcLodData.PhysicalMeshData.WeightMaps;
-				TMap<uint32, FPointWeightMap>& NewPhysWeightMaps = NewLodData.PhysicalMeshData.WeightMaps;
+			UCustomizableObjectClothingAsset* NewClothingAsset = NewObject<UCustomizableObjectClothingAsset>(SkeletalMesh, ClothingAssetObjectName);
+			NewClothingAsset->AssignNewGuid();
+		
+			const int32 NumClothLODs = ClothingAssetsData[AssetIndex].LodData.Num() - ClothAssetData.MinLOD;
 
-				for (const TPair<uint32, FPointWeightMap>& WeightMap : SrcPhysWeightMaps)
-				{
-					if (WeightMap.Value.Values.Num() > 0)
-					{
-						FPointWeightMap& NewWeightMap = NewLodData.PhysicalMeshData.AddWeightMap(WeightMap.Key);
-						NewWeightMap.Values.SetNum(NewPhysicalMeshVerticesNum);
+			NewClothingAsset->LodData.SetNum(NumClothLODs);
+			for (int32 LODIndex = 0; LODIndex < NumClothLODs; ++LODIndex)
+			{
+				NewClothingAsset->LodData[LODIndex] = ClothingAssetsData[AssetIndex].LodData[LODIndex + ClothAssetData.MinLOD];
+			}
 
-						CopyIfUsed(NewWeightMap.Values, WeightMap.Value.Values);
-					}
-				}
-				// Remap render mesh to mesh indices.
-				if (bHasVerticesRemoved)
+			// Reconstruct clothing asset lod map.
+			NewClothingAsset->LodMap.Init(INDEX_NONE, LODCount);
+			for (int32 LODIndex = 0; LODIndex < NumClothLODs; ++LODIndex)
+			{
+				const int32 NumClothSections = ClothAssetData.AttachedSections.Num(); 
+				for (int32 SectionWithClothIndex : ClothAssetData.AttachedSections[LODIndex])
 				{
-					for (FMeshToMeshVertData& VertClothData : SectionWithCloth.MappingData)
-					{
-						uint16* Indices = VertClothData.SourceMeshVertIndices;
-
-						// It could happen with missing data.
-						if (IndexMap.IsValidIndex(Indices[0]))
-						{
-							Indices[0] = (uint16)IndexMap[Indices[0]];
-						}
-						if (IndexMap.IsValidIndex(Indices[1]))
-						{
-							Indices[1] = (uint16)IndexMap[Indices[1]];
-						}
-						if (IndexMap.IsValidIndex(Indices[2]))
-						{
-							Indices[2] = (uint16)IndexMap[Indices[2]];
-						}
-					}
+					NewClothingAsset->LodMap[SectionsWithCloth[SectionWithClothIndex].Lod] = LODIndex; 
 				}
+			}
 
-				// Remap and trim physics mesh vertices and self collision indices. 
-				
-				// Returns the final size of Dst.
-				const auto TrimAndRemapTriangles = [&IndexMap](TArray<uint32>& Dst, const TArray<uint32>& Src) -> int32
-				{
-					check(Src.Num() % 3 == 0);
-
-					const int32 SrcNumElems = Src.Num();
-					if (!IndexMap.Num())
-					{
-						//for (int32 Idx = 0; Idx < SrcNumElems; ++Idx)
-						//{
-						//	Dst[Idx] = Src[Idx];
-						//}
-
-						FMemory::Memcpy( Dst.GetData(), Src.GetData(), SrcNumElems*sizeof(uint32) );
-						return SrcNumElems;
-					}
-
-					int32 DstNumElems = 0;
-					for (int32 Idx = 0; Idx < SrcNumElems; Idx += 3)
-					{
-						const int32 Idx0 = IndexMap[Src[Idx + 0]];
-						const int32 Idx1 = IndexMap[Src[Idx + 1]];
-						const int32 Idx2 = IndexMap[Src[Idx + 2]];
-
-						// triangles are only copied if all vertices are used.
-						if (!((Idx0 < 0) | (Idx1 < 0) | (Idx2 < 0)))
-						{
-							Dst[DstNumElems + 0] = Idx0;
-							Dst[DstNumElems + 1] = Idx1;
-							Dst[DstNumElems + 2] = Idx2;
-
-							DstNumElems += 3;
-						}
-					}
-
-					return DstNumElems;
-				};
-
-				const TArray<uint32>& SrcPhysicalMeshIndices = SrcLodData.PhysicalMeshData.Indices;
-				TArray<uint32>& NewPhysicalMeshIndices = NewLodData.PhysicalMeshData.Indices;
-				NewPhysicalMeshIndices.SetNum(SrcPhysicalMeshIndices.Num());
-				NewPhysicalMeshIndices.SetNum(TrimAndRemapTriangles(NewPhysicalMeshIndices, SrcPhysicalMeshIndices), EAllowShrinking::No);
-			
-				const auto TrimAndRemapVertexSet = [&IndexMap](TSet<int32>& Dst, const TSet<int32>& Src)
-				{	
-					if (!IndexMap.Num())
-					{
-						Dst = Src;
-						return;
-					}
-
-					Dst.Reserve(Src.Num());
-					for(const int32 SrcIdx : Src)
-					{
-						const int32 MappedIdx = IndexMap[SrcIdx];
+			NewClothingAsset->UsedBoneIndices = ClothingAssetsData[AssetIndex].UsedBoneIndices;
+			NewClothingAsset->UsedBoneNames = ClothingAssetsData[AssetIndex].UsedBoneNames;
+			NewClothingAsset->ReferenceBoneIndex = ClothingAssetsData[AssetIndex].ReferenceBoneIndex;
+			NewClothingAsset->RefreshBoneMapping(SkeletalMesh);
+			NewClothingAsset->CalculateReferenceBoneIndex();	
+			NewClothingAsset->PhysicsAsset = ClothAssetData.PhysicsAsset;
 
-						if (MappedIdx >= 0)
-						{
-							Dst.Add(MappedIdx);
-						}
-					}
-				};
-
-				const TSet<int32>& SrcSelfCollisionVertexSet = SrcLodData.PhysicalMeshData.SelfCollisionVertexSet;
-				TSet<int32>& NewSelfCollisionVertexSet = NewLodData.PhysicalMeshData.SelfCollisionVertexSet;
-				TrimAndRemapVertexSet(NewSelfCollisionVertexSet, SrcSelfCollisionVertexSet);
-							
+			for (const FCustomizableObjectClothConfigData& ConfigData : ClothingAssetsData[AssetIndex].ConfigsData)
+			{
+				UClothConfigCommon* ClothConfig = CreateNewClothConfigFromData(NewClothingAsset, ConfigData);
+				if (ClothConfig)
 				{
-					MUTABLE_CPUPROFILER_SCOPE(BuildClothTetherData)
-
-					auto TrimAndRemapTethers = [&IndexMap](FClothTetherData& Dst, const FClothTetherData& Src)
-					{
-						if (!IndexMap.Num())
-						{
-							Dst.Tethers = Src.Tethers;
-							return;
-						}
-
-						Dst.Tethers.Reserve(Src.Tethers.Num());
-						for ( const TArray<TTuple<int32, int32, float>>& SrcTetherCluster : Src.Tethers )
-						{
-							TArray<TTuple<int32, int32, float>>& DstTetherCluster = Dst.Tethers.Emplace_GetRef();
-							DstTetherCluster.Reserve(SrcTetherCluster.Num());
-							for ( const TTuple<int32, int32, float>& Tether : SrcTetherCluster )
-							{
-								const int32 Index0 = IndexMap[Tether.Get<0>()];
-								const int32 Index1 = IndexMap[Tether.Get<1>()];
-								if ((Index0 >= 0) & (Index1 >= 0))
-								{
-									DstTetherCluster.Emplace(Index0, Index1, Tether.Get<2>());
-								}
-							}
-
-							if (!DstTetherCluster.Num())
-							{
-								Dst.Tethers.RemoveAt( Dst.Tethers.Num() - 1, 1, EAllowShrinking::No );
-							}
-						}
-					};
-
-					TrimAndRemapTethers(NewLodData.PhysicalMeshData.GeodesicTethers, SrcLodData.PhysicalMeshData.GeodesicTethers);
-					TrimAndRemapTethers(NewLodData.PhysicalMeshData.EuclideanTethers, SrcLodData.PhysicalMeshData.EuclideanTethers);
+					NewClothingAsset->ClothConfigs.Add(ConfigData.ConfigName, ClothConfig);
 				}
 			}
-		}
 
-		// Try to find plausible values for LodTransitionData vertices that have lost the triangle to which are attached.
-		{
-			MUTABLE_CPUPROFILER_SCOPE(BuildLodTransitionData)
-			
-			for (const FSectionWithClothData& SectionWithCloth : SectionsWithCloth)
+			for (const TTuple<FName, UClothConfigCommon*>& SharedConfig : SharedConfigs)
 			{
+				NewClothingAsset->ClothConfigs.Add(SharedConfig);
+			}
 
-				FCustomizableObjectClothingAssetData& NewClothingAssetData = NewClothingAssetsData[SectionWithCloth.ClothAssetIndex];
-				FClothLODDataCommon& NewLodData = NewClothingAssetData.LodData[SectionWithCloth.ClothAssetLodIndex];
-
-				const int32 PhysicalMeshVerticesNum = NewLodData.PhysicalMeshData.Vertices.Num();
-
-				if (!PhysicalMeshVerticesNum)
-				{
-					// Nothing to do.
-					continue;
-				}
-			
-				auto RemapTransitionMeshToMeshVertData = []( TArray<FMeshToMeshVertData>& InOutVertData, TArray<int32>& IndexMap )
-				{
-					for (FMeshToMeshVertData& VertData : InOutVertData)
-					{
-						uint16* Indices = VertData.SourceMeshVertIndices;
-						Indices[0] = (uint16)IndexMap[Indices[0]];
-						Indices[1] = (uint16)IndexMap[Indices[1]];
-						Indices[2] = (uint16)IndexMap[Indices[2]];
-					}
-				};
-
-				if (NewLodData.TransitionDownSkinData.Num() > 0)
-				{	
-					TArray<int32>& IndexMap = PhysicsSectionLodsIndicesRemaps[SectionWithCloth.ClothAssetIndex][SectionWithCloth.ClothAssetLodIndex + 1];
-
-					if (IndexMap.Num())
-					{
-						RemapTransitionMeshToMeshVertData( NewLodData.TransitionDownSkinData, IndexMap );
-					}
-				}
-				
-				if (NewLodData.TransitionUpSkinData.Num() > 0)
-				{	
-					TArray<int32>& IndexMap = PhysicsSectionLodsIndicesRemaps[SectionWithCloth.ClothAssetIndex][SectionWithCloth.ClothAssetLodIndex - 1];
-					if (IndexMap.Num())
-					{
-						RemapTransitionMeshToMeshVertData( NewLodData.TransitionUpSkinData, IndexMap );
-					}
-				}
-
-				struct FMeshPhysicsDesc
-				{
-					const TArray<FVector3f>& Vertices;
-					const TArray<FVector3f>& Normals;
-					const TArray<uint32>& Indices;
-				};	
-			
-				auto RebindVertex = [](const FMeshPhysicsDesc& Mesh, const FVector3f& InPosition, const FVector3f& InNormal, FMeshToMeshVertData& Out)
-				{
-					const FVector3f Normal = InNormal;
-
-					// We don't have the mesh tangent, find something plausible.
-					FVector3f Tan0, Tan1;
-					Normal.FindBestAxisVectors(Tan0, Tan1);
-					const FVector3f Tangent = Tan0;
-					
-					// Some of the math functions take as argument FVector, we'd want to be FVector3f. 
-					// This should be changed once support for the single type in the FMath functions is added. 
-					const FVector Position = (FVector)InPosition;
-					int32 BestBaseTriangleIdx = INDEX_NONE;
-					FVector::FReal BestDistanceSq = TNumericLimits<FVector::FReal>::Max();
-					
-					const int32 NumIndices = Mesh.Indices.Num();
-					check(NumIndices % 3 == 0);
-
-					for (int32 I = 0; I < NumIndices; I += 3)
-					{
-						const FVector& A = (FVector)Mesh.Vertices[Mesh.Indices[I + 0]];
-						const FVector& B = (FVector)Mesh.Vertices[Mesh.Indices[I + 1]];
-						const FVector& C = (FVector)Mesh.Vertices[Mesh.Indices[I + 2]];
-
-						FVector ClosestTrianglePoint = FMath::ClosestPointOnTriangleToPoint(Position, (FVector)A, (FVector)B, (FVector)C);
-
-						const FVector::FReal CurrentDistSq = (ClosestTrianglePoint - Position).SizeSquared();
-						if (CurrentDistSq < BestDistanceSq)
-						{
-							BestDistanceSq = CurrentDistSq;
-							BestBaseTriangleIdx = I;
-						}
-					}
-
-					check(BestBaseTriangleIdx >= 0);
-
-					auto ComputeBaryCoordsAndDist = [](const FVector3f& A, const FVector3f& B, const FVector3f& C, const FVector3f& P) -> FVector4f
-					{
-						FPlane4f TrianglePlane(A, B, C);
-
-						const FVector3f PointOnTriPlane = FVector3f::PointPlaneProject(P, TrianglePlane);
-						const FVector3f BaryCoords = (FVector3f)FMath::ComputeBaryCentric2D((FVector)PointOnTriPlane, (FVector)A, (FVector)B, (FVector)C);
-
-						return FVector4f(BaryCoords, TrianglePlane.PlaneDot((FVector3f)P));
-					};
-
-					const FVector3f& A = Mesh.Vertices[Mesh.Indices[BestBaseTriangleIdx + 0]];
-					const FVector3f& B = Mesh.Vertices[Mesh.Indices[BestBaseTriangleIdx + 1]];
-					const FVector3f& C = Mesh.Vertices[Mesh.Indices[BestBaseTriangleIdx + 2]];
-
-					Out.PositionBaryCoordsAndDist = ComputeBaryCoordsAndDist(A, B, C, (FVector3f)Position );
-					Out.NormalBaryCoordsAndDist = ComputeBaryCoordsAndDist(A, B, C, (FVector3f)Position + Normal );
-					Out.TangentBaryCoordsAndDist = ComputeBaryCoordsAndDist(A, B, C, (FVector3f)Position + Tangent );
-					Out.SourceMeshVertIndices[0] = (uint16)Mesh.Indices[BestBaseTriangleIdx + 0];
-					Out.SourceMeshVertIndices[1] = (uint16)Mesh.Indices[BestBaseTriangleIdx + 1]; 
-					Out.SourceMeshVertIndices[2] = (uint16)Mesh.Indices[BestBaseTriangleIdx + 2];
-				};
-			
-				auto RecreateTransitionData = [&PhysicsSectionLodsIndicesRemaps, &RebindVertex]( 
-					const FMeshPhysicsDesc& ToMesh, const FMeshPhysicsDesc& FromMesh, const TArray<int32>& IndexMap, TArray<FMeshToMeshVertData>& InOutTransitionData )
-				{
-					if (!IndexMap.Num())
-					{
-						return;
-					}
-
-					if (!InOutTransitionData.Num())
-					{
-						return;
-					}
-
-					const int32 TransitionDataNum = InOutTransitionData.Num();
-					
-					for (int32 I = 0; I < TransitionDataNum; ++I)
-					{
-						FMeshToMeshVertData& VertData = InOutTransitionData[I];
-						uint16* Indices = VertData.SourceMeshVertIndices;
+			ClothAssetData.ClothingAsset = NewClothingAsset;
+			SkeletalMesh->GetMeshClothingAssets().AddUnique(NewClothingAsset);
+		}
+	}
 
-						// If any original indices are missing but the vertex is still alive rebind the vertex.
-						// In general, the number of rebinds should be small.
+	IMutableClothingModule* MutableClothingModule = nullptr;
+	if (bAllowClothingPhysicsEdits)
+	{
+		MutableClothingModule = FModuleManager::GetModulePtr<IMutableClothingModule>(MUTABLE_CLOTHING_MODULE_NAME);	
 
-						// Currently, if any index is missing we rebind to the closest triangle but it could be nice to use the remaining indices, 
-						// if any, to find the most appropriate triangle to bind to. 
-						const bool bNeedsRebind = (Indices[0] == 0xFFFF) | (Indices[1] == 0xFFFF) | (Indices[2] == 0xFFFF);
+		if (!MutableClothingModule)
+		{
+			UE_LOG(LogMutable, Warning, TEXT("MutableClothing plugin could not be found. Make sure the plugin is enabled if you want to use advanced clothing features."));
+		}
+	}
 
-						if (bNeedsRebind)
-						{
-							RebindVertex( ToMesh, FromMesh.Vertices[I], FromMesh.Normals[I], VertData );
-						}
-					}
-				};
+	if (MutableClothingModule)
+	{
+		for (TPair<int32, FPerClothAssetData>& Data : PerClothAssetData)
+		{
+			int32 ClothAssetIndex = Data.Key;
+			FPerClothAssetData& ClothAssetData = Data.Value;
 
-				const FMeshPhysicsDesc CurrentPhysicsMesh {
-					NewClothingAssetData.LodData[SectionWithCloth.ClothAssetLodIndex].PhysicalMeshData.Vertices,
-					NewClothingAssetData.LodData[SectionWithCloth.ClothAssetLodIndex].PhysicalMeshData.Normals,
-					NewClothingAssetData.LodData[SectionWithCloth.ClothAssetLodIndex].PhysicalMeshData.Indices };
+			const int32 NumLODs = ClothAssetData.AttachedSections.Num() - ClothAssetData.MinLOD;
 
-				const TArray<TArray<int32>>& SectionIndexRemaps = PhysicsSectionLodsIndicesRemaps[SectionWithCloth.ClothAssetIndex];
-				
-				if (SectionWithCloth.ClothAssetLodIndex < SectionIndexRemaps.Num() - 1)
-				{
-					const TArray<int32>& IndexMap = SectionIndexRemaps[SectionWithCloth.ClothAssetLodIndex + 1];
-					
-					const FMeshPhysicsDesc TransitionDownTarget {  
-						NewClothingAssetData.LodData[SectionWithCloth.ClothAssetLodIndex + 1].PhysicalMeshData.Vertices,
-						NewClothingAssetData.LodData[SectionWithCloth.ClothAssetLodIndex + 1].PhysicalMeshData.Normals,
-						NewClothingAssetData.LodData[SectionWithCloth.ClothAssetLodIndex + 1].PhysicalMeshData.Indices };
-						
-					RecreateTransitionData( TransitionDownTarget, CurrentPhysicsMesh, IndexMap, NewLodData.TransitionDownSkinData );
-				}
-				
-				if (SectionWithCloth.ClothAssetLodIndex > 0)
+			for (int32 LODIndex = ClothAssetData.MinLOD; LODIndex < NumLODs; ++LODIndex)
+			{
+				for (int32 AttachedSectionIndex : ClothAssetData.AttachedSections[LODIndex])
 				{
-					const TArray<int32>& IndexMap = SectionIndexRemaps[SectionWithCloth.ClothAssetLodIndex - 1];
-					
-					FMeshPhysicsDesc TransitionUpTarget{  
-						NewClothingAssetData.LodData[SectionWithCloth.ClothAssetLodIndex - 1].PhysicalMeshData.Vertices,
-						NewClothingAssetData.LodData[SectionWithCloth.ClothAssetLodIndex - 1].PhysicalMeshData.Normals,
-						NewClothingAssetData.LodData[SectionWithCloth.ClothAssetLodIndex - 1].PhysicalMeshData.Indices };
-			
-					RecreateTransitionData( TransitionUpTarget, CurrentPhysicsMesh, IndexMap, NewLodData.TransitionUpSkinData );
+					const FSectionWithClothData& ClothLODRenderSection = SectionsWithCloth[AttachedSectionIndex];
+					const int32 SectionVerticesNum = ClothLODRenderSection.ClothingDataIndicesView.Num();
 				}
 			}
 		}
-	}
 
-	// Remove empty LODs and redo LodMaps.
-	{
-		const int32 NumClothingAssets = NewClothingAssetsData.Num();
-		for (int32 I = 0; I < NumClothingAssets; ++I)
+		for (TPair<int32, FPerClothAssetData>& Data : PerClothAssetData)
 		{
-			// Skip assets not set.
-			if (!NewClothingAssetsData[I].LodData.Num())
+			FPerClothAssetData& ClothAssetData = Data.Value; 
+			UClothingAssetCommon* ClothingAsset = ClothAssetData.ClothingAsset;
+			
+			if (!ClothingAsset)
 			{
 				continue;
 			}
 
-			TArray<int32>& LodMap = NewClothingAssetsData[I].LodMap;
-			TArray<FClothLODDataCommon>& LodData = NewClothingAssetsData[I].LodData;
-
-			const int32 LodMapNum = LodMap.Num();
-			for (int32 LODIndex = 0; LODIndex < LodMapNum; ++LODIndex)
+			bool bNeedsLodTransitionUpdate = false;
+			for (int32 LODIndex = 0; LODIndex < ClothingAsset->LodData.Num(); ++LODIndex)
 			{
-				int32 MappedLODIndex = LodMap[LODIndex];
-				
-				if (MappedLODIndex != INDEX_NONE && !NewClothingAssetsData[I].LodData[MappedLODIndex].PhysicalMeshData.Vertices.Num())
+				TArray<TArrayView<FMeshToMeshVertData>, TInlineAllocator<8>> MeshToMeshDataViews;
+
+				for (int32 AttachedSectionIndex : ClothAssetData.AttachedSections[LODIndex])
 				{
-					LodMap[LODIndex] = INDEX_NONE;
+					MeshToMeshDataViews.Add(MakeArrayView(SectionsWithCloth[AttachedSectionIndex].MappingData));
 				}
-			}
-
-			TArray<int32> RemappedLODDataIndices;
-			RemappedLODDataIndices.Init(INDEX_NONE, LodData.Num());
 
-			TArray<FClothLODDataCommon> TrimmedLodData;
-			TrimmedLodData.Reserve(NewClothingAssetsData[I].LodData.Num());
+				bool bModified = MutableClothingModule->UpdateClothSimulationLOD(
+						LODIndex, *ClothingAsset, MakeConstArrayView(MeshToMeshDataViews));
+				
+				bNeedsLodTransitionUpdate = bNeedsLodTransitionUpdate || bModified;
+			}
 
-			for (int32 LODIndex = 0; LODIndex < LodMapNum; ++LODIndex)
+			if (bNeedsLodTransitionUpdate)
 			{
-				const int32 MappedLODIndex = LodMap[LODIndex];
-				if (MappedLODIndex != INDEX_NONE)
+				// This needs to happen after all LODs have been processed.
+				for (int32 LODIndex = 0; LODIndex < ClothingAsset->LodData.Num(); ++LODIndex)
 				{
-					if (RemappedLODDataIndices[MappedLODIndex] == INDEX_NONE)
-					{
-						const int32 RemappedIndex = TrimmedLodData.Emplace(MoveTemp(LodData[MappedLODIndex]));
-						LodMap[LODIndex] = RemappedIndex; 
-						RemappedLODDataIndices[MappedLODIndex] = RemappedIndex;
-					}
-					else
-					{
-						LodMap[LODIndex] = RemappedLODDataIndices[MappedLODIndex];
-					}
+					MutableClothingModule->FixLODTransitionMappings(LODIndex, *ClothingAsset);
 				}
 			}
-
-			LodData = MoveTemp(TrimmedLodData);
-		}
+		}	
 	}
-	
-	// From here up, could be moved to an async task similar to what is done with the other prepare tasks.  
- 
-	// Create Clothing Assets.
 
 	// Based on FSkeletalMeshLODModel::GetClothMappingData().
 	TArray<TArray<FMeshToMeshVertData>> LodMappingData;
@@ -5062,131 +4697,25 @@ void UCustomizableInstancePrivate::BuildOrCopyClothingData(const TSharedRef<FUpd
 		}
 	}
 
-	TArray<UCustomizableObjectClothingAsset*> NewClothingAssets;
-	NewClothingAssets.Init(nullptr, ClothingAssetsData.Num());
-
-	{ 
-		MUTABLE_CPUPROFILER_SCOPE(CreateClothingAssets)
-
-		auto CreateNewClothConfigFromData = [](UObject* Outer, const FCustomizableObjectClothConfigData& ConfigData) -> UClothConfigCommon* 
-		{
-			UClass* ClothConfigClass = FindObject<UClass>(nullptr, *ConfigData.ClassPath);
-			if (ClothConfigClass)
-			{
-				UClothConfigCommon* ClothConfig = NewObject<UClothConfigCommon>(Outer, ClothConfigClass);
-				if (ClothConfig)
-				{
-					FMemoryReaderView MemoryReader(ConfigData.ConfigBytes);
-					ClothConfig->Serialize(MemoryReader);
-
-					return ClothConfig;
-				}
-			}
-
-			return nullptr;
-		};
-
-		TArray<TTuple<FName, UClothConfigCommon*>> SharedConfigs;
-		SharedConfigs.Reserve(ClothSharedConfigsData.Num());
- 
-		for (const FCustomizableObjectClothConfigData& ConfigData : ClothSharedConfigsData)
-		{
-			UClothConfigCommon* ClothConfig = CreateNewClothConfigFromData(SkeletalMesh, ConfigData);
-			if (ClothConfig)
-			{
-				SharedConfigs.Emplace(ConfigData.ConfigName, ClothConfig);
-			}
-		}
-
-		const int32 NumClothingAssets = NewClothingAssetsData.Num(); 
-		check(NumClothingAssets == ClothingPhysicsAssets.Num());
-	
-		bool bAllNamesUnique = true;
-		TArray<FName, TInlineAllocator<8>> UniqueAssetNames;
-
-		for (int32 I = 0; I < NumClothingAssets; ++I)
-		{
-			// Skip assets not set.
-			if (!NewClothingAssetsData[I].LodData.Num())
-			{
-				continue;
-			}
-
-			const int32 PrevNumUniqueElems = UniqueAssetNames.Num();
-			const int32 ElemIndex = UniqueAssetNames.AddUnique(NewClothingAssetsData[I].Name);
-
-			if (ElemIndex < PrevNumUniqueElems)
-			{
-				bAllNamesUnique = false;
-				break;
-			}
-		}
-
-		for (int32 I = 0; I < NumClothingAssets; ++I)
-		{
-			// Skip assets not set.
-			if (!NewClothingAssetsData[I].LodData.Num())
-			{
-				continue;
-			}
-
-			FName ClothingAssetObjectName = bAllNamesUnique 
-					? NewClothingAssetsData[I].Name
-					: FName(FString::Printf(TEXT("%s_%d"), *NewClothingAssetsData[I].Name.ToString(), I));;
-
-			NewClothingAssets[I] = NewObject<UCustomizableObjectClothingAsset>(SkeletalMesh, ClothingAssetObjectName);
-
-			// The data can be moved to the actual asset since it will not be used anymore.
-			NewClothingAssets[I]->LodMap = MoveTemp(NewClothingAssetsData[I].LodMap);
-			NewClothingAssets[I]->LodData = MoveTemp(NewClothingAssetsData[I].LodData);
-			NewClothingAssets[I]->UsedBoneIndices = MoveTemp(NewClothingAssetsData[I].UsedBoneIndices);
-			NewClothingAssets[I]->UsedBoneNames = MoveTemp(NewClothingAssetsData[I].UsedBoneNames);
-			NewClothingAssets[I]->ReferenceBoneIndex = NewClothingAssetsData[I].ReferenceBoneIndex;
-			NewClothingAssets[I]->AssignNewGuid();
-			NewClothingAssets[I]->RefreshBoneMapping(SkeletalMesh);
-			NewClothingAssets[I]->CalculateReferenceBoneIndex();	
-			NewClothingAssets[I]->PhysicsAsset = ClothingPhysicsAssets[I];
-
-			for (const FCustomizableObjectClothConfigData& ConfigData : ClothingAssetsData[I].ConfigsData)
-			{
-				UClothConfigCommon* ClothConfig = CreateNewClothConfigFromData(NewClothingAssets[I], ConfigData);
-				if (ClothConfig)
-				{
-					NewClothingAssets[I]->ClothConfigs.Add(ConfigData.ConfigName, ClothConfig);
-				}
-			}
-
-			for (const TTuple<FName, UClothConfigCommon*>& SharedConfig : SharedConfigs)
-			{
-				NewClothingAssets[I]->ClothConfigs.Add(SharedConfig);
-			}
-
-			SkeletalMesh->GetMeshClothingAssets().AddUnique(NewClothingAssets[I]);
-		}	
-	}
-
-	for ( FSectionWithClothData& SectionWithCloth : SectionsWithCloth )
+	for (FSectionWithClothData& SectionWithCloth : SectionsWithCloth)
 	{
 		FSkeletalMeshLODRenderData& LODModel = RenderResource->LODRenderData[SectionWithCloth.Lod];
 		FSkelMeshRenderSection& SectionData = LODModel.RenderSections[SectionWithCloth.Section];
         
-		UClothingAssetCommon* NewClothingAsset = NewClothingAssets[SectionWithCloth.ClothAssetIndex];
-		if (!NewClothingAsset)
-		{
-			continue;
-		}
-       
+		FPerClothAssetData& AssetData = PerClothAssetData.FindChecked(SectionWithCloth.ClothAssetIndex);
+ 
 		SectionData.ClothMappingDataLODs.AddDefaulted(1);
 		SectionData.ClothMappingDataLODs[0] = MoveTemp(SectionWithCloth.MappingData);
 		
-		SectionData.CorrespondClothAssetIndex = SkeletalMesh->GetClothingAssetIndex(NewClothingAsset);
-		SectionData.ClothingData.AssetGuid = NewClothingAsset->GetAssetGuid();
-		SectionData.ClothingData.AssetLodIndex = SectionWithCloth.ClothAssetLodIndex;
+		SectionData.CorrespondClothAssetIndex = SkeletalMesh->GetClothingAssetIndex(AssetData.ClothingAsset);
+		SectionData.ClothingData.AssetGuid = AssetData.ClothingAsset->GetAssetGuid();
+		SectionData.ClothingData.AssetLodIndex = AssetData.ClothingAsset->LodMap[SectionWithCloth.Lod];
 	}
 
-	SkeletalMesh->SetHasActiveClothingAssets( static_cast<bool>( SectionsWithCloth.Num() ) );
+	SkeletalMesh->SetHasActiveClothingAssets(static_cast<bool>( SectionsWithCloth.Num()));
 }
 
+
 bool UCustomizableInstancePrivate::BuildOrCopyRenderData(const TSharedRef<FUpdateContextPrivate>& OperationData, USkeletalMesh* SkeletalMesh, const USkeletalMesh* LastUpdateSkeletalMesh, 
 	UCustomizableObjectInstance* Public, FCustomizableObjectInstanceComponentIndex InstanceComponentIndex)
 {
@@ -6115,7 +5644,6 @@ UE::Tasks::FTask UCustomizableInstancePrivate::LoadAdditionalAssetsAndData(const
 	}
 }
 
-
 FCustomizableObjectInstanceDescriptor& UCustomizableInstancePrivate::GetDescriptor() const
 {
 	return GetPublic()->Descriptor;
diff --git a/Engine/Plugins/Mutable/Source/CustomizableObject/Public/MuCO/Plugins/IMutableClothingModule.h b/Engine/Plugins/Mutable/Source/CustomizableObject/Public/MuCO/Plugins/IMutableClothingModule.h
new file mode 100644
--- /dev/null
+++ b/Engine/Plugins/Mutable/Source/CustomizableObject/Public/MuCO/Plugins/IMutableClothingModule.h
@@ -0,0 +1,51 @@
+// Copyright Epic Games, Inc. All Rights Reserved.
+
+#pragma once
+
+#include "Modules/ModuleManager.h"
+#include "CoreMinimal.h"
+
+#define MUTABLE_CLOTHING_MODULE_NAME "MutableClothing"
+
+struct FMeshToMeshVertData;
+class UClothingAssetCommon;
+
+class IMutableClothingModule : public IModuleInterface
+{
+public:
+
+	/**
+	 * Updates in-place the InSimulationLODIndex of InOutClothingAsset physics mesh to conform 
+	 * to InOutAttachedLODsRenderData. Remaps InOutAttachedLODsRenderData indices to be valid 
+	 * after the modifications have been applied.
+	 *
+	 * NOTE: After calling UpdateClothSimulationLOD LOD transition mapping may contain invalid data 
+	 * if any modification has taken place. In that case, FixLODTransitionMappings() must be called 
+	 * after all InOutClothingAsset LODs have been updated.
+	 *
+	 * @param InSimulationLODIndex 		  Lod index in InOutClothingAsset to modify.
+	 * @param InOutClothingAsset 		  Clothing asset to modify.
+	 * @param InOutAttachedLODsRenderData List of mapping datas for all meshes attached to InOutClothingAsset's InSimulationLODIndex.
+	 *
+	 * @return true if the clothing asset has been modified, false otherwise. 
+	 */
+	virtual bool UpdateClothSimulationLOD(
+			int32 InSimulationLODIndex,	
+			UClothingAssetCommon& InOutClothingAsset,
+			TConstArrayView<TArrayView<FMeshToMeshVertData>> InOutAttachedLODsRenderData) = 0;
+
+	/**
+	 * Fixes in-place LOD transition mappings for InSimulationLODIndex of InOutClothingAsset  
+	 * 
+	 * NOTE: FixLODTransitionMappings needs to run after all InOutClothingAsset LODs have been 
+	 * updated.
+	 *
+	 * @param InSimulationLODIndex LOD in InOutClothingAsset to fix.
+	 * @param InOutClothingAsset   Clothing asset to fix.
+	 */
+	virtual void FixLODTransitionMappings(
+			int32 InSimulationLODIndex, 
+			UClothingAssetCommon& InOutClothingAsset) = 0;
+
+};
+
