[mutable] Reduced system settings change scope so it only affects the baked instance.
- Added system settings override object to run the update with some custom system settings. It clears the changes out once the update has completed.
- Moved system settings change from pre-schedule time to on start update.
- Moved system settings restore from during Update callbacks to just before those callbacks.
- Fixed some typos

diff --git a/Engine/Plugins/Mutable/Source/CustomizableObject/CustomizableObject.Build.cs b/Engine/Plugins/Mutable/Source/CustomizableObject/CustomizableObject.Build.cs
--- a/Engine/Plugins/Mutable/Source/CustomizableObject/CustomizableObject.Build.cs
+++ b/Engine/Plugins/Mutable/Source/CustomizableObject/CustomizableObject.Build.cs
@@ -13,8 +13,8 @@ public class CustomizableObject : ModuleRules
 		bAllowConfidentialPlatformDefines = true;
 		//bUseUnity = false;
 
-		PrivateDependencyModuleNames.AddRange(new string[] {
-            "InputCore",
+		PrivateDependencyModuleNames.AddRange([
+			"InputCore",
             "SlateCore",
             "CoreUObject",
             "RenderCore",
@@ -25,19 +25,19 @@ public class CustomizableObject : ModuleRules
 			"ClothingSystemRuntimeInterface",
 			//"ClothingSystemEditor",
 			"UMG",
-			"ImageCore",
-		});
+			"ImageCore"
+		]);
 
-        PublicDependencyModuleNames.AddRange(new string[] {
-            "Slate",
+        PublicDependencyModuleNames.AddRange([
+	        "Slate",
 			"Core",
             "Engine",
 			"SkeletalMerging",
 			"ClothingSystemRuntimeCommon",
             "GameplayTags",
 			"MutableRuntime",
-			"AnimGraphRuntime",
-		});
+			"AnimGraphRuntime"
+        ]);
 
 		if (TargetRules.bBuildEditor == true)
         {
@@ -45,6 +45,8 @@ public class CustomizableObject : ModuleRules
             PublicDependencyModuleNames.Add("DerivedDataCache");
             PublicDependencyModuleNames.Add("EditorStyle");
             PublicDependencyModuleNames.Add("MessageLog");
+
+            PrivateDependencyModuleNames.Add("MutableTools");		// COI Baking muSystem settings override
         }
     }
 }
diff --git a/Engine/Plugins/Mutable/Source/CustomizableObject/Internal/MuCO/CustomizableObjectInstancePrivate.h b/Engine/Plugins/Mutable/Source/CustomizableObject/Internal/MuCO/CustomizableObjectInstancePrivate.h
--- a/Engine/Plugins/Mutable/Source/CustomizableObject/Internal/MuCO/CustomizableObjectInstancePrivate.h
+++ b/Engine/Plugins/Mutable/Source/CustomizableObject/Internal/MuCO/CustomizableObjectInstancePrivate.h
@@ -366,6 +366,10 @@ public:
 
 	/** Finds in ProjectorParameters a parameter with name ParamName, returns the index if found, -1 otherwise. */
 	int32 FindProjectorParameterNameIndex(const FString& ParamName) const;
+
+#if WITH_EDITOR
+	void UpdateSkeletalMeshAsyncResult(FInstanceUpdateNativeDelegate Callback, bool bIgnoreCloseDist = false, bool bForceHighPriority = false, TSharedPtr<FMutableSystemSettingsOverrides> MutableSystemSettingsOverride = nullptr);
+#endif
 	
 	UCustomizableObjectInstance* GetPublic() const;
 
diff --git a/Engine/Plugins/Mutable/Source/CustomizableObject/Internal/MuCO/CustomizableObjectSystemPrivate.h b/Engine/Plugins/Mutable/Source/CustomizableObject/Internal/MuCO/CustomizableObjectSystemPrivate.h
--- a/Engine/Plugins/Mutable/Source/CustomizableObject/Internal/MuCO/CustomizableObjectSystemPrivate.h
+++ b/Engine/Plugins/Mutable/Source/CustomizableObject/Internal/MuCO/CustomizableObjectSystemPrivate.h
@@ -50,6 +50,40 @@ extern TAutoConsoleVariable<bool> CVarMutableHighPriorityLoading;
 void FinishUpdateGlobal(const TSharedRef<FUpdateContextPrivate>& Context);
 
 
+#if WITH_EDITOR
+
+/**
+ * Class used to hold some MutableSystem settings to be used during the update of a given instance.
+ * It will store the settings used by the system during the setup of this object and so later can be reverted back.
+ */
+class CUSTOMIZABLEOBJECT_API FMutableSystemSettingsOverrides
+{
+public:
+	FMutableSystemSettingsOverrides(bool bUseProgressiveMipStreaming, bool bOnlyGenerateRequestedLODs, mu::FImageOperator::FImagePixelFormatFunc InImagePixelFormatFunc);
+	
+	/** Apply the settings set in this object to the mutable system. */
+	void ApplySettingsOverrides() const;
+
+	/** Restore the mutable system settings to a state prior to the setup of this object */
+	void RestoreSettings() const;
+
+
+private:
+	
+	// Requested values
+	bool bIsProgressiveMipStreamingEnabled = false;
+	bool bIsOnlyGenerateRequestedLODsEnabled = false;
+	mu::FImageOperator::FImagePixelFormatFunc ImagePixelFormatFunc;
+	
+	// previous values. Cached at the time of creating this object
+	bool bOldIsProgressiveMipStreamingEnabled = false;
+	bool bOldIsOnlyGenerateRequestedLODsEnabled = false;
+	mu::FImageOperator::FImagePixelFormatFunc OldImagePixelFormatFunc;
+};
+
+#endif
+
+
 /** Strongly typed index for the index of a component in a UCustomizableObjectInstance. */
 class FCustomizableObjectInstanceComponentIndex
 {
@@ -664,7 +698,11 @@ public:
 	TSet<TWeakObjectPtr<UCustomizableObjectInstanceUsage>> AttachedParentUpdated;
 	
 	/** Hard references to objects. Avoids GC to collect them. */
-	TArray<TStrongObjectPtr<const UObject>> Objects;	
+	TArray<TStrongObjectPtr<const UObject>> Objects;
+
+#if WITH_EDITORONLY_DATA
+	TSharedPtr<class FMutableSystemSettingsOverrides> UpdateSettingsOverride = nullptr;
+#endif
 };
 
 
diff --git a/Engine/Plugins/Mutable/Source/CustomizableObject/Private/MuCO/CustomizableObjectInstance.cpp b/Engine/Plugins/Mutable/Source/CustomizableObject/Private/MuCO/CustomizableObjectInstance.cpp
--- a/Engine/Plugins/Mutable/Source/CustomizableObject/Private/MuCO/CustomizableObjectInstance.cpp
+++ b/Engine/Plugins/Mutable/Source/CustomizableObject/Private/MuCO/CustomizableObjectInstance.cpp
@@ -2914,6 +2914,23 @@ int32 UCustomizableInstancePrivate::FindProjectorParameterNameIndex(const FStrin
 }
 
 
+#if WITH_EDITOR
+
+void UCustomizableInstancePrivate::UpdateSkeletalMeshAsyncResult(FInstanceUpdateNativeDelegate Callback, bool bIgnoreCloseDist, bool bForceHighPriority, TSharedPtr<FMutableSystemSettingsOverrides> MutableSystemSettingsOverride)
+{
+	UCustomizableObjectSystemPrivate* SystemPrivate = UCustomizableObjectSystem::GetInstance()->GetPrivate();
+
+	const TSharedRef<FUpdateContextPrivate> Context = MakeShared<FUpdateContextPrivate>(*GetPublic());
+	Context->bIgnoreCloseDist = bIgnoreCloseDist;
+	Context->bForceHighPriority = bForceHighPriority;
+	Context->UpdateNativeCallback = Callback;
+	Context->UpdateSettingsOverride = MutableSystemSettingsOverride;
+	
+	SystemPrivate->EnqueueUpdateSkeletalMesh(Context);
+}
+
+#endif
+
 void UCustomizableObjectInstance::SetRandomValues()
 {
 	Descriptor.SetRandomValues();
diff --git a/Engine/Plugins/Mutable/Source/CustomizableObject/Private/MuCO/CustomizableObjectSystem.cpp b/Engine/Plugins/Mutable/Source/CustomizableObject/Private/MuCO/CustomizableObjectSystem.cpp
--- a/Engine/Plugins/Mutable/Source/CustomizableObject/Private/MuCO/CustomizableObjectSystem.cpp
+++ b/Engine/Plugins/Mutable/Source/CustomizableObject/Private/MuCO/CustomizableObjectSystem.cpp
@@ -38,6 +38,7 @@
 #include "Interfaces/ITargetPlatform.h"
 #include "MuCO/ICustomizableObjectEditorModule.h"
 #include "MuCO/EditorImageProvider.h"
+#include "MuT/UnrealPixelFormatOverride.h"
 #else
 #include "Engine/Engine.h"
 #endif
@@ -176,6 +177,73 @@ static void CVarMutableSinkFunction()
 	}
 }
 
+#if WITH_EDITOR
+
+FMutableSystemSettingsOverrides::FMutableSystemSettingsOverrides(bool bUseProgressiveMipStreaming, bool bOnlyGenerateRequestedLODs,
+	mu::FImageOperator::FImagePixelFormatFunc InImagePixelFormatFunc)
+{
+	// Cache previous settings
+	UCustomizableObjectSystem* System = UCustomizableObjectSystem::GetInstanceChecked();
+	bOldIsProgressiveMipStreamingEnabled = System->IsProgressiveMipStreamingEnabled();
+	bOldIsOnlyGenerateRequestedLODsEnabled = System->IsOnlyGenerateRequestedLODsEnabled();
+	OldImagePixelFormatFunc = System->GetPrivate()->ImageFormatOverrideFunc;
+	
+	// Cache desired settings
+	bIsProgressiveMipStreamingEnabled = bUseProgressiveMipStreaming;
+	bIsOnlyGenerateRequestedLODsEnabled = bOnlyGenerateRequestedLODs;
+	ImagePixelFormatFunc = InImagePixelFormatFunc;
+}
+
+
+void FMutableSystemSettingsOverrides::ApplySettingsOverrides() const
+{
+	UE_LOG(LogMutable, Display, TEXT("Overriding system settings for this instance update."))
+	
+	UCustomizableObjectSystem* System = UCustomizableObjectSystem::GetInstanceChecked();
+	System->SetProgressiveMipStreamingEnabled(bIsProgressiveMipStreamingEnabled);
+	System->SetOnlyGenerateRequestedLODsEnabled(bIsOnlyGenerateRequestedLODsEnabled);
+	
+	if (ensure(IsInGameThread()))
+	{
+		PrepareUnrealCompression();
+		
+		UCustomizableObjectSystemPrivate* Private = System->GetPrivate();
+		check(Private)
+		Private->ImageFormatOverrideFunc = ImagePixelFormatFunc;
+	}
+	else
+	{
+		UE_LOG(LogMutable, Error, TEXT("Failed to apply the MutableSystem Pixel Format override. Invocation of method performed outside the MainThread"))
+	}
+}
+
+
+void FMutableSystemSettingsOverrides::RestoreSettings() const
+{
+	UE_LOG(LogMutable, Display, TEXT("Restoring system settings."))
+	
+ 	// apply the settings to the system using the old data
+ 	UCustomizableObjectSystem* System = UCustomizableObjectSystem::GetInstanceChecked();
+ 	System->SetProgressiveMipStreamingEnabled(bOldIsProgressiveMipStreamingEnabled);
+ 	System->SetOnlyGenerateRequestedLODsEnabled(bOldIsOnlyGenerateRequestedLODsEnabled);
+
+	if (ensure(IsInGameThread()))
+	{
+		PrepareUnrealCompression();
+		
+		UCustomizableObjectSystemPrivate* Private = System->GetPrivate();
+		check(Private)
+		Private->ImageFormatOverrideFunc = OldImagePixelFormatFunc;
+	}
+	else
+	{
+		UE_LOG(LogMutable, Error, TEXT("Failed to restore the MutableSystem Pixel Format override. Invocation of method performed outside the MainThread"))
+	}
+}
+
+#endif
+
+
 static FAutoConsoleVariableSink CVarMutableSink(FConsoleCommandDelegate::CreateStatic(&CVarMutableSinkFunction));
 
 
@@ -1081,7 +1149,16 @@ void FinishUpdateGlobal(const TSharedRef<FUpdateContextPrivate>& Context)
 
 	FUpdateContext ContextPublic;
 	ContextPublic.UpdateResult = Context->UpdateResult;
-		
+
+#if WITH_EDITORONLY_DATA
+	// Restore the overriden system settings (bake operation)
+	if (Context->UpdateSettingsOverride)
+	{
+		Context->UpdateSettingsOverride->RestoreSettings();
+		Context->UpdateSettingsOverride = nullptr;
+	}
+#endif
+	
 	Context->UpdateCallback.ExecuteIfBound(ContextPublic);
 	Context->UpdateNativeCallback.Broadcast(ContextPublic);
 
@@ -4172,6 +4249,14 @@ int32 UCustomizableObjectSystem::TickInternal(const bool bBlocking)
 				UCustomizableObjectInstance* PendingInstance = PendingInstanceUpdateFound->Context->Instance.Get();
 				check(PendingInstance);
 
+#if WITH_EDITORONLY_DATA
+				// Used for the baking of instances so we do not want it to be possible to run in a package
+				if (PendingInstanceUpdateFound->Context->UpdateSettingsOverride)
+				{
+					PendingInstanceUpdateFound->Context->UpdateSettingsOverride->ApplySettingsOverrides();
+				}
+#endif
+				
 				// Maybe there's a LODUpdate that has the same instance, merge both updates as an optimization
 				FMutableUpdateCandidate* LODUpdateWithSameInstance = RequestedLODUpdates.Find(PendingInstance);
 				
@@ -4528,6 +4613,7 @@ void UCustomizableObjectSystem::SetOnlyGenerateRequestedLODsEnabled(bool bIsEnab
 
 
 #if WITH_EDITOR
+UE_DEPRECATED(5.6, "This method has been deprecated.")
 void UCustomizableObjectSystem::SetImagePixelFormatOverride(const mu::FImageOperator::FImagePixelFormatFunc& InFunc)
 {
 	if (Private != nullptr)
diff --git a/Engine/Plugins/Mutable/Source/CustomizableObjectEditor/Internal/MuCOE/CustomizableObjectInstanceBakingUtils.h b/Engine/Plugins/Mutable/Source/CustomizableObjectEditor/Internal/MuCOE/CustomizableObjectInstanceBakingUtils.h
--- a/Engine/Plugins/Mutable/Source/CustomizableObjectEditor/Internal/MuCOE/CustomizableObjectInstanceBakingUtils.h
+++ b/Engine/Plugins/Mutable/Source/CustomizableObjectEditor/Internal/MuCOE/CustomizableObjectInstanceBakingUtils.h
@@ -10,7 +10,7 @@
  * @param InInstanceUpdateDelegate Delegate to be called after the instance gets updated.
  */
 CUSTOMIZABLEOBJECTEDITOR_API
-void UpdateInstanceForBaking(UCustomizableObjectInstance& InInstance, FInstanceUpdateNativeDelegate& InInstanceUpdateDelegate);
+void ScheduleInstanceUpdateForBaking(UCustomizableObjectInstance& InInstance, FInstanceUpdateNativeDelegate& InInstanceUpdateDelegate);
 
 
 /**
diff --git a/Engine/Plugins/Mutable/Source/CustomizableObjectEditor/Private/MuCOE/CustomizableObjectEditorViewportClient.cpp b/Engine/Plugins/Mutable/Source/CustomizableObjectEditor/Private/MuCOE/CustomizableObjectEditorViewportClient.cpp
--- a/Engine/Plugins/Mutable/Source/CustomizableObjectEditor/Private/MuCOE/CustomizableObjectEditorViewportClient.cpp
+++ b/Engine/Plugins/Mutable/Source/CustomizableObjectEditor/Private/MuCOE/CustomizableObjectEditorViewportClient.cpp
@@ -1571,7 +1571,7 @@ void FCustomizableObjectEditorViewportClient::BakeInstance()
 	// Call the instance update async method
 	FInstanceUpdateNativeDelegate UpdateDelegate;
 	UpdateDelegate.AddRaw(this, &FCustomizableObjectEditorViewportClient::OnInstanceForBakingUpdate);
-	UpdateInstanceForBaking(*BakeTempInstance, UpdateDelegate);
+	ScheduleInstanceUpdateForBaking(*BakeTempInstance, UpdateDelegate);
 }
 
 
diff --git a/Engine/Plugins/Mutable/Source/CustomizableObjectEditor/Private/MuCOE/CustomizableObjectInstanceBaker.cpp b/Engine/Plugins/Mutable/Source/CustomizableObjectEditor/Private/MuCOE/CustomizableObjectInstanceBaker.cpp
--- a/Engine/Plugins/Mutable/Source/CustomizableObjectEditor/Private/MuCOE/CustomizableObjectInstanceBaker.cpp
+++ b/Engine/Plugins/Mutable/Source/CustomizableObjectEditor/Private/MuCOE/CustomizableObjectInstanceBaker.cpp
@@ -5,7 +5,6 @@
 #include "MuCO/CustomizableObject.h"
 #include "MuCO/CustomizableObjectInstance.h"
 #include "MuCO/CustomizableObjectInstancePrivate.h"
-#include "MuCO/CustomizableObjectSystem.h"
 #include "MuCOE/CustomizableObjectInstanceBakingUtils.h"
 
 
@@ -27,7 +26,7 @@ void UCustomizableObjectInstanceBaker::BakeInstance(UCustomizableObjectInstance*
 		return;
 	}
 
-	// We clone the InTargetInstance to ensure that the baked it is not open in a COI edior.
+	// We clone the InTargetInstance to ensure that the baked it is not open in a COI editor.
 	InstanceToBake = InTargetInstance->Clone();
 	
 	if (Configuration.OutputPath.IsEmpty())
@@ -52,7 +51,7 @@ void UCustomizableObjectInstanceBaker::UpdateInstance()
 	UpdateDelegate.AddUObject(this, &UCustomizableObjectInstanceBaker::OnInstanceUpdate);
 	
 	// Once the update finishes successfully the baking operation will engage
-	UpdateInstanceForBaking(*InstanceToBake, UpdateDelegate);
+	ScheduleInstanceUpdateForBaking(*InstanceToBake, UpdateDelegate);
 }
 
 
diff --git a/Engine/Plugins/Mutable/Source/CustomizableObjectEditor/Private/MuCOE/CustomizableObjectInstanceBaker.h b/Engine/Plugins/Mutable/Source/CustomizableObjectEditor/Private/MuCOE/CustomizableObjectInstanceBaker.h
--- a/Engine/Plugins/Mutable/Source/CustomizableObjectEditor/Private/MuCOE/CustomizableObjectInstanceBaker.h
+++ b/Engine/Plugins/Mutable/Source/CustomizableObjectEditor/Private/MuCOE/CustomizableObjectInstanceBaker.h
@@ -25,7 +25,7 @@ DECLARE_DELEGATE(FOnBakerFinishedWork);
  *
  * Since the update of the instance for baking requires of some static changes in the instance updating system so the parallel update of instances
  * for baking is not allowed. This prohibition is enforced by the update methods that we call prior to the baking operation. No more than one update for baking can
- * be ran at the same time.
+ * be run at the same time.
  */
 UCLASS()
 class CUSTOMIZABLEOBJECTEDITOR_API UCustomizableObjectInstanceBaker : public UObject
@@ -37,8 +37,7 @@ public:
 	 * Execute this method in order to bake the provided instance. It will schedule a special type of instance update before proceeding with the bake itself.
 	 * @param InTargetInstance The instance we want to bake
 	 * @param InBakingConfig Configuration structure that determines how the baking is going to be made
-	 * @param InOnBakerFinishedWork Delegate used to determine if the baker has completed it's operation. Does not provide any other data than that so we can remove it from
-	 * the root safely (or perform any other cleanup)
+	 * @param InOnBakerFinishedWork Delegate used to determine if the baker has completed its operation. Does not provide any other data than that so we can perform some post-bake actions.
 	 */
 	void BakeInstance(UCustomizableObjectInstance* InTargetInstance, const FBakingConfiguration& InBakingConfig, const TSharedPtr<const FOnBakerFinishedWork> InOnBakerFinishedWork = nullptr);
 	
@@ -81,7 +80,7 @@ private:
 	TObjectPtr<UCustomizableObjectInstance> InstanceToBake;	
 	
 	/**
-	 * Optional delegate used to report when the baker has completed it's work (and it not running or making other systems run)
+	 * Optional delegate used to report when the baker has completed it's work (and it is not running or making other systems run)
 	 * It will be called after we have given the end user the data about the success state of the baking and also the resources saved.
 	 */
 	TSharedPtr<const FOnBakerFinishedWork> OnWorkFinished;
diff --git a/Engine/Plugins/Mutable/Source/CustomizableObjectEditor/Private/MuCOE/CustomizableObjectInstanceBakingUtils.cpp b/Engine/Plugins/Mutable/Source/CustomizableObjectEditor/Private/MuCOE/CustomizableObjectInstanceBakingUtils.cpp
--- a/Engine/Plugins/Mutable/Source/CustomizableObjectEditor/Private/MuCOE/CustomizableObjectInstanceBakingUtils.cpp
+++ b/Engine/Plugins/Mutable/Source/CustomizableObjectEditor/Private/MuCOE/CustomizableObjectInstanceBakingUtils.cpp
@@ -283,72 +283,42 @@ bool ManageBakingAction(const FString& InAssetSavePath, const FString& InObjName
 	return true;
 }
 
+static bool bIsBakeOperationAlreadyScheduled = false;
 
-namespace PreBakeSystemSettings
-{
-	bool bIsProgressiveMipStreamingEnabled = false;
-	bool bIsOnlyGenerateRequestedLODsEnabled = false;
-}
-
-
-// Prevents the execution of the baking in parallel for the baking operation. It will not prevent other updates from running (not baking updates)
-// So you are encouraged to halt all other updates while you are baking instances
-static bool bIsUpdateForBakingRunning = false;
 
-
-void PrepareForBaking()
+void ResetBakeOperationAlreadyRunning(const FUpdateContext& Result)
 {
-	// Implementation of the bake operation
-	UCustomizableObjectSystem* System = UCustomizableObjectSystem::GetInstanceChecked();
-	
-	// The instance in the editor viewport does not have high quality mips in the platform data because streaming is enabled.
-	// Disable streaming and retry with a newly generated temp instance.
-	PreBakeSystemSettings::bIsProgressiveMipStreamingEnabled = System->IsProgressiveMipStreamingEnabled();
-	System->SetProgressiveMipStreamingEnabled(false);
-	// Disable requested LOD generation as it will prevent the new instance from having all the LODs
-	PreBakeSystemSettings::bIsOnlyGenerateRequestedLODsEnabled = System->IsOnlyGenerateRequestedLODsEnabled();
-	System->SetOnlyGenerateRequestedLODsEnabled(false);
-	// Force high quality texture compression for this instance
-	PrepareUnrealCompression();
-	System->SetImagePixelFormatOverride(UnrealPixelFormatFunc);
+	bIsBakeOperationAlreadyScheduled = false;
 }
 
 
-void RestoreCustomizableObjectSettings(const FUpdateContext& Result)
+void ScheduleInstanceUpdateForBaking(UCustomizableObjectInstance& InInstance, FInstanceUpdateNativeDelegate& InInstanceUpdateDelegate)
 {
-	// Reenable Mutable texture streaming and requested LOD generation as they had been disabled to bake the textures
-	UCustomizableObjectSystem* System =  UCustomizableObjectSystem::GetInstanceChecked();
-	System->SetProgressiveMipStreamingEnabled(PreBakeSystemSettings::bIsProgressiveMipStreamingEnabled);
-	System->SetOnlyGenerateRequestedLODsEnabled(PreBakeSystemSettings::bIsOnlyGenerateRequestedLODsEnabled);
-	System->SetImagePixelFormatOverride(nullptr);
-
-	// Tell the system we have finished the update and that we can schedule another "for baking" update
-	bIsUpdateForBakingRunning = false;
-}
-
-
-void UpdateInstanceForBaking(UCustomizableObjectInstance& InInstance, FInstanceUpdateNativeDelegate& InInstanceUpdateDelegate)
-{
-	if (bIsUpdateForBakingRunning)
+	// This prevents the queue of updates to have more than one instance for baking.
+	if (bIsBakeOperationAlreadyScheduled)
 	{
 		UE_LOG(LogMutable, Error, TEXT("The COInstance update for baking could not be scheduled. Another instance is being updated for baking."));
 		InInstanceUpdateDelegate.Broadcast({EUpdateResult::Error});
 		return;
 	}
 
-	// Set the update for the baking of the instance as running so we prevent other baking updates while we do run our own
-	bIsUpdateForBakingRunning = true;
+	// Allow the baking of more instances once the bake of this one has completed (or at least until the callbacks have been broadcast)
+	InInstanceUpdateDelegate.AddStatic(&ResetBakeOperationAlreadyRunning);
 	
-	// Prepare the customizable object system for baking
-	PrepareForBaking();
-	
-	// Ensure we clear the changes in the COSystem after performing the update so later updates do not get affected
-	InInstanceUpdateDelegate.AddStatic(&RestoreCustomizableObjectSettings);
+	// Desired System settings for the update part of a baking operation
+	TSharedPtr<FMutableSystemSettingsOverrides> SystemSettingsOverride =
+		MakeShared<FMutableSystemSettingsOverrides>(
+			false,
+			false,
+			UnrealPixelFormatFunc);
 	
 	// Schedule the update
-	InInstance.UpdateSkeletalMeshAsyncResult(InInstanceUpdateDelegate,true,true);
-
-	UE_LOG(LogMutable, Display, TEXT("The COInstance Update operation for baking was succesfuly scheduled."));
+	check (InInstance.GetPrivate());
+	InInstance.GetPrivate()->UpdateSkeletalMeshAsyncResult(InInstanceUpdateDelegate,true,true, SystemSettingsOverride);
+	
+	// Tell the baking system that an instance for baking will be processed
+	bIsBakeOperationAlreadyScheduled = true;
+	UE_LOG(LogMutable, Display, TEXT("The COInstance Update operation for baking has been scheduled for the instance '%s'."), *InInstance.GetName());
 }
 
 
@@ -371,7 +341,7 @@ bool BakeCustomizableObjectInstance(
 	if (!InstanceCO || InstanceCO->GetPrivate()->IsLocked())
 	{
 		FCustomizableObjectEditorLogger::CreateLog(
-		LOCTEXT("CustomizableObjecInstanceBakingUtils_LockedObject", "Please wait until the Customizable Object is compiled"))
+		LOCTEXT("CustomizableObjectInstanceBakingUtils_LockedObject", "Please wait until the Customizable Object is compiled"))
 		.Category(ELoggerCategory::COInstanceBaking)
 		.CustomNotification()
 		.Notification(true)
@@ -383,7 +353,7 @@ bool BakeCustomizableObjectInstance(
 	if (InstanceCO->GetPrivate()->Status.Get() == FCustomizableObjectStatus::EState::Loading)
 	{
 		FCustomizableObjectEditorLogger::CreateLog(
-			LOCTEXT("CustomizableObjecInstanceBakingUtils_LoadingObject","Please wait unitl Customizable Object is loaded"))
+			LOCTEXT("CustomizableObjectInstanceBakingUtils_LoadingObject","Please wait until the Customizable Object is loaded"))
 		.Category(ELoggerCategory::COInstanceBaking)
 		.CustomNotification()
 		.Notification(true)
@@ -425,7 +395,7 @@ bool BakeCustomizableObjectInstance(
 	if (ModelResources->bCompiledWithHDTextureCompression == false)
 	{
 		FCustomizableObjectEditorLogger::CreateLog(
-		LOCTEXT("CustomizableObjecInstanceBakingUtils_LowQualityTextures", "The Customizable Object wasn't compiled with high quality textures. For the best baking results, change the Texture Compression setting and recompile it."))
+		LOCTEXT("CustomizableObjectInstanceBakingUtils_LowQualityTextures", "The Customizable Object wasn't compiled with high quality textures. For the best baking results, change the Texture Compression setting and recompile it."))
 		.Category(ELoggerCategory::COInstanceBaking)
 		.CustomNotification()
 		.Notification(true)
diff --git a/Engine/Plugins/Mutable/Source/CustomizableObjectEditor/Private/MuCOE/Nodes/CustomizableObjectNodeExposePin.h b/Engine/Plugins/Mutable/Source/CustomizableObjectEditor/Private/MuCOE/Nodes/CustomizableObjectNodeExposePin.h
--- a/Engine/Plugins/Mutable/Source/CustomizableObjectEditor/Private/MuCOE/Nodes/CustomizableObjectNodeExposePin.h
+++ b/Engine/Plugins/Mutable/Source/CustomizableObjectEditor/Private/MuCOE/Nodes/CustomizableObjectNodeExposePin.h
@@ -54,7 +54,7 @@ public:
 		return FindPin(InputPinName);
 	}
 
-	/** Broadcasted when the UPROPERTY Name changes. */
+	/** Will be broadcast when the UPROPERTY Name changes. */
 	FOnNameChangedDelegate OnNameChangedDelegate;
 	
 private:
diff --git a/Engine/Plugins/Mutable/Source/MutableValidation/Private/MuV/COIBakingTestCommandlet.cpp b/Engine/Plugins/Mutable/Source/MutableValidation/Private/MuV/COIBakingTestCommandlet.cpp
--- a/Engine/Plugins/Mutable/Source/MutableValidation/Private/MuV/COIBakingTestCommandlet.cpp
+++ b/Engine/Plugins/Mutable/Source/MutableValidation/Private/MuV/COIBakingTestCommandlet.cpp
@@ -71,7 +71,7 @@ int32 UCOIBakingTestCommandlet::Main(const FString& Params)
 			return 1;
 		}
 
-		UE_LOG(LogMutable,Display,TEXT("Succesfully loaded %s Customizable Object Instance!"), *TargetInstance->GetName());
+		UE_LOG(LogMutable,Display,TEXT("Successfully loaded %s Customizable Object Instance!"), *TargetInstance->GetName());
 	}
 
 	// Perform a blocking search to ensure all assets used by mutable are reachable using the AssetRegistry
@@ -99,7 +99,7 @@ int32 UCOIBakingTestCommandlet::Main(const FString& Params)
 	TSharedRef<FCustomizableObjectCompilationUtility> CompilationUtility = MakeShared<FCustomizableObjectCompilationUtility>();
 	if( !CompilationUtility->CompileCustomizableObject(InstanceCustomizableObject, true, &CompilationOptions ))
 	{
-		UE_LOG(LogMutable,Error,TEXT("Failed to compile the target CO. Exitting commandlet."));
+		UE_LOG(LogMutable,Error,TEXT("Failed to compile the target CO. Exiting commandlet."));
 		return 1;
 	}
 	
@@ -117,7 +117,7 @@ int32 UCOIBakingTestCommandlet::Main(const FString& Params)
 		
 		bIsInstanceBeingUpdated = true;
 		UE_LOG(LogMutable, Display, TEXT("Scheduling instance update."));
-		UpdateInstanceForBaking(*TargetInstance, InstanceUpdateDelegate);
+		ScheduleInstanceUpdateForBaking(*TargetInstance, InstanceUpdateDelegate);
 		
 		// Now tick the engine so the instance gets updated while running in the commandlet context
         while (bIsInstanceBeingUpdated)
