[Mutable] UCustomizableObjectSystem Ticker, removed Core Ticker alternative.

We have been using a Streaming Manager for a long time and we had no issues.

diff --git a/Engine/Plugins/Mutable/Source/CustomizableObject/Internal/MuCO/CustomizableObjectSystemPrivate.h b/Engine/Plugins/Mutable/Source/CustomizableObject/Internal/MuCO/CustomizableObjectSystemPrivate.h
--- a/Engine/Plugins/Mutable/Source/CustomizableObject/Internal/MuCO/CustomizableObjectSystemPrivate.h
+++ b/Engine/Plugins/Mutable/Source/CustomizableObject/Internal/MuCO/CustomizableObjectSystemPrivate.h
@@ -834,7 +834,6 @@ public:
 	TSharedPtr<FUpdateContextPrivate> CurrentMutableOperation = nullptr;
 
 	// Handle to the registered TickDelegate.
-	FTSTicker::FDelegateHandle TickDelegateHandle;
 	FTSTicker::FDelegateHandle TickWarningsDelegateHandle;
 
 	/** Change the current status of Mutable. Enabling/Disabling core features.	
diff --git a/Engine/Plugins/Mutable/Source/CustomizableObject/Private/MuCO/CustomizableObjectSettings.h b/Engine/Plugins/Mutable/Source/CustomizableObject/Private/MuCO/CustomizableObjectSettings.h
deleted file mode 100644
diff --git a/Engine/Plugins/Mutable/Source/CustomizableObject/Private/MuCO/CustomizableObjectSystem.cpp b/Engine/Plugins/Mutable/Source/CustomizableObject/Private/MuCO/CustomizableObjectSystem.cpp
--- a/Engine/Plugins/Mutable/Source/CustomizableObject/Private/MuCO/CustomizableObjectSystem.cpp
+++ b/Engine/Plugins/Mutable/Source/CustomizableObject/Private/MuCO/CustomizableObjectSystem.cpp
@@ -28,7 +28,6 @@
 #include "ContentStreaming.h"
 #include "Components/SkeletalMeshComponent.h"
 #include "MuCO/CustomizableObjectSystemPrivate.h"
-#include "CustomizableObjectSettings.h"
 
 #if WITH_EDITOR
 #include "Editor.h"
@@ -860,15 +859,8 @@ void UCustomizableObjectSystem::BeginDestroy()
 #endif
 
 #if !UE_SERVER
-		if (GetMutableDefault<UCustomizableObjectSettings>()->bEnableStreamingManager)
-		{
-			FStreamingManagerCollection::Get().RemoveStreamingManager(GetPrivate());
-		}
-		else
-		{
-			FTSTicker::GetCoreTicker().RemoveTicker(Private->TickDelegateHandle);			
-		}
-
+		FStreamingManagerCollection::Get().RemoveStreamingManager(GetPrivate());
+		
 		FTSTicker::GetCoreTicker().RemoveTicker(Private->TickWarningsDelegateHandle);
 #endif // !UE_SERVER
 		
@@ -4308,16 +4300,7 @@ int32 UCustomizableObjectSystem::TickInternal(const bool bBlocking)
 
 	if (!bIsMutableEnabled && !Private->CurrentMutableOperation)
 	{
-		if (GetMutableDefault<UCustomizableObjectSettings>()->bEnableStreamingManager)
-		{
-			FStreamingManagerCollection::Get().RemoveStreamingManager(GetPrivate());
-		}
-		else
-		{
-			// Mutable has been disabled. Unregister the ticker if there is no CurrentMutableOperation.
-			FTSTicker::GetCoreTicker().RemoveTicker(Private->TickDelegateHandle);
-			Private->TickDelegateHandle.Reset();
-		}
+		FStreamingManagerCollection::Get().RemoveStreamingManager(GetPrivate());
 	}
 
 	int32 RemainingWork = Private->CurrentMutableOperation.IsValid() + 
@@ -4810,18 +4793,8 @@ void UCustomizableObjectSystemPrivate::OnMutableEnabledChanged(IConsoleVariable*
 	if (bIsMutableEnabled)
 	{
 #if !UE_SERVER
-		if (GetMutableDefault<UCustomizableObjectSettings>()->bEnableStreamingManager)
-		{
-			FStreamingManagerCollection::Get().RemoveStreamingManager(SystemPrivate); // Avoid being added twice
-			FStreamingManagerCollection::Get().AddStreamingManager(SystemPrivate);
-		}
-		else
-		{
-			if (!SystemPrivate->TickDelegateHandle.IsValid())
-			{
-				SystemPrivate->TickDelegateHandle = FTSTicker::GetCoreTicker().AddTicker(FTickerDelegate::CreateUObject(System, &UCustomizableObjectSystem::Tick), 0.f);
-			}
-		}
+		FStreamingManagerCollection::Get().RemoveStreamingManager(SystemPrivate); // Avoid being added twice
+        FStreamingManagerCollection::Get().AddStreamingManager(SystemPrivate);
 
 		if (!SystemPrivate->TickWarningsDelegateHandle.IsValid())
 		{
