[Mutable] Fix cook indeterminism. Data from Clothing and RealTimeMorphs stomped data from other resource types.

#rnx

diff --git a/Engine/Plugins/Mutable/Source/CustomizableObject/Internal/MuCO/CustomizableObjectPrivate.h b/Engine/Plugins/Mutable/Source/CustomizableObject/Internal/MuCO/CustomizableObjectPrivate.h
--- a/Engine/Plugins/Mutable/Source/CustomizableObject/Internal/MuCO/CustomizableObjectPrivate.h
+++ b/Engine/Plugins/Mutable/Source/CustomizableObject/Internal/MuCO/CustomizableObjectPrivate.h
@@ -1846,6 +1846,6 @@ public:
 	
 	static constexpr int32 CurrentSupportedVersion = ECustomizableObjectVersions::LastCustomizableObjectVersion;
 	
-	static constexpr int32 DerivedDataVersion = 3;
+	static constexpr int32 DerivedDataVersion = 4;
 };
 
diff --git a/Engine/Plugins/Mutable/Source/CustomizableObject/Private/MuCO/CustomizableObject.cpp b/Engine/Plugins/Mutable/Source/CustomizableObject/Private/MuCO/CustomizableObject.cpp
--- a/Engine/Plugins/Mutable/Source/CustomizableObject/Private/MuCO/CustomizableObject.cpp
+++ b/Engine/Plugins/Mutable/Source/CustomizableObject/Private/MuCO/CustomizableObject.cpp
@@ -2592,7 +2592,7 @@ namespace MutablePrivate
 						NewFile.DataType = CategoryID.DataType;
 						NewFile.ResourceType = CategoryID.ResourceType;
 						NewFile.Flags = CategoryID.Flags;
-						NewFile.Id = FileID; 
+						NewFile.Id = FileID + Category->FirstFile;
 
 						FileIndex = NumBulkDataFiles;
 						++NumBulkDataFiles;
