[Mutable] Fix crash when the mip stream in petition is aborted while prefetching data.
Reduce the number of rooms to prefetch.

#rnx

diff --git a/Engine/Plugins/Mutable/Source/CustomizableObject/Private/MuCO/CustomizableObjectMipDataProvider.cpp b/Engine/Plugins/Mutable/Source/CustomizableObject/Private/MuCO/CustomizableObjectMipDataProvider.cpp
--- a/Engine/Plugins/Mutable/Source/CustomizableObject/Private/MuCO/CustomizableObjectMipDataProvider.cpp
+++ b/Engine/Plugins/Mutable/Source/CustomizableObject/Private/MuCO/CustomizableObjectMipDataProvider.cpp
@@ -434,11 +434,13 @@ int32 FMutableTextureMipDataProvider::GetMips(const FTextureUpdateContext& Conte
 			TSharedPtr<mu::FModel, ESPMode::ThreadSafe> Model = UpdateContext->GetModel();
 			check(Model);
 
-			for (int32 ImageToLoadIndex : ImageRef.ConstantImagesNeededToGenerate)
+			if (!OperationData->Levels.IsEmpty())
 			{
-				for (const FMutableMipUpdateLevel& Level : OperationData->Levels)
+				const int32 LODIndex = OperationData->Levels[0].MipLevel;
+
+				for (int32 ImageToLoadIndex : ImageRef.ConstantImagesNeededToGenerate)
 				{
-					int32 RomId = Model->GetConstantImageRomId(ImageToLoadIndex, Level.MipLevel);
+					const int32 RomId = Model->GetConstantImageRomId(ImageToLoadIndex, LODIndex);
 
 					if (RomId >= 0 && Model->IsRomHighRes(RomId))
 					{
@@ -487,7 +489,6 @@ int32 FMutableTextureMipDataProvider::GetMips(const FTextureUpdateContext& Conte
 				if (OperationData->bIsCancelled)
 				{
 					check(OperationData->Counter == 0);
-					OperationData->RescheduleCallback();
 					return;
 				}
 
