[Mutable] Morph Streaming.

Add Morph support to Mutable Mesh Streaming.
All LOD UMorphTargets are will be created in the Initial Generation. Data render data of those will be streamed directly in the FCustomizableObjectMeshStreamIn.

- Decoupled ReconstructMorphTargets from Subtask_Mutable_PrepareRealTimeMorphData.
- Decoupled FMutableDataStreamer and LoadMorphTargets from LoadAdditionalAssetsAndData.
- Replaced RealTimeMorphTargets ObjectComponentIndex with ComponentName. Now is a TMap.
- Removed GetComponentMorphTargets.
- Moved CVarMutableHighPriorityLoading.

#rnx

diff --git a/Engine/Plugins/Mutable/Source/CustomizableObject/Internal/MuCO/CustomizableObjectPrivate.h b/Engine/Plugins/Mutable/Source/CustomizableObject/Internal/MuCO/CustomizableObjectPrivate.h
--- a/Engine/Plugins/Mutable/Source/CustomizableObject/Internal/MuCO/CustomizableObjectPrivate.h
+++ b/Engine/Plugins/Mutable/Source/CustomizableObject/Internal/MuCO/CustomizableObjectPrivate.h
@@ -1484,7 +1484,7 @@ public:
 	void SetModelStreamableBulkData(const TSharedPtr<FModelStreamableBulkData>& StreamableData, bool bIsCooking);
 #endif
 	
-	TSharedPtr<FModelStreamableBulkData> GetModelStreamableBulkData(bool bIsCooking = false);
+	TSharedPtr<FModelStreamableBulkData> GetModelStreamableBulkData(bool bIsCooking = false) const;
 
 	// See UCustomizableObjectSystem::LockObject()
 	bool IsLocked() const;
@@ -1517,7 +1517,7 @@ public:
 	static FString GetDataTypeExtension(MutablePrivate::EStreamableDataType DataType);
 
 	/** Compose file name. */
-	FString GetCompiledDataFileName(const ITargetPlatform* InTargetPlatform = nullptr, bool bIsDiskStreamer = false);
+	FString GetCompiledDataFileName(const ITargetPlatform* InTargetPlatform = nullptr, bool bIsDiskStreamer = false) const;
 
 	/** DDC helpers. BuildDerivedDataKey is expensive, try to cache it as much as possible. */
 	TArray<uint8> BuildDerivedDataKey(FCompilationOptions Options);
diff --git a/Engine/Plugins/Mutable/Source/CustomizableObject/Internal/MuCO/CustomizableObjectSystemPrivate.h b/Engine/Plugins/Mutable/Source/CustomizableObject/Internal/MuCO/CustomizableObjectSystemPrivate.h
--- a/Engine/Plugins/Mutable/Source/CustomizableObject/Internal/MuCO/CustomizableObjectSystemPrivate.h
+++ b/Engine/Plugins/Mutable/Source/CustomizableObject/Internal/MuCO/CustomizableObjectSystemPrivate.h
@@ -33,10 +33,8 @@
 
 #include "CustomizableObjectSystemPrivate.generated.h"
 
-#if WITH_EDITORONLY_DATA
+class FMutableStreamRequest;
 class UEditorImageProvider;
-#endif
-
 class UCustomizableObjectSystem;
 namespace LowLevelTasks { enum class ETaskPriority : int8; }
 struct FTexturePlatformData;
@@ -44,8 +42,46 @@ struct FTexturePlatformData;
 // Split StreamedBulkData into chunks smaller than MUTABLE_STREAMED_DATA_MAXCHUNKSIZE
 #define MUTABLE_STREAMED_DATA_MAXCHUNKSIZE		(512 * 1024 * 1024)
 
+
 extern TAutoConsoleVariable<bool> CVarMutableHighPriorityLoading;
 
+
+/** Final data to create a UMorphTarget. */
+struct FMorphTargetMeshData
+{
+	/** Names of the Morph Targets. */
+	TArray<FName> NameResolutionMap;
+
+	/** Vertex data. Each vertex contains an index to the NameResolutionMap. */
+	TArray<FMorphTargetVertexData> Data;
+};
+
+
+/** Mapping of FMorphTargetVertexData local names to global names. */
+struct FMappedMorphTargetMeshData
+{
+	/** Index is the local name. Value is the index to a global name table. */
+	TArray<int32> NameResolutionMap;
+
+	/** Pointer to the original data with indices to local names. */
+	const TArray<FMorphTargetVertexData>* DataView = nullptr;
+};
+
+
+/* Reconstruct the final Morph Targets using the global names.
+ * @param Mesh Used to know which MappedMorphTargets vertices must be removed.
+ * @param GlobalNames
+ * @param OutMorphTargets resulting Morph Targets with the indices pointing to GlobalNames. */
+void ReconstructMorphTargets(const mu::FMesh& Mesh, const TArray<FName>& GlobalNames, const TMap<uint32, FMappedMorphTargetMeshData>& MappedMorphTargets, TArray<FMorphTargetLODModel>& OutMorphTargets);
+
+
+/** Request the Mutable Data Streamer to load Morph Target blocks. Does not stream them.
+ *  @param MutableDataStreamer Streamer to request the blocks.
+ *  @param Mesh Mutable mesh to obtain which blocks to load.
+ *  @param StreamingResult Key is the block. Value is the container where the data will be deserialized once streamed. */
+void LoadMorphTargets(FMutableStreamRequest& MutableDataStreamer, const TSharedRef<const mu::FMesh>& Mesh, TMap<uint32, FMorphTargetMeshData>& StreamingResult);
+
+
 /** End a Customizable Object Instance Update. All code paths of an update have to end here. */
 void FinishUpdateGlobal(const TSharedRef<FUpdateContextPrivate>& Context);
 
@@ -348,6 +384,17 @@ struct FGeneratedMaterial
 };
 
 
+/** Final data per component. */
+struct FSkeletalMeshMorphTargets
+{
+	/** Name of the Morph Target. */
+	TArray<FName> RealTimeMorphTargetNames;
+
+	/** First index is the Morph Target (in sync with RealTimeMorphTargetNames). Second index is the LOD. */
+	TArray<TArray<FMorphTargetLODModel>> RealTimeMorphsLODData; 
+};
+
+
 // Mutable data generated during the update steps.
 // We keep it from begin to end update, and it is used in several steps.
 struct FInstanceUpdateData
@@ -460,28 +507,12 @@ struct FInstanceUpdateData
 
 		bool operator==(const mu::FBoneName& OtherName) const { return Name == OtherName; };
 	};
-
 	
-	struct FMorphTargetMeshData
-	{
-		TArray<FName> NameResolutionMap;
-		TArray<FMorphTargetVertexData> Data;
-	};
-
+	/** Key is the Block Id. Value is the LoadAdditionalAssetsAndData read destination. */
 	TMap<uint32, FMorphTargetMeshData> RealTimeMorphTargetMeshData;
-
-
-	struct FRealTimeMorphsComponentData
-	{
-		FCustomizableObjectComponentIndex ObjectComponentIndex = FCustomizableObjectComponentIndex(INDEX_NONE);
-
-		TArray<FName> RealTimeMorphTargetNames;
-		TArray<TArray<FMorphTargetLODModel>> RealTimeMorphsLODData; 
-	};
-
-	// Index in this array is not related to any other index.
-	TArray<FRealTimeMorphsComponentData> RealTimeMorphTargets;
 	
+	/** Key is the Component Name. Value is the final Morph Target data to copy into the Skeletal Mesh. */
+	TMap<FName, FSkeletalMeshMorphTargets> RealTimeMorphTargets;
 
 	struct FClothingMeshData
 	{
@@ -573,9 +604,7 @@ public:
 
 	/** Return the object component index associated with a component in this instance. */
 	FCustomizableObjectComponentIndex GetObjectComponentIndex(FCustomizableObjectInstanceComponentIndex InstanceComponentIndex);
-
-	FInstanceUpdateData::FRealTimeMorphsComponentData* GetComponentMorphTargets(FCustomizableObjectComponentIndex ObjectComponentIndex);
-
+	
 private:
 	/** Descriptor which the update will be performed on. */
 	FCustomizableObjectInstanceDescriptor CapturedDescriptor;
diff --git a/Engine/Plugins/Mutable/Source/CustomizableObject/Internal/MuCO/UnrealConversionUtils.h b/Engine/Plugins/Mutable/Source/CustomizableObject/Internal/MuCO/UnrealConversionUtils.h
--- a/Engine/Plugins/Mutable/Source/CustomizableObject/Internal/MuCO/UnrealConversionUtils.h
+++ b/Engine/Plugins/Mutable/Source/CustomizableObject/Internal/MuCO/UnrealConversionUtils.h
@@ -11,6 +11,8 @@
 #include "Containers/Array.h"
 #include "Containers/Map.h"
 
+struct FMorphTargetMeshData;
+
 namespace mu
 {
 	class FMesh;
@@ -130,4 +132,6 @@ namespace UnrealConversionUtils
 	CUSTOMIZABLEOBJECT_API void UpdateSkeletalMeshLODRenderDataBuffersSize(
 		FSkeletalMeshLODRenderData& LODResource
 	);
+
+	void MorphTargetVertexInfoBuffers(FSkeletalMeshLODRenderData& LODResource, const USkeletalMesh& Owner, const mu::FMesh& MutableMesh, const TMap<uint32, FMorphTargetMeshData>& MorphTargetMeshData, int32 LODIndex);
 }
diff --git a/Engine/Plugins/Mutable/Source/CustomizableObject/Private/MuCO/CustomizableObject.cpp b/Engine/Plugins/Mutable/Source/CustomizableObject/Private/MuCO/CustomizableObject.cpp
--- a/Engine/Plugins/Mutable/Source/CustomizableObject/Private/MuCO/CustomizableObject.cpp
+++ b/Engine/Plugins/Mutable/Source/CustomizableObject/Private/MuCO/CustomizableObject.cpp
@@ -786,7 +786,7 @@ FString UCustomizableObjectPrivate::GetCompiledDataFolderPath()
 }
 
 
-FString UCustomizableObjectPrivate::GetCompiledDataFileName(const ITargetPlatform* InTargetPlatform, bool bIsDiskStreamer)
+FString UCustomizableObjectPrivate::GetCompiledDataFileName(const ITargetPlatform* InTargetPlatform, bool bIsDiskStreamer) const
 {
 	const FString FilePath = GetCompiledDataFolderPath();
 	const FString PlatformName = InTargetPlatform ? InTargetPlatform->PlatformName() : FPlatformProperties::PlatformName();
@@ -2203,7 +2203,7 @@ void UCustomizableObjectPrivate::SetModelStreamableBulkData(const TSharedPtr<FMo
 }
 #endif
 
-TSharedPtr<FModelStreamableBulkData> UCustomizableObjectPrivate::GetModelStreamableBulkData(bool bIsCooking)
+TSharedPtr<FModelStreamableBulkData> UCustomizableObjectPrivate::GetModelStreamableBulkData(bool bIsCooking) const
 {
 #if WITH_EDITOR
 	if (bIsCooking)
diff --git a/Engine/Plugins/Mutable/Source/CustomizableObject/Private/MuCO/CustomizableObjectInstance.cpp b/Engine/Plugins/Mutable/Source/CustomizableObject/Private/MuCO/CustomizableObjectInstance.cpp
--- a/Engine/Plugins/Mutable/Source/CustomizableObject/Private/MuCO/CustomizableObjectInstance.cpp
+++ b/Engine/Plugins/Mutable/Source/CustomizableObject/Private/MuCO/CustomizableObjectInstance.cpp
@@ -17,6 +17,7 @@
 #include "Engine/SkeletalMeshSocket.h"
 #include "Engine/SkeletalMeshLODSettings.h"
 #include "MaterialDomain.h"
+#include "MutableStreamRequest.h"
 #include "Materials/Material.h"
 #include "Materials/MaterialInstanceDynamic.h"
 #include "Misc/Optional.h"
@@ -4165,45 +4166,77 @@ void UCustomizableInstancePrivate::BuildOrCopyMorphTargetsData(const TSharedRef<
 	FCustomizableObjectComponentIndex ObjectComponentIndex = OperationData->GetObjectComponentIndex(InstanceComponentIndex);
 	const FName ComponentName = OperationData->ComponentNames[ObjectComponentIndex.GetValue()];
 
-	FInstanceUpdateData::FRealTimeMorphsComponentData* ComponentMorphTargets = OperationData->GetComponentMorphTargets(ObjectComponentIndex);
+	FSkeletalMeshMorphTargets* ComponentMorphTargets = OperationData->InstanceUpdateData.RealTimeMorphTargets.Find(ComponentName);
 	if (!ComponentMorphTargets)
 	{
 		return;
 	}
 
+	const int32 NumMorphTargets = ComponentMorphTargets->RealTimeMorphTargetNames.Num();
+	
 	TArray<TObjectPtr<UMorphTarget>>& MorphTargets = SkeletalMesh->GetMorphTargets();
-
-	MorphTargets.Empty(ComponentMorphTargets->RealTimeMorphTargetNames.Num());
+	MorphTargets.Empty(NumMorphTargets);
 	
-	const int32 NumMorphs = ComponentMorphTargets->RealTimeMorphTargetNames.Num();
-	for (int32 I = 0; I < NumMorphs; ++I)
+	for (int32 MorphTargetIndex = 0; MorphTargetIndex < NumMorphTargets; ++MorphTargetIndex)
 	{
-		if (ComponentMorphTargets->RealTimeMorphsLODData[I].IsEmpty())
+		TArray<FMorphTargetLODModel>& MorphTargetData = ComponentMorphTargets->RealTimeMorphsLODData[MorphTargetIndex];
+
+		if (MorphTargetData.IsEmpty())
 		{
 			continue;
 		}
+		
+		const FName& MorphTargetName = ComponentMorphTargets->RealTimeMorphTargetNames[MorphTargetIndex];
 
-		UMorphTarget* NewMorphTarget = NewObject<UMorphTarget>(SkeletalMesh, ComponentMorphTargets->RealTimeMorphTargetNames[I]);
+		UMorphTarget* NewMorphTarget = NewObject<UMorphTarget>(SkeletalMesh, MorphTargetName);
 		NewMorphTarget->BaseSkelMesh = SkeletalMesh;
-		NewMorphTarget->GetMorphLODModels() = MoveTemp(ComponentMorphTargets->RealTimeMorphsLODData[I]);
+
+		TArray<FMorphTargetLODModel>& MorphLODModels = NewMorphTarget->GetMorphLODModels();
+
+		if (OperationData->bStreamMeshLODs)
+		{
+			MorphLODModels.SetNum(ComponentMorphTargets->RealTimeMorphsLODData[MorphTargetIndex].Num());
+			
+			// Streamed LODs
+			const int32 FirstGeneratedLOD = FMath::Max(OperationData->GetRequestedLODs()[ComponentName], OperationData->GetMinLOD()[ComponentName]);
+			for (int32 LODIndex = FirstGeneratedLOD; LODIndex < OperationData->FirstResidentLOD[ComponentName]; ++LODIndex)
+			{
+				// Copy data required for streaming
+				MorphLODModels[LODIndex].NumVertices = 1; // Trick the engine
+				MorphLODModels[LODIndex].SectionIndices = MoveTemp(MorphTargetData[LODIndex].SectionIndices);
+			}
+			
+			// Residents LODs
+			for (int32 LODIndex = OperationData->FirstResidentLOD[ComponentName]; LODIndex < OperationData->NumLODsAvailable[ComponentName]; ++LODIndex)
+			{
+				MorphLODModels[LODIndex] = ComponentMorphTargets->RealTimeMorphsLODData[MorphTargetIndex][LODIndex];
+			}
+		}
+		else
+		{
+			MorphLODModels = MoveTemp(ComponentMorphTargets->RealTimeMorphsLODData[MorphTargetIndex]);
+		}
+
 		MorphTargets.Add(NewMorphTarget);
 	}
-
-	// Copy MorphTargets from the FirstGeneratedLOD to the LODs below
-	const int32 FirstGeneratedLOD = FMath::Max(OperationData->GetRequestedLODs()[ComponentName], OperationData->GetMinLOD()[ComponentName]);
-	for (int32 LODIndex = OperationData->FirstLODAvailable[ComponentName]; LODIndex < FirstGeneratedLOD; ++LODIndex)
+	
+	if (!OperationData->bStreamMeshLODs)
 	{
-		MUTABLE_CPUPROFILER_SCOPE(CopyMorphTargetsData);
-
-		const int32 NumMorphTargets = MorphTargets.Num();
-		for (int32 MorphTargetIndex = 0; MorphTargetIndex < NumMorphTargets; ++MorphTargetIndex)
+		// Copy MorphTargets from the FirstGeneratedLOD to the LODs below
+		const int32 FirstGeneratedLOD = FMath::Max(OperationData->GetRequestedLODs()[ComponentName], OperationData->GetMinLOD()[ComponentName]);
+		for (int32 LODIndex = OperationData->FirstLODAvailable[ComponentName]; LODIndex < FirstGeneratedLOD; ++LODIndex)
 		{
-			MorphTargets[MorphTargetIndex]->GetMorphLODModels()[LODIndex] =
-					MorphTargets[MorphTargetIndex]->GetMorphLODModels()[FirstGeneratedLOD];
+			MUTABLE_CPUPROFILER_SCOPE(CopyMorphTargetsData);
+
+			for (int32 MorphTargetIndex = 0; MorphTargetIndex < MorphTargets.Num(); ++MorphTargetIndex)
+			{
+				MorphTargets[MorphTargetIndex]->GetMorphLODModels()[LODIndex] = MorphTargets[MorphTargetIndex]->GetMorphLODModels()[FirstGeneratedLOD];
+			}
 		}
 	}
 
-	SkeletalMesh->InitMorphTargets();
+	const bool bInKeepEmptyMorphTargets = OperationData->bStreamMeshLODs;
+	SkeletalMesh->InitMorphTargets(bInKeepEmptyMorphTargets); // True to avoid removing streamed Morph Targets.
 }
 
 namespace 
@@ -4983,12 +5016,13 @@ UE::Tasks::FTask UCustomizableInstancePrivate::LoadAdditionalAssetsAndData(const
 
 	const UModelResources& ModelResources = CustomizableObject->GetPrivate()->GetModelResourcesChecked();
 	const TSharedPtr<FModelStreamableBulkData>& ModelStreamableBulkData = CustomizableObject->GetPrivate()->GetModelStreamableBulkData();
+	
+	FMutableStreamRequest StreamRequest(ModelStreamableBulkData);
 
 	TArray<FSoftObjectPath> AssetsToStream;
 	TArray<uint32> RealTimeMorphStreamableBlocksToStream;
 	TArray<uint32> ClothingStreamableBlocksToStream;
 
-	TArray<FInstanceUpdateData::FLOD>& LODs = OperationData->InstanceUpdateData.LODs;
 	TArray<FInstanceUpdateData::FComponent>& Components = OperationData->InstanceUpdateData.Components;
 
 	ObjectToInstanceIndexMap.Empty();
@@ -5090,8 +5124,9 @@ UE::Tasks::FTask UCustomizableInstancePrivate::LoadAdditionalAssetsAndData(const
 			}
 		}
 	}
-
-	const bool bMorphTargetsEnabled = CVarEnableRealTimeMorphTargets.GetValueOnAnyThread();
+	
+	bool bHasInvalidMesh = false;
+	const bool bUpdateMeshes = DoComponentsNeedUpdate(GetPublic(), OperationData, bHasInvalidMesh);
 
 	// Load assets coming from SubMeshes of the newly generated Mesh
 	if (OperationData->InstanceUpdateData.LODs.Num())
@@ -5124,6 +5159,11 @@ UE::Tasks::FTask UCustomizableInstancePrivate::LoadAdditionalAssetsAndData(const
 					continue;
 				}
 
+				if (bUpdateMeshes)
+				{
+					LoadMorphTargets(StreamRequest, MutableMesh.ToSharedRef(), OperationData->InstanceUpdateData.RealTimeMorphTargetMeshData);
+				}
+				
 				const TArray<uint64>& StreamedResources = MutableMesh->GetStreamedResources();
 
 				for (uint64 ResourceId : StreamedResources)
@@ -5138,23 +5178,6 @@ UE::Tasks::FTask UCustomizableInstancePrivate::LoadAdditionalAssetsAndData(const
 #endif
 						ComponentData->StreamedResourceIndex.AddUnique(ResourceIndex);
 					}
-					else if (TypedResourceId.Type == (uint8)FCustomizableObjectStreameableResourceId::EType::RealTimeMorphTarget)
-					{
-						if (bMorphTargetsEnabled)
-						{
-							check(TypedResourceId.Id != 0 && TypedResourceId.Id <= TNumericLimits<uint32>::Max());
-
-							const TMap<uint32, FRealTimeMorphStreamable>& MorphsStreamables = ModelStreamableBulkData->RealTimeMorphStreamables;
-							if (MorphsStreamables.Contains((uint32)TypedResourceId.Id))
-							{
-								RealTimeMorphStreamableBlocksToStream.AddUnique(TypedResourceId.Id);
-							}
-							else
-							{
-								UE_LOG(LogMutable, Error, TEXT("Invalid streamed real time morph target data block [%d] found."), TypedResourceId.Id);
-							}
-						}
-					}
 					else if (TypedResourceId.Type == (uint8)FCustomizableObjectStreameableResourceId::EType::Clothing)
 					{
 						check(TypedResourceId.Id != 0 && TypedResourceId.Id <= TNumericLimits<uint32>::Max());
@@ -5193,15 +5216,8 @@ UE::Tasks::FTask UCustomizableInstancePrivate::LoadAdditionalAssetsAndData(const
 							UE_LOG(LogMutable, Error, TEXT("Invalid streamed clothing data block [%d] found."), TypedResourceId.Id);
 						}
 					}
-					else
-					{
-						UE_LOG(LogMutable, Error, TEXT("Unknown streamed resource type found."));
-						check(false);
-					}
 				}
-
-				const bool bReplacePhysicsAssets = HasCOInstanceFlags(ReplacePhysicsAssets);
-
+				
 				for (int32 TagIndex = 0; TagIndex < MutableMesh->GetTagCount(); ++TagIndex)
 				{
 					FString Tag = MutableMesh->GetTag(TagIndex);
@@ -5381,13 +5397,10 @@ UE::Tasks::FTask UCustomizableInstancePrivate::LoadAdditionalAssetsAndData(const
 #endif
 	}
 	
-	const bool bEnableHighPriorityLoading = CVarMutableHighPriorityLoading.GetValueOnAnyThread();
-
-	TArray<UE::Tasks::FTaskEvent> StreamingCompletionEvents;
 	if (AssetsToStream.Num() > 0)
 	{
 		check(AssetAsyncLoadCompletionEvent.IsCompleted());
-		AssetAsyncLoadCompletionEvent = StreamingCompletionEvents.Emplace_GetRef(TEXT("AssetAsyncLoadCompletionEvent"));
+		AssetAsyncLoadCompletionEvent = UE::Tasks::FTaskEvent(TEXT("AssetsStreamed"));
 
 #if WITH_EDITOR
 		// TODO: Remove with UE-217665 when the underlying bug in the ColorPicker is solved
@@ -5401,120 +5414,21 @@ UE::Tasks::FTask UCustomizableInstancePrivate::LoadAdditionalAssetsAndData(const
 		PrivateSystem->StreamableManager->RequestAsyncLoad(
 			AssetsToStream,
 			FStreamableDelegate::CreateUObject(this, &UCustomizableInstancePrivate::AdditionalAssetsAsyncLoaded, OperationData),
-			bEnableHighPriorityLoading ? FStreamableManager::AsyncLoadHighPriority : FStreamableManager::DefaultAsyncLoadPriority);
+			CVarMutableHighPriorityLoading.GetValueOnAnyThread() ? FStreamableManager::AsyncLoadHighPriority : FStreamableManager::DefaultAsyncLoadPriority);
 	}
 	
-	// File handles will end up owned by the gather task.
-	TArray<TUniquePtr<IAsyncReadFileHandle>> OpenFileHandles;
-	TArray<UE::Tasks::TTask<TUniquePtr<IAsyncReadRequest>>> ReadRequestTasks;
-	TArray<UE::Tasks::TTask<TUniquePtr<IBulkDataIORequest>>> BulkReadRequestTasks;
-	
-	bool bHasInvalidMesh = false;
-	bool bUpdateMeshes = DoComponentsNeedUpdate(GetPublic(), OperationData, bHasInvalidMesh);
-
-	bool bIsDataBlocksStreamNeeded = RealTimeMorphStreamableBlocksToStream.Num() || ClothingStreamableBlocksToStream.Num();
-
-	if (bIsDataBlocksStreamNeeded && bUpdateMeshes)
+	if (bUpdateMeshes)
 	{
-#if WITH_EDITOR
-		const int32 BulkDataFileOffset = sizeof(MutableCompiledDataStreamHeader);
-		const bool bUseFBulkData = false;
-
-		const FString FullFileName = CustomizableObject->GetPrivate()->GetCompiledDataFileName(nullptr, true);
-
-		TArray<uint32> OpenFilesIds;
-		auto OpenOrGetFileHandleForBlock = [&OpenFilesIds, &OpenFileHandles, FullFileName](const FMutableStreamableBlock& Block, MutablePrivate::EStreamableDataType DataType)
-			-> IAsyncReadFileHandle*
-		{
-			int32 FileHandleIndex = OpenFilesIds.Find(Block.FileId);
-			if (FileHandleIndex == INDEX_NONE)
-			{
-				IAsyncReadFileHandle* ReadFileHandle = FPlatformFileManager::Get().GetPlatformFile().OpenAsyncRead(*(FullFileName + UCustomizableObjectPrivate::GetDataTypeExtension(DataType)));
-				OpenFileHandles.Emplace(ReadFileHandle);
-
-				FileHandleIndex = OpenFilesIds.Add(Block.FileId);
-				check(OpenFileHandles.Num() == OpenFilesIds.Num());
-			}
-
-			if (OpenFileHandles.IsValidIndex(FileHandleIndex))
-			{
-				return OpenFileHandles[FileHandleIndex].Get();
-			}
-
-			return nullptr;
-		};
-#else
-		const int32 BulkDataFileOffset = 0;
-		const bool bUseFBulkData = !ModelStreamableBulkData->StreamableBulkData.IsEmpty();
-#endif
-
-
-		MUTABLE_CPUPROFILER_SCOPE(RealTimeMorphStreaming);
-
-		struct FBlockReadInfo
-		{
-			uint64 Offset;
-			IAsyncReadFileHandle* FileHandle;
-			TArrayView<uint8> AllocatedMemoryView;
-			uint32 FileId;
-		};
-
-		TArray<FBlockReadInfo> BlockReadInfos;
-		BlockReadInfos.Reserve(16);
-
-		if (bMorphTargetsEnabled)
-		{
-			const int32 NumMorphBlocks = RealTimeMorphStreamableBlocksToStream.Num();
-			for (int32 I = 0; I < NumMorphBlocks; ++I)
-			{
-				MUTABLE_CPUPROFILER_SCOPE(RealTimeMorphStreamingRequest_Alloc);
-
-				const int32 BlockId = RealTimeMorphStreamableBlocksToStream[I];
-				const FRealTimeMorphStreamable& Streamable = ModelStreamableBulkData->RealTimeMorphStreamables[BlockId];
-				const FMutableStreamableBlock& Block = Streamable.Block; 
-			
-				FInstanceUpdateData::FMorphTargetMeshData& ReadDestData = 
-						OperationData->InstanceUpdateData.RealTimeMorphTargetMeshData.FindOrAdd(BlockId);
-
-				// Only request blocks once.
-				if (ReadDestData.Data.Num())
-				{
-					continue;
-				}
-
-				ReadDestData.NameResolutionMap = ModelStreamableBulkData->RealTimeMorphStreamables[BlockId].NameResolutionMap;
-
-				check(Streamable.Size % sizeof(FMorphTargetVertexData) == 0);
-				uint32 NumElems = Streamable.Size / sizeof(FMorphTargetVertexData);
-
-				ReadDestData.Data.SetNumUninitialized(NumElems);
-
-#if WITH_EDITOR
-				IAsyncReadFileHandle* FileHandle = OpenOrGetFileHandleForBlock(Block, MutablePrivate::EStreamableDataType::RealTimeMorph);
-#else
-				IAsyncReadFileHandle* FileHandle = nullptr;
-#endif
-
-				BlockReadInfos.Emplace(FBlockReadInfo
-				{
-					BulkDataFileOffset + Block.Offset,
-					FileHandle,
-					MakeArrayView(reinterpret_cast<uint8*>(ReadDestData.Data.GetData()), ReadDestData.Data.Num()*sizeof(FMorphTargetVertexData)),
-					Block.FileId
-				});
-			}
-		}
-		const int32 NumClothBlocks = ClothingStreamableBlocksToStream.Num();
-		for (int32 I = 0; I < NumClothBlocks; ++I)
+		// Clothing blocks to stream 
+		for (const int32 BlockId : ClothingStreamableBlocksToStream)
 		{
 			MUTABLE_CPUPROFILER_SCOPE(ClothingStreamingRequest_Alloc);
 
-			const int32 BlockId = ClothingStreamableBlocksToStream[I];
 			const FClothingStreamable& ClothingStreamable = ModelStreamableBulkData->ClothingStreamables[BlockId];
 			const FMutableStreamableBlock& Block = ClothingStreamable.Block;
-		
+
 			FInstanceUpdateData::FClothingMeshData& ReadDestData = 
-					OperationData->InstanceUpdateData.ClothingMeshData.FindOrAdd(BlockId);
+				OperationData->InstanceUpdateData.ClothingMeshData.FindOrAdd(BlockId);
 
 			// Only request blocks once.
 			if (ReadDestData.Data.Num())
@@ -5531,134 +5445,19 @@ UE::Tasks::FTask UCustomizableInstancePrivate::LoadAdditionalAssetsAndData(const
 
 			ReadDestData.Data.SetNumUninitialized(NumElems);
 
-#if WITH_EDITOR
-			IAsyncReadFileHandle* FileHandle = OpenOrGetFileHandleForBlock(Block, MutablePrivate::EStreamableDataType::Clothing);
-#else
-			IAsyncReadFileHandle* FileHandle = nullptr;
-#endif
-
-			BlockReadInfos.Emplace(FBlockReadInfo
-			{
-				BulkDataFileOffset + Block.Offset,
-				FileHandle,
-				MakeArrayView(reinterpret_cast<uint8*>(ReadDestData.Data.GetData()), ReadDestData.Data.Num()*sizeof(FCustomizableObjectMeshToMeshVertData)),
-				Block.FileId
-			});
-		}
-
-		for (const FBlockReadInfo& BlockReadInfo : BlockReadInfos)
-		{
-			if (bUseFBulkData)
-			{
-				BulkReadRequestTasks.Emplace(UE::Tasks::Launch(TEXT("CustomizableObjectInstanceBulkReadRequestTask"),
-					[
-						OwnedOperationData = OperationData.ToSharedPtr(), // Keep a reference to make sure allocated memory is always alive.
-							ModelStreamableBulkData,
-							ReadDataReadyEvent = StreamingCompletionEvents.Emplace_GetRef(TEXT("AsyncReadDataReadyEvent")),
-							Block = BlockReadInfo,
-							Priority = bEnableHighPriorityLoading ? AIOP_High : AIOP_Normal
-					]() -> TUniquePtr<IBulkDataIORequest>
-					{
-						MUTABLE_CPUPROFILER_SCOPE(CustomizableInstanceLoadBlocksAsyncRead_Request);
-						FBulkDataIORequestCallBack IOCallback =
-							[OwnedOperationData, ReadDataReadyEvent, FileId = Block.FileId](bool bWasCancelled, IBulkDataIORequest*) mutable
-							{
-								if (bWasCancelled)
-								{
-									UE_LOG(LogMutable, Warning, TEXT("An AsyncReadRequest to file %08x was cancelled. The file may not exist."), FileId);
-								}
-								ReadDataReadyEvent.Trigger();
-							};
-
-						check(ModelStreamableBulkData->StreamableBulkData.IsValidIndex(Block.FileId));
-						FByteBulkData& ByteBulkData = ModelStreamableBulkData->StreamableBulkData[Block.FileId];
-
-						return TUniquePtr<IBulkDataIORequest>(ByteBulkData.CreateStreamingRequest(
-							Block.Offset,
-							(int64)Block.AllocatedMemoryView.Num(),
-							Priority,
-							&IOCallback,
-							Block.AllocatedMemoryView.GetData()));
-					},
-					bEnableHighPriorityLoading ? UE::Tasks::ETaskPriority::High : UE::Tasks::ETaskPriority::Normal));
-			}
-			else if (BlockReadInfo.FileHandle)
-			{
-				ReadRequestTasks.Emplace(UE::Tasks::Launch(TEXT("CustomizableObjectInstanceReadRequestTask"),
-					[
-						OwnedOperationData = OperationData.ToSharedPtr(), // Keep a reference to make sure allocated memory is always alive.
-							ReadDataReadyEvent = StreamingCompletionEvents.Emplace_GetRef(TEXT("AsyncReadDataReadyEvent")),
-							Block = BlockReadInfo,
-							Priority = bEnableHighPriorityLoading ? AIOP_High : AIOP_Normal
-					]() -> TUniquePtr<IAsyncReadRequest>
-					{
-						MUTABLE_CPUPROFILER_SCOPE(CustomizableInstanceLoadBlocksAsyncRead_Request);
-						FAsyncFileCallBack ReadRequestCallBack =
-							[OwnedOperationData, ReadDataReadyEvent, FileId = Block.FileId](bool bWasCancelled, IAsyncReadRequest*) mutable
-							{
-								if (bWasCancelled)
-								{
-									UE_LOG(LogMutable, Warning, TEXT("An AsyncReadRequest to file %08x was cancelled. The file may not exist."), FileId);
-								}
-								ReadDataReadyEvent.Trigger();
-							};
-
-						return TUniquePtr<IAsyncReadRequest>(Block.FileHandle->ReadRequest(
-							Block.Offset,
-							(int64)Block.AllocatedMemoryView.Num(),
-							Priority,
-							&ReadRequestCallBack,
-							Block.AllocatedMemoryView.GetData()));
-					},
-					bEnableHighPriorityLoading ? UE::Tasks::ETaskPriority::High : UE::Tasks::ETaskPriority::Normal));
-			}
+			StreamRequest.AddBlock(
+				Block,
+				MutablePrivate::EStreamableDataType::Clothing,
+				MakeArrayView(reinterpret_cast<uint8*>(ReadDestData.Data.GetData()), ReadDestData.Data.Num() * sizeof(FCustomizableObjectMeshToMeshVertData)));
 		}
 	}
 	
-	if (AssetsToStream.Num() > 0 || OpenFileHandles.Num() > 0 || BulkReadRequestTasks.Num() > 0)
-	{
-		return UE::Tasks::Launch(TEXT("GatherStreamingRequestsCompletionTask"),
-				[
-					ReadRequestTasks = MoveTemp(ReadRequestTasks),
-					OpenFileHandles  = MoveTemp(OpenFileHandles),
-					BulkReadRequestTasks = MoveTemp(BulkReadRequestTasks)
-				]() mutable 
-				{
-					for (UE::Tasks::TTask<TUniquePtr<IAsyncReadRequest>>& ReadRequestTask : ReadRequestTasks)
-					{
-						// GetResult() may wait for the task to complete, this should not be a problem as this task
-						// prerequisites guarantee ReadRequestTask has at least started execution.
-						TUniquePtr<IAsyncReadRequest>& ReadRequest = ReadRequestTask.GetResult();
-						if (ReadRequest)
-						{
-							ReadRequest->WaitCompletion();
-						}
-					}
-					
-					for (UE::Tasks::TTask<TUniquePtr<IBulkDataIORequest>>& BulkReadRequestTask : BulkReadRequestTasks)
-					{
-						// GetResult() may wait for the task to complete, this should not be a problem as this task
-						// prerequisites guarantee ReadRequestTask has at least started execution.
-						TUniquePtr<IBulkDataIORequest>& ReadRequest = BulkReadRequestTask.GetResult();
-						if (ReadRequest)
-						{
-							ReadRequest->WaitCompletion();
-						}
-					}
-
-					ReadRequestTasks.Empty();
-					OpenFileHandles.Empty();
-					BulkReadRequestTasks.Empty();
-				},
-				StreamingCompletionEvents,
-				bEnableHighPriorityLoading ? UE::Tasks::ETaskPriority::High : UE::Tasks::ETaskPriority::Normal);
-	}
-	else
-	{
-		check(ReadRequestTasks.Num() == 0);
-		check(BulkReadRequestTasks.Num() == 0);
-		return UE::Tasks::MakeCompletedTask<void>();
-	}
+	// Stream files
+	UE::Tasks::FTask StreamingTask = StreamRequest.Stream();
+	
+	return UE::Tasks::Launch(TEXT("CaptureOperationData"), [OperationData]() {}, // Keep a reference to make sure allocated memory is always alive.
+	UE::Tasks::Prerequisites(AssetAsyncLoadCompletionEvent, StreamingTask),
+	UE::Tasks::ETaskPriority::Inherit);
 }
 
 FCustomizableObjectInstanceDescriptor& UCustomizableInstancePrivate::GetDescriptor() const
diff --git a/Engine/Plugins/Mutable/Source/CustomizableObject/Private/MuCO/CustomizableObjectMeshUpdate.cpp b/Engine/Plugins/Mutable/Source/CustomizableObject/Private/MuCO/CustomizableObjectMeshUpdate.cpp
--- a/Engine/Plugins/Mutable/Source/CustomizableObject/Private/MuCO/CustomizableObjectMeshUpdate.cpp
+++ b/Engine/Plugins/Mutable/Source/CustomizableObject/Private/MuCO/CustomizableObjectMeshUpdate.cpp
@@ -7,6 +7,7 @@ CustomizableObjectMeshUpdate.cpp: Helpers to stream in CustomizableObject skelet
 
 #include "MuCO/CustomizableObjectMeshUpdate.h"
 
+#include "MutableStreamRequest.h"
 #include "MuCO/CustomizableObjectSystem.h"
 #include "MuCO/CustomizableObjectSystemPrivate.h"
 #include "MuCO/CustomizableObjectSkeletalMesh.h"
@@ -32,8 +33,12 @@ FAutoConsoleVariableRef CVarMutableEnableGCHangFix(
 	TEXT("If true, TaskSynchronization decrement will happen in the Abort method instead of DoCancelMeshUpdate.")
 );
 
-FCustomizableObjectMeshStreamIn::FCustomizableObjectMeshStreamIn(const UCustomizableObjectSkeletalMesh* InMesh, bool bInHighPrio, EThreadType CreateResourcesThread)
-	: FSkeletalMeshStreamIn(InMesh, CreateResourcesThread),
+FCustomizableObjectMeshStreamIn::FCustomizableObjectMeshStreamIn(
+	const UCustomizableObjectSkeletalMesh* InMesh,
+	bool bInHighPrio,
+	EThreadType CreateResourcesThread,
+	const TSharedPtr<FModelStreamableBulkData>& ModelStreamableBulkData) :
+	FSkeletalMeshStreamIn(InMesh, CreateResourcesThread),
 	bHighPriority(bInHighPrio)
 {
 	OperationData = MakeShared<FMutableMeshOperationData>();
@@ -46,6 +51,8 @@ FCustomizableObjectMeshStreamIn::FCustomizableObjectMeshStreamIn(const UCustomiz
 	OperationData->Parameters = InMesh->Parameters;
 	OperationData->State = InMesh->State;
 
+	OperationData->ModelStreamableBulkData = ModelStreamableBulkData;
+	
 	OperationData->MeshIDs = InMesh->MeshIDs;
 	OperationData->Meshes.SetNum(InMesh->MeshIDs.Num());
 
@@ -156,7 +163,7 @@ namespace impl
 			OperationData->System->ClearWorkingMemory();
 		}
 
-		Task->OnUpdateMeshFinished();
+		OperationData->Event.Trigger();
 		
 		TRACE_END_REGION(UE_MUTABLE_UPDATE_MESH_REGION);
 	}
@@ -207,6 +214,7 @@ namespace impl
 		if (!Model || !Model->IsValid())
 		{
 			TRACE_END_REGION(UE_MUTABLE_UPDATE_MESH_REGION);
+			Task->Abort();
 			return;
 		}
 #endif
@@ -265,6 +273,39 @@ void FCustomizableObjectMeshStreamIn::RequestMeshUpdate(const FContext& Context)
 			impl::Task_Mutable_UpdateMesh(SharedOperationData, RefThis);
 		});
 
+	// Stream data (morphs, clothing...).
+	UE::Tasks::TTask StreamData = UE::Tasks::Launch(TEXT("StreamData"),
+		[RefThis = TRefCountPtr(this)]()
+		{
+			FMutableStreamRequest StreamRequest(RefThis->OperationData->ModelStreamableBulkData);
+
+			// Get morphs to be streamed.
+			for (TSharedPtr<const mu::FMesh>& Mesh : RefThis->OperationData->Meshes)
+			{
+				if (!Mesh)
+				{
+					continue;
+				}
+				
+				LoadMorphTargets(StreamRequest, Mesh.ToSharedRef(), RefThis->OperationData->MorphTargetMeshData);
+			}
+
+			// Stream data.
+			UE::Tasks::AddNested(StreamRequest.Stream());
+		},
+		OperationData->Event,
+		UE::Tasks::ETaskPriority::Inherit);
+
+	// Go to next step.
+	UE::Tasks::Launch(TEXT("OnUpdateMeshFinished"),
+		[RefThis = TRefCountPtr(this)]()
+		{
+			RefThis->OnUpdateMeshFinished();
+		},
+		StreamData,
+		UE::Tasks::ETaskPriority::Inherit,
+		UE::Tasks::EExtendedTaskPriority::Inline);
+	
 	if (IsCancelled() && TaskSynchronization.GetValue() > 0 && bEnableGCHangFix)
 	{
 		TaskSynchronization.Set(0);
@@ -337,10 +378,12 @@ void FCustomizableObjectMeshStreamIn::ConvertMesh(const FContext& Context)
 		UnrealConversionUtils::CopyMutableVertexBuffers(LODResource, MutableMesh.Get(), bNeedsCPUAccess);
 		UnrealConversionUtils::CopyMutableIndexBuffers(LODResource, MutableMesh.Get());
 		UnrealConversionUtils::CopyMutableSkinWeightProfilesBuffers(LODResource, *const_cast<UCustomizableObjectSkeletalMesh*>(Mesh), LODIndex, MutableMesh.Get(), Mesh->SkinWeightProfileIDs);
+		UnrealConversionUtils::MorphTargetVertexInfoBuffers(LODResource, *Mesh, *MutableMesh.Get(), OperationData->MorphTargetMeshData, LODIndex);
 
-		// TODO PERE: Implement missing buffers
-		// MorphTargets (MorphTargetVertexInfoBuffers)
 		// Clothing (ClothVertexBuffer)
+		{
+			// TODO PERE
+		}
 
 		UnrealConversionUtils::UpdateSkeletalMeshLODRenderDataBuffersSize(LODResource);
 	}
diff --git a/Engine/Plugins/Mutable/Source/CustomizableObject/Private/MuCO/CustomizableObjectMeshUpdate.h b/Engine/Plugins/Mutable/Source/CustomizableObject/Private/MuCO/CustomizableObjectMeshUpdate.h
--- a/Engine/Plugins/Mutable/Source/CustomizableObject/Private/MuCO/CustomizableObjectMeshUpdate.h
+++ b/Engine/Plugins/Mutable/Source/CustomizableObject/Private/MuCO/CustomizableObjectMeshUpdate.h
@@ -7,11 +7,14 @@ CustomizableObjectMeshUpdate.h: Helpers to stream in skeletal mesh LODs.
 #pragma once
 
 #include "CoreMinimal.h"
+#include "MuCO/CustomizableObjectPrivate.h"
+#include "MuCO/CustomizableObjectSystemPrivate.h"
 #include "Streaming/SkeletalMeshUpdate.h"
 
 #include "MuR/System.h"
 #include "MuR/Mesh.h"
 
+struct FModelStreamableBulkData;
 class UCustomizableObjectSkeletalMesh;
 
 namespace mu 
@@ -21,10 +24,12 @@ namespace mu
 
 struct FMutableMeshOperationData
 {
-	/** Mutable */
 	TSharedPtr<mu::FSystem> System;
 	TSharedPtr<mu::FModel, ESPMode::ThreadSafe> Model;
 	TSharedPtr<const mu::FParameters> Parameters;
+
+	TSharedPtr<FModelStreamableBulkData> ModelStreamableBulkData;
+	
 	int32 State = -1;
 
 	// IDs of the meshes to generate per LOD.
@@ -33,6 +38,11 @@ struct FMutableMeshOperationData
 	// Meshes generated per LOD
 	TArray<TSharedPtr<const mu::FMesh>> Meshes;
 
+	/** Serialized data per block. */
+	TMap<uint32, FMorphTargetMeshData> MorphTargetMeshData;
+
+	UE::Tasks::FTaskEvent Event = UE::Tasks::FTaskEvent(TEXT("Task_Mutable_UpdateMesh_End"));
+
 	/** SkeletalMesh */
 	/** The resident first LOD resource index.With domain = [0, ResourceState.NumLODs[.NOT THE ASSET LOD INDEX! */ 
 	int32 CurrentFirstLODIdx = INDEX_NONE;
@@ -53,7 +63,11 @@ extern template class TRenderAssetUpdate<FSkelMeshUpdateContext>;
 class FCustomizableObjectMeshStreamIn : public FSkeletalMeshStreamIn
 {
 public:
-	FCustomizableObjectMeshStreamIn(const UCustomizableObjectSkeletalMesh* InMesh, bool bHighPrio, EThreadType CreateResourcesThread);
+	FCustomizableObjectMeshStreamIn(
+		const UCustomizableObjectSkeletalMesh* InMesh,
+		bool bInHighPrio,
+		EThreadType CreateResourcesThread,
+		const TSharedPtr<FModelStreamableBulkData>& ModelStreamableBulkData);
 
 	void OnUpdateMeshFinished();
 
diff --git a/Engine/Plugins/Mutable/Source/CustomizableObject/Private/MuCO/CustomizableObjectSkeletalMesh.cpp b/Engine/Plugins/Mutable/Source/CustomizableObject/Private/MuCO/CustomizableObjectSkeletalMesh.cpp
--- a/Engine/Plugins/Mutable/Source/CustomizableObject/Private/MuCO/CustomizableObjectSkeletalMesh.cpp
+++ b/Engine/Plugins/Mutable/Source/CustomizableObject/Private/MuCO/CustomizableObjectSkeletalMesh.cpp
@@ -35,6 +35,8 @@ void UCustomizableObjectSkeletalMesh::InitMutableStreamingData(const TSharedRef<
 {
 	UCustomizableObject& CustomizableObject = *InOperationData->Object;
 
+	ModelStreamableBulkData = CustomizableObject.GetPrivate()->GetModelStreamableBulkData();
+	
 	// Debug info
 	CustomizableObjectPathName = GetNameSafe(&CustomizableObject);
 
@@ -76,7 +78,7 @@ bool UCustomizableObjectSkeletalMesh::StreamIn(int32 NewMipCount, bool bHighPrio
 			? FRenderAssetUpdate::TT_Async
 			: FRenderAssetUpdate::TT_Render;
 
-		PendingUpdate = new FCustomizableObjectMeshStreamIn(this, bHighPrio, CreateResourcesThread);
+		PendingUpdate = new FCustomizableObjectMeshStreamIn(this, bHighPrio, CreateResourcesThread, ModelStreamableBulkData);
 
 		return !PendingUpdate->IsCancelled();
 	}
diff --git a/Engine/Plugins/Mutable/Source/CustomizableObject/Private/MuCO/CustomizableObjectSystem.cpp b/Engine/Plugins/Mutable/Source/CustomizableObject/Private/MuCO/CustomizableObjectSystem.cpp
--- a/Engine/Plugins/Mutable/Source/CustomizableObject/Private/MuCO/CustomizableObjectSystem.cpp
+++ b/Engine/Plugins/Mutable/Source/CustomizableObject/Private/MuCO/CustomizableObjectSystem.cpp
@@ -42,6 +42,7 @@
 #include "Engine/Engine.h"
 #endif
 
+#include "MutableStreamRequest.h"
 #include "MuCO/LoadUtils.h"
 
 #include UE_INLINE_GENERATED_CPP_BY_NAME(CustomizableObjectSystem)
@@ -160,6 +161,7 @@ TAutoConsoleVariable<bool> CVarMutableHighPriorityLoading(
 	true,
 	TEXT("If enabled, the request to load additional assets will have high priority."));
 
+
 int32 UCustomizableObjectSystemPrivate::SkeletalMeshMinLodQualityLevel = -1;
 
 
@@ -536,16 +538,6 @@ FCustomizableObjectComponentIndex FUpdateContextPrivate::GetObjectComponentIndex
 }
 
 
-FInstanceUpdateData::FRealTimeMorphsComponentData* FUpdateContextPrivate::GetComponentMorphTargets(FCustomizableObjectComponentIndex ObjectComponentIndex)
-{
-	return InstanceUpdateData.RealTimeMorphTargets.FindByPredicate(
-		[ObjectComponentIndex](const FInstanceUpdateData::FRealTimeMorphsComponentData& Elem)
-		{
-			return Elem.ObjectComponentIndex == ObjectComponentIndex;
-		});
-}
-
-
 const FInstanceUpdateData::FComponent* FUpdateContextPrivate::GetComponentUpdateData(FCustomizableObjectInstanceComponentIndex InstanceComponentIndex)
 {
 	if (InstanceUpdateData.Components.IsValidIndex(InstanceComponentIndex.GetValue()))
@@ -2234,13 +2226,14 @@ namespace impl
 		}
 	}
 
-	void Subtask_Mutable_PrepareRealTimeMorphData(const TSharedRef<FUpdateContextPrivate>& OperationData)
+	
+	void Subtask_Mutable_PrepareRealTimeMorphData(const TSharedRef<FUpdateContextPrivate>& OperationData) 
 	{
 		MUTABLE_CPUPROFILER_SCOPE(BuildMorphTargetsData);
 
 		FInstanceUpdateData& UpdateData = OperationData->InstanceUpdateData;
 
-		const TMap<uint32, FInstanceUpdateData::FMorphTargetMeshData>& ResourceIdToMeshDataMap =
+		const TMap<uint32, FMorphTargetMeshData>& ResourceIdToMeshDataMap =
 				UpdateData.RealTimeMorphTargetMeshData;
 
 		if (ResourceIdToMeshDataMap.IsEmpty())
@@ -2248,63 +2241,49 @@ namespace impl
 			return;
 		}
 		
-		int32 NumNotFoundLoadedMorphsResources = 0;
-		
 		const int32 NumInstanceComponents = OperationData->InstanceUpdateData.Components.Num();
 		check(OperationData->InstanceUpdateData.RealTimeMorphTargets.IsEmpty());
 		OperationData->InstanceUpdateData.RealTimeMorphTargets.Reserve(NumInstanceComponents);
 		for ( int32 InstanceComponentIndex = 0; InstanceComponentIndex< NumInstanceComponents; ++InstanceComponentIndex)
 		{
 			FInstanceUpdateData::FComponent& Component = OperationData->InstanceUpdateData.Components[InstanceComponentIndex];
-			FInstanceUpdateData::FRealTimeMorphsComponentData& ComponentMorphTargetsData = OperationData->InstanceUpdateData.RealTimeMorphTargets.Emplace_GetRef();
-			ComponentMorphTargetsData.ObjectComponentIndex = Component.Id;
-			check(ComponentMorphTargetsData.ObjectComponentIndex.IsValid());
+			check(Component.Id.IsValid());
 
-			struct FMorphTargetMeshData
-			{
-				TArray<int32> NameResolutionMap;
-				TArrayView<const FMorphTargetVertexData> DataView;
-			};
+			const FName& ComponentName = OperationData->ComponentNames[Component.Id.GetValue()];
 
-			TArray<FName>& MorphTargetNames = ComponentMorphTargetsData.RealTimeMorphTargetNames;
-			MorphTargetNames.Empty();
+			FSkeletalMeshMorphTargets& ComponentMorphTargetsData = OperationData->InstanceUpdateData.RealTimeMorphTargets.Add(ComponentName);
+			
+			ComponentMorphTargetsData.RealTimeMorphTargetNames.Empty();
 
-			TMap<uint32, FMorphTargetMeshData> MorphTargetMeshData;
+			TMap<uint32, FMappedMorphTargetMeshData> MorphTargetMeshData;
 			MorphTargetMeshData.Reserve(ResourceIdToMeshDataMap.Num());
 			
-			for (const TPair<uint32, FInstanceUpdateData::FMorphTargetMeshData>& MorphTargetResource : ResourceIdToMeshDataMap)
+			for (const TPair<uint32, FMorphTargetMeshData>& MorphTargetResource : ResourceIdToMeshDataMap)
 			{
-				FMorphTargetMeshData& MeshData = MorphTargetMeshData.FindOrAdd(MorphTargetResource.Key);
-				MeshData.DataView = MakeArrayView(MorphTargetResource.Value.Data);
+				FMappedMorphTargetMeshData& MeshData = MorphTargetMeshData.FindOrAdd(MorphTargetResource.Key);
+				MeshData.DataView = &MorphTargetResource.Value.Data;
 
 				const int32 NumMorphNames = MorphTargetResource.Value.NameResolutionMap.Num();
 				MeshData.NameResolutionMap.SetNumUninitialized(NumMorphNames);
 
 				for (int32 NameIndex = 0; NameIndex < NumMorphNames; ++NameIndex)
 				{
-					const int32 ResolvedNameIndex = MorphTargetNames.AddUnique(MorphTargetResource.Value.NameResolutionMap[NameIndex]);
+					const int32 ResolvedNameIndex = ComponentMorphTargetsData.RealTimeMorphTargetNames.AddUnique(MorphTargetResource.Value.NameResolutionMap[NameIndex]);
 					MeshData.NameResolutionMap[NameIndex] = ResolvedNameIndex;
 				}
 			}
 
 			// Allocate Morph data for used morphs.
 			TArray<TArray<FMorphTargetLODModel>>& MorphsData = ComponentMorphTargetsData.RealTimeMorphsLODData;
-			const int32 NumMorphs = MorphTargetNames.Num();
-
-			const FName& ComponentName = OperationData->ComponentNames[ComponentMorphTargetsData.ObjectComponentIndex.GetValue()];
-
+			const int32 NumMorphs = ComponentMorphTargetsData.RealTimeMorphTargetNames.Num();
+			
 			const int32 NumLODsInThisComponent = OperationData->NumLODsAvailable[ComponentName];
 
-			MorphsData.SetNum(MorphTargetNames.Num());
+			MorphsData.SetNum(ComponentMorphTargetsData.RealTimeMorphTargetNames.Num());
 			for (int32 MorphIndex = 0; MorphIndex < NumMorphs; ++MorphIndex)
 			{
 				MorphsData[MorphIndex].SetNum(NumLODsInThisComponent);
 			}
-
-			TArray<int32> SectionMorphTargetVerticesCount;
-			SectionMorphTargetVerticesCount.SetNumZeroed(ComponentMorphTargetsData.RealTimeMorphTargetNames.Num());
-
-			int32 NumInvalidVertexMorphNamesFound = 0;
 			
 			const int32 FirstGeneratedLOD = FMath::Max(OperationData->GetRequestedLODs()[ComponentName], OperationData->GetMinLOD()[ComponentName]);
 			for (int32 LODIndex = FirstGeneratedLOD; LODIndex < NumLODsInThisComponent; ++LODIndex)
@@ -2316,141 +2295,20 @@ namespace impl
 				{
 					continue;
 				}
-				
-				const mu::FMeshBufferSet& MeshSet = LOD.Mesh->GetVertexBuffers();
-
-				int32 VertexMorphsInfoIndexAndCountBufferIndex, VertexMorphsInfoIndexAndCountBufferChannel;
-				MeshSet.FindChannel(mu::EMeshBufferSemantic::Other, 0, &VertexMorphsInfoIndexAndCountBufferIndex, &VertexMorphsInfoIndexAndCountBufferChannel);
 
-				int32 VertexMorphsResourceIdBufferIndex, VertexMorphsResourceIdBufferChannel;
-				MeshSet.FindChannel(mu::EMeshBufferSemantic::Other, 1, &VertexMorphsResourceIdBufferIndex, &VertexMorphsResourceIdBufferChannel);
+				TArray<FMorphTargetLODModel> MorphTargets;
+				ReconstructMorphTargets(*LOD.Mesh.Get(), ComponentMorphTargetsData.RealTimeMorphTargetNames, MorphTargetMeshData, MorphTargets);
 
-				if (VertexMorphsInfoIndexAndCountBufferIndex < 0 || VertexMorphsResourceIdBufferIndex < 0)
-				{
-					continue;
-				}
-
-				const uint32* const VertexMorphsInfoIndexAndCountBuffer = reinterpret_cast<const uint32*>(MeshSet.GetBufferData(VertexMorphsInfoIndexAndCountBufferIndex));
-				TArrayView<const uint32> VertexMorphsInfoIndexAndCountView(VertexMorphsInfoIndexAndCountBuffer, MeshSet.GetElementCount());
-
-				const uint32* const VertexMorphsResourceIdBuffer = reinterpret_cast<const uint32*>(MeshSet.GetBufferData(VertexMorphsResourceIdBufferIndex));
-				TArrayView<const uint32> VertexMorphsResourceIdView(VertexMorphsResourceIdBuffer, MeshSet.GetElementCount());
-					
-				const int32 SurfaceCount = LOD.Mesh->GetSurfaceCount();
-				for (int32 Section = 0; Section < SurfaceCount; ++Section)
+				for (int32 NameIndex = 0; NameIndex < ComponentMorphTargetsData.RealTimeMorphTargetNames.Num(); NameIndex++)
 				{
-					// Reset SectionMorphTargets.
-					for (int32& Elem : SectionMorphTargetVerticesCount)
-					{
-						Elem = 0;
-					}
-
-					int32 FirstVertex, VerticesCount, FirstIndex, IndiciesCount, UnusedBoneIndex, UnusedBoneCount;
-					LOD.Mesh->GetSurface(Section, FirstVertex, VerticesCount, FirstIndex, IndiciesCount, UnusedBoneIndex, UnusedBoneCount);
-
-					for (int32 VertexIdx = FirstVertex; VertexIdx < FirstVertex + VerticesCount;)
-					{
-						// Find a span with the same VertexMorphResourceId to amortise the cost of finding 
-						// in the loaded resources map. It is expected to find large consecutive mesh sections pointing to
-						// the same loaded resource.
-
-						const int32 SpanStart = VertexIdx++;
-						const uint32 CurrentResourceId = VertexMorphsResourceIdView[SpanStart];
-
-						// Vertex with no morphs are marked with 0, skip vertex if the case.
-						if (CurrentResourceId == 0)
-						{
-							continue;
-						}
-
-						for (; VertexIdx < FirstVertex + VerticesCount; ++VertexIdx)
-						{
-							const int32 VertexResourceId = VertexMorphsResourceIdView[VertexIdx];
-							// we can skip vertices with no morph without breaking the span.
-							if (VertexResourceId == 0)
-							{
-								continue;
-							}
-
-							if (CurrentResourceId != VertexResourceId)
-							{
-								break;
-							}
-						}
-						const int32 SpanEnd = VertexIdx;
-
-						const FMorphTargetMeshData* MorphTargetReconstructionData = MorphTargetMeshData.Find(CurrentResourceId);
-
-						if (!MorphTargetReconstructionData)
-						{
-							++NumNotFoundLoadedMorphsResources;
-							continue;
-						}
-
-						TArrayView<const FMorphTargetVertexData> SpanMorphData = MorphTargetReconstructionData->DataView;
-						const int32 NumNamesInResolutionMap = MorphTargetReconstructionData->NameResolutionMap.Num();
-
-						for (int32 SpanVertexIdx = SpanStart; SpanVertexIdx < SpanEnd; ++SpanVertexIdx)
-						{
-							const uint32 MorphOffsetAndCount = VertexMorphsInfoIndexAndCountView[SpanVertexIdx];
-							if (MorphOffsetAndCount == 0)
-							{
-								continue;
-							}
-
-							// See encoding in GenerateMutableSourceMesh.cpp.
-							constexpr uint32 Log2MaxNumVerts = 23;
-							
-							TArrayView<const FMorphTargetVertexData> MorphsVertexDataView = MakeArrayView(
-									SpanMorphData.GetData() + (MorphOffsetAndCount & ((1 << Log2MaxNumVerts) - 1)), 
-									MorphOffsetAndCount >> Log2MaxNumVerts);
-
-							for (const FMorphTargetVertexData& SourceVertex : MorphsVertexDataView)
-							{
-								if (SourceVertex.MorphNameIndex >= (uint32)NumNamesInResolutionMap)
-								{
-									++NumInvalidVertexMorphNamesFound;
-									continue;
-								}
-
-								const uint32 ResolvedNameIndex =
-										MorphTargetReconstructionData->NameResolutionMap[SourceVertex.MorphNameIndex];
-
-								FMorphTargetLODModel& DestMorphLODModel = MorphsData[ResolvedNameIndex][LODIndex];
-
-								DestMorphLODModel.Vertices.Emplace(
-										FMorphTargetDelta 
-										{ 
-											SourceVertex.PositionDelta, 
-											SourceVertex.TangentZDelta, 
-											static_cast<uint32>(SpanVertexIdx) 
-										});
-
-								++SectionMorphTargetVerticesCount[ResolvedNameIndex];
-							}
-						}
-					}
-
-					const int32 SectionMorphTargetsNum = SectionMorphTargetVerticesCount.Num();
-					for (int32 MorphIdx = 0; MorphIdx < SectionMorphTargetsNum; ++MorphIdx)
+					if (MorphTargets.IsValidIndex(NameIndex))
 					{
-						if (SectionMorphTargetVerticesCount[MorphIdx] > 0)
-						{
-							FMorphTargetLODModel& MorphTargetLodModel = MorphsData[MorphIdx][LODIndex];
-
-							MorphTargetLodModel.SectionIndices.Add(Section);
-							MorphTargetLodModel.NumVertices += SectionMorphTargetVerticesCount[MorphIdx];
-						}
+						ComponentMorphTargetsData.RealTimeMorphsLODData[NameIndex][LODIndex] = MoveTemp(MorphTargets[NameIndex]);
 					}
 				}
 			}
-	
-			if (NumInvalidVertexMorphNamesFound > 0)
-			{
-				UE_LOG(LogMutable, Warning, TEXT("Invalid real-time morphs names found in instance vertices. Some morph may not work as expected."));
-			}
-	
-			// Remove empty morph targets;
+			
+			// Remove empty morph targets
 			for (int32 MorphIndex = 0; MorphIndex < NumMorphs; ++MorphIndex)
 			{
 				const int32 NumLODs = MorphsData[MorphIndex].Num();
@@ -2469,16 +2327,10 @@ namespace impl
 					MorphsData[MorphIndex].Empty();
 				}
 			}
-
 		}
 
 		// Free unneeded data memory.
 		UpdateData.RealTimeMorphTargetMeshData.Empty();
-		
-		if (NumNotFoundLoadedMorphsResources > 0)
-		{
-			UE_LOG(LogMutable, Warning, TEXT("Needed realtime morph reconstruction data was not loaded properly. Some realtime morphs may not work correctly."));
-		}
 	}
 
 	/** End of the GetMeshes tasks. */
@@ -4929,3 +4781,200 @@ int32 UCustomizableObjectSystemPrivate::BlockTillAllRequestsFinished(float TimeL
 
 	return 0;
 }
+
+
+void LoadMorphTargets(FMutableStreamRequest& MutableDataStreamer, const TSharedRef<const mu::FMesh>& Mesh, TMap<uint32, FMorphTargetMeshData>& StreamingResult)
+{
+	MUTABLE_CPUPROFILER_SCOPE(RealTimeMorphStreaming);
+
+	if (!CVarEnableRealTimeMorphTargets.GetValueOnAnyThread())
+	{
+		return;
+	}
+
+	const TSharedPtr<FModelStreamableBulkData>& ModelStreamableBulkData = MutableDataStreamer.GetModelStreamableBulkData();
+	
+	TArray<uint32> RealTimeMorphStreamableBlocksToStream;
+	
+	for (uint64 ResourceId : Mesh->GetStreamedResources())
+	{
+		const FCustomizableObjectStreameableResourceId TypedResourceId = BitCast<FCustomizableObjectStreameableResourceId>(ResourceId);
+			
+		if (TypedResourceId.Type == static_cast<uint8>(FCustomizableObjectStreameableResourceId::EType::RealTimeMorphTarget))
+		{
+			check(TypedResourceId.Id != 0 && TypedResourceId.Id <= TNumericLimits<uint32>::Max());
+
+			if (ModelStreamableBulkData->RealTimeMorphStreamables.Contains(static_cast<uint32>(TypedResourceId.Id)))
+			{
+				RealTimeMorphStreamableBlocksToStream.AddUnique(TypedResourceId.Id);
+			}
+			else
+			{
+				UE_LOG(LogMutable, Error, TEXT("Invalid streamed real time morph target data block [%d] found."), TypedResourceId.Id);
+			}
+		}
+	}
+	
+	for (const int32 BlockId : RealTimeMorphStreamableBlocksToStream)
+	{
+		MUTABLE_CPUPROFILER_SCOPE(RealTimeMorphStreamingRequest_Alloc);
+
+		const FRealTimeMorphStreamable& Streamable = ModelStreamableBulkData->RealTimeMorphStreamables[BlockId];
+		const FMutableStreamableBlock& Block = Streamable.Block; 
+		
+		FMorphTargetMeshData& ReadDestData = StreamingResult.FindOrAdd(BlockId);
+
+		// Only request blocks once.
+		if (ReadDestData.Data.Num())
+		{
+			continue;
+		}
+
+		ReadDestData.NameResolutionMap = ModelStreamableBulkData->RealTimeMorphStreamables[BlockId].NameResolutionMap;
+
+		check(Streamable.Size % sizeof(FMorphTargetVertexData) == 0);
+		uint32 NumElems = Streamable.Size / sizeof(FMorphTargetVertexData);
+
+		ReadDestData.Data.SetNumUninitialized(NumElems);
+
+		MutableDataStreamer.AddBlock(
+			Block,
+			MutablePrivate::EStreamableDataType::RealTimeMorph,
+			MakeArrayView(reinterpret_cast<uint8*>(ReadDestData.Data.GetData()), ReadDestData.Data.Num() * sizeof(FMorphTargetVertexData)));
+	}
+}
+
+
+void ReconstructMorphTargets(const mu::FMesh& Mesh, const TArray<FName>& GlobalNames, const TMap<uint32, FMappedMorphTargetMeshData>& MappedMorphTargets, TArray<FMorphTargetLODModel>& OutMorphTargets)
+{
+	const mu::FMeshBufferSet& MeshSet = Mesh.GetVertexBuffers();
+
+	int32 VertexMorphsInfoIndexAndCountBufferIndex, VertexMorphsInfoIndexAndCountBufferChannel;
+	MeshSet.FindChannel(mu::EMeshBufferSemantic::Other, 0, &VertexMorphsInfoIndexAndCountBufferIndex, &VertexMorphsInfoIndexAndCountBufferChannel);
+
+	int32 VertexMorphsResourceIdBufferIndex, VertexMorphsResourceIdBufferChannel;
+	MeshSet.FindChannel(mu::EMeshBufferSemantic::Other, 1, &VertexMorphsResourceIdBufferIndex, &VertexMorphsResourceIdBufferChannel);
+
+	if (VertexMorphsInfoIndexAndCountBufferIndex < 0 || VertexMorphsResourceIdBufferIndex < 0)
+	{
+		return;
+	}
+
+	OutMorphTargets.SetNum(GlobalNames.Num());
+
+	const uint32* const VertexMorphsInfoIndexAndCountBuffer = reinterpret_cast<const uint32*>(MeshSet.GetBufferData(VertexMorphsInfoIndexAndCountBufferIndex));
+	TArrayView<const uint32> VertexMorphsInfoIndexAndCountView(VertexMorphsInfoIndexAndCountBuffer, MeshSet.GetElementCount());
+
+	const uint32* const VertexMorphsResourceIdBuffer = reinterpret_cast<const uint32*>(MeshSet.GetBufferData(VertexMorphsResourceIdBufferIndex));
+	TArrayView<const uint32> VertexMorphsResourceIdView(VertexMorphsResourceIdBuffer, MeshSet.GetElementCount());
+
+	TArray<int32> SectionMorphTargetVerticesCount;
+	SectionMorphTargetVerticesCount.SetNumZeroed(GlobalNames.Num());
+	
+	const int32 SurfaceCount = Mesh.GetSurfaceCount();
+	for (int32 Section = 0; Section < SurfaceCount; ++Section)
+	{
+		// Reset SectionMorphTargets.
+		for (int32& Elem : SectionMorphTargetVerticesCount)
+		{
+			Elem = 0;
+		}
+
+		int32 FirstVertex, VerticesCount, FirstIndex, IndiciesCount, UnusedBoneIndex, UnusedBoneCount;
+		Mesh.GetSurface(Section, FirstVertex, VerticesCount, FirstIndex, IndiciesCount, UnusedBoneIndex, UnusedBoneCount);
+
+		for (int32 VertexIdx = FirstVertex; VertexIdx < FirstVertex + VerticesCount;)
+		{
+			// Find a span with the same VertexMorphResourceId to amortise the cost of finding 
+			// in the loaded resources map. It is expected to find large consecutive mesh sections pointing to
+			// the same loaded resource.
+
+			const int32 SpanStart = VertexIdx++;
+			const uint32 CurrentResourceId = VertexMorphsResourceIdView[SpanStart];
+
+			// Vertex with no morphs are marked with 0, skip vertex if the case.
+			if (CurrentResourceId == 0)
+			{
+				continue;
+			}
+
+			for (; VertexIdx < FirstVertex + VerticesCount; ++VertexIdx)
+			{
+				const int32 VertexResourceId = VertexMorphsResourceIdView[VertexIdx];
+				// we can skip vertices with no morph without breaking the span.
+				if (VertexResourceId == 0)
+				{
+					continue;
+				}
+
+				if (CurrentResourceId != VertexResourceId)
+				{
+					break;
+				}
+			}
+			const int32 SpanEnd = VertexIdx;
+
+			const FMappedMorphTargetMeshData* MorphTargetReconstructionData = MappedMorphTargets.Find(CurrentResourceId);
+
+			if (!MorphTargetReconstructionData)
+			{
+				ensureMsgf(false, TEXT("Needed realtime morph reconstruction data was not loaded properly. Some realtime morphs may not work correctly."));
+				continue;
+			}
+
+			const TArray<FMorphTargetVertexData>& SpanMorphData = *MorphTargetReconstructionData->DataView;
+			const int32 NumNamesInResolutionMap = MorphTargetReconstructionData->NameResolutionMap.Num();
+
+			for (int32 SpanVertexIdx = SpanStart; SpanVertexIdx < SpanEnd; ++SpanVertexIdx)
+			{
+				const uint32 MorphOffsetAndCount = VertexMorphsInfoIndexAndCountView[SpanVertexIdx];
+				if (MorphOffsetAndCount == 0)
+				{
+					continue;
+				}
+
+				// See encoding in GenerateMutableSourceMesh.cpp.
+				constexpr uint32 Log2MaxNumVerts = 23;
+				
+				TArrayView<const FMorphTargetVertexData> MorphsVertexDataView = MakeArrayView(
+						SpanMorphData.GetData() + (MorphOffsetAndCount & ((1 << Log2MaxNumVerts) - 1)), 
+						MorphOffsetAndCount >> Log2MaxNumVerts);
+
+				for (const FMorphTargetVertexData& SourceVertex : MorphsVertexDataView)
+				{
+					if (SourceVertex.MorphNameIndex >= (uint32)NumNamesInResolutionMap)
+					{
+						ensureMsgf(false, TEXT("Invalid real-time morphs names found in instance vertices. Some morph may not work as expected."));
+						continue;
+					}
+
+					const uint32 ResolvedNameIndex = MorphTargetReconstructionData->NameResolutionMap[SourceVertex.MorphNameIndex];
+
+					FMorphTargetLODModel& DestMorphLODModel = OutMorphTargets[ResolvedNameIndex];
+
+					DestMorphLODModel.Vertices.Emplace(
+							FMorphTargetDelta 
+							{ 
+								SourceVertex.PositionDelta, 
+								SourceVertex.TangentZDelta, 
+								static_cast<uint32>(SpanVertexIdx) 
+							});
+
+					++SectionMorphTargetVerticesCount[ResolvedNameIndex];
+				}
+			}
+		}
+
+		const int32 SectionMorphTargetsNum = SectionMorphTargetVerticesCount.Num();
+		for (int32 MorphIdx = 0; MorphIdx < SectionMorphTargetsNum; ++MorphIdx)
+		{
+			if (SectionMorphTargetVerticesCount[MorphIdx] > 0)
+			{
+				FMorphTargetLODModel& MorphTargetLodModel = OutMorphTargets[MorphIdx];
+
+				MorphTargetLodModel.SectionIndices.Add(Section);
+				MorphTargetLodModel.NumVertices += SectionMorphTargetVerticesCount[MorphIdx];
+			}
+		}
+	}
+}
diff --git a/Engine/Plugins/Mutable/Source/CustomizableObject/Private/MuCO/UnrealConversionUtils.cpp b/Engine/Plugins/Mutable/Source/CustomizableObject/Private/MuCO/UnrealConversionUtils.cpp
--- a/Engine/Plugins/Mutable/Source/CustomizableObject/Private/MuCO/UnrealConversionUtils.cpp
+++ b/Engine/Plugins/Mutable/Source/CustomizableObject/Private/MuCO/UnrealConversionUtils.cpp
@@ -8,6 +8,7 @@
 #include "GPUSkinVertexFactory.h"
 #include "MuCO/CustomizableObject.h"
 #include "MuCO/CustomizableObjectPrivate.h"
+#include "MuCO/CustomizableObjectSystemPrivate.h"
 #include "MuR/Skeleton.h"
 #include "MuR/Mesh.h"
 #include "MuR/MeshBufferSet.h"
@@ -748,4 +749,57 @@ namespace UnrealConversionUtils
 			LODResource.BuffersSize += LODResource.MultiSizeIndexContainer.GetIndexBuffer()->GetResourceDataSize();
 		}
 	}
+
+
+	void MorphTargetVertexInfoBuffers(FSkeletalMeshLODRenderData& LODResource, const USkeletalMesh& Owner, const mu::FMesh& MutableMesh, const TMap<uint32, FMorphTargetMeshData>& MorphTargetMeshData, int32 LODIndex)
+	{
+		// Get the global names.
+		const TMap<FName, int32>& IndexMap = Owner.GetMorphTargetIndexMap();
+	
+		int32 MaxIndex = 0;
+		for (const TPair<FName, int32>& Pair : IndexMap)
+		{
+			MaxIndex = FMath::Max(MaxIndex, Pair.Value);
+		}
+
+		TArray<FName> GlobalNames;
+		GlobalNames.SetNum(MaxIndex + 1);
+
+		for (const TPair<FName, int32>& Pair : IndexMap)
+		{
+			GlobalNames[Pair.Value] = Pair.Key;
+		}
+
+		// Map the local names to the global names.
+		TMap<uint32, FMappedMorphTargetMeshData> GlobalMorphTargets; // Key is the block id. Value data needed to reconstruct the Morp Targets.
+		for (const TPair<uint32, FMorphTargetMeshData>& Pair : MorphTargetMeshData)
+		{
+			FMappedMorphTargetMeshData& MappedMorphsData = GlobalMorphTargets.FindOrAdd(Pair.Key);
+
+			// Remapping to global names.
+			MappedMorphsData.NameResolutionMap.SetNum(Pair.Value.NameResolutionMap.Num());
+			for (int32 NameIndex = 0; NameIndex < Pair.Value.NameResolutionMap.Num(); ++NameIndex)
+			{
+				MappedMorphsData.NameResolutionMap[NameIndex] = GlobalNames.Find(Pair.Value.NameResolutionMap[NameIndex]);
+			}
+
+			// Pointer to the serialized data
+			MappedMorphsData.DataView = &Pair.Value.Data;
+		}
+
+		// Reconstruct the final Morph Targets using the global names.
+		TArray<FMorphTargetLODModel> MorphTargetLODs;
+		ReconstructMorphTargets(MutableMesh, GlobalNames, GlobalMorphTargets, MorphTargetLODs);
+
+		TArray<const FMorphTargetLODModel*> MorphsToInit;
+		MorphsToInit.Reserve(MorphTargetLODs.Num());
+		for (FMorphTargetLODModel& MorphTargetLOD : MorphTargetLODs)
+		{
+			MorphsToInit.Add(&MorphTargetLOD);
+		}
+
+		const TBitArray UsesBuiltinMorphTargetCompression(true, MorphTargetLODs.Num());
+		const float ErrorTolerance = Owner.GetLODInfo(LODIndex)->MorphTargetPositionErrorTolerance;
+		LODResource.MorphTargetVertexInfoBuffers.InitMorphResourcesStreaming(LODResource.RenderSections, MorphsToInit, LODResource.StaticVertexBuffers.StaticMeshVertexBuffer.GetNumVertices(), ErrorTolerance);
+	}
 }
diff --git a/Engine/Plugins/Mutable/Source/CustomizableObject/Public/MuCO/CustomizableObjectSkeletalMesh.h b/Engine/Plugins/Mutable/Source/CustomizableObject/Public/MuCO/CustomizableObjectSkeletalMesh.h
--- a/Engine/Plugins/Mutable/Source/CustomizableObject/Public/MuCO/CustomizableObjectSkeletalMesh.h
+++ b/Engine/Plugins/Mutable/Source/CustomizableObject/Public/MuCO/CustomizableObjectSkeletalMesh.h
@@ -19,8 +19,8 @@ class UCustomizableObjectInstance;
 class UCustomizableObject;
 class UCustomizableInstancePrivate;
 class FUpdateContextPrivate;
-
 class FCustomizableObjectMeshStreamIn;
+struct FModelStreamableBulkData;
 
 FName GenerateUniqueNameFromCOInstance(const UCustomizableObjectInstance& Instance, FName ComponentName);
 
@@ -48,6 +48,8 @@ private:
 
 	FString CustomizableObjectPathName;
 
+	TSharedPtr<FModelStreamableBulkData> ModelStreamableBulkData;
+
 	TSharedPtr<mu::FModel, ESPMode::ThreadSafe> Model;
 	
 	/** Instance parameters at the time of creation. */
@@ -57,5 +59,4 @@ private:
 	/** Mesh IDs per LOD */
 	TArray<mu::FResourceID> MeshIDs;
 	TArray<TPair<uint32, FName>> SkinWeightProfileIDs;
-
 };
diff --git a/Engine/Source/Runtime/Engine/Private/Rendering/MorphTargetVertexInfoBuffers.cpp b/Engine/Source/Runtime/Engine/Private/Rendering/MorphTargetVertexInfoBuffers.cpp
--- a/Engine/Source/Runtime/Engine/Private/Rendering/MorphTargetVertexInfoBuffers.cpp
+++ b/Engine/Source/Runtime/Engine/Private/Rendering/MorphTargetVertexInfoBuffers.cpp
@@ -172,7 +172,7 @@ private:
 	int32 			NumPendingBits;
 };
 
-void FMorphTargetVertexInfoBuffers::InitMorphResources(EShaderPlatform ShaderPlatform, const TArray<FSkelMeshRenderSection>& RenderSections, const TArray<UMorphTarget*>& MorphTargets, int NumVertices, int32 LODIndex, float TargetPositionErrorTolerance)
+void FMorphTargetVertexInfoBuffers::InitMorphResources(EShaderPlatform ShaderPlatform, const TArray<FSkelMeshRenderSection>& RenderSections, const TArray<UMorphTarget*>& MorphTargets, int32 NumVertices, int32 LODIndex, float TargetPositionErrorTolerance)
 {
 	check(!IsRHIInitialized());
 	check(!IsMorphResourcesInitialized());
@@ -188,6 +188,41 @@ void FMorphTargetVertexInfoBuffers::InitMorphResources(EShaderPlatform ShaderPla
 
 	bIsMorphCPUDataValid = true;
 
+	TArray<const FMorphTargetLODModel*> MorphTargetLODs;
+	MorphTargetLODs.Reserve(MorphTargets.Num());
+	
+	TBitArray UsesBuiltinMorphTargetCompression;
+	UsesBuiltinMorphTargetCompression.Reserve(MorphTargets.Num());
+
+	for (UMorphTarget* MorphTarget : MorphTargets)
+	{
+		TArray<FMorphTargetLODModel>& MorphLODModels = MorphTarget->GetMorphLODModels();
+		
+		const FMorphTargetLODModel* MorphTargetLOD = MorphLODModels.IsValidIndex(LODIndex) ? &MorphLODModels[LODIndex] : nullptr;
+		MorphTargetLODs.Add(MorphTargetLOD);
+		UsesBuiltinMorphTargetCompression.Add(MorphTarget->UsesBuiltinMorphTargetCompression());
+	}
+
+	Compress(RenderSections, MorphTargetLODs, UsesBuiltinMorphTargetCompression, NumVertices, TargetPositionErrorTolerance);
+	
+	ValidateVertexBuffers(true);
+}
+
+void FMorphTargetVertexInfoBuffers::InitMorphResourcesStreaming(const TArray<FSkelMeshRenderSection>& RenderSections, const TArray<const FMorphTargetLODModel*>& MorphTargets, int32 NumVertices, float TargetPositionErrorTolerance)
+{
+	ResetCPUData();
+
+	const TBitArray UsesBuiltinMorphTargetCompression(true, MorphTargets.Num());
+	Compress(RenderSections, MorphTargets, UsesBuiltinMorphTargetCompression, NumVertices, TargetPositionErrorTolerance);
+	
+	bRHIInitialized = false;
+	bIsMorphCPUDataValid = true;
+	bResourcesInitialized = true;
+	ValidateVertexBuffers(true);
+}
+
+void FMorphTargetVertexInfoBuffers::Compress(const TArray<FSkelMeshRenderSection>& RenderSections, const TArray<const FMorphTargetLODModel*>& MorphTargets, const TBitArray<>& UsesBuiltinMorphTargetCompression, int32 NumVertices, float TargetPositionErrorTolerance)
+{
 	// Simple Morph compression 0.1
 	// Instead of storing vertex deltas individually they are organized into batches of 64.
 	// Each batch has a header that describes how many bits are allocated to each of the vertex components.
@@ -273,11 +308,11 @@ void FMorphTargetVertexInfoBuffers::InitMorphResources(EShaderPlatform ShaderPla
 
 		float MaximumValues[4] = { -FLT_MAX, -FLT_MAX, -FLT_MAX, -FLT_MAX };
 		float MinimumValues[4] = { +FLT_MAX, +FLT_MAX, +FLT_MAX, +FLT_MAX };
-		UMorphTarget* MorphTarget = MorphTargets[AnimIdx];
-		int32 NumSrcDeltas = 0;
-		const FMorphTargetDelta* MorphDeltas = MorphTarget->GetMorphTargetDelta(LODIndex, NumSrcDeltas);
 
-		if (NumSrcDeltas == 0 || !MorphTarget->UsesBuiltinMorphTargetCompression())
+		const FMorphTargetLODModel* MorphModel = MorphTargets[AnimIdx];
+		int32 NumSrcDeltas = MorphModel ? MorphModel->Vertices.Num() : 0;
+		
+		if (NumSrcDeltas == 0 || !UsesBuiltinMorphTargetCompression[AnimIdx])
 		{
 			MaximumValues[0] = 0.0f;
 			MaximumValues[1] = 0.0f;
@@ -291,6 +326,8 @@ void FMorphTargetVertexInfoBuffers::InitMorphResources(EShaderPlatform ShaderPla
 		}
 		else
 		{
+			const FMorphTargetDelta* MorphDeltas = MorphModel->Vertices.GetData();
+
 			struct FQuantizedDelta
 			{
 				FIntVector	Position;
@@ -487,8 +524,6 @@ void FMorphTargetVertexInfoBuffers::InitMorphResources(EShaderPlatform ShaderPla
 	}
 
 	// UE_LOG(LogStaticMesh, Log, TEXT("Morph compression time: [%.2fs]"), FPlatformTime::ToMilliseconds(FPlatformTime::Cycles() - StartTime) / 1000.0f);
-
-	ValidateVertexBuffers(true);
 }
 
 bool FMorphTargetVertexInfoBuffers::IsPlatformShaderSupported(EShaderPlatform ShaderPlatform)
diff --git a/Engine/Source/Runtime/Engine/Public/Rendering/MorphTargetVertexInfoBuffers.h b/Engine/Source/Runtime/Engine/Public/Rendering/MorphTargetVertexInfoBuffers.h
--- a/Engine/Source/Runtime/Engine/Public/Rendering/MorphTargetVertexInfoBuffers.h
+++ b/Engine/Source/Runtime/Engine/Public/Rendering/MorphTargetVertexInfoBuffers.h
@@ -7,6 +7,7 @@
 #include "Math/Vector4.h"
 
 struct FSkelMeshRenderSection;
+struct FMorphTargetLODModel;
 class UMorphTarget;
 
 class FMorphTargetVertexInfoBuffers : public FRenderResource
@@ -15,7 +16,9 @@ public:
 	ENGINE_API FMorphTargetVertexInfoBuffers();
 	ENGINE_API ~FMorphTargetVertexInfoBuffers();
 
-	ENGINE_API void InitMorphResources(EShaderPlatform ShaderPlatform, const TArray<FSkelMeshRenderSection>& RenderSections, const TArray<UMorphTarget*>& MorphTargets, int NumVertices, int32 LODIndex, float TargetPositionErrorTolerance);
+	ENGINE_API void InitMorphResources(EShaderPlatform ShaderPlatform, const TArray<FSkelMeshRenderSection>& RenderSections, const TArray<UMorphTarget*>& MorphTargets, int32 NumVertices, int32 LODIndex, float TargetPositionErrorTolerance);
+
+	ENGINE_API void InitMorphResourcesStreaming(const TArray<FSkelMeshRenderSection>& RenderSections, const TArray<const FMorphTargetLODModel*>& MorphTargets, int32 NumVertices, float TargetPositionErrorTolerance);
 
 	inline bool IsMorphResourcesInitialized() const { return bResourcesInitialized; }
 	inline bool IsRHIInitialized() const { return bRHIInitialized; }
@@ -87,6 +90,8 @@ protected:
 	ENGINE_API void ValidateVertexBuffers(bool bMorphTargetsShouldBeValid);
 	ENGINE_API void Serialize(FArchive& Ar);
 
+	void Compress(const TArray<FSkelMeshRenderSection>& RenderSections, const TArray<const FMorphTargetLODModel*>& MorphTargets, const TBitArray<>& UsesBuiltinMorphTargetCompression, int32 NumVertices, float TargetPositionErrorTolerance);
+	
 	// Transient data. Gets deleted as soon as the GPU resource has been initialized (unless excplicitly disabled by bEmptyMorphCPUDataOnInitRHI).
 	TResourceArray<uint32> MorphData;
 
