[mutable] Fixed COI asset not being shown in the content browser upon creation.
- Removed some includes

#rnx

diff --git a/Engine/Plugins/Mutable/Source/CustomizableObjectEditor/Private/MuCOE/AssetDefinition_CustomizableObject.cpp b/Engine/Plugins/Mutable/Source/CustomizableObjectEditor/Private/MuCOE/AssetDefinition_CustomizableObject.cpp
--- a/Engine/Plugins/Mutable/Source/CustomizableObjectEditor/Private/MuCOE/AssetDefinition_CustomizableObject.cpp
+++ b/Engine/Plugins/Mutable/Source/CustomizableObjectEditor/Private/MuCOE/AssetDefinition_CustomizableObject.cpp
@@ -15,7 +15,6 @@
 #include "MuCO/CustomizableObject.h"
 #include "MuCO/CustomizableObjectInstance.h"
 #include "MuCO/CustomizableObjectSystem.h"
-#include "MuCO/ICustomizableObjectEditorModule.h"
 #include "Subsystems/AssetEditorSubsystem.h"
 #include "ToolMenus.h"
 #include "MuCO/LoadUtils.h"
@@ -120,6 +119,8 @@ namespace MenuExtension_CustomizableObject
 
 					Instance->SetObject(Object);
 
+					FAssetRegistryModule::AssetCreated(Instance);
+					
 					ObjectsToSync.Add(Instance);
 				}
 			}
diff --git a/Engine/Plugins/Mutable/Source/CustomizableObjectEditor/Private/MuCOE/CustomizableObjectInstanceFactory.cpp b/Engine/Plugins/Mutable/Source/CustomizableObjectEditor/Private/MuCOE/CustomizableObjectInstanceFactory.cpp
--- a/Engine/Plugins/Mutable/Source/CustomizableObjectEditor/Private/MuCOE/CustomizableObjectInstanceFactory.cpp
+++ b/Engine/Plugins/Mutable/Source/CustomizableObjectEditor/Private/MuCOE/CustomizableObjectInstanceFactory.cpp
@@ -5,17 +5,13 @@
 #include "AssetRegistry/AssetData.h"
 #include "AssetRegistry/AssetRegistryModule.h"
 #include "Components/SkeletalMeshComponent.h"
-#include "MuCO/CustomizableSkeletalComponentPrivate.h"
 #include "Engine/SkeletalMesh.h"
 #include "Modules/ModuleManager.h"
 #include "MuCO/CustomizableObjectInstance.h"
-#include "MuCO/CustomizableObjectPrivate.h"
 #include "MuCO/CustomizableSkeletalComponent.h"
 #include "MuCO/CustomizableSkeletalMeshActor.h"
 #include "MuCO/CustomizableSkeletalMeshActorPrivate.h"
 #include "MuCO/LoadUtils.h"
-#include "MuCOE/Nodes/CustomizableObjectNodeObject.h"
-#include "MuCO/UnrealPortabilityHelpers.h"
 
 #define LOCTEXT_NAMESPACE "CustomizableObjectInstanceFactory"
 
@@ -62,7 +58,7 @@ UObject* UCustomizableObjectInstanceFactory::GetAssetFromActorInstance(AActor* A
 		const TArray<TObjectPtr<UCustomizableSkeletalComponent>>& Components = CSMActor->GetPrivate()->GetComponents();
 		if (Components.IsValidIndex(0))
 		{
-			return Components[0];
+			return Components[0]; 
 		}
 	}
 	
