[Mutable] Renamed LOD variables.

diff --git a/Engine/Plugins/Mutable/Source/CustomizableObject/Internal/MuCO/CustomizableObjectPrivate.h b/Engine/Plugins/Mutable/Source/CustomizableObject/Internal/MuCO/CustomizableObjectPrivate.h
--- a/Engine/Plugins/Mutable/Source/CustomizableObject/Internal/MuCO/CustomizableObjectPrivate.h
+++ b/Engine/Plugins/Mutable/Source/CustomizableObject/Internal/MuCO/CustomizableObjectPrivate.h
@@ -1306,15 +1306,15 @@ public:
 	
 	/** Max number of LODs in the compiled Model. */
 	UPROPERTY()
-	TMap<FName, uint8> ComponentLODs;
+	TMap<FName, uint8> NumLODsAvailable;
 
 	/** Max number of LODs to stream. Mutable will always generate at least one LOD. */
 	UPROPERTY()
-	TMap<FName, uint8> ComponentNumLODsToStream;
+	TMap<FName, uint8> NumLODsToStream;
 
 	/** First LOD available, some platforms may remove lower LODs when cooking, this MinLOD represents the first LOD we can generate */
 	UPROPERTY()
-	TMap<FName, uint8> ComponentFirstLODAvailable;
+	TMap<FName, uint8> FirstLODAvailable;
 
 	/** Name of all possible components. Index is the ObjectComponentIndex. */
 	UPROPERTY()
diff --git a/Engine/Plugins/Mutable/Source/CustomizableObject/Internal/MuCO/CustomizableObjectSystemPrivate.h b/Engine/Plugins/Mutable/Source/CustomizableObject/Internal/MuCO/CustomizableObjectSystemPrivate.h
--- a/Engine/Plugins/Mutable/Source/CustomizableObject/Internal/MuCO/CustomizableObjectSystemPrivate.h
+++ b/Engine/Plugins/Mutable/Source/CustomizableObject/Internal/MuCO/CustomizableObjectSystemPrivate.h
@@ -755,6 +755,21 @@ struct FPendingReleaseSkeletalMeshInfo
 };
 
 
+/** Private part, hidden from outside the plugin.
+ *
+ * Small diagram to help understanding Mutable LOD nomenclature:
+ * 
+ * [- NumLODsAvailable -----------------------------] = 8
+ * [- Stripped --[- Packaged -----------------------]
+ * 0      1      2      3      4      5      6      7    
+ * |------|------|------|------|------|------|------|
+ *               [- Streaming --------[- Residents -]
+ *               ^                    ^
+ *               |                    |
+ *               FirstLODAvailable    FirstResidentLOD
+ *               MinLOD (when not using Quality Levels)
+ * 
+ * [- NumLODsToStream ----------------] = 6 */
 UCLASS()
 class UCustomizableObjectSystemPrivate : public UObject, public IStreamingManager
 {
diff --git a/Engine/Plugins/Mutable/Source/CustomizableObject/Private/MuCO/CustomizableObject.cpp b/Engine/Plugins/Mutable/Source/CustomizableObject/Private/MuCO/CustomizableObject.cpp
--- a/Engine/Plugins/Mutable/Source/CustomizableObject/Private/MuCO/CustomizableObject.cpp
+++ b/Engine/Plugins/Mutable/Source/CustomizableObject/Private/MuCO/CustomizableObject.cpp
@@ -1987,7 +1987,7 @@ uint8 UCustomizableObjectPrivate::GetMinLODIndex(const FName& ComponentName) con
 
 	if (LocalModelResources)
 	{
-		if (const uint8* Result = LocalModelResources->ComponentFirstLODAvailable.Find(ComponentName))
+		if (const uint8* Result = LocalModelResources->FirstLODAvailable.Find(ComponentName))
 		{
 			FirstLODAvailable = *Result;
 		}
diff --git a/Engine/Plugins/Mutable/Source/CustomizableObject/Private/MuCO/CustomizableObjectInstance.cpp b/Engine/Plugins/Mutable/Source/CustomizableObject/Private/MuCO/CustomizableObjectInstance.cpp
--- a/Engine/Plugins/Mutable/Source/CustomizableObject/Private/MuCO/CustomizableObjectInstance.cpp
+++ b/Engine/Plugins/Mutable/Source/CustomizableObject/Private/MuCO/CustomizableObjectInstance.cpp
@@ -6716,7 +6716,7 @@ void UCustomizableObjectInstance::SetRequestedLODs(const TMap<FName, uint8>& InM
 		const uint8 MinLODIdx = CustomizableObject->GetPrivate()->GetMinLODIndex(ComponentName);
 		MutableUpdateCandidate.QualitySettingMinLODs.Add(ComponentName, MinLODIdx);
 
-		const int32 MaxLODIdx = ModelResources->ComponentLODs[ComponentName] - 1;
+		const int32 MaxLODIdx = ModelResources->NumLODsAvailable[ComponentName] - 1;
 		InMinLOD = FMath::Clamp(InMinLOD, MinLODIdx, MaxLODIdx);
 
 		uint8 DescriptorMinLOD = 0;
@@ -6737,7 +6737,7 @@ void UCustomizableObjectInstance::SetRequestedLODs(const TMap<FName, uint8>& InM
 
 			PrivateData->SetCOInstanceFlags(InMinLOD > CurrentMinLOD ? PendingLODsDowngrade : ECONone);
 
-			const uint8 FirstNonStreamedLODIndex = FMath::Clamp(ModelResources->ComponentNumLODsToStream[ComponentName], 0, MaxLODIdx);
+			const uint8 FirstNonStreamedLODIndex = FMath::Clamp(ModelResources->NumLODsToStream[ComponentName], 0, MaxLODIdx);
 
 			MutableUpdateCandidate.RequestedLODLevels.Add(ComponentName, FirstNonStreamedLODIndex);
 
diff --git a/Engine/Plugins/Mutable/Source/CustomizableObject/Private/MuCO/CustomizableObjectSystem.cpp b/Engine/Plugins/Mutable/Source/CustomizableObject/Private/MuCO/CustomizableObjectSystem.cpp
--- a/Engine/Plugins/Mutable/Source/CustomizableObject/Private/MuCO/CustomizableObjectSystem.cpp
+++ b/Engine/Plugins/Mutable/Source/CustomizableObject/Private/MuCO/CustomizableObjectSystem.cpp
@@ -313,8 +313,8 @@ FUpdateContextPrivate::FUpdateContextPrivate(UCustomizableObjectInstance& InInst
 	
 	if (const UModelResources* ModelResources = InObject->GetPrivate()->GetModelResources())
 	{
-		FirstLODAvailable = ModelResources->ComponentFirstLODAvailable;
-		FirstResidentLOD = ModelResources->ComponentNumLODsToStream;
+		FirstLODAvailable = ModelResources->FirstLODAvailable;
+		FirstResidentLOD = ModelResources->NumLODsToStream;
 		ComponentNames = ModelResources->ComponentNamesPerObjectComponent;
 	}
 
diff --git a/Engine/Plugins/Mutable/Source/CustomizableObjectEditor/Private/MuCOE/CustomizableObjectCompiler.cpp b/Engine/Plugins/Mutable/Source/CustomizableObjectEditor/Private/MuCOE/CustomizableObjectCompiler.cpp
--- a/Engine/Plugins/Mutable/Source/CustomizableObjectEditor/Private/MuCOE/CustomizableObjectCompiler.cpp
+++ b/Engine/Plugins/Mutable/Source/CustomizableObjectEditor/Private/MuCOE/CustomizableObjectCompiler.cpp
@@ -974,21 +974,21 @@ void FCustomizableObjectCompiler::CompileInternal(bool bAsync)
 			ModelResources->ReleaseVersion = VersionBridge->GetCurrentVersionAsString();
 		}
 
-		ModelResources->ComponentLODs = GenerationContext.NumLODs;
+		ModelResources->NumLODsAvailable = GenerationContext.NumLODs;
 
 		if (GenerationContext.bEnableLODStreaming)
 		{
-			ModelResources->ComponentNumLODsToStream = GenerationContext.NumMaxLODsToStream;
+			ModelResources->NumLODsToStream = GenerationContext.NumMaxLODsToStream;
 		}
 		else
 		{
-			for (TTuple<FName, uint8>& Pair : ModelResources->ComponentNumLODsToStream)
+			for (TTuple<FName, uint8>& Pair : ModelResources->NumLODsToStream)
 			{
 				Pair.Value = 0;
 			}
 		}
 
-		ModelResources->ComponentFirstLODAvailable = GenerationContext.FirstLODAvailable;
+		ModelResources->FirstLODAvailable = GenerationContext.FirstLODAvailable;
 		
 		ModelResources->ParticipatingObjects = MoveTemp(ParticipatingObjects);
 		
