[Mutable] Reenabled TickWarnings.

Improved its performance. Work split in between Game Thread ticks.

- Reenabled TickWarnings ticker.
- Optimized FollowPinArray.
- Optimized GetRootNode (removed bOutMultipleBaseObjectsFound since is no longer valid).
- GetParticipatingObjects now uses VisitNodes.
- GetParticipatingObjects now discovers all assets + is more granual.
- Implemented an async version of IsCompilationOutOfDate.
- Added GraphTraversal functions inside the namespace.
- Renamed GetRootObject to GetObject.
- Changed GetFullGraphRootNode to start from a Customizable Object instead of a node.

diff --git a/Engine/Plugins/Mutable/Source/CustomizableObject/Internal/MuCO/ICustomizableObjectEditorModule.h b/Engine/Plugins/Mutable/Source/CustomizableObject/Internal/MuCO/ICustomizableObjectEditorModule.h
--- a/Engine/Plugins/Mutable/Source/CustomizableObject/Internal/MuCO/ICustomizableObjectEditorModule.h
+++ b/Engine/Plugins/Mutable/Source/CustomizableObject/Internal/MuCO/ICustomizableObjectEditorModule.h
@@ -70,8 +70,14 @@ public:
 	  * @param RemovedPackages list of removed packages since the last compilation.
 	  * @param bReleaseVersionDiff true if the Release Version has changed since the last compilation.
 	  * @return true if the compilation is out of date. */
-	virtual bool IsCompilationOutOfDate(UCustomizableObject& Object, bool bSkipIndirectReferences, TArray<FName>& OutOfDatePackages, TArray<FName>& AddedPackages, TArray<FName>& RemovedPackages, bool& bReleaseVersionDiff) const = 0;
+	virtual bool IsCompilationOutOfDate(const UCustomizableObject& Object, bool bSkipIndirectReferences, TArray<FName>& OutOfDatePackages, TArray<FName>& AddedPackages, TArray<FName>& RemovedPackages, bool& bReleaseVersionDiff) const = 0;
 
+	using IsCompilationOutOfDateCallback = TFunction<void(bool bOutOfDate, bool bVersionDiff, const TArray<FName>& OutOfDatePackages, const TArray<FName>& AddedPackages, const TArray<FName>& RemovedPackages)>;
+
+	/** Async version. See IsCompilationOutOfDate.
+	 * @param MaxTime. Max time it can spend on each Game Thread tick. Use MAX_flt for sync call. */
+	virtual void IsCompilationOutOfDate(const UCustomizableObject& Object, bool bSkipIndirectReferences, float MaxTime, const IsCompilationOutOfDateCallback& Callback) const = 0;
+	
 	/** See GraphTraversal::IsRootObject(...) */
 	virtual bool IsRootObject(const UCustomizableObject& Object) const = 0;
 
@@ -110,7 +116,7 @@ public:
 	virtual USkeletalMesh* GetReferenceSkeletalMesh(const UCustomizableObject& Object, const FName& Component) const = 0;
 	
 	/** Perform a fast compilation pass to get all participating objects. */
-	virtual TMap<FName, FGuid> GetParticipatingObjects(UCustomizableObject* Object, const FCompilationOptions* Options = nullptr) const = 0;
+	virtual TMap<FName, FGuid> GetParticipatingObjects(const UCustomizableObject* Object, const FCompilationOptions* Options = nullptr) const = 0;
 
 	virtual void BackwardsCompatibleFixup(UEdGraph& Graph, int32 CustomizableObjectCustomVersion) = 0;
 	
diff --git a/Engine/Plugins/Mutable/Source/CustomizableObject/Private/MuCO/CustomizableObjectSystem.cpp b/Engine/Plugins/Mutable/Source/CustomizableObject/Private/MuCO/CustomizableObjectSystem.cpp
--- a/Engine/Plugins/Mutable/Source/CustomizableObject/Private/MuCO/CustomizableObjectSystem.cpp
+++ b/Engine/Plugins/Mutable/Source/CustomizableObject/Private/MuCO/CustomizableObjectSystem.cpp
@@ -248,13 +248,17 @@ void FMutableSystemSettingsOverrides::RestoreSettings() const
 static FAutoConsoleVariableSink CVarMutableSink(FConsoleCommandDelegate::CreateStatic(&CVarMutableSinkFunction));
 
 
-constexpr float OnScreenWarningsTickerTime = 1.0f;
+/** How often update the on screen warnings (seconds). */
+constexpr float OnScreenWarningsTickerTime = 5.0f;
 
+/** Duration of the on screen warning messages (seconds). */
+constexpr float WarningDisplayTime = OnScreenWarningsTickerTime * 2.0f;
 
-int64 GetOnScreenMessageKey(const UCustomizableObject& Object, TMap<TWeakObjectPtr<const UCustomizableObject>, uint64>& KeyMap)
+
+int64 GetOnScreenMessageKey(const TWeakObjectPtr<const UCustomizableObject>& Object, TMap<TWeakObjectPtr<const UCustomizableObject>, uint64>& KeyMap)
 {
 	uint64 Key;
-	if (uint64* Result = KeyMap.Find(&Object))
+	if (uint64* Result = KeyMap.Find(Object))
 	{
 		Key = *Result;
 	}
@@ -266,7 +270,7 @@ int64 GetOnScreenMessageKey(const UCustomizableObject& Object, TMap<TWeakObjectP
 			++Key;
 		}
 		
-		KeyMap.Add(TWeakObjectPtr(&Object), Key);
+		KeyMap.Add(Object, Key);
 	}
 
 	return Key;
@@ -328,11 +332,120 @@ FUpdateContextPrivate::FUpdateContextPrivate(UCustomizableObjectInstance& InInst
 }
 
 
+#if WITH_EDITOR
+struct FOutOfDateWarningContext
+{
+	TArray<TWeakObjectPtr<const UCustomizableObject>> Objects;
+
+	int32 IndexObject = 0;
+
+	double StartTime = 0.0f;
+};
+
+
+/** If true, the warning is being executed asynchronously. */
+bool bOutOfDateAsync  = false;
+
+
+/** Async because work is split in between ticks. */
+void OutOfDateWarning_Async(const TSharedRef<FOutOfDateWarningContext>& Context)
+{
+	MUTABLE_CPUPROFILER_SCOPE(OutOfDateWarning_Async)
+
+	check(IsInGameThread());
+	
+	static TMap<TWeakObjectPtr<const UCustomizableObject>, uint64> KeysOutOfDate;
+
+	constexpr float MaxTime = 1.0 / 1000.0 * 2.0; // 2ms
+
+	if (FPlatformTime::Seconds() - Context->StartTime >= MaxTime)
+	{
+		// Time limit reached. Reschedule itself.
+		GEditor->GetTimerManager()->SetTimerForNextTick([=]()
+		{
+			Context->StartTime = FPlatformTime::Seconds();
+			OutOfDateWarning_Async(Context);
+		});
+
+		return;
+	}
+
+	const ICustomizableObjectEditorModule* Module = ICustomizableObjectEditorModule::Get();	
+	if (!Module)
+	{
+		bOutOfDateAsync = false; // End async task.
+		return;	
+	}
+	
+	// Find the next Customizable Object still alive.
+	const UCustomizableObject* Object = nullptr;
+	while (Context->IndexObject < Context->Objects.Num())
+	{
+		Object = Context->Objects[Context->IndexObject].Get();
+
+		if (Object)
+		{
+			break;
+		}
+		
+		++Context->IndexObject;
+	}
+	
+	// If all Customizable Objects processed, end async task.
+	if (Context->IndexObject == Context->Objects.Num())
+	{
+		RemoveUnusedOnScreenMessages(KeysOutOfDate); // Clean old message keys.
+
+		bOutOfDateAsync = false; // End async task.
+		return;
+	}
+	
+	ICustomizableObjectEditorModule::IsCompilationOutOfDateCallback Callback = [Context](bool bOutOfDate, bool bVersionDiff, const TArray<FName>& OutOfDatePackages,
+		const TArray<FName>& AddedPackages, const TArray<FName>& RemovedPackages)
+	{
+		check(IsInGameThread());
+
+		const TWeakObjectPtr<const UCustomizableObject>& WeakObject = Context->Objects[Context->IndexObject];
+
+		if (const UCustomizableObject* Object = WeakObject.Get())
+		{
+			if (bOutOfDate)
+			{
+				const uint64 Key = GetOnScreenMessageKey(WeakObject, KeysOutOfDate);
+			
+				if (!GEngine->OnScreenDebugMessageExists(Key))
+				{
+					UE_LOG(LogMutable, Display, TEXT("Customizable Object [%s] compilation out of date. Changes since last compilation:"), *Object->GetName());
+					PrintParticipatingPackagesDiff(OutOfDatePackages, AddedPackages, RemovedPackages, bVersionDiff);
+				}
+			
+				FString Msg = FString::Printf(TEXT("Customizable Object [%s] compilation out of date. See the Output Log for more information."), *Object->GetName());
+				GEngine->AddOnScreenDebugMessage(Key, WarningDisplayTime, FColor::Yellow, Msg);
+			}
+			else if (uint64* Key = KeysOutOfDate.Find(WeakObject))
+			{
+				GEngine->RemoveOnScreenDebugMessage(*Key);
+			}
+		}
+
+		// Process the next Customizable Object.
+		++Context->IndexObject;
+		OutOfDateWarning_Async(Context);
+	};
+
+	Module->IsCompilationOutOfDate(*Object, true, MaxTime, Callback);
+}
+#endif
+
+
 bool TickWarnings(float DeltaTime)
 {
+	MUTABLE_CPUPROFILER_SCOPE(TickWarnings);
+
+	const double StartTime = FPlatformTime::Seconds();
+	
 	static TMap<TWeakObjectPtr<const UCustomizableObject>, uint64> KeysNotCompiled;
 	static TMap<TWeakObjectPtr<const UCustomizableObject>, uint64> KeysNotOptimized;
-	static TMap<TWeakObjectPtr<const UCustomizableObject>, uint64> KeysOutOfDate;
 	
 	TSet<const UCustomizableObject*> Objects;
 	
@@ -372,78 +485,76 @@ bool TickWarnings(float DeltaTime)
 		{
 			continue;
 		}
-
-		Objects.Add(Object);
-	}
-
-	for (const UCustomizableObject* Object : Objects)
-	{
-		constexpr float DisplayTime = OnScreenWarningsTickerTime * 2.0f;
 		
-		if (!Object->IsLoading() && !Object->IsCompiled())
-		{
-			const uint64 Key = GetOnScreenMessageKey(*Object, KeysNotCompiled);
-			FString Msg = FString::Printf(TEXT("Customizable Object [%s] not compiled."), *Object->GetName());
-			GEngine->AddOnScreenDebugMessage(Key, DisplayTime, FColor::Red, Msg);
-		}
-		else if (uint64* Key = KeysNotCompiled.Find(TWeakObjectPtr(Object)))
-		{
-			GEngine->RemoveOnScreenDebugMessage(*Key);
-		}
-
-		const UModelResources* ModelResources = Object->GetPrivate()->GetModelResources();
-		if (!ModelResources)
+		if (Object->GetPrivate()->Status.Get() != FCustomizableObjectStatus::EState::ModelLoaded)
 		{
 			continue;
 		}
 
-		if (Object->GetPrivate()->Status.Get() != FCustomizableObjectStatus::EState::ModelLoaded)
+		Objects.Add(Object);
+	}
+
+	// Not compiled warning.
+	{
+		for (const UCustomizableObject* Object : Objects)
 		{
-			continue;
+			TWeakObjectPtr<const UCustomizableObject> WeakObject = TWeakObjectPtr(Object);
+		
+			if (!Object->IsLoading() && !Object->IsCompiled())
+			{
+				const uint64 Key = GetOnScreenMessageKey(WeakObject, KeysNotCompiled);
+				FString Msg = FString::Printf(TEXT("Customizable Object [%s] not compiled."), *Object->GetName());
+				GEngine->AddOnScreenDebugMessage(Key, WarningDisplayTime, FColor::Red, Msg);
+			}
+			else if (uint64* Key = KeysNotCompiled.Find(TWeakObjectPtr(Object)))
+			{
+				GEngine->RemoveOnScreenDebugMessage(*Key);
+			}
 		}
+		
+		RemoveUnusedOnScreenMessages(KeysNotCompiled);
+	}
 
-#if WITH_EDITORONLY_DATA
-		if (!ModelResources->bIsCompiledWithOptimization)
+	// Compiled without optimizations warning.
+#if WITH_EDITOR
+	for (const UCustomizableObject* Object : Objects)
+	{
+		TWeakObjectPtr<const UCustomizableObject> WeakObject = TWeakObjectPtr(Object);
+		
+		if (!Object->GetPrivate()->GetModelResourcesChecked().bIsCompiledWithOptimization)
 		{
-			const uint64 Key = GetOnScreenMessageKey(*Object, KeysNotOptimized);
+			const uint64 Key = GetOnScreenMessageKey(WeakObject, KeysNotOptimized);
 			FString Msg = FString::Printf(TEXT("Customizable Object [%s] was compiled without optimization."), *Object->GetName());
-			GEngine->AddOnScreenDebugMessage(Key, DisplayTime, FColor::Yellow, Msg);
+			GEngine->AddOnScreenDebugMessage(Key, WarningDisplayTime, FColor::Yellow, Msg);
 		}
 		else if (uint64* Key = KeysNotOptimized.Find(TWeakObjectPtr(Object)))
 		{
 			GEngine->RemoveOnScreenDebugMessage(*Key);
 		}
+	}
+
+	RemoveUnusedOnScreenMessages(KeysNotOptimized);
 #endif
+	
 
+	// Is compilation out of date warning.
 #if WITH_EDITOR
-		TArray<FName> OutOfDatePackages;
-		TArray<FName> AddedPackages;
-		TArray<FName> RemovedPackages;
-		bool bReleaseVersion;
-		if (Object->GetPrivate()->IsCompilationOutOfDate(true, OutOfDatePackages, AddedPackages, RemovedPackages, bReleaseVersion))
-		{
-			const uint64 Key = GetOnScreenMessageKey(*Object, KeysOutOfDate);
-			
-			if (!GEngine->OnScreenDebugMessageExists(Key))
-			{
-				UE_LOG(LogMutable, Display, TEXT("Customizable Object [%s] compilation out of date. Changes since last compilation:"), *Object->GetName());
-				PrintParticipatingPackagesDiff(OutOfDatePackages, AddedPackages, RemovedPackages, bReleaseVersion);
-			}
-			
-			FString Msg = FString::Printf(TEXT("Customizable Object [%s] compilation out of date. See the Output Log for more information."), *Object->GetName());
-			GEngine->AddOnScreenDebugMessage(Key, DisplayTime, FColor::Yellow, Msg);
-		}
-		else if (uint64* Key = KeysOutOfDate.Find(TWeakObjectPtr(Object)))
+	if (!bOutOfDateAsync)
+	{
+		bOutOfDateAsync = true;
+
+		TSharedRef<FOutOfDateWarningContext> Context = MakeShareable(new FOutOfDateWarningContext);
+		Context->StartTime = StartTime;
+		
+		for (const UCustomizableObject* Object : Objects)
 		{
-			GEngine->RemoveOnScreenDebugMessage(*Key);
+			Context->Objects.Add(Object);
 		}
-#endif
+		
+		OutOfDateWarning_Async(Context);
 	}
-
-	RemoveUnusedOnScreenMessages(KeysNotCompiled);
-	RemoveUnusedOnScreenMessages(KeysNotOptimized);
-	RemoveUnusedOnScreenMessages(KeysOutOfDate);
-
+#endif
+	
 	return true;
 }
 
@@ -4650,7 +4761,7 @@ void UCustomizableObjectSystemPrivate::OnMutableEnabledChanged(IConsoleVariable*
 
 		if (!SystemPrivate->TickWarningsDelegateHandle.IsValid())
 		{
-			// SystemPrivate->TickWarningsDelegateHandle = FTSTicker::GetCoreTicker().AddTicker(FTickerDelegate::CreateStatic(&TickWarnings), OnScreenWarningsTickerTime); // TODO GMT Temporally disabled due to performance
+			SystemPrivate->TickWarningsDelegateHandle = FTSTicker::GetCoreTicker().AddTicker(FTickerDelegate::CreateStatic(&TickWarnings), OnScreenWarningsTickerTime);
 		}
 #endif // !UE_SERVER
 	}
diff --git a/Engine/Plugins/Mutable/Source/CustomizableObjectEditor/Private/MuCOE/CustomizableObjectCompiler.cpp b/Engine/Plugins/Mutable/Source/CustomizableObjectEditor/Private/MuCOE/CustomizableObjectCompiler.cpp
--- a/Engine/Plugins/Mutable/Source/CustomizableObjectEditor/Private/MuCOE/CustomizableObjectCompiler.cpp
+++ b/Engine/Plugins/Mutable/Source/CustomizableObjectEditor/Private/MuCOE/CustomizableObjectCompiler.cpp
@@ -456,8 +456,6 @@ void ProcessChildObjectsRecursively(const UCustomizableObject* ParentObject, FMu
 		return;
 	}
 	
-	bool bMultipleBaseObjectsFound = false;
-	
 	TArray<FAssetData> AssetDataArray;
 
 	FARFilter Filter;
@@ -481,7 +479,7 @@ void ProcessChildObjectsRecursively(const UCustomizableObject* ParentObject, FMu
 			continue;
 		}
 
-		UCustomizableObjectNodeObject* Root = GetRootNode(ChildObject, bMultipleBaseObjectsFound);
+		UCustomizableObjectNodeObject* Root = GetRootNode(ChildObject);
 		if (!Root)
 		{
 			continue;
@@ -511,32 +509,29 @@ void ProcessChildObjectsRecursively(const UCustomizableObject* ParentObject, FMu
 			}
 		}
 		
-		if (!bMultipleBaseObjectsFound)
+		if (const FGroupNodeIdsTempData* GroupGuid = GenerationContext.DuplicatedGroupNodeIds.FindPair(ParentObject, FGroupNodeIdsTempData(Root->ParentObjectGroupId)))
 		{
-			if (const FGroupNodeIdsTempData* GroupGuid = GenerationContext.DuplicatedGroupNodeIds.FindPair(ParentObject, FGroupNodeIdsTempData(Root->ParentObjectGroupId)))
-			{
-				Root->ParentObjectGroupId = GroupGuid->NewGroupNodeId;
-			}
+			Root->ParentObjectGroupId = GroupGuid->NewGroupNodeId;
+		}
 
-			GenerationContext.GroupIdToExternalNodeMap.Add(Root->ParentObjectGroupId, Root);
+		GenerationContext.GroupIdToExternalNodeMap.Add(Root->ParentObjectGroupId, Root);
 
-			TArray<UCustomizableObjectNodeObjectGroup*> GroupNodes;
-			ChildObject->GetPrivate()->GetSource()->GetNodesOfClass<UCustomizableObjectNodeObjectGroup>(GroupNodes);
+		TArray<UCustomizableObjectNodeObjectGroup*> GroupNodes;
+		ChildObject->GetPrivate()->GetSource()->GetNodesOfClass<UCustomizableObjectNodeObjectGroup>(GroupNodes);
 
-			if (GroupNodes.Num() > 0) // Only grafs with group nodes should have child grafs
+		if (GroupNodes.Num() > 0) // Only grafs with group nodes should have child grafs
+		{
+			for (int32 i = 0; i < GroupNodes.Num(); ++i)
 			{
-				for (int32 i = 0; i < GroupNodes.Num(); ++i)
+				const FGuid NodeId = GenerationContext.GetNodeIdUnique(GroupNodes[i]);
+				if (NodeId != GroupNodes[i]->NodeGuid)
 				{
-					const FGuid NodeId = GenerationContext.GetNodeIdUnique(GroupNodes[i]);
-					if (NodeId != GroupNodes[i]->NodeGuid)
-					{
-						GenerationContext.DuplicatedGroupNodeIds.Add(ChildObject, FGroupNodeIdsTempData(GroupNodes[i]->NodeGuid, NodeId));
-						GroupNodes[i]->NodeGuid = NodeId;
-					}
+					GenerationContext.DuplicatedGroupNodeIds.Add(ChildObject, FGroupNodeIdsTempData(GroupNodes[i]->NodeGuid, NodeId));
+					GroupNodes[i]->NodeGuid = NodeId;
 				}
-
-				ProcessChildObjectsRecursively(ChildObject, GenerationContext);
 			}
+
+			ProcessChildObjectsRecursively(ChildObject, GenerationContext);
 		}
 	}
 }
@@ -562,15 +557,7 @@ mu::Ptr<mu::NodeObject> GenerateMutableRoot(
 		return nullptr;
 	}
 
-	bool bMultipleBaseObjectsFound;
-	UCustomizableObjectNodeObject* LocalRootNodeObject = GetRootNode(Object, bMultipleBaseObjectsFound);
-
-	if (bMultipleBaseObjectsFound)
-	{
-		GenerationContext.Log(LOCTEXT("MultipleBaseRoot","Multiple base object nodes found."));
-		return nullptr;
-	}
-
+	UCustomizableObjectNodeObject* LocalRootNodeObject = GetRootNode(Object);
 	if (!LocalRootNodeObject)
 	{
 		GenerationContext.Log(LOCTEXT("NoRootBase","No base object node found. Object not built."));
@@ -582,14 +569,8 @@ mu::Ptr<mu::NodeObject> GenerateMutableRoot(
 
 	GenerationContext.RootVersionBridge = RootObject->VersionBridge;
 
-	UCustomizableObjectNodeObject* RootNodeObject = GetRootNode(RootObject, bMultipleBaseObjectsFound);
+	UCustomizableObjectNodeObject* RootNodeObject = GetRootNode(RootObject);
 	GenerationContext.Root = RootNodeObject;
-	
-	if (bMultipleBaseObjectsFound)
-	{
-		GenerationContext.Log(LOCTEXT("MultipleBaseActualRoot", "Multiple base object nodes found."));
-		return nullptr;
-	}
 
 	if (!RootNodeObject)
 	{
@@ -676,13 +657,14 @@ mu::Ptr<mu::NodeObject> GenerateMutableRoot(
 	// First pass. Only used to recollect info required for the primary pass.
 	// Notice that the traversal is different form the primary pass. Here we follow all pins indiscriminately,
 	// while the primary pass follows the Mutable Source structure (which may cut branches).
-	GraphTraversal::VisitNodes(*RootNodeObject, GenerationContext.GroupIdToExternalNodeMap, [&GenerationContext](UCustomizableObjectNode& Node)
+	GraphTraversal::VisitNodes(*RootNodeObject, [&GenerationContext](UCustomizableObjectNode& Node)
 	{
 		if (UCustomizableObjectNodeComponentMesh* NodeComponentMesh = Cast<UCustomizableObjectNodeComponentMesh>(&Node))
 		{
 			FirstPass(*NodeComponentMesh, GenerationContext);
 		}
-	});
+	},
+	&GenerationContext.GroupIdToExternalNodeMap);
 	
     GenerationContext.CompilationContext.RealTimeMorphTargetsOverrides = RootNodeObject->RealTimeMorphSelectionOverrides;
 
diff --git a/Engine/Plugins/Mutable/Source/CustomizableObjectEditor/Private/MuCOE/CustomizableObjectEditorModule.cpp b/Engine/Plugins/Mutable/Source/CustomizableObjectEditor/Private/MuCOE/CustomizableObjectEditorModule.cpp
--- a/Engine/Plugins/Mutable/Source/CustomizableObjectEditor/Private/MuCOE/CustomizableObjectEditorModule.cpp
+++ b/Engine/Plugins/Mutable/Source/CustomizableObjectEditor/Private/MuCOE/CustomizableObjectEditorModule.cpp
@@ -2,6 +2,7 @@
 
 #include "MuCOE/CustomizableObjectEditorModule.h"
 
+#include "Animation/AnimBlueprint.h"
 #include "AssetToolsModule.h"
 #include "GameFramework/Pawn.h"
 #include "ISettingsModule.h"
@@ -97,6 +98,7 @@
 #include "Nodes/CustomizableObjectNodeComponentMeshDetails.h"
 #include "Nodes/CustomizableObjectNodeModifierTransformInMeshDetails.h"
 #include "MuCOE/GenerateMutableSource/GenerateMutableSourceTable.h"
+#include "PhysicsEngine/PhysicsAsset.h"
 
 
 class AActor;
@@ -512,26 +514,59 @@ void GetReferencingPackages(const UCustomizableObject& Object, TArray<FAssetData
 }
 
 
-bool FCustomizableObjectEditorModule::IsCompilationOutOfDate(UCustomizableObject& Object, bool bSkipIndirectReferences, TArray<FName>& OutOfDatePackages, TArray<FName>& AddedPackages, TArray<FName>& RemovedPackages, bool& bReleaseVersion) const
+bool FCustomizableObjectEditorModule::IsCompilationOutOfDate(const UCustomizableObject& Object, bool bSkipIndirectReferences, TArray<FName>& OutOfDatePackages, TArray<FName>& OutAddedPackages, TArray<FName>& OutRemovedPackages, bool& bOutReleaseVersion) const
 {
-	MUTABLE_CPUPROFILER_SCOPE(FCustomizableObjectEditorModule::IsCompilationOutOfDate)
+	bool Result;
 	
-	// TODO CO Custom version
-	// TODO List of plugins and their custom versions
-	// Maybe use BuildDerivedDataKey? BuildDerivedDataKey should also consider bSkipIndirectReferences
+	auto Callback = [&](bool bOutOfDate, bool bVersionDiff, TArray<FName> OfDatePackages, TArray<FName> AddedPackages, TArray<FName> RemovedPackages)
+	{
+		Result = bOutOfDate;
+		OutOfDatePackages = OfDatePackages;
+		OutAddedPackages = AddedPackages;
+		OutRemovedPackages = RemovedPackages;
+		bOutReleaseVersion = bVersionDiff;
+	};
+
+	IsCompilationOutOfDate(Object, bSkipIndirectReferences, TNumericLimits<float>::Max(), Callback);
 	
-	const FAssetRegistryModule& AssetRegistryModule = FModuleManager::LoadModuleChecked<FAssetRegistryModule>("AssetRegistry");
+	return Result;
+}
 
-	const UModelResources* ModelResources = Object.GetPrivate()->GetModelResources();
-	if (!ModelResources)
-	{
-		return true;
-	}
 
-	const TMap<FName, FGuid>& OldParticipatingObjects = ModelResources->ParticipatingObjects;
+struct FCompilationOutOfDateContext
+{
+	TArray<TPair<FName, FGuid>> ParticipatingObjects;
+
+	int32 IndexParticipatingObject = 0;
+
+	float MaxTime = 0;
+	
+	TArray<FName> OutOfDatePackages;
+	TArray<FName> AddedPackages;
+	TArray<FName> RemovedPackages;
+
+	bool bVersionDiff = false;
 	
-	for (const TTuple<FName, FGuid>& ParticipatingObject : OldParticipatingObjects)
+	ICustomizableObjectEditorModule::IsCompilationOutOfDateCallback Callback;
+};
+
+
+/** Async because work is split in between ticks. */
+void IsCompilationOutOfDate_Async(const TSharedRef<FCompilationOutOfDateContext>& Context)
+{
+	MUTABLE_CPUPROFILER_SCOPE(IsCompilationOutOfDate_Async)
+
+	check(IsInGameThread());
+	
+	IAssetRegistry& AssetRegistry = FModuleManager::LoadModuleChecked<FAssetRegistryModule>("AssetRegistry").Get();
+	
+	const double StartTime = FPlatformTime::Seconds();
+	
+	while (FPlatformTime::Seconds() - StartTime < Context->MaxTime &&
+		Context->IndexParticipatingObject < Context->ParticipatingObjects.Num())
 	{
+		const TPair<FName, FGuid>& ParticipatingObject = Context->ParticipatingObjects[Context->IndexParticipatingObject];
+			
 		TSoftObjectPtr SoftObjectPtr(FSoftObjectPath(ParticipatingObject.Key.ToString()));
 		if (SoftObjectPtr) // If loaded
 		{
@@ -542,17 +577,17 @@ bool FCustomizableObjectEditorModule::IsCompilationOutOfDate(UCustomizableObject
 			
 			if (PackageGuid != ParticipatingObject.Value)
 			{
-				OutOfDatePackages.AddUnique(ParticipatingObject.Key);
+				Context->OutOfDatePackages.Add(ParticipatingObject.Key);
 			}
 		}
 		else // Not loaded
 		{
 			FAssetPackageData AssetPackageData;
-			const UE::AssetRegistry::EExists Result = AssetRegistryModule.Get().TryGetAssetPackageData(ParticipatingObject.Key, AssetPackageData);
+			const UE::AssetRegistry::EExists Result = AssetRegistry.TryGetAssetPackageData(ParticipatingObject.Key, AssetPackageData);
 				
 			if (Result != UE::AssetRegistry::EExists::Exists)
 			{
-				OutOfDatePackages.AddUnique(ParticipatingObject.Key);
+				Context->RemovedPackages.Add(ParticipatingObject.Key);
 			}
 
 			PRAGMA_DISABLE_DEPRECATION_WARNINGS
@@ -561,49 +596,98 @@ bool FCustomizableObjectEditorModule::IsCompilationOutOfDate(UCustomizableObject
 			
 			if (PackageGuid != ParticipatingObject.Value)
 			{
-				OutOfDatePackages.AddUnique(ParticipatingObject.Key);
+				Context->OutOfDatePackages.Add(ParticipatingObject.Key);
 			}
 		}
+
+		++Context->IndexParticipatingObject;
+	}
+
+	if (Context->IndexParticipatingObject == Context->ParticipatingObjects.Num())
+	{
+		const bool bOutOfDate = Context->bVersionDiff || !Context->OutOfDatePackages.IsEmpty() || !Context->AddedPackages.IsEmpty() || !Context->RemovedPackages.IsEmpty();
+		Context->Callback(bOutOfDate, Context->bVersionDiff, Context->OutOfDatePackages, Context->AddedPackages, Context->RemovedPackages);
+	}
+	else
+	{
+		GEditor->GetTimerManager()->SetTimerForNextTick([=]()
+		{
+			IsCompilationOutOfDate_Async(Context);
+		});
+	}
+}
+
+
+void FCustomizableObjectEditorModule::IsCompilationOutOfDate(const UCustomizableObject& Object, bool bSkipIndirectReferences, float MaxTime, const IsCompilationOutOfDateCallback& Callback) const
+{
+	MUTABLE_CPUPROFILER_SCOPE(FCustomizableObjectEditorModule::IsCompilationOutOfDate)
+
+	check(IsInGameThread());
+	
+	// TODO CO Custom version
+	// TODO List of plugins and their custom versions
+	// Maybe use BuildDerivedDataKey? BuildDerivedDataKey should also consider bSkipIndirectReferences
+	
+	const UModelResources* ModelResources = Object.GetPrivate()->GetModelResources();
+	if (!ModelResources)
+	{
+		Callback(true, false, {}, {}, {});
+		return;
 	}
+
+	TSharedRef<FCompilationOutOfDateContext> Context = MakeShared<FCompilationOutOfDateContext>();
+	Context->ParticipatingObjects = ModelResources->ParticipatingObjects.Array();
+	Context->MaxTime = MaxTime;
+	Context->Callback = Callback;
 	
+	Context->bVersionDiff = false;
+	if (ICustomizableObjectVersionBridgeInterface* VersionBridge = Cast<ICustomizableObjectVersionBridgeInterface>(Object.VersionBridge))
+	{
+		Context->bVersionDiff = ModelResources->ReleaseVersion != VersionBridge->GetCurrentVersionAsString();
+	}
+
+	check(MaxTime == TNumericLimits<float>::Max()  || bSkipIndirectReferences); // If async, bSkipIndirectReferences must be false since it is very expensive and can not be split in subtasks.
+
 	// Check that we have the exact same set of participating object as before. This can change due to indirect references and versioning.
 	if (!bSkipIndirectReferences)
 	{
+		IAssetRegistry& AssetRegistry = FModuleManager::LoadModuleChecked<FAssetRegistryModule>("AssetRegistry").Get();
+
+		const TMap<FName, FGuid>& OldParticipatingObjects = ModelResources->ParticipatingObjects;
+
+		const int32 Num = OldParticipatingObjects.Num();
+
+		Context->AddedPackages.Reserve(Num);
+		
 		// Due to performance issues, we will skip loading all objects. We can do that since loading/not loading objects do not affect the number of indirect objects discovered
 		// (e.g., we will traverse the same number of COs/Tables regardless if we do not load meshes/textures...). 
 		TMap<FName, FGuid> ParticipatingObjects = GetParticipatingObjects(&Object);
 		
 		for (const TTuple<FName, FGuid>& ParticipatingObject : ParticipatingObjects)
 		{
-			// Since here we are if the smaller set (objects found now without loading all objects) is contained in the larger set (objects found in the compilation pass),
+			// Since here we are checking if the smaller set (objects found now without loading all objects) is contained in the larger set (objects found in the compilation pass),
 			// there is no need to check if the asset is a indirect reference (CO or Table).
 			if (!OldParticipatingObjects.Contains(ParticipatingObject.Key))
 			{
-				AddedPackages.AddUnique(ParticipatingObject.Key);
+				Context->AddedPackages.AddUnique(ParticipatingObject.Key);
 			}
 		}
 
 		for (const TTuple<FName, FGuid>& OldParticipatingObject : OldParticipatingObjects)
 		{
-			FAssetData AssetData = AssetRegistryModule.Get().GetAssetByObjectPath(FSoftObjectPath(OldParticipatingObject.Key.ToString()));
+			FAssetData AssetData = AssetRegistry.GetAssetByObjectPath(FSoftObjectPath(OldParticipatingObject.Key.ToString()));
 			if (AssetData.AssetClassPath == UCustomizableObject::StaticClass()->GetClassPathName() ||
 				AssetData.AssetClassPath == UDataTable::StaticClass()->GetClassPathName())
 			{
 				if (!ParticipatingObjects.Contains(OldParticipatingObject.Key))
 				{
-					RemovedPackages.AddUnique(OldParticipatingObject.Key);
+					Context->RemovedPackages.AddUnique(OldParticipatingObject.Key);
 				}
 			}
 		}
 	}
 
-	bReleaseVersion = false;
-	if (ICustomizableObjectVersionBridgeInterface* VersionBridge = Cast<ICustomizableObjectVersionBridgeInterface>(Object.VersionBridge))
-	{
-		bReleaseVersion = ModelResources->ReleaseVersion != VersionBridge->GetCurrentVersionAsString();
-	}
-
-	return bReleaseVersion || !OutOfDatePackages.IsEmpty() || !AddedPackages.IsEmpty() || !RemovedPackages.IsEmpty();
+	IsCompilationOutOfDate_Async(Context);
 }
 
 
@@ -691,211 +775,209 @@ USkeletalMesh* FCustomizableObjectEditorModule::GetReferenceSkeletalMesh(const U
 }
 
 
-TMap<FName, FGuid> FCustomizableObjectEditorModule::GetParticipatingObjects(UCustomizableObject* Object, const FCompilationOptions* InOptions) const
+TMap<FName, FGuid> FCustomizableObjectEditorModule::GetParticipatingObjects(const UCustomizableObject* Object, const FCompilationOptions* InOptions) const
 {
-	MUTABLE_CPUPROFILER_SCOPE(FCustomizableObjectEditorModule::GetParticipatingObjectsNew);
+	MUTABLE_CPUPROFILER_SCOPE(FCustomizableObjectEditorModule::GetParticipatingObjects);
 
 	FCompilationOptions Options = InOptions ? *InOptions : Object->GetPrivate()->GetCompileOptions();
 	FMutableCompilationContext CompilationContext(Object, nullptr, Options);
 	FMutableGraphGenerationContext Context(CompilationContext);
 
-
-	FAssetRegistryModule& AssetRegistryModule = FModuleManager::LoadModuleChecked<FAssetRegistryModule>("AssetRegistry");
-
+	IAssetRegistry& AssetRegistry = FModuleManager::LoadModuleChecked<FAssetRegistryModule>("AssetRegistry").Get();
+	
 	// Store the list of participating assets here
 	TMap<FName, FGuid> Result;
+		
+	auto VisitDependencies = [&AssetRegistry](const FSoftObjectPath& RootAsset, const TArray<UClass*>& FilterClasses, bool bRecursively, const TFunction<void(const FSoftObjectPath&)>& VisitFunc)
+	{
+		TArray<FSoftObjectPath> References;
+		
+		TArray<FName> PackageDependencies;
+		AssetRegistry.GetDependencies(RootAsset.GetLongPackageFName(), PackageDependencies);
 
-	// Use these 2 helpers to add assets to the list.
-	auto AddObject = [&Result]
-	(UObject* Candidate)
+		FARFilter Filter;
+		Filter.bRecursivePaths = bRecursively;
+		
+		for (const FName& PackageName : PackageDependencies)
 		{
-			if (Candidate)
+			if (!PackageName.ToString().StartsWith(TEXT("/TempAutosave")))
 			{
-				// TODO UE-226453
-				PRAGMA_DISABLE_DEPRECATION_WARNINGS
-				const FGuid PackageGuid = Candidate->GetPackage()->GetGuid();
-				PRAGMA_ENABLE_DEPRECATION_WARNINGS
-
-				FName PackageName = Candidate->GetPackage()->GetFName();
-				Result.Emplace(PackageName, PackageGuid);
+				Filter.PackageNames.Add(PackageName);
 			}
-		};
+		}
+
+		TArray<FAssetData> Assets;
+		AssetRegistry.GetAssets(Filter, Assets);
 
-	auto AddObjectSoft = [&Result, &AssetRegistryModule]
-	(const FSoftObjectPath& SoftPath)
+		for (const FAssetData& Asset : Assets)
 		{
-			FAssetPackageData AssetPackageData;
-			const UE::AssetRegistry::EExists Query = AssetRegistryModule.Get().TryGetAssetPackageData(SoftPath.GetLongPackageFName(), AssetPackageData);
-			if (Query != UE::AssetRegistry::EExists::Exists)
+			if (!FilterClasses.Contains(Asset.GetClass(EResolveClass::Yes)))
 			{
-				return;
+				VisitFunc(Asset.GetSoftObjectPath());
 			}
-
+		}
+	};
+	
+	// Use these 2 helpers to add assets to the list.
+	auto AddObject = [&Result](UObject* Candidate)
+	{
+		if (Candidate)
+		{
+			// TODO UE-226453
 			PRAGMA_DISABLE_DEPRECATION_WARNINGS
-			const FGuid PackageGuid = AssetPackageData.PackageGuid;
+			const FGuid PackageGuid = Candidate->GetPackage()->GetGuid();
 			PRAGMA_ENABLE_DEPRECATION_WARNINGS
 
-			Result.Emplace(SoftPath.GetLongPackageFName(), PackageGuid);
-		};
-
-	// Gather all related COs
-	TSet<UCustomizableObject*> AllObjects;
-	GetAllObjectsInGraph(Object, AllObjects);
+			FName PackageName = Candidate->GetPackage()->GetFName();
+			Result.Emplace(PackageName, PackageGuid);
+		}
+	};
 
-	// Scan all nodes in each CO
-	for (UCustomizableObject* CO : AllObjects)
+	auto AddObjectSoft = [&Result, &AssetRegistry](const FSoftObjectPath& SoftPath)
 	{
-		AddObject(CO);
+		FAssetPackageData AssetPackageData;
+		const UE::AssetRegistry::EExists Query = AssetRegistry.TryGetAssetPackageData(SoftPath.GetLongPackageFName(), AssetPackageData);
+		if (Query != UE::AssetRegistry::EExists::Exists)
+		{
+			return;
+		}
 
-		UEdGraph* Graph = CO->GetPrivate()->GetSource();
+		PRAGMA_DISABLE_DEPRECATION_WARNINGS
+		const FGuid PackageGuid = AssetPackageData.PackageGuid;
+		PRAGMA_ENABLE_DEPRECATION_WARNINGS
 
-		// For each node, based on its type.
-		// This could also be implemented as a virtual method on the nodes, but only a few of them are relevant.
-		for ( UEdGraphNode* GenericNode : Graph->Nodes )
+		Result.Emplace(SoftPath.GetLongPackageFName(), PackageGuid);
+	};
+	
+	auto Visit = [&](UCustomizableObjectNode& Node)
+	{
+		if (UCustomizableObjectNodeObject* NodeObject = Cast<UCustomizableObjectNodeObject>(&Node))
+		{
+			AddObject(GraphTraversal::GetObject(*NodeObject));
+		}
+		
+		if (UCustomizableObjectNodeSkeletalMesh* SkeletalNode = Cast<UCustomizableObjectNodeSkeletalMesh>(&Node))
 		{
-			if (UCustomizableObjectNodeSkeletalMesh* SkeletalNode = Cast<UCustomizableObjectNodeSkeletalMesh>(GenericNode))
 			{
-				AddObjectSoft(SkeletalNode->SkeletalMesh.ToSoftObjectPath());
+				const FSoftObjectPath SoftObjectPath = SkeletalNode->AnimInstance.ToSoftObjectPath();
+				AddObjectSoft(SoftObjectPath);
+				VisitDependencies(SoftObjectPath, { UAnimBlueprint::StaticClass() }, false, AddObjectSoft);
 			}
 
-			else if (UCustomizableObjectNodeStaticMesh* StaticMeshNode = Cast<UCustomizableObjectNodeStaticMesh>(GenericNode))
 			{
-				AddObjectSoft(StaticMeshNode->StaticMesh.ToSoftObjectPath());
+				const FSoftObjectPath SoftObjectPath = SkeletalNode->SkeletalMesh.ToSoftObjectPath();
+				AddObjectSoft(SoftObjectPath);
+				VisitDependencies(SoftObjectPath, { UPhysicsAsset::StaticClass(), USkeleton::StaticClass() }, false, AddObjectSoft);
+				VisitDependencies(SoftObjectPath, { UTexture::StaticClass() }, true, AddObjectSoft);
 			}
+		}
 
-			else if (UCustomizableObjectNodeAnimationPose* PoseNode = Cast<UCustomizableObjectNodeAnimationPose>(GenericNode))
-			{
-				AddObjectSoft(PoseNode->PoseAsset);
-			}
+		else if (UCustomizableObjectNodeStaticMesh* StaticMeshNode = Cast<UCustomizableObjectNodeStaticMesh>(&Node))
+		{
+			const FSoftObjectPath SoftObjectPath = StaticMeshNode->StaticMesh.ToSoftObjectPath();
 
-			else if (UCustomizableObjectNodeCurve* CurveNode = Cast<UCustomizableObjectNodeCurve>(GenericNode))
-			{
-				AddObjectSoft(CurveNode->CurveAsset);
-			}
+			AddObjectSoft(SoftObjectPath);
+		}
 
-			else if (UCustomizableObjectNodeComponentMesh* CompNode = Cast<UCustomizableObjectNodeComponentMesh>(GenericNode))
-			{
-				AddObjectSoft(CompNode->ReferenceSkeletalMesh);
-			}
+		else if (UCustomizableObjectNodeAnimationPose* PoseNode = Cast<UCustomizableObjectNodeAnimationPose>(&Node))
+		{
+			AddObjectSoft(PoseNode->PoseAsset);
+		}
 
-			else if (UCustomizableObjectNodeTexture* TexNode = Cast<UCustomizableObjectNodeTexture>(GenericNode))
-			{
-				AddObjectSoft(TexNode->Texture);
-			}
+		else if (UCustomizableObjectNodeCurve* CurveNode = Cast<UCustomizableObjectNodeCurve>(&Node))
+		{
+			AddObjectSoft(CurveNode->CurveAsset);
+		}
+
+		else if (UCustomizableObjectNodeComponentMesh* CompNode = Cast<UCustomizableObjectNodeComponentMesh>(&Node))
+		{
+			AddObjectSoft(CompNode->ReferenceSkeletalMesh);
+		}
 
-			else if (UCustomizableObjectNodeTextureProject* ProjectNode = Cast<UCustomizableObjectNodeTextureProject>(GenericNode))
+		else if (UCustomizableObjectNodeTexture* TexNode = Cast<UCustomizableObjectNodeTexture>(&Node))
+		{
+			AddObjectSoft(TexNode->Texture);
+		}
+
+		else if (UCustomizableObjectNodeTextureProject* ProjectNode = Cast<UCustomizableObjectNodeTextureProject>(&Node))
+		{
+			AddObjectSoft(ProjectNode->ReferenceTexture);
+		}
+
+		else if (UCustomizableObjectNodeTextureTransform* TransformNode = Cast<UCustomizableObjectNodeTextureTransform>(&Node))
+		{
+			AddObjectSoft(TransformNode->ReferenceTexture);
+		}
+
+		else if (UCustomizableObjectNodeMaterial* MatNode = Cast<UCustomizableObjectNodeMaterial>(&Node))
+		{
+			AddObjectSoft(MatNode->GetMaterial());
+			int32 NumImages = MatNode->GetNumParameters(EMaterialParameterType::Texture);
+			for (int32 ImageIndex = 0; ImageIndex < NumImages; ++ImageIndex)
 			{
-				AddObjectSoft(ProjectNode->ReferenceTexture);
+				AddObject(MatNode->GetImageReferenceTexture(ImageIndex));
 			}
+		}
 
-			else if (UCustomizableObjectNodeTextureTransform* TransformNode = Cast<UCustomizableObjectNodeTextureTransform>(GenericNode))
+		else if (UCustomizableObjectNodeTable* TableNode = Cast<UCustomizableObjectNodeTable>(&Node))
+		{
+			UDataTable* DataTable = nullptr;
+				
+			if (TableNode->TableDataGatheringMode == ETableDataGatheringSource::ETDGM_AssetRegistry)
 			{
-				AddObjectSoft(TransformNode->ReferenceTexture);
+				DataTable = GenerateDataTableFromStruct(TableNode, Context);
 			}
-
-			else if (UCustomizableObjectNodeMaterial* MatNode = Cast<UCustomizableObjectNodeMaterial>(GenericNode))
+			else
 			{
-				AddObjectSoft(MatNode->GetMaterial());
-				int32 NumImages = MatNode->GetNumParameters(EMaterialParameterType::Texture);
-				for (int32 ImageIndex = 0; ImageIndex < NumImages; ++ImageIndex)
-				{
-					AddObject(MatNode->GetImageReferenceTexture(ImageIndex));
-				}
+				DataTable = TableNode->Table;
 			}
 
-			else if (UCustomizableObjectNodeTable* TableNode = Cast<UCustomizableObjectNodeTable>(GenericNode))
+			if (DataTable)
 			{
-				UDataTable* DataTable = nullptr;
-					
-				if (TableNode->TableDataGatheringMode == ETableDataGatheringSource::ETDGM_AssetRegistry)
-				{
-					DataTable = GenerateDataTableFromStruct(TableNode, Context);
-				}
-				else
-				{
-					DataTable = TableNode->Table;
-				}
+				AddObject(DataTable);
 
-				if (DataTable)
-				{
-					AddObject(DataTable);
-
-					TArray<FName> RowNames = TableNode->GetEnabledRows(*DataTable);
+				TArray<FName> RowNames = TableNode->GetEnabledRows(*DataTable);
 
-					TArray<FName> ExpectedPropNames = DataTableUtils::GetStructPropertyNames(DataTable->RowStruct);
-					for (const FName& ColumnName : ExpectedPropNames)
+				TArray<FName> ExpectedPropNames = DataTableUtils::GetStructPropertyNames(DataTable->RowStruct);
+				for (const FName& ColumnName : ExpectedPropNames)
+				{
+					FProperty* ColumnProperty = DataTable->FindTableProperty(ColumnName);
+					if (!ColumnProperty)
 					{
-						FProperty* ColumnProperty = DataTable->FindTableProperty(ColumnName);
-						if (!ColumnProperty)
-						{
-							continue;
-						}
+						continue;
+					}
 
-						const FSoftObjectProperty* SoftObjectProperty = CastField<FSoftObjectProperty>(ColumnProperty);
-						if (!SoftObjectProperty)
-						{
-							continue;
-						}
+					const FSoftObjectProperty* SoftObjectProperty = CastField<FSoftObjectProperty>(ColumnProperty);
+					if (!SoftObjectProperty)
+					{
+						continue;
+					}
 
-						for (const FName& RowName : RowNames )
+					for (const FName& RowName : RowNames )
+					{
+						uint8* RowData = DataTable->FindRowUnchecked(RowName);
+						if (RowData)
 						{
-							uint8* RowData = DataTable->FindRowUnchecked(RowName);
-							if (RowData)
+							uint8* CellData = ColumnProperty->ContainerPtrToValuePtr<uint8>(RowData, 0);
+							if (CellData)
 							{
-								uint8* CellData = ColumnProperty->ContainerPtrToValuePtr<uint8>(RowData, 0);
-								if (CellData)
-								{
-									const FSoftObjectPtr& Path = SoftObjectProperty->GetPropertyValue(CellData);
-									AddObjectSoft(Path.ToSoftObjectPath());
-								}
+								const FSoftObjectPtr& Path = SoftObjectProperty->GetPropertyValue(CellData);
+								AddObjectSoft(Path.ToSoftObjectPath());
 							}
 						}
 					}
 				}
 			}
 		}
-	}
-
-	// Get one level of references for each participating object
-	// This will add some unnecessary assets. It could be further refined by filtering asset types or ignoring assets in Engine.
-	// If more fine-tunning is needed, assets could specificy a recursion level when added to the list from nodes.
-
-	// Copy the original asset list since we will be modifying it on every iteration.
-	TMap<FName, FGuid> ResultCopy = Result;
-	{
-		TArray<FAssetData> Assets;
-		TArray<FSoftObjectPath> References;
-		for (const TPair<FName, FGuid>& Entry : ResultCopy)
-		{
-			Assets.Reset();
-			References.Reset();
-			TArray<FName> PackageDependencies;
-			AssetRegistryModule.Get().GetDependencies(Entry.Key, PackageDependencies);
-
-			FARFilter Filter;
-			for (const FName& PackageName : PackageDependencies)
-			{
-				if (!PackageName.ToString().StartsWith(TEXT("/TempAutosave")))
-				{
-					Filter.PackageNames.Add(PackageName);
-				}
-			}
-			AssetRegistryModule.Get().GetAssets(Filter, Assets);
-
-			for (const FAssetData& Dep : Assets)
-			{
-				// Assets could be filtered like this:
-				//const UClass* AssetClass = Dep.GetClass(EResolveClass::Yes);
-				//if (AssetClass == USkeleton::StaticClass())
-
-				AddObjectSoft(Dep.GetSoftObjectPath());
-			}
-		}
+	};
 
-	}
+	TArray<const UCustomizableObject*> VisitedObjects;
+    UCustomizableObjectNodeObject* RootNode = GraphTraversal::GetFullGraphRootNode(Object, VisitedObjects);
 
+	GraphTraversal::VisitNodes(*RootNode, Visit);
+	
 	// Done
-	return MoveTemp(Result);
+	return Result;
 }
 
 
diff --git a/Engine/Plugins/Mutable/Source/CustomizableObjectEditor/Private/MuCOE/CustomizableObjectEditorModule.h b/Engine/Plugins/Mutable/Source/CustomizableObjectEditor/Private/MuCOE/CustomizableObjectEditorModule.h
--- a/Engine/Plugins/Mutable/Source/CustomizableObjectEditor/Private/MuCOE/CustomizableObjectEditorModule.h
+++ b/Engine/Plugins/Mutable/Source/CustomizableObjectEditor/Private/MuCOE/CustomizableObjectEditorModule.h
@@ -43,7 +43,8 @@ public:
 
 	// ICustomizableObjectEditorModule interface
 	virtual FCustomizableObjectEditorLogger& GetLogger() override;
-	virtual bool IsCompilationOutOfDate(UCustomizableObject& Object, bool bSkipIndirectReferences, TArray<FName>& OutOfDatePackages, TArray<FName>& AddedPackages, TArray<FName>& RemovedPackages, bool& bVersionDiff) const override;
+	virtual bool IsCompilationOutOfDate(const UCustomizableObject& Object, bool bSkipIndirectReferences, TArray<FName>& OutOfDatePackages, TArray<FName>& OutAddedPackages, TArray<FName>& OutRemovedPackages, bool& bOutVersionDiff) const override;
+	virtual void IsCompilationOutOfDate(const UCustomizableObject& Object, bool bSkipIndirectReferences, float MaxTime, const IsCompilationOutOfDateCallback& Callback) const override;
 	virtual bool IsRootObject(const UCustomizableObject& Object) const override;
 	virtual FString GetCurrentReleaseVersionForObject(const UCustomizableObject& Object) const override;
 	virtual UCustomizableObject* GetRootObject(UCustomizableObject* ChildObject) const override;
@@ -51,7 +52,7 @@ public:
 	virtual void GetRelatedObjects(UCustomizableObject*, TSet<UCustomizableObject*>& OutRelated) const override;
 	virtual void BakeCustomizableObjectInstance(UCustomizableObjectInstance* InTargetInstance, const FBakingConfiguration& InBakingConfig) override;
 	virtual USkeletalMesh* GetReferenceSkeletalMesh(const UCustomizableObject& Object, const FName& ComponentName) const override;
-	virtual TMap<FName, FGuid> GetParticipatingObjects(UCustomizableObject* Object, const FCompilationOptions* Options = nullptr) const override;
+	virtual TMap<FName, FGuid> GetParticipatingObjects(const UCustomizableObject* Object, const FCompilationOptions* Options = nullptr) const override;
 	virtual void BackwardsCompatibleFixup(UEdGraph& Graph, int32 CustomizableObjectCustomVersion) override;
 	virtual void PostBackwardsCompatibleFixup(UEdGraph& Graph) override;
 	virtual void CancelCompileRequests() override;
diff --git a/Engine/Plugins/Mutable/Source/CustomizableObjectEditor/Private/MuCOE/CustomizableObjectIdentifierCustomization.cpp b/Engine/Plugins/Mutable/Source/CustomizableObjectEditor/Private/MuCOE/CustomizableObjectIdentifierCustomization.cpp
--- a/Engine/Plugins/Mutable/Source/CustomizableObjectEditor/Private/MuCOE/CustomizableObjectIdentifierCustomization.cpp
+++ b/Engine/Plugins/Mutable/Source/CustomizableObjectEditor/Private/MuCOE/CustomizableObjectIdentifierCustomization.cpp
@@ -55,10 +55,9 @@ void FCustomizableObjectIdentifierCustomization::CustomizeHeader(TSharedRef<clas
 
 			if (CustomizableObject)
 			{
-				bool bOutMultipleBaseObjectsFound = false;
-				UCustomizableObjectNodeObject* RootNode = GetRootNode(CustomizableObject, bOutMultipleBaseObjectsFound);
+				UCustomizableObjectNodeObject* RootNode = GetRootNode(CustomizableObject);
 
-				if (RootNode && !bOutMultipleBaseObjectsFound && RootNode->Identifier == Guid)
+				if (RootNode && RootNode->Identifier == Guid)
 				{
 					SelectedCustomizableObject = CustomizableObject;
 					break;
@@ -100,10 +99,7 @@ void FCustomizableObjectIdentifierCustomization::OnEventChanged(const FAssetData
 
 	if (SelectedCustomizableObject)
 	{
-		bool bOutMultipleBaseObjectsFound = false;
-		UCustomizableObjectNodeObject* RootNode = GetRootNode(SelectedCustomizableObject, bOutMultipleBaseObjectsFound);
-
-		if (RootNode && !bOutMultipleBaseObjectsFound)
+		if (UCustomizableObjectNodeObject* RootNode = GetRootNode(SelectedCustomizableObject))
 		{
 			if (!RootNode->Identifier.IsValid())
 			{
diff --git a/Engine/Plugins/Mutable/Source/CustomizableObjectEditor/Private/MuCOE/GenerateMutableSource/GenerateMutableSourceComponent.cpp b/Engine/Plugins/Mutable/Source/CustomizableObjectEditor/Private/MuCOE/GenerateMutableSource/GenerateMutableSourceComponent.cpp
--- a/Engine/Plugins/Mutable/Source/CustomizableObjectEditor/Private/MuCOE/GenerateMutableSource/GenerateMutableSourceComponent.cpp
+++ b/Engine/Plugins/Mutable/Source/CustomizableObjectEditor/Private/MuCOE/GenerateMutableSource/GenerateMutableSourceComponent.cpp
@@ -440,7 +440,7 @@ void FirstPass(UCustomizableObjectNodeComponentMesh& Node, FMutableGraphGenerati
 		return Element.ComponentName == Node.ComponentName;
 	}))
 	{
-		FText Msg = FText::Format(LOCTEXT("ComponentNodeWithSameNameExists", "Error! It already exists a Mesh Component node with the same name in Customizable Object [{0}]"), FText::FromString(GetRootObject(*Result->Node)->GetName()));
+		FText Msg = FText::Format(LOCTEXT("ComponentNodeWithSameNameExists", "Error! It already exists a Mesh Component node with the same name in Customizable Object [{0}]"), FText::FromString(GraphTraversal::GetObject(*Result->Node)->GetName()));
 		GenerationContext.Log(Msg, &Node, EMessageSeverity::Error);
 		return;
 	}
diff --git a/Engine/Plugins/Mutable/Source/CustomizableObjectEditor/Private/MuCOE/GraphTraversal.cpp b/Engine/Plugins/Mutable/Source/CustomizableObjectEditor/Private/MuCOE/GraphTraversal.cpp
--- a/Engine/Plugins/Mutable/Source/CustomizableObjectEditor/Private/MuCOE/GraphTraversal.cpp
+++ b/Engine/Plugins/Mutable/Source/CustomizableObjectEditor/Private/MuCOE/GraphTraversal.cpp
@@ -30,13 +30,29 @@
 
 TArray<UEdGraphPin*> FollowPinArray(const UEdGraphPin& Pin, bool bIgnoreOrphan, bool* bOutCycleDetected)
 {
+	MUTABLE_CPUPROFILER_SCOPE(FollowPinArray)
+	
 	bool bCycleDetected = false;
 		
 	TArray<UEdGraphPin*> Result;
+
+	// Early out optimization
+	if (Pin.LinkedTo.IsEmpty())
+	{
+		if (bOutCycleDetected)
+		{
+			*bOutCycleDetected = bCycleDetected;	
+		}
+
+		return Result;
+	}
 	
 	TSet<const UEdGraphPin*> Visited;
+	Visited.Reserve(32);
 
 	TArray<const UEdGraphPin*> PinsToVisit;
+	PinsToVisit.Reserve(32);
+	
 	PinsToVisit.Add(&Pin);
 	while (PinsToVisit.Num())
 	{
@@ -59,24 +75,10 @@ TArray<UEdGraphPin*> FollowPinArray(const UEdGraphPin& Pin, bool bIgnoreOrphan,
 			{
 				continue;
 			}
-			
-			if (const UCustomizableObjectNodeExposePin* ExposePinNode = Cast<UCustomizableObjectNodeExposePin>(LinkedPin->GetOwningNodeUnchecked()))
-			{
-				check(Pin.Direction == EGPD_Output);
 
-				for (TObjectIterator<UCustomizableObjectNodeExternalPin> It; It; ++It)
-				{
-					const UCustomizableObjectNodeExternalPin* LinkedNode = *It;
-					
-					if (IsValid(LinkedNode) &&
-						!LinkedNode->IsTemplate() &&
-						LinkedNode->GetNodeExposePin() == ExposePinNode)
-					{
-						const UEdGraphPin* ExternalPin = LinkedNode->GetExternalPin();
-						check(ExternalPin);
-						PinsToVisit.Add(ExternalPin);
-					}
-				}
+			if (const UCustomizableObjectNodeReroute* NodeReroute = Cast<UCustomizableObjectNodeReroute>(LinkedPin->GetOwningNodeUnchecked()))
+			{
+				PinsToVisit.Add(Pin.Direction == EGPD_Input ? NodeReroute->GetInputPin() : NodeReroute->GetOutputPin());
 			}
 			else if (const UCustomizableObjectNodeExternalPin* ExternalPinNode = Cast<UCustomizableObjectNodeExternalPin>(LinkedPin->GetOwningNodeUnchecked()))
 			{
@@ -89,9 +91,22 @@ TArray<UEdGraphPin*> FollowPinArray(const UEdGraphPin& Pin, bool bIgnoreOrphan,
 					PinsToVisit.Add(ExposePin);
 				}
 			}
-			else if (const UCustomizableObjectNodeReroute* NodeReroute = Cast<UCustomizableObjectNodeReroute>(LinkedPin->GetOwningNodeUnchecked()))
+			else if (const UCustomizableObjectNodeExposePin* ExposePinNode = Cast<UCustomizableObjectNodeExposePin>(LinkedPin->GetOwningNodeUnchecked()))
 			{
-				PinsToVisit.Add(Pin.Direction == EGPD_Input ? NodeReroute->GetInputPin() : NodeReroute->GetOutputPin());
+				check(Pin.Direction == EGPD_Output);
+
+				for (TObjectIterator<UCustomizableObjectNodeExternalPin> It; It; ++It)
+				{
+					const UCustomizableObjectNodeExternalPin* LinkedNode = *It;
+					
+					if (LinkedNode &&
+						LinkedNode->GetNodeExposePin() == ExposePinNode)
+					{
+						const UEdGraphPin* ExternalPin = LinkedNode->GetExternalPin();
+						check(ExternalPin);
+						PinsToVisit.Add(ExternalPin);
+					}
+				}
 			}
 			else
 			{
@@ -244,53 +259,35 @@ TArray<UEdGraphPin*> ReverseFollowPinArray(const UEdGraphPin& Pin, bool bIgnoreO
 }
 
 
-UCustomizableObjectNodeObject* GetRootNode(const UCustomizableObject* Object, bool& bOutMultipleBaseObjectsFound)
+UCustomizableObjectNodeObject* GetRootNode(const UCustomizableObject* Object)
 {
-	// Look for the base object node
-	UCustomizableObjectNodeObject* Root = nullptr;
-	TArray<UCustomizableObjectNodeObject*> ObjectNodes;
-
-	if (Object)
+	if (!Object)
 	{
-		if (const UCustomizableObjectPrivate* Private = Object->GetPrivate())
-		{
-			if (TObjectPtr<UEdGraph> Source = Private->GetSource())
-			{
-				Source->GetNodesOfClass<UCustomizableObjectNodeObject>(ObjectNodes);
-			}
-		}
+		return nullptr;
 	}
-
-	bOutMultipleBaseObjectsFound = false;
-
-	for (TArray<UCustomizableObjectNodeObject*>::TIterator It(ObjectNodes); It; ++It)
+	
+	for (UEdGraphNode* Node : Object->GetPrivate()->GetSource()->Nodes)
 	{
-		if ((*It)->bIsBase)
+		if (UCustomizableObjectNodeObject* NodeObject = Cast<UCustomizableObjectNodeObject>(Node))
 		{
-			if (Root)
-			{
-				bOutMultipleBaseObjectsFound = true;
-				break;
-			}
-			else
+			if (NodeObject->bIsBase)
 			{
-				Root = *It;
+				return NodeObject;
 			}
 		}
 	}
 
-	return Root;
+	return nullptr;
 }
 
 
 bool GetParentsUntilRoot(const UCustomizableObject* Object, TArray<UCustomizableObjectNodeObject*>& ArrayNodeObject, TArray<const  UCustomizableObject*>& ArrayCustomizableObject)
 {
-	bool MultipleBaseObjectsFound;
-	UCustomizableObjectNodeObject* Root = GetRootNode(Object, MultipleBaseObjectsFound);
+	UCustomizableObjectNodeObject* Root = GetRootNode(Object);
 
 	bool bSuccess = true;
 
-	if (!MultipleBaseObjectsFound && (Root != nullptr))
+	if (Root)
 	{
 		if (!ArrayCustomizableObject.Contains(Object))
 		{
@@ -322,16 +319,15 @@ bool HasCandidateAsParent(UCustomizableObjectNodeObject* Node, UCustomizableObje
 
 	if (Node->ParentObject != nullptr)
 	{
-		bool MultipleBaseObjectsFound;
-		UCustomizableObjectNodeObject* ParentNodeObject = GetRootNode(Node->ParentObject, MultipleBaseObjectsFound);
+		UCustomizableObjectNodeObject* ParentNodeObject = GetRootNode(Node->ParentObject);
 
-		if ((ParentNodeObject->ParentObject == nullptr) || MultipleBaseObjectsFound)
+		if (ParentNodeObject->ParentObject)
 		{
-			return false;
+			return HasCandidateAsParent(ParentNodeObject, ParentCandidate);
 		}
 		else
 		{
-			return HasCandidateAsParent(ParentNodeObject, ParentCandidate);
+			return false;
 		}
 	}
 
@@ -345,19 +341,11 @@ UCustomizableObject* GetFullGraphRootObject(const UCustomizableObjectNodeObject*
 	{
 		VisitedObjects.Add(Node->ParentObject);
 
-		bool MultipleBaseObjectsFound;
-		UCustomizableObjectNodeObject* Root = GetRootNode(Node->ParentObject, MultipleBaseObjectsFound);
+		UCustomizableObjectNodeObject* Root = GetRootNode(Node->ParentObject);
 
 		if (Root->ParentObject == nullptr)
 		{
-			if (MultipleBaseObjectsFound)
-			{
-				return nullptr;
-			}
-			else
-			{
-				return Node->ParentObject;
-			}
+			return Node->ParentObject;
 		}
 		else
 		{
@@ -377,7 +365,7 @@ UCustomizableObject* GetFullGraphRootObject(const UCustomizableObjectNodeObject*
 }
 
 
-UCustomizableObject* GetRootObject(const UCustomizableObjectNode& Node)
+UCustomizableObject* GraphTraversal::GetObject(const UCustomizableObjectNode& Node)
 {
 	return CastChecked<UCustomizableObject>(Node.GetGraph()->GetOuter());
 }
@@ -393,8 +381,7 @@ UCustomizableObject* GraphTraversal::GetRootObject(UCustomizableObject* ChildObj
 const UCustomizableObject* GraphTraversal::GetRootObject(const UCustomizableObject* ChildObject)
 {
 	// Grab a node to start the search -> Get the root since it should be always present
-	bool bMultipleBaseObjectsFound = false;
-	UCustomizableObjectNodeObject* ObjectRootNode = GetRootNode(ChildObject, bMultipleBaseObjectsFound);
+	UCustomizableObjectNodeObject* ObjectRootNode = GetRootNode(ChildObject);
 
 	if (ObjectRootNode && ObjectRootNode->ParentObject)
 	{
@@ -408,13 +395,20 @@ const UCustomizableObject* GraphTraversal::GetRootObject(const UCustomizableObje
 }
 
 
-void GraphTraversal::VisitNodes(UCustomizableObjectNode& StartNode, const TMultiMap<FGuid, UCustomizableObjectNodeObject*>& ObjectGroupMap, const TFunction<void(UCustomizableObjectNode&)>& VisitFunction)
+void GraphTraversal::VisitNodes(UCustomizableObjectNode& StartNode, const TFunction<void(UCustomizableObjectNode&)>& VisitFunction, const TMultiMap<FGuid, UCustomizableObjectNodeObject*>* ObjectGroupMap)
 {
 	TSet<UCustomizableObjectNode*> VisitedNodes;
 	
 	TArray<UCustomizableObjectNode*> NodesToVisit;
 	NodesToVisit.Add(&StartNode);
 
+	TMultiMap<FGuid, UCustomizableObjectNodeObject*> LocalObjectGroupMap;
+	if (!ObjectGroupMap)
+	{
+		LocalObjectGroupMap = GetNodeGroupObjectNodeMapping(GetObject(StartNode));
+		ObjectGroupMap = &LocalObjectGroupMap;
+	}
+	
 	while (!NodesToVisit.IsEmpty())
 	{
 		UCustomizableObjectNode* CurrentNode = NodesToVisit.Pop();
@@ -428,13 +422,18 @@ void GraphTraversal::VisitNodes(UCustomizableObjectNode& StartNode, const TMulti
 
 		VisitFunction(*CurrentNode);
 
-		for (UEdGraphPin* Pin : CurrentNode->GetAllNonOrphanPins())
+		for (UEdGraphPin* Pin : CurrentNode->Pins)
 		{
 			if (Pin->Direction != EGPD_Input)
 			{
 				continue;
 			}
 
+			if (IsPinOrphan(*Pin))
+			{
+				continue;
+			}
+			
 			for (UEdGraphPin* ConnectedPin : FollowInputPinArray(*Pin))
 			{
 				UEdGraphNode* ConnectedNode = ConnectedPin->GetOwningNode();
@@ -442,16 +441,16 @@ void GraphTraversal::VisitNodes(UCustomizableObjectNode& StartNode, const TMulti
 				if (UCustomizableObjectNode* Node = Cast<UCustomizableObjectNode>(ConnectedNode))
 				{
 					NodesToVisit.Add(Node);
-				}
-
-				if (UCustomizableObjectNodeObjectGroup* ObjectGroupNode = Cast<UCustomizableObjectNodeObjectGroup>(ConnectedNode))
-				{
-					TArray<UCustomizableObjectNodeObject*> ChildObjectNodes;
-					ObjectGroupMap.MultiFind(ObjectGroupNode->NodeGuid, ChildObjectNodes);
-						
-					for (UCustomizableObjectNodeObject* ChildObjectNode : ChildObjectNodes)
+					
+					if (UCustomizableObjectNodeObjectGroup* ObjectGroupNode = Cast<UCustomizableObjectNodeObjectGroup>(ConnectedNode))
 					{
-						NodesToVisit.Add(ChildObjectNode);
+						TArray<UCustomizableObjectNodeObject*> ChildObjectNodes;
+						ObjectGroupMap->MultiFind(ObjectGroupNode->NodeGuid, ChildObjectNodes);
+						
+						for (UCustomizableObjectNodeObject* ChildObjectNode : ChildObjectNodes)
+						{
+							NodesToVisit.Add(ChildObjectNode);
+						}
 					}
 				}
 			}
@@ -460,25 +459,17 @@ void GraphTraversal::VisitNodes(UCustomizableObjectNode& StartNode, const TMulti
 }
 
 
-UCustomizableObjectNodeObject* GetFullGraphRootNodeObject(const UCustomizableObjectNodeObject* Node, TArray<const UCustomizableObject*>& VisitedObjects)
+UCustomizableObjectNodeObject* GraphTraversal::GetFullGraphRootNode(const UCustomizableObject* Object, TArray<const UCustomizableObject*>& VisitedObjects)
 {
-	if (Node->ParentObject != nullptr)
+	if (Object != nullptr)
 	{
-		VisitedObjects.Add(Node->ParentObject);
+		VisitedObjects.Add(Object);
 
-		bool MultipleBaseObjectsFound;
-		UCustomizableObjectNodeObject* Root = GetRootNode(Node->ParentObject, MultipleBaseObjectsFound);
+		UCustomizableObjectNodeObject* Root = GetRootNode(Object);
 
 		if (Root->ParentObject == nullptr)
 		{
-			if (MultipleBaseObjectsFound)
-			{
-				return nullptr;
-			}
-			else
-			{
-				return Root;
-			}
+			return Root;
 		}
 		else
 		{
@@ -489,7 +480,7 @@ UCustomizableObjectNodeObject* GetFullGraphRootNodeObject(const UCustomizableObj
 			}
 			else
 			{
-				return GetFullGraphRootNodeObject(Root, VisitedObjects);
+				return GetFullGraphRootNode(Root->ParentObject, VisitedObjects);
 			}
 		}
 	}
@@ -648,14 +639,16 @@ const UEdGraphPin* FindMeshBaseSource(const UEdGraphPin& Pin, const bool bOnlyLo
 }
 
 
-void GetNodeGroupObjectNodeMappingImmersive(UCustomizableObject* Object, FAssetRegistryModule& AssetRegistryModule, TSet<UCustomizableObject*>& Visited, TMultiMap<FGuid, UCustomizableObjectNodeObject*>& Mapping)
+void GetNodeGroupObjectNodeMappingImmersive(UCustomizableObject* Object, IAssetRegistry& AssetRegistry, TSet<UCustomizableObject*>& Visited, TMultiMap<FGuid, UCustomizableObjectNodeObject*>& Mapping)
 {
 	Visited.Add(Object);
 
 	TArray<FName> ArrayReferenceNames;
-	AssetRegistryModule.Get().GetReferencers(*Object->GetOuter()->GetPathName(), ArrayReferenceNames, UE::AssetRegistry::EDependencyCategory::Package, UE::AssetRegistry::EDependencyQuery::Hard);
+	AssetRegistry.GetReferencers(*Object->GetOuter()->GetPathName(), ArrayReferenceNames, UE::AssetRegistry::EDependencyCategory::Package, UE::AssetRegistry::EDependencyQuery::Hard);
 
 	FARFilter Filter;
+	Filter.bIncludeOnlyOnDiskAssets = false;
+	
 	for (const FName& ReferenceName : ArrayReferenceNames)
 	{
 		if (!ReferenceName.ToString().StartsWith(TEXT("/TempAutosave")))
@@ -664,10 +657,8 @@ void GetNodeGroupObjectNodeMappingImmersive(UCustomizableObject* Object, FAssetR
 		}
 	}
 	
-	Filter.bIncludeOnlyOnDiskAssets = false;
-	
 	TArray<FAssetData> ArrayAssetData;
-	AssetRegistryModule.Get().GetAssets(Filter, ArrayAssetData);
+	AssetRegistry.GetAssets(Filter, ArrayAssetData);
 
 	for (FAssetData& AssetData : ArrayAssetData)
 	{
@@ -679,10 +670,7 @@ void GetNodeGroupObjectNodeMappingImmersive(UCustomizableObject* Object, FAssetR
 
 		if (ChildObject != Object && !ChildObject->HasAnyFlags(RF_Transient))
 		{
-			bool bMultipleBaseObjectsFound = false;
-			UCustomizableObjectNodeObject* ChildRoot = GetRootNode(ChildObject, bMultipleBaseObjectsFound);
-
-			if (ChildRoot && !bMultipleBaseObjectsFound)
+			if (UCustomizableObjectNodeObject* ChildRoot = GetRootNode(ChildObject))
 			{
 				if (ChildRoot->ParentObject == Object)
 				{
@@ -693,7 +681,7 @@ void GetNodeGroupObjectNodeMappingImmersive(UCustomizableObject* Object, FAssetR
 
 		if (!Visited.Contains(ChildObject))
 		{
-			GetNodeGroupObjectNodeMappingImmersive(ChildObject, AssetRegistryModule, Visited, Mapping);
+			GetNodeGroupObjectNodeMappingImmersive(ChildObject, AssetRegistry, Visited, Mapping);
 		}
 	}
 }
@@ -701,12 +689,14 @@ void GetNodeGroupObjectNodeMappingImmersive(UCustomizableObject* Object, FAssetR
 
 TMultiMap<FGuid, UCustomizableObjectNodeObject*> GetNodeGroupObjectNodeMapping(UCustomizableObject* Object)
 {
-	FAssetRegistryModule& AssetRegistryModule = FModuleManager::LoadModuleChecked<FAssetRegistryModule>("AssetRegistry");
+	MUTABLE_CPUPROFILER_SCOPE(GetNodeGroupObjectNodeMapping);
+
+	IAssetRegistry& AssetRegistry = FModuleManager::LoadModuleChecked<FAssetRegistryModule>("AssetRegistry").Get();
 
 	TSet<UCustomizableObject*> Visited;
 	TMultiMap<FGuid, UCustomizableObjectNodeObject*> Mapping;
 
-	GetNodeGroupObjectNodeMappingImmersive(Object, AssetRegistryModule, Visited, Mapping);
+	GetNodeGroupObjectNodeMappingImmersive(Object, AssetRegistry, Visited, Mapping);
 	
 	return Mapping;
 }
@@ -723,8 +713,8 @@ void GetAllObjectsInGraph(UCustomizableObject* Object, TSet<UCustomizableObject*
 	UCustomizableObject* RootObject = GraphTraversal::GetRootObject(Object);
 	TMultiMap<FGuid, UCustomizableObjectNodeObject*> DummyMap;
 
-	FAssetRegistryModule& AssetRegistryModule = FModuleManager::LoadModuleChecked<FAssetRegistryModule>("AssetRegistry");
-	GetNodeGroupObjectNodeMappingImmersive(RootObject, AssetRegistryModule, OutObjects, DummyMap);
+	IAssetRegistry& AssetRegistry = FModuleManager::LoadModuleChecked<FAssetRegistryModule>("AssetRegistry").Get();
+	GetNodeGroupObjectNodeMappingImmersive(RootObject, AssetRegistry, OutObjects, DummyMap);
 }
 
 
diff --git a/Engine/Plugins/Mutable/Source/CustomizableObjectEditor/Private/MuCOE/GraphTraversal.h b/Engine/Plugins/Mutable/Source/CustomizableObjectEditor/Private/MuCOE/GraphTraversal.h
--- a/Engine/Plugins/Mutable/Source/CustomizableObjectEditor/Private/MuCOE/GraphTraversal.h
+++ b/Engine/Plugins/Mutable/Source/CustomizableObjectEditor/Private/MuCOE/GraphTraversal.h
@@ -53,7 +53,7 @@ UEdGraphPin* FollowOutputPin(const UEdGraphPin& Pin, bool* CycleDetected = nullp
 TArray<UEdGraphPin*> ReverseFollowPinArray(const UEdGraphPin& Pin, bool bIgnoreOrphan = false, bool* bOutCycleDetected = nullptr);
 
 /** Returns the root Object Node of the Customizable Object's graph */
-UCustomizableObjectNodeObject* GetRootNode(const UCustomizableObject* Object, bool& bOutMultipleBaseObjectsFound);
+UCustomizableObjectNodeObject* GetRootNode(const UCustomizableObject* Object);
 
 /** Return in ArrayNodeObject the roots nodes in each Customizable Object graph until the whole root node is found (i.e. the one with parent = nullptr)
  * return false if a cycle is found between Customizable Objects */
@@ -65,11 +65,11 @@ bool HasCandidateAsParent(UCustomizableObjectNodeObject* Node, UCustomizableObje
 /** Return the full graph Customizable Object root of the given node. */
 UCustomizableObject* GetFullGraphRootObject(const UCustomizableObjectNodeObject* Node, TArray<UCustomizableObject*>& VisitedObjects);
 
-/** Return Customizable Object of the given node. */
-UCustomizableObject* GetRootObject(const UCustomizableObjectNode& Node);
-
 namespace GraphTraversal
 {
+	/** Return Customizable Object of the given node. */
+	UCustomizableObject* GetObject(const UCustomizableObjectNode& Node);
+
 	/** Provided a CO object it provides the root CO it is connected. In other words : it returns the root of the entire
 		 * mutable graph.
 		 * @param InObject Customizable object whose root CO we are asking for.
@@ -83,12 +83,12 @@ namespace GraphTraversal
 	 *   @param StartNode Node to start the visit.
 	 *   @param ObjectGroupMap Key is the Object Group node id, values are the attached Child Object nodes. Used to prune the traversal.
 	 *   @param VisitFunction Called for each UCustomizableObjectNode node found. Order is not guaranteed. */
-	void VisitNodes(UCustomizableObjectNode& StartNode, const TMultiMap<FGuid, UCustomizableObjectNodeObject*>& ObjectGroupMap, const TFunction<void(UCustomizableObjectNode&)>& VisitFunction);
+	void VisitNodes(UCustomizableObjectNode& StartNode, const TFunction<void(UCustomizableObjectNode&)>& VisitFunction, const TMultiMap<FGuid, UCustomizableObjectNodeObject*>* ObjectGroupMap = nullptr);
+	
+	/** Return the full graph Customizable Object Node root of the node given as parameter */
+	UCustomizableObjectNodeObject* GetFullGraphRootNode(const UCustomizableObject* Object, TArray<const UCustomizableObject*>& VisitedObjects);
 }
 
-/** Return the full graph Customizable Object Node root of the node given as parameter */
-UCustomizableObjectNodeObject* GetFullGraphRootNodeObject(const UCustomizableObjectNodeObject* Node, TArray<const UCustomizableObject*>& VisitedObjects);
-
 /** Given an output pin, return the output pin where the mesh is located. */
 const UEdGraphPin* FindMeshBaseSource(const UEdGraphPin& Pin, const bool bOnlyLookForStaticMesh);
 
diff --git a/Engine/Plugins/Mutable/Source/CustomizableObjectEditor/Private/MuCOE/Nodes/CustomizableObjectNode.cpp b/Engine/Plugins/Mutable/Source/CustomizableObjectEditor/Private/MuCOE/Nodes/CustomizableObjectNode.cpp
--- a/Engine/Plugins/Mutable/Source/CustomizableObjectEditor/Private/MuCOE/Nodes/CustomizableObjectNode.cpp
+++ b/Engine/Plugins/Mutable/Source/CustomizableObjectEditor/Private/MuCOE/Nodes/CustomizableObjectNode.cpp
@@ -493,7 +493,7 @@ bool UCustomizableObjectNode::FindNodeForInternalTag(const FString& Tag, UCustom
 	}
 
 	// Scan all potential receivers
-	UCustomizableObject* ThisNodeObject = GetRootObject(*this);
+	UCustomizableObject* ThisNodeObject = GraphTraversal::GetObject(*this);
 	UCustomizableObject* RootObject = GraphTraversal::GetRootObject(ThisNodeObject);
 
 	TSet<UCustomizableObject*> AllCustomizableObject;
diff --git a/Engine/Plugins/Mutable/Source/CustomizableObjectEditor/Private/MuCOE/Nodes/CustomizableObjectNodeComponentMesh.cpp b/Engine/Plugins/Mutable/Source/CustomizableObjectEditor/Private/MuCOE/Nodes/CustomizableObjectNodeComponentMesh.cpp
--- a/Engine/Plugins/Mutable/Source/CustomizableObjectEditor/Private/MuCOE/Nodes/CustomizableObjectNodeComponentMesh.cpp
+++ b/Engine/Plugins/Mutable/Source/CustomizableObjectEditor/Private/MuCOE/Nodes/CustomizableObjectNodeComponentMesh.cpp
@@ -35,7 +35,7 @@ void UCustomizableObjectNodeComponentMesh::BackwardsCompatibleFixup(int32 Custom
 	
 	if (CustomizableObjectCustomVersion == FCustomizableObjectCustomVersion::ComponentsArray)
 	{
-		UCustomizableObject* Object = GetRootObject(*this);
+		UCustomizableObject* Object = GraphTraversal::GetObject(*this);
 
 		if (FMutableMeshComponentData* Result = Object->GetPrivate()->MutableMeshComponents_DEPRECATED.FindByPredicate([&](const FMutableMeshComponentData& ComponentData)
 		{
diff --git a/Engine/Plugins/Mutable/Source/CustomizableObjectEditor/Private/MuCOE/Nodes/CustomizableObjectNodeMeshClipWithMesh.cpp b/Engine/Plugins/Mutable/Source/CustomizableObjectEditor/Private/MuCOE/Nodes/CustomizableObjectNodeMeshClipWithMesh.cpp
--- a/Engine/Plugins/Mutable/Source/CustomizableObjectEditor/Private/MuCOE/Nodes/CustomizableObjectNodeMeshClipWithMesh.cpp
+++ b/Engine/Plugins/Mutable/Source/CustomizableObjectEditor/Private/MuCOE/Nodes/CustomizableObjectNodeMeshClipWithMesh.cpp
@@ -69,7 +69,7 @@ void UCustomizableObjectNodeModifierClipWithMesh::BackwardsCompatibleFixup(int32
 
 				// If the tag was added to another CO, keep track of the addition to be able to repeat it in case
 				// that CO is not re-saved but this one is.
-				UCustomizableObject* ThisNodeObject = GetRootObject(*this);
+				UCustomizableObject* ThisNodeObject = GraphTraversal::GetObject(*this);
 				if (CustomizableObjectToClipWith_DEPRECATED != ThisNodeObject)
 				{
 					FLegacyTag LegacyTag;
diff --git a/Engine/Plugins/Mutable/Source/CustomizableObjectEditor/Private/MuCOE/Nodes/CustomizableObjectNodeModifierBase.cpp b/Engine/Plugins/Mutable/Source/CustomizableObjectEditor/Private/MuCOE/Nodes/CustomizableObjectNodeModifierBase.cpp
--- a/Engine/Plugins/Mutable/Source/CustomizableObjectEditor/Private/MuCOE/Nodes/CustomizableObjectNodeModifierBase.cpp
+++ b/Engine/Plugins/Mutable/Source/CustomizableObjectEditor/Private/MuCOE/Nodes/CustomizableObjectNodeModifierBase.cpp
@@ -75,7 +75,7 @@ bool UCustomizableObjectNodeModifierBase::IsApplicableTo(UCustomizableObjectNode
 void UCustomizableObjectNodeModifierBase::GetPossiblyModifiedNodes(TArray<UCustomizableObjectNode*>& CandidateNodes)
 {
 	// Scan all potential receivers
-	UCustomizableObject* ThisNodeObject = GetRootObject(*this);
+	UCustomizableObject* ThisNodeObject = GraphTraversal::GetObject(*this);
 	UCustomizableObject* RootObject = GraphTraversal::GetRootObject(ThisNodeObject);
 
 	TSet<UCustomizableObject*> AllCustomizableObject;
diff --git a/Engine/Plugins/Mutable/Source/CustomizableObjectEditor/Private/MuCOE/Nodes/CustomizableObjectNodeModifierEditMeshSectionBase.cpp b/Engine/Plugins/Mutable/Source/CustomizableObjectEditor/Private/MuCOE/Nodes/CustomizableObjectNodeModifierEditMeshSectionBase.cpp
--- a/Engine/Plugins/Mutable/Source/CustomizableObjectEditor/Private/MuCOE/Nodes/CustomizableObjectNodeModifierEditMeshSectionBase.cpp
+++ b/Engine/Plugins/Mutable/Source/CustomizableObjectEditor/Private/MuCOE/Nodes/CustomizableObjectNodeModifierEditMeshSectionBase.cpp
@@ -45,7 +45,7 @@ void UCustomizableObjectNodeModifierEditMeshSectionBase::BackwardsCompatibleFixu
 
 			// If the tag was added to another CO, keep track of the addition to be able to repeat it in case
 			// that CO is not re-saved but this one is.
-			UCustomizableObject* ThisNodeObject = GetRootObject(*this);
+			UCustomizableObject* ThisNodeObject = GraphTraversal::GetObject(*this);
 			if (ParentMaterialObject_DEPRECATED != ThisNodeObject)
 			{
 				FLegacyTag LegacyTag;
diff --git a/Engine/Plugins/Mutable/Source/CustomizableObjectEditor/Private/MuCOE/Nodes/CustomizableObjectNodeModifierExtendMeshSection.cpp b/Engine/Plugins/Mutable/Source/CustomizableObjectEditor/Private/MuCOE/Nodes/CustomizableObjectNodeModifierExtendMeshSection.cpp
--- a/Engine/Plugins/Mutable/Source/CustomizableObjectEditor/Private/MuCOE/Nodes/CustomizableObjectNodeModifierExtendMeshSection.cpp
+++ b/Engine/Plugins/Mutable/Source/CustomizableObjectEditor/Private/MuCOE/Nodes/CustomizableObjectNodeModifierExtendMeshSection.cpp
@@ -88,7 +88,7 @@ void UCustomizableObjectNodeModifierExtendMeshSection::BackwardsCompatibleFixup(
 
 			// If the tag was added to another CO, keep track of the addition to be able to repeat it in case
 			// that CO is not re-saved but this one is.
-			UCustomizableObject* ThisNodeObject = GetRootObject(*this);
+			UCustomizableObject* ThisNodeObject = GraphTraversal::GetObject(*this);
 			if (ParentMaterialObject_DEPRECATED != ThisNodeObject)
 			{
 				FLegacyTag LegacyTag;
diff --git a/Engine/Plugins/Mutable/Source/CustomizableObjectEditor/Private/MuCOE/Nodes/CustomizableObjectNodeModifierMorphMeshSectionDetails.cpp b/Engine/Plugins/Mutable/Source/CustomizableObjectEditor/Private/MuCOE/Nodes/CustomizableObjectNodeModifierMorphMeshSectionDetails.cpp
--- a/Engine/Plugins/Mutable/Source/CustomizableObjectEditor/Private/MuCOE/Nodes/CustomizableObjectNodeModifierMorphMeshSectionDetails.cpp
+++ b/Engine/Plugins/Mutable/Source/CustomizableObjectEditor/Private/MuCOE/Nodes/CustomizableObjectNodeModifierMorphMeshSectionDetails.cpp
@@ -135,7 +135,7 @@ void FCustomizableObjectNodeModifierMorphMeshSectionDetails::RefreshMorphOptions
 	// Add all morphs if no candidate is found
 	if (MorphOptionsSource.IsEmpty())
 	{
-		UCustomizableObject* ThisNodeObject = GetRootObject(*Node);
+		UCustomizableObject* ThisNodeObject = GraphTraversal::GetObject(*Node);
 		UCustomizableObject* RootObject = GraphTraversal::GetRootObject(ThisNodeObject);
 
 		TSet<UCustomizableObject*> AllCustomizableObject;
diff --git a/Engine/Plugins/Mutable/Source/CustomizableObjectEditor/Private/MuCOE/Nodes/CustomizableObjectNodeObject.cpp b/Engine/Plugins/Mutable/Source/CustomizableObjectEditor/Private/MuCOE/Nodes/CustomizableObjectNodeObject.cpp
--- a/Engine/Plugins/Mutable/Source/CustomizableObjectEditor/Private/MuCOE/Nodes/CustomizableObjectNodeObject.cpp
+++ b/Engine/Plugins/Mutable/Source/CustomizableObjectEditor/Private/MuCOE/Nodes/CustomizableObjectNodeObject.cpp
@@ -254,7 +254,7 @@ void UCustomizableObjectNodeObject::BackwardsCompatibleFixup(int32 CustomizableO
 				NodePosX += 400; // Move it a bit to make space for the new component nodes.
 			}
 
-			for (const FMutableMeshComponentData& MeshComponent : GetRootObject(*this)->GetPrivate()->MutableMeshComponents_DEPRECATED)
+			for (const FMutableMeshComponentData& MeshComponent : GraphTraversal::GetObject(*this)->GetPrivate()->MutableMeshComponents_DEPRECATED)
 			{
 				UCustomizableObjectNodeComponentMesh* NewNode = NewObject<UCustomizableObjectNodeComponentMesh>(this);
 				UEdGraphNode* Node = CreateNewNode(GetGraph(), ComponentsPin(), FVector2D(NodePosX - 300.0,  NodePosY + 200.0 * NodesCreated), NewNode);
diff --git a/Engine/Plugins/Mutable/Source/CustomizableObjectEditor/Private/MuCOE/Nodes/CustomizableObjectNodeObjectGroup.cpp b/Engine/Plugins/Mutable/Source/CustomizableObjectEditor/Private/MuCOE/Nodes/CustomizableObjectNodeObjectGroup.cpp
--- a/Engine/Plugins/Mutable/Source/CustomizableObjectEditor/Private/MuCOE/Nodes/CustomizableObjectNodeObjectGroup.cpp
+++ b/Engine/Plugins/Mutable/Source/CustomizableObjectEditor/Private/MuCOE/Nodes/CustomizableObjectNodeObjectGroup.cpp
@@ -53,7 +53,7 @@ void UCustomizableObjectNodeObjectGroup::PostEditChangeProperty(FPropertyChanged
 				if (CurrentRootNode->ParentObject)
 				{
 					TArray<const UCustomizableObject*> VisitedObjects;
-					CurrentRootNode = GetFullGraphRootNodeObject(CurrentRootNode, VisitedObjects);
+					CurrentRootNode = GraphTraversal::GetFullGraphRootNode(CurrentRootNode->ParentObject, VisitedObjects);
 				}
 				if (!CurrentRootNode->ParentObject)
 				{
diff --git a/Engine/Plugins/Mutable/Source/CustomizableObjectEditor/Private/MuCOE/SMutableTagListWidget.cpp b/Engine/Plugins/Mutable/Source/CustomizableObjectEditor/Private/MuCOE/SMutableTagListWidget.cpp
--- a/Engine/Plugins/Mutable/Source/CustomizableObjectEditor/Private/MuCOE/SMutableTagListWidget.cpp
+++ b/Engine/Plugins/Mutable/Source/CustomizableObjectEditor/Private/MuCOE/SMutableTagListWidget.cpp
@@ -181,7 +181,7 @@ void SMutableTagComboBox::RefreshOptions()
 	TagComboOptionsSource.SetNum(0, EAllowShrinking::No);
 
 	// Scan all potential receivers
-	UCustomizableObject* ThisNodeObject = GetRootObject(*Node);
+	UCustomizableObject* ThisNodeObject = GraphTraversal::GetObject(*Node);
 	UCustomizableObject* RootObject = GraphTraversal::GetRootObject(ThisNodeObject);
 
 	TSet<UCustomizableObject*> AllCustomizableObject;
