[mutable] Avoid crashing in reshape operations with malformed data.
#rnx

diff --git a/Engine/Plugins/Mutable/Source/MutableRuntime/Private/MuR/OpMeshClipDeform.h b/Engine/Plugins/Mutable/Source/MutableRuntime/Private/MuR/OpMeshClipDeform.h
--- a/Engine/Plugins/Mutable/Source/MutableRuntime/Private/MuR/OpMeshClipDeform.h
+++ b/Engine/Plugins/Mutable/Source/MutableRuntime/Private/MuR/OpMeshClipDeform.h
@@ -70,6 +70,13 @@ namespace mu
 		
 		const UE::Geometry::FIndex3i& Triangle = ShapeMesh.Triangles[Binding.Triangle];
 	
+		if (!ShapeMesh.Positions.IsValidIndex(Triangle.A) || !ShapeMesh.Positions.IsValidIndex(Triangle.B) || !ShapeMesh.Positions.IsValidIndex(Triangle.C) )
+		{
+			// It seems to happen with some test objects.
+			ensure(false);
+			return false;
+		}
+
 		const FVector3f ProjectedVertexPosition = 
 				ShapeMesh.Positions[Triangle.A] * (1.0f - BaryCoords.X - BaryCoords.Y) +
 				ShapeMesh.Positions[Triangle.B] * BaryCoords.X +
