Fix intermittent compilations failures in some platforms due to missing header.
#rnx

diff --git a/Engine/Plugins/Mutable/Source/MutableTools/Public/MuT/NodeRangeFromScalar.h b/Engine/Plugins/Mutable/Source/MutableTools/Public/MuT/NodeRangeFromScalar.h
--- a/Engine/Plugins/Mutable/Source/MutableTools/Public/MuT/NodeRangeFromScalar.h
+++ b/Engine/Plugins/Mutable/Source/MutableTools/Public/MuT/NodeRangeFromScalar.h
@@ -5,6 +5,7 @@
 #include "MuR/Ptr.h"
 #include "MuT/Node.h"
 #include "MuT/NodeRange.h"
+#include "Containers/UnrealString.h"
 
 namespace mu
 {
@@ -12,7 +13,6 @@ namespace mu
 	// Forward definitions
     class NodeScalar;
 
-
     class MUTABLETOOLS_API NodeRangeFromScalar : public NodeRange
 	{
 	public:
