[Mutable] Fix incorrect check in FMutableTextureMipDataProvider::GetMips().

#rnx

diff --git a/Engine/Plugins/Mutable/Source/CustomizableObject/Private/MuCO/CustomizableObjectMipDataProvider.cpp b/Engine/Plugins/Mutable/Source/CustomizableObject/Private/MuCO/CustomizableObjectMipDataProvider.cpp
--- a/Engine/Plugins/Mutable/Source/CustomizableObject/Private/MuCO/CustomizableObjectMipDataProvider.cpp
+++ b/Engine/Plugins/Mutable/Source/CustomizableObject/Private/MuCO/CustomizableObjectMipDataProvider.cpp
@@ -488,7 +488,7 @@ int32 FMutableTextureMipDataProvider::GetMips(const FTextureUpdateContext& Conte
 
 				if (OperationData->bIsCancelled)
 				{
-					check(OperationData->Counter == 0);
+					check(OperationData->Counter->GetValue() == 0);
 					return;
 				}
 
