[mutable] Initial boilerplate for the implementation of Mesh Parameters.
- added stub for Mesh Parameter node.
- Cleanup old unnecessary pimpl pattern on mu::Node objects.
- Remove unnecessary files.
#rnx

diff --git a/Engine/Plugins/Mutable/Source/CustomizableObject/Private/MuCO/CustomizableObjectInstance.cpp b/Engine/Plugins/Mutable/Source/CustomizableObject/Private/MuCO/CustomizableObjectInstance.cpp
--- a/Engine/Plugins/Mutable/Source/CustomizableObject/Private/MuCO/CustomizableObjectInstance.cpp
+++ b/Engine/Plugins/Mutable/Source/CustomizableObject/Private/MuCO/CustomizableObjectInstance.cpp
@@ -6716,7 +6716,12 @@ void UCustomizableObjectInstance::SetRequestedLODs(const TMap<FName, uint8>& InM
 		const uint8 MinLODIdx = CustomizableObject->GetPrivate()->GetMinLODIndex(ComponentName);
 		MutableUpdateCandidate.QualitySettingMinLODs.Add(ComponentName, MinLODIdx);
 
-		const int32 MaxLODIdx = ModelResources->NumLODsAvailable[ComponentName] - 1;
+		int32 MaxLODIdx = 0;
+		if (const uint8* Found = ModelResources->NumLODsAvailable.Find(ComponentName))
+		{
+			MaxLODIdx = *Found - 1;
+		}
+
 		InMinLOD = FMath::Clamp(InMinLOD, MinLODIdx, MaxLODIdx);
 
 		uint8 DescriptorMinLOD = 0;
@@ -6737,7 +6742,11 @@ void UCustomizableObjectInstance::SetRequestedLODs(const TMap<FName, uint8>& InM
 
 			PrivateData->SetCOInstanceFlags(InMinLOD > CurrentMinLOD ? PendingLODsDowngrade : ECONone);
 
-			const uint8 FirstNonStreamedLODIndex = FMath::Clamp(ModelResources->NumLODsToStream[ComponentName], 0, MaxLODIdx);
+			uint8 FirstNonStreamedLODIndex = 0;
+			if (const uint8* Found = ModelResources->NumLODsToStream.Find(ComponentName))
+			{
+				FirstNonStreamedLODIndex = FMath::Clamp(*Found, 0, MaxLODIdx);
+			}
 
 			MutableUpdateCandidate.RequestedLODLevels.Add(ComponentName, FirstNonStreamedLODIndex);
 
diff --git a/Engine/Plugins/Mutable/Source/CustomizableObjectEditor/Private/MuCOE/CustomizableObjectCompiler.cpp b/Engine/Plugins/Mutable/Source/CustomizableObjectEditor/Private/MuCOE/CustomizableObjectCompiler.cpp
--- a/Engine/Plugins/Mutable/Source/CustomizableObjectEditor/Private/MuCOE/CustomizableObjectCompiler.cpp
+++ b/Engine/Plugins/Mutable/Source/CustomizableObjectEditor/Private/MuCOE/CustomizableObjectCompiler.cpp
@@ -680,7 +680,7 @@ mu::Ptr<mu::NodeObject> GenerateMutableRoot(
 
 	// Generate the object expression
 	UE_LOG(LogMutable, Verbose, TEXT("PROFILE: [ %16.8f ] GenerateMutableSource start."), FPlatformTime::Seconds());
-	mu::NodeObjectPtr MutableRoot = GenerateMutableSource(RootNodeObject->OutputPin(), GenerationContext);
+	mu::Ptr<mu::NodeObject> MutableRoot = GenerateMutableSource(RootNodeObject->OutputPin(), GenerationContext);
 	UE_LOG(LogMutable, Verbose, TEXT("PROFILE: [ %16.8f ] GenerateMutableSource end."), FPlatformTime::Seconds());
 
 	GenerationContext.GenerateSharedSurfacesUniqueIds();
@@ -824,7 +824,7 @@ void FCustomizableObjectCompiler::CompileInternal(bool bAsync)
 	CompilationLogsContainer.ClearMessagesArray();
 
 	// Generate the mutable node expression
-	mu::NodeObjectPtr MutableRoot = GenerateMutableRoot(CurrentObject, GenerationContext);
+	mu::Ptr<mu::NodeObject> MutableRoot = GenerateMutableRoot(CurrentObject, GenerationContext);
 	if (!MutableRoot)
 	{
 		CompilerLog(FText(LOCTEXT("FailedToGenerateRoot", "Failed to generate the mutable node graph. Object not built.")));
@@ -1245,7 +1245,7 @@ mu::Ptr<mu::Node> FCustomizableObjectCompiler::Export(UCustomizableObject* Objec
 	FMutableGraphGenerationContext GenerationContext(*CompilationContext);
 	
 	// Generate the mutable node expression
-	mu::NodeObjectPtr MutableRoot = GenerateMutableRoot(Object, GenerationContext);
+	mu::Ptr<mu::NodeObject> MutableRoot = GenerateMutableRoot(Object, GenerationContext);
 	if (!MutableRoot)
 	{
 		CompilerLog(LOCTEXT("FailedToExport", "Failed to generate the mutable node graph. Object not built."), nullptr);
diff --git a/Engine/Plugins/Mutable/Source/CustomizableObjectEditor/Private/MuCOE/EdGraphSchema_CustomizableObject.cpp b/Engine/Plugins/Mutable/Source/CustomizableObjectEditor/Private/MuCOE/EdGraphSchema_CustomizableObject.cpp
--- a/Engine/Plugins/Mutable/Source/CustomizableObjectEditor/Private/MuCOE/EdGraphSchema_CustomizableObject.cpp
+++ b/Engine/Plugins/Mutable/Source/CustomizableObjectEditor/Private/MuCOE/EdGraphSchema_CustomizableObject.cpp
@@ -52,6 +52,7 @@
 #include "MuCOE/Nodes/CustomizableObjectNodeMeshReshape.h"
 #include "MuCOE/Nodes/CustomizableObjectNodeMeshSwitch.h"
 #include "MuCOE/Nodes/CustomizableObjectNodeMeshVariation.h"
+#include "MuCOE/Nodes/CustomizableObjectNodeMeshParameter.h"
 #include "MuCOE/Nodes/CustomizableObjectNodeObjectChild.h"
 #include "MuCOE/Nodes/CustomizableObjectNodeObjectGroup.h"
 #include "MuCOE/Nodes/CustomizableObjectNodeComponentPassthroughMesh.h"
@@ -595,6 +596,7 @@ void UEdGraphSchema_CustomizableObject::GetGraphContextActions(FGraphContextMenu
 			ContextMenuBuilder.CreateTemplateNode<UCustomizableObjectNodeModifierClipDeform>(),
 			ContextMenuBuilder.CreateTemplateNode<UCustomizableObjectNodeModifierTransformInMesh>(),
 			ContextMenuBuilder.CreateTemplateNode<UCustomizableObjectNodeTextureParameter>(),
+			ContextMenuBuilder.CreateTemplateNode<UCustomizableObjectNodeMeshParameter>(),
 		};
 
 		AddNewNodeCategoryActionsFiltered(ExperimentalTemplateNodes, ContextMenuBuilder, NC_Experimental.ToString(), 2, Filter);
diff --git a/Engine/Plugins/Mutable/Source/CustomizableObjectEditor/Private/MuCOE/GenerateMutableSource/GenerateMutableSource.cpp b/Engine/Plugins/Mutable/Source/CustomizableObjectEditor/Private/MuCOE/GenerateMutableSource/GenerateMutableSource.cpp
--- a/Engine/Plugins/Mutable/Source/CustomizableObjectEditor/Private/MuCOE/GenerateMutableSource/GenerateMutableSource.cpp
+++ b/Engine/Plugins/Mutable/Source/CustomizableObjectEditor/Private/MuCOE/GenerateMutableSource/GenerateMutableSource.cpp
@@ -665,13 +665,13 @@ UTexture2D* FindReferenceImage(const UEdGraphPin* Pin, FMutableGraphGenerationCo
 }
 
 
-mu::NodeMeshApplyPosePtr CreateNodeMeshApplyPose(FMutableGraphGenerationContext& GenerationContext, mu::NodeMeshPtr InputMeshNode, const TArray<FName>& ArrayBoneName, const TArray<FTransform>& ArrayTransform)
+mu::Ptr<mu::NodeMeshApplyPose> CreateNodeMeshApplyPose(FMutableGraphGenerationContext& GenerationContext, mu::NodeMeshPtr InputMeshNode, const TArray<FName>& ArrayBoneName, const TArray<FTransform>& ArrayTransform)
 {
 	check(ArrayBoneName.Num() == ArrayTransform.Num());
 
 	TSharedPtr<mu::FMesh> MutableMesh = MakeShared<mu::FMesh>();
-	mu::NodeMeshConstantPtr PoseNodeMesh = new mu::NodeMeshConstant;
-	PoseNodeMesh->SetValue(MutableMesh);
+	mu::Ptr<mu::NodeMeshConstant> PoseNodeMesh = new mu::NodeMeshConstant;
+	PoseNodeMesh->Value = MutableMesh;
 
 	TSharedPtr<mu::FSkeleton> MutableSkeleton = MakeShared<mu::FSkeleton>();
 	MutableMesh->SetSkeleton(MutableSkeleton);
@@ -688,9 +688,9 @@ mu::NodeMeshApplyPosePtr CreateNodeMeshApplyPose(FMutableGraphGenerationContext&
 		MutableMesh->SetBonePose(i, BoneId, (FTransform3f)ArrayTransform[i], mu::EBoneUsageFlags::Skinning);
 	}
 
-	mu::NodeMeshApplyPosePtr NodeMeshApplyPose = new mu::NodeMeshApplyPose;
-	NodeMeshApplyPose->SetBase(InputMeshNode);
-	NodeMeshApplyPose->SetPose(PoseNodeMesh);
+	mu::Ptr<mu::NodeMeshApplyPose> NodeMeshApplyPose = new mu::NodeMeshApplyPose;
+	NodeMeshApplyPose->Base = InputMeshNode;
+	NodeMeshApplyPose->Pose = PoseNodeMesh;
 
 	return NodeMeshApplyPose;
 }
@@ -925,7 +925,7 @@ mu::Ptr<mu::NodeObject> GenerateMutableSource(const UEdGraphPin * Pin, FMutableG
 
 	else if (const UCustomizableObjectNodeObjectGroup* TypedNodeGroup = Cast<UCustomizableObjectNodeObjectGroup>(Node))
 	{
-		mu::NodeObjectGroupPtr GroupNode = new mu::NodeObjectGroup();
+		mu::Ptr<mu::NodeObjectGroup> GroupNode = new mu::NodeObjectGroup();
 		Result = GroupNode;
 
 		// All sockets from all mesh parts plugged into this group node will have the following priority when there's a socket name clash
@@ -953,7 +953,7 @@ mu::Ptr<mu::NodeObject> GenerateMutableSource(const UEdGraphPin * Pin, FMutableG
 			GenerationContext.CurrentGroupProjectors.Pop();
 		};
 
-		mu::NodeObjectGroup::CHILD_SELECTION Type = mu::NodeObjectGroup::CS_ALWAYS_ALL;
+		mu::NodeObjectGroup::EChildSelection Type = mu::NodeObjectGroup::CS_ALWAYS_ALL;
 		switch (TypedNodeGroup->GroupType)
 		{
 		case ECustomizableObjectGroupType::COGT_ALL: Type = mu::NodeObjectGroup::CS_ALWAYS_ALL; break;
@@ -964,7 +964,7 @@ mu::Ptr<mu::NodeObject> GenerateMutableSource(const UEdGraphPin * Pin, FMutableG
 			GenerationContext.Log(LOCTEXT("UnsupportedGroupType", "Object Group Type not supported. Setting to 'ALL'."), Node);
 			break;
 		}
-		GroupNode->SetSelectionType(Type);
+		GroupNode->Type = Type;
 
 		// External children
 		TArray<UCustomizableObjectNodeObject*> ExternalChildNodes;
@@ -976,8 +976,8 @@ mu::Ptr<mu::NodeObject> GenerateMutableSource(const UEdGraphPin * Pin, FMutableG
 		const int32 NumChildren = ConnectedChildrenPins.Num();
 		const int32 TotalNumChildren = NumChildren + ExternalChildNodes.Num();
 
-		GroupNode->SetChildCount(TotalNumChildren);
-		GroupNode->SetDefaultValue(Type == mu::NodeObjectGroup::CS_ONE_OR_NONE ? -1 : 0);
+		GroupNode->Children.SetNum(TotalNumChildren);
+		GroupNode->DefaultValue = (Type == mu::NodeObjectGroup::CS_ONE_OR_NONE ? -1 : 0);
 		int32 ChildIndex = 0;
 
 		// UI data
@@ -998,14 +998,14 @@ mu::Ptr<mu::NodeObject> GenerateMutableSource(const UEdGraphPin * Pin, FMutableG
 			UCustomizableObjectNodeObject* CustomizableObjectNodeObject = Cast<UCustomizableObjectNodeObject>(ConnectedChildrenPins[ChildIndex]->GetOwningNode());
 
 			const FString* SelectedOptionName = GenerationContext.CompilationContext.Options.ParamNamesToSelectedOptions.Find(TypedNodeGroup->GroupName); // If the param is in the map restrict to only the selected option
-			mu::NodeObjectPtr ChildNode = nullptr;
+			mu::Ptr<mu::NodeObject> ChildNode;
 
 			if (bConnectAtLeastTheLastChild || !SelectedOptionName || (CustomizableObjectNodeObject && *SelectedOptionName == CustomizableObjectNodeObject->ObjectName) )
 			{
 				bAtLeastOneConnected = true;
 
 				ChildNode = GenerateMutableSource(ConnectedChildrenPins[ChildIndex], GenerationContext);
-				GroupNode->SetChild(ChildIndex, ChildNode.get());
+				GroupNode->Children[ChildIndex] = ChildNode;
 
 				if (CustomizableObjectNodeObject)
 				{
@@ -1030,14 +1030,14 @@ mu::Ptr<mu::NodeObject> GenerateMutableSource(const UEdGraphPin * Pin, FMutableG
 			{
 				ChildNode = new mu::NodeObjectNew;
 				ChildNode->SetName(CustomizableObjectNodeObject->ObjectName);
-				GroupNode->SetChild(ChildIndex, ChildNode.get());
+				GroupNode->Children[ChildIndex] = ChildNode;
 			}
 
 			if ((TypedNodeGroup->GroupType == ECustomizableObjectGroupType::COGT_ONE ||
 				TypedNodeGroup->GroupType == ECustomizableObjectGroupType::COGT_ONE_OR_NONE)
 				&& TypedNodeGroup->DefaultValue == ChildNode->GetName())
 			{
-				GroupNode->SetDefaultValue(ChildIndex);
+				GroupNode->DefaultValue = ChildIndex;
 			}
 		}
 
@@ -1111,14 +1111,14 @@ mu::Ptr<mu::NodeObject> GenerateMutableSource(const UEdGraphPin * Pin, FMutableG
 			UCustomizableObjectNodeObject* CustomizableObjectNodeObject = Cast<UCustomizableObjectNodeObject>(ExternalChildNode->OutputPin()->GetOwningNode());
 
 			const FString* SelectedOptionName = GenerationContext.CompilationContext.Options.ParamNamesToSelectedOptions.Find(TypedNodeGroup->GroupName); // If the param is in the map restrict to only the selected option
-			mu::NodeObjectPtr ChildNode = nullptr;
+			mu::Ptr<mu::NodeObject> ChildNode;
 
 			if (bConnectAtLeastTheLastChild || !SelectedOptionName || (CustomizableObjectNodeObject  && *SelectedOptionName == CustomizableObjectNodeObject->ObjectName) )
 			{
 				bAtLeastOneConnected = true;
 
 				ChildNode = GenerateMutableSource(ExternalChildNode->OutputPin(), GenerationContext);
-				GroupNode->SetChild(ChildIndex, ChildNode.get());
+				GroupNode->Children[ChildIndex] = ChildNode;
 
 				if (CustomizableObjectNodeObject)
 				{
@@ -1148,14 +1148,14 @@ mu::Ptr<mu::NodeObject> GenerateMutableSource(const UEdGraphPin * Pin, FMutableG
 			{
 				ChildNode = new mu::NodeObjectNew;
 				ChildNode->SetName(CustomizableObjectNodeObject->ObjectName);
-				GroupNode->SetChild(ChildIndex, ChildNode.get());
+				GroupNode->Children[ChildIndex] = ChildNode;
 			}
 
 			if ((TypedNodeGroup->GroupType == ECustomizableObjectGroupType::COGT_ONE ||
 				TypedNodeGroup->GroupType == ECustomizableObjectGroupType::COGT_ONE_OR_NONE)
 				&& TypedNodeGroup->DefaultValue == ChildNode->GetName())
 			{
-				GroupNode->SetDefaultValue(ChildIndex);
+				GroupNode->DefaultValue = ChildIndex;
 			}
 
 			ChildIndex++;
diff --git a/Engine/Plugins/Mutable/Source/CustomizableObjectEditor/Private/MuCOE/GenerateMutableSource/GenerateMutableSource.h b/Engine/Plugins/Mutable/Source/CustomizableObjectEditor/Private/MuCOE/GenerateMutableSource/GenerateMutableSource.h
--- a/Engine/Plugins/Mutable/Source/CustomizableObjectEditor/Private/MuCOE/GenerateMutableSource/GenerateMutableSource.h
+++ b/Engine/Plugins/Mutable/Source/CustomizableObjectEditor/Private/MuCOE/GenerateMutableSource/GenerateMutableSource.h
@@ -308,12 +308,12 @@ struct FClothingMeshData
 struct FGroupProjectorTempData
 {
 	class UCustomizableObjectNodeGroupProjectorParameter* CustomizableObjectNodeGroupProjectorParameter;
-	mu::NodeProjectorParameterPtr NodeProjectorParameterPtr;
-	mu::NodeImagePtr NodeImagePtr;
-	mu::NodeRangePtr NodeRange;
-	mu::NodeScalarParameterPtr NodeOpacityParameter;
+	mu::Ptr<mu::NodeProjectorParameter> NodeProjectorParameterPtr;
+	mu::Ptr<mu::NodeImage> NodeImagePtr;
+	mu::Ptr<mu::NodeRange> NodeRange;
+	mu::Ptr<mu::NodeScalarParameter> NodeOpacityParameter;
 
-	mu::NodeScalarEnumParameterPtr PoseOptionsParameter;
+	mu::Ptr<mu::NodeScalarEnumParameter> PoseOptionsParameter;
 	TArray<FPoseBoneData> PoseBoneDataArray;
 
 	bool bAlternateResStateNameWarningDisplayed = false; // Used to display this warning only once
@@ -943,7 +943,7 @@ void CheckNumOutputs(const UEdGraphPin& Pin, const FMutableGraphGenerationContex
 UTexture2D* FindReferenceImage(const UEdGraphPin* Pin, FMutableGraphGenerationContext& GenerationContext);
 
 
-mu::NodeMeshApplyPosePtr CreateNodeMeshApplyPose(FMutableGraphGenerationContext& GenerationContext, mu::NodeMeshPtr InputMeshNode, const TArray<FName>& ArrayBoneName, const TArray<FTransform>& ArrayTransform);
+mu::Ptr<mu::NodeMeshApplyPose> CreateNodeMeshApplyPose(FMutableGraphGenerationContext& GenerationContext, mu::NodeMeshPtr InputMeshNode, const TArray<FName>& ArrayBoneName, const TArray<FTransform>& ArrayTransform);
 
 
 /** Adds Tag to MutableMesh uniquely, returns the index were the tag has been inserted or the index where an intance of the tag has been found */
diff --git a/Engine/Plugins/Mutable/Source/CustomizableObjectEditor/Private/MuCOE/GenerateMutableSource/GenerateMutableSourceColor.cpp b/Engine/Plugins/Mutable/Source/CustomizableObjectEditor/Private/MuCOE/GenerateMutableSource/GenerateMutableSourceColor.cpp
--- a/Engine/Plugins/Mutable/Source/CustomizableObjectEditor/Private/MuCOE/GenerateMutableSource/GenerateMutableSourceColor.cpp
+++ b/Engine/Plugins/Mutable/Source/CustomizableObjectEditor/Private/MuCOE/GenerateMutableSource/GenerateMutableSourceColor.cpp
@@ -92,7 +92,7 @@ mu::Ptr<mu::NodeColour> GenerateMutableSourceColor(const UEdGraphPin* Pin, FMuta
 			}
 
 			const UEdGraphPin* EnumPin = FollowInputPin(*TypedNodeColorSwitch->SwitchParameter());
-			mu::NodeScalarPtr SwitchParam = EnumPin ? GenerateMutableSourceFloat(EnumPin, GenerationContext) : nullptr;
+			mu::Ptr<mu::NodeScalar> SwitchParam = EnumPin ? GenerateMutableSourceFloat(EnumPin, GenerationContext) : nullptr;
 
 			// Switch Param not generated
 			if (!SwitchParam)
@@ -118,7 +118,7 @@ mu::Ptr<mu::NodeColour> GenerateMutableSourceColor(const UEdGraphPin* Pin, FMuta
 			const int32 NumSwitchOptions = TypedNodeColorSwitch->GetNumElements();
 
 			mu::NodeScalarEnumParameter* EnumParameter = static_cast<mu::NodeScalarEnumParameter*>(SwitchParam.get());
-			if (NumSwitchOptions != EnumParameter->GetValueCount())
+			if (NumSwitchOptions != EnumParameter->Options.Num())
 			{
 				const FText Message = LOCTEXT("MismatchedSwitch", "Switch enum and switch node have different number of options. Please refresh the switch node to make sure the outcomes are labeled properly.");
 				GenerationContext.Log(Message, Node);
@@ -154,13 +154,13 @@ mu::Ptr<mu::NodeColour> GenerateMutableSourceColor(const UEdGraphPin* Pin, FMuta
 
 		if (const UEdGraphPin* ConnectedPin = FollowInputPin(*TypedNodeTexSample->XPin()))
 		{
-			mu::NodeScalarPtr XNode = GenerateMutableSourceFloat(ConnectedPin, GenerationContext);
+			mu::Ptr<mu::NodeScalar> XNode = GenerateMutableSourceFloat(ConnectedPin, GenerationContext);
 			ColorNode->X = XNode;
 		}
 
 		if (const UEdGraphPin* ConnectedPin = FollowInputPin(*TypedNodeTexSample->YPin()))
 		{
-			mu::NodeScalarPtr YNode = GenerateMutableSourceFloat(ConnectedPin, GenerationContext);
+			mu::Ptr<mu::NodeScalar> YNode = GenerateMutableSourceFloat(ConnectedPin, GenerationContext);
 			ColorNode->Y = YNode;
 		}
 	}
@@ -207,25 +207,25 @@ mu::Ptr<mu::NodeColour> GenerateMutableSourceColor(const UEdGraphPin* Pin, FMuta
 
 		if (const UEdGraphPin* ConnectedPin = FollowInputPin(*TypedNodeFrom->RPin()))
 		{
-			mu::NodeScalarPtr FloatNode = GenerateMutableSourceFloat(ConnectedPin, GenerationContext);
+			mu::Ptr<mu::NodeScalar> FloatNode = GenerateMutableSourceFloat(ConnectedPin, GenerationContext);
 			OpNode->X = FloatNode;
 		}
 
 		if (const UEdGraphPin* ConnectedPin = FollowInputPin(*TypedNodeFrom->GPin()))
 		{
-			mu::NodeScalarPtr FloatNode = GenerateMutableSourceFloat(ConnectedPin, GenerationContext);
+			mu::Ptr<mu::NodeScalar> FloatNode = GenerateMutableSourceFloat(ConnectedPin, GenerationContext);
 			OpNode->Y = FloatNode;
 		}
 
 		if (const UEdGraphPin* ConnectedPin = FollowInputPin(*TypedNodeFrom->BPin()))
 		{
-			mu::NodeScalarPtr FloatNode = GenerateMutableSourceFloat(ConnectedPin, GenerationContext);
+			mu::Ptr<mu::NodeScalar> FloatNode = GenerateMutableSourceFloat(ConnectedPin, GenerationContext);
 			OpNode->Z = FloatNode;
 		}
 
 		if (const UEdGraphPin* ConnectedPin = FollowInputPin(*TypedNodeFrom->APin()))
 		{
-			mu::NodeScalarPtr FloatNode = GenerateMutableSourceFloat(ConnectedPin, GenerationContext);
+			mu::Ptr<mu::NodeScalar> FloatNode = GenerateMutableSourceFloat(ConnectedPin, GenerationContext);
 			OpNode->W = FloatNode;
 		}
 	}
diff --git a/Engine/Plugins/Mutable/Source/CustomizableObjectEditor/Private/MuCOE/GenerateMutableSource/GenerateMutableSourceComponent.cpp b/Engine/Plugins/Mutable/Source/CustomizableObjectEditor/Private/MuCOE/GenerateMutableSource/GenerateMutableSourceComponent.cpp
--- a/Engine/Plugins/Mutable/Source/CustomizableObjectEditor/Private/MuCOE/GenerateMutableSource/GenerateMutableSourceComponent.cpp
+++ b/Engine/Plugins/Mutable/Source/CustomizableObjectEditor/Private/MuCOE/GenerateMutableSource/GenerateMutableSourceComponent.cpp
@@ -154,7 +154,7 @@ mu::Ptr<mu::NodeComponent> GenerateMutableSourceComponent(const UEdGraphPin* Pin
 		else if(UMaterialInterface* OverlayMaterial = GenerationContext.LoadObject(TypedComponentMesh->GetOverlayMaterial()))
 		{
 			mu::Ptr<mu::NodeScalarConstant> ScalarNode = new mu::NodeScalarConstant();
-			ScalarNode->SetValue(GenerationContext.ReferencedMaterials.AddUnique(OverlayMaterial));
+			ScalarNode->Value = GenerationContext.ReferencedMaterials.AddUnique(OverlayMaterial);
 			NodeComponentNew->OverlayMaterial = ScalarNode;
 		}
 
@@ -187,7 +187,7 @@ mu::Ptr<mu::NodeComponent> GenerateMutableSourceComponent(const UEdGraphPin* Pin
 			mu::Ptr<mu::NodeComponent> ParentNodeComponent = GenerateMutableSourceComponent(TypedParentComponentMesh->OutputPin.Get(), GenerationContext);
 
 			mu::Ptr<mu::NodeComponentEdit> NodeComponentEdit = new mu::NodeComponentEdit();
-			NodeComponentEdit->SetParent(ParentNodeComponent.get());
+			NodeComponentEdit->Parent = ParentNodeComponent.get();
 			NodeComponentEdit->SetMessageContext(TypedComponentMeshExtend);
 		
 			GenerationContext.CurrentMeshComponent = TypedParentComponentMesh->ComponentName;
@@ -246,7 +246,7 @@ mu::Ptr<mu::NodeComponent> GenerateMutableSourceComponent(const UEdGraphPin* Pin
 			TSharedPtr<mu::FMesh> MutableMesh = GenerateMutableSkeletalMesh(SkeletalMesh, AnimInstance, 0, 0, 0, 0, MeshUniqueTags, 0, 
 																GenerationContext, TypedComponentPassthroughMesh, nullptr, bIsPassthrough);
 
-			MeshNode->SetValue(MutableMesh);
+			MeshNode->Value = MutableMesh;
 		}
 
 		// Create the component node
@@ -333,7 +333,7 @@ mu::Ptr<mu::NodeComponent> GenerateMutableSourceComponent(const UEdGraphPin* Pin
 					const int32 NumSwitchOptions = TypedNodeSwitch->GetNumElements();
 
 					mu::NodeScalarEnumParameter* EnumParameter = static_cast<mu::NodeScalarEnumParameter*>(SwitchParam.get());
-					if (NumSwitchOptions != EnumParameter->GetValueCount())
+					if (NumSwitchOptions != EnumParameter->Options.Num())
 					{
 						const FText Message = LOCTEXT("MismatchedSwitch", "Switch enum and switch node have different number of options. Please refresh the switch node to make sure the outcomes are labeled properly.");
 						GenerationContext.Log(Message, Node);
diff --git a/Engine/Plugins/Mutable/Source/CustomizableObjectEditor/Private/MuCOE/GenerateMutableSource/GenerateMutableSourceFloat.cpp b/Engine/Plugins/Mutable/Source/CustomizableObjectEditor/Private/MuCOE/GenerateMutableSource/GenerateMutableSourceFloat.cpp
--- a/Engine/Plugins/Mutable/Source/CustomizableObjectEditor/Private/MuCOE/GenerateMutableSource/GenerateMutableSourceFloat.cpp
+++ b/Engine/Plugins/Mutable/Source/CustomizableObjectEditor/Private/MuCOE/GenerateMutableSource/GenerateMutableSourceFloat.cpp
@@ -29,7 +29,7 @@
 
 
 
-mu::NodeScalarPtr GenerateMutableSourceFloat(const UEdGraphPin* Pin, FMutableGraphGenerationContext& GenerationContext)
+mu::Ptr<mu::NodeScalar> GenerateMutableSourceFloat(const UEdGraphPin* Pin, FMutableGraphGenerationContext& GenerationContext)
 {
 	check(Pin)
 	RETURN_ON_CYCLE(*Pin, GenerationContext)
@@ -53,24 +53,24 @@ mu::NodeScalarPtr GenerateMutableSourceFloat(const UEdGraphPin* Pin, FMutableGra
 
 	bool bDoNotAddToGeneratedCache = false;
 
-	mu::NodeScalarPtr Result;
+	mu::Ptr<mu::NodeScalar> Result;
 	
 	if (const UCustomizableObjectNodeFloatConstant* FloatConstantNode = Cast<UCustomizableObjectNodeFloatConstant>(Node))
 	{
-		mu::NodeScalarConstantPtr ScalarNode = new mu::NodeScalarConstant();
+		mu::Ptr<mu::NodeScalarConstant> ScalarNode = new mu::NodeScalarConstant();
 		Result = ScalarNode;
 
-		ScalarNode->SetValue(FloatConstantNode->Value);
+		ScalarNode->Value = FloatConstantNode->Value;
 	}
 
 	else if (const UCustomizableObjectNodeFloatParameter* FloatParameterNode = Cast<UCustomizableObjectNodeFloatParameter>(Node))
 	{
-		mu::NodeScalarParameterPtr ScalarNode = new mu::NodeScalarParameter();
+		mu::Ptr<mu::NodeScalarParameter> ScalarNode = new mu::NodeScalarParameter();
 		Result = ScalarNode;
 
-		ScalarNode->SetName(FloatParameterNode->ParameterName);
-		ScalarNode->SetUid(GenerationContext.GetNodeIdUnique(Node).ToString());
-		ScalarNode->SetDefaultValue(FloatParameterNode->DefaultValue);
+		ScalarNode->Name = FloatParameterNode->ParameterName;
+		ScalarNode->UID = GenerationContext.GetNodeIdUnique(Node).ToString();
+		ScalarNode->DefaultValue = FloatParameterNode->DefaultValue;
 
 		GenerationContext.ParameterUIDataMap.Add(FloatParameterNode->ParameterName, FMutableParameterData(
 			FloatParameterNode->ParamUIMetadata,
@@ -79,23 +79,24 @@ mu::NodeScalarPtr GenerateMutableSourceFloat(const UEdGraphPin* Pin, FMutableGra
 
 	else if (const UCustomizableObjectNodeEnumParameter* EnumParamNode = Cast<UCustomizableObjectNodeEnumParameter>(Node))
 	{
-		mu::NodeScalarEnumParameterPtr EnumParameterNode = new mu::NodeScalarEnumParameter;
+		mu::Ptr<mu::NodeScalarEnumParameter> EnumParameterNode = new mu::NodeScalarEnumParameter;
 
 		const int32 NumSelectors = EnumParamNode->Values.Num();
 
 		int32 DefaultValue = FMath::Clamp(EnumParamNode->DefaultIndex, 0, NumSelectors - 1);
 
-		EnumParameterNode->SetName(EnumParamNode->ParameterName);
-		EnumParameterNode->SetUid(GenerationContext.GetNodeIdUnique(Node).ToString());
-		EnumParameterNode->SetValueCount(NumSelectors);
-		EnumParameterNode->SetDefaultValueIndex(DefaultValue);
+		EnumParameterNode->Name = EnumParamNode->ParameterName;
+		EnumParameterNode->UID = GenerationContext.GetNodeIdUnique(Node).ToString();
+		EnumParameterNode->Options.SetNum(NumSelectors);
+		EnumParameterNode->DefaultValue = DefaultValue;
 
 		FMutableParameterData ParameterUIData(EnumParamNode->ParamUIMetadata, EMutableParameterType::Int);
 		ParameterUIData.IntegerParameterGroupType = ECustomizableObjectGroupType::COGT_ONE;
 
 		for (int SelectorIndex = 0; SelectorIndex < NumSelectors; ++SelectorIndex)
 		{
-			EnumParameterNode->SetValue(SelectorIndex, (float)SelectorIndex, EnumParamNode->Values[SelectorIndex].Name);
+			EnumParameterNode->Options[SelectorIndex].Name = EnumParamNode->Values[SelectorIndex].Name;
+			EnumParameterNode->Options[SelectorIndex].Value = (float)SelectorIndex;
 
 			ParameterUIData.ArrayIntegerParameterOption.Add(
 				EnumParamNode->Values[SelectorIndex].Name,
@@ -127,7 +128,7 @@ mu::NodeScalarPtr GenerateMutableSourceFloat(const UEdGraphPin* Pin, FMutableGra
 			}
 
 			const UEdGraphPin* EnumPin = FollowInputPin(*SwitchParameter);
-			mu::NodeScalarPtr SwitchParam = GenerateMutableSourceFloat(EnumPin, GenerationContext);
+			mu::Ptr<mu::NodeScalar> SwitchParam = GenerateMutableSourceFloat(EnumPin, GenerationContext);
 
 			// Switch Param not generated
 			if (!SwitchParam)
@@ -153,7 +154,7 @@ mu::NodeScalarPtr GenerateMutableSourceFloat(const UEdGraphPin* Pin, FMutableGra
 			const int32 NumSwitchOptions = TypedNodeFloatSwitch->GetNumElements();
 
 			mu::NodeScalarEnumParameter* EnumParameter = static_cast<mu::NodeScalarEnumParameter*>(SwitchParam.get());
-			if (NumSwitchOptions != EnumParameter->GetValueCount())
+			if (NumSwitchOptions != EnumParameter->Options.Num())
 			{
 				const FText Message = LOCTEXT("MismatchedSwitch", "Switch enum and switch node have different number of options. Please refresh the switch node to make sure the outcomes are labeled properly.");
 				GenerationContext.Log(Message, Node);
@@ -227,15 +228,15 @@ mu::NodeScalarPtr GenerateMutableSourceFloat(const UEdGraphPin* Pin, FMutableGra
 
 	else if (const UCustomizableObjectNodeFloatVariation* TypedNodeFloatVar = Cast<const UCustomizableObjectNodeFloatVariation>(Node))
 	{
-		mu::NodeScalarVariationPtr FloatNode = new mu::NodeScalarVariation();
+		mu::Ptr<mu::NodeScalarVariation> FloatNode = new mu::NodeScalarVariation();
 		Result = FloatNode;
 
 		if (const UEdGraphPin* ConnectedPin = FollowInputPin(*TypedNodeFloatVar->DefaultPin()))
 		{
-			mu::NodeScalarPtr ChildNode = GenerateMutableSourceFloat(ConnectedPin, GenerationContext);
+			mu::Ptr<mu::NodeScalar> ChildNode = GenerateMutableSourceFloat(ConnectedPin, GenerationContext);
 			if (ChildNode)
 			{
-				FloatNode->SetDefaultScalar(ChildNode.get());
+				FloatNode->DefaultScalar = ChildNode;
 			}
 			else
 			{
@@ -244,54 +245,54 @@ mu::NodeScalarPtr GenerateMutableSourceFloat(const UEdGraphPin* Pin, FMutableGra
 		}
 
 		const int32 NumVariations = TypedNodeFloatVar->GetNumVariations();
-		FloatNode->SetVariationCount(NumVariations);
+		FloatNode->Variations.SetNum(NumVariations);
 		for (int VariationIndex = 0; VariationIndex < NumVariations; ++VariationIndex)
 		{
 			UEdGraphPin* VariationPin = TypedNodeFloatVar->VariationPin(VariationIndex);
 			if (!VariationPin) continue;
 
-			FloatNode->SetVariationTag(VariationIndex, TypedNodeFloatVar->GetVariation(VariationIndex).Tag);
+			FloatNode->Variations[VariationIndex].Tag = TypedNodeFloatVar->GetVariation(VariationIndex).Tag;
 			if (const UEdGraphPin* ConnectedPin = FollowInputPin(*VariationPin))
 			{
-				mu::NodeScalarPtr ChildNode = GenerateMutableSourceFloat(ConnectedPin, GenerationContext);
-				FloatNode->SetVariationScalar(VariationIndex, ChildNode.get());
+				mu::Ptr<mu::NodeScalar> ChildNode = GenerateMutableSourceFloat(ConnectedPin, GenerationContext);
+				FloatNode->Variations[VariationIndex].Scalar = ChildNode;
 			}
 		}
 	}
 
 	else if (const UCustomizableObjectNodeFloatArithmeticOp* TypedNodeFloatArith = Cast<UCustomizableObjectNodeFloatArithmeticOp>(Node))
 	{
-		mu::NodeScalarArithmeticOperationPtr OpNode = new mu::NodeScalarArithmeticOperation();
+		mu::Ptr<mu::NodeScalarArithmeticOperation> OpNode = new mu::NodeScalarArithmeticOperation();
 		Result = OpNode;
 
 		if (const UEdGraphPin* ConnectedPin = FollowInputPin(*TypedNodeFloatArith->XPin()))
 		{
-			mu::NodeScalarPtr XNode = GenerateMutableSourceFloat(ConnectedPin, GenerationContext);
-			OpNode->SetA(XNode);
+			mu::Ptr<mu::NodeScalar> XNode = GenerateMutableSourceFloat(ConnectedPin, GenerationContext);
+			OpNode->A = XNode;
 		}
 
 		if (const UEdGraphPin* ConnectedPin = FollowInputPin(*TypedNodeFloatArith->YPin()))
 		{
-			mu::NodeScalarPtr YNode = GenerateMutableSourceFloat(ConnectedPin, GenerationContext);
-			OpNode->SetB(YNode);
+			mu::Ptr<mu::NodeScalar> YNode = GenerateMutableSourceFloat(ConnectedPin, GenerationContext);
+			OpNode->B = YNode;
 		}
 
 		switch (TypedNodeFloatArith->Operation)
 		{
 		case EFloatArithmeticOperation::E_Add:
-			OpNode->SetOperation(mu::NodeScalarArithmeticOperation::OPERATION::AO_ADD);
+			OpNode->Operation = mu::NodeScalarArithmeticOperation::EOperation::AO_ADD;
 			break;
 
 		case EFloatArithmeticOperation::E_Sub:
-			OpNode->SetOperation(mu::NodeScalarArithmeticOperation::OPERATION::AO_SUBTRACT);
+			OpNode->Operation = mu::NodeScalarArithmeticOperation::EOperation::AO_SUBTRACT;
 			break;
 
 		case EFloatArithmeticOperation::E_Mul:
-			OpNode->SetOperation(mu::NodeScalarArithmeticOperation::OPERATION::AO_MULTIPLY);
+			OpNode->Operation = mu::NodeScalarArithmeticOperation::EOperation::AO_MULTIPLY;
 			break;
 
 		case EFloatArithmeticOperation::E_Div:
-			OpNode->SetOperation(mu::NodeScalarArithmeticOperation::OPERATION::AO_DIVIDE);
+			OpNode->Operation = mu::NodeScalarArithmeticOperation::EOperation::AO_DIVIDE;
 			break;
 
 		default:
@@ -302,8 +303,8 @@ mu::NodeScalarPtr GenerateMutableSourceFloat(const UEdGraphPin* Pin, FMutableGra
 	else if (const UCustomizableObjectNodeTable* TypedNodeTable = Cast<UCustomizableObjectNodeTable>(Node))
 	{
 		//This node will add a default value in case of error
-		mu::NodeScalarConstantPtr ConstantValue = new mu::NodeScalarConstant();
-		ConstantValue->SetValue(1.0f);
+		mu::Ptr<mu::NodeScalarConstant> ConstantValue = new mu::NodeScalarConstant();
+		ConstantValue->Value = 1.0f;
 
 		Result = ConstantValue;
 
@@ -337,7 +338,7 @@ mu::NodeScalarPtr GenerateMutableSourceFloat(const UEdGraphPin* Pin, FMutableGra
 
 				if (Table)
 				{
-					mu::NodeScalarTablePtr ScalarTableNode = new mu::NodeScalarTable();
+					mu::Ptr<mu::NodeScalarTable> ScalarTableNode = new mu::NodeScalarTable();
 
 					if (Pin->PinType.PinCategory == Schema->PC_MaterialAsset)
 					{
@@ -364,10 +365,10 @@ mu::NodeScalarPtr GenerateMutableSourceFloat(const UEdGraphPin* Pin, FMutableGra
 						Result = ScalarTableNode;
 
 						ScalarTableNode->Table = Table;
-						ScalarTableNode->SetColumn(ColumnName);
-						ScalarTableNode->SetParameterName(TypedNodeTable->ParameterName);
-						ScalarTableNode->SetNoneOption(TypedNodeTable->bAddNoneOption);
-						ScalarTableNode->SetDefaultRowName(TypedNodeTable->DefaultRowName.ToString());
+						ScalarTableNode->ColumnName = ColumnName;
+						ScalarTableNode->ParameterName = TypedNodeTable->ParameterName;
+						ScalarTableNode->bNoneOption = TypedNodeTable->bAddNoneOption;
+						ScalarTableNode->DefaultRowName = TypedNodeTable->DefaultRowName.ToString();
 					}
 				}
 				else
diff --git a/Engine/Plugins/Mutable/Source/CustomizableObjectEditor/Private/MuCOE/GenerateMutableSource/GenerateMutableSourceFloat.h b/Engine/Plugins/Mutable/Source/CustomizableObjectEditor/Private/MuCOE/GenerateMutableSource/GenerateMutableSourceFloat.h
--- a/Engine/Plugins/Mutable/Source/CustomizableObjectEditor/Private/MuCOE/GenerateMutableSource/GenerateMutableSourceFloat.h
+++ b/Engine/Plugins/Mutable/Source/CustomizableObjectEditor/Private/MuCOE/GenerateMutableSource/GenerateMutableSourceFloat.h
@@ -8,4 +8,4 @@ class UEdGraphPin;
 struct FMutableGraphGenerationContext;
 
 /** Convert a CustomizableObject Source Graph into a mutable source graph. */
-mu::NodeScalarPtr GenerateMutableSourceFloat(const UEdGraphPin* Pin, FMutableGraphGenerationContext& GenerationContext);
+mu::Ptr<mu::NodeScalar> GenerateMutableSourceFloat(const UEdGraphPin* Pin, FMutableGraphGenerationContext& GenerationContext);
diff --git a/Engine/Plugins/Mutable/Source/CustomizableObjectEditor/Private/MuCOE/GenerateMutableSource/GenerateMutableSourceGroupProjector.cpp b/Engine/Plugins/Mutable/Source/CustomizableObjectEditor/Private/MuCOE/GenerateMutableSource/GenerateMutableSourceGroupProjector.cpp
--- a/Engine/Plugins/Mutable/Source/CustomizableObjectEditor/Private/MuCOE/GenerateMutableSource/GenerateMutableSourceGroupProjector.cpp
+++ b/Engine/Plugins/Mutable/Source/CustomizableObjectEditor/Private/MuCOE/GenerateMutableSource/GenerateMutableSourceGroupProjector.cpp
@@ -129,7 +129,7 @@ mu::NodeImagePtr GenerateMutableSourceGroupProjector(const int32 NodeLOD, const
 				
 				mu::Ptr<mu::NodeImageProject> ImageNode = new mu::NodeImageProject();
 				bIsGroupProjectorImage = true;
-				ImageNode->SetLayout(ProjectorTempData.CustomizableObjectNodeGroupProjectorParameter->UVLayout);
+				ImageNode->Layout = ProjectorTempData.CustomizableObjectNodeGroupProjectorParameter->UVLayout;
 
 				{
 					if (TypedNodeMat &&
@@ -154,27 +154,27 @@ mu::NodeImagePtr GenerateMutableSourceGroupProjector(const int32 NodeLOD, const
 				}
 
 				{
-					mu::NodeScalarConstantPtr ScalarNode = new mu::NodeScalarConstant;
-					ScalarNode->SetValue(120.f);
-					ImageNode->SetAngleFadeStart(ScalarNode);
+					mu::Ptr<mu::NodeScalarConstant> ScalarNode = new mu::NodeScalarConstant;
+					ScalarNode->Value = 120.f;
+					ImageNode->AngleFadeStart = ScalarNode;
 				}
 
 				{
-					mu::NodeScalarConstantPtr ScalarNode = new mu::NodeScalarConstant;
-					ScalarNode->SetValue(150.f);
-					ImageNode->SetAngleFadeEnd(ScalarNode);
+					mu::Ptr<mu::NodeScalarConstant> ScalarNode = new mu::NodeScalarConstant;
+					ScalarNode->Value = 150.f;
+					ImageNode->AngleFadeEnd = ScalarNode;
 				}
 
-				mu::NodeMeshSwitchPtr MeshSwitchNode = new mu::NodeMeshSwitch;
-				MeshSwitchNode->SetParameter(ProjectorTempData.PoseOptionsParameter);
-				MeshSwitchNode->SetOptionCount(ProjectorTempData.CustomizableObjectNodeGroupProjectorParameter->OptionPoses.Num() + 1);
-				MeshSwitchNode->SetOption(0, MeshNode);
+				mu::Ptr<mu::NodeMeshSwitch> MeshSwitchNode = new mu::NodeMeshSwitch;
+				MeshSwitchNode->Parameter = ProjectorTempData.PoseOptionsParameter;
+				MeshSwitchNode->Options.SetNum(ProjectorTempData.CustomizableObjectNodeGroupProjectorParameter->OptionPoses.Num() + 1);
+				MeshSwitchNode->Options[0] = MeshNode;
 
 				for (int32 SelectorIndex = 0; SelectorIndex < ProjectorTempData.CustomizableObjectNodeGroupProjectorParameter->OptionPoses.Num(); ++SelectorIndex)
 				{
 					if (ProjectorTempData.CustomizableObjectNodeGroupProjectorParameter->OptionPoses[SelectorIndex].OptionPose)
 					{
-						mu::NodeMeshApplyPosePtr NodeMeshApplyPose = CreateNodeMeshApplyPose(GenerationContext, MeshNode,
+						mu::Ptr<mu::NodeMeshApplyPose> NodeMeshApplyPose = CreateNodeMeshApplyPose(GenerationContext, MeshNode,
 							ProjectorTempData.PoseBoneDataArray[SelectorIndex].ArrayBoneName, ProjectorTempData.PoseBoneDataArray[SelectorIndex].ArrayTransform);
 
 						if (!NodeMeshApplyPose)
@@ -183,20 +183,20 @@ mu::NodeImagePtr GenerateMutableSourceGroupProjector(const int32 NodeLOD, const
 							GenerationContext.Log(FText::FromString(msg), TypedNodeMat);
 						}
 
-						MeshSwitchNode->SetOption(SelectorIndex + 1, NodeMeshApplyPose);
+						MeshSwitchNode->Options[SelectorIndex + 1] = NodeMeshApplyPose;
 					}
 					else
 					{
-						MeshSwitchNode->SetOption(SelectorIndex + 1, MeshNode);
+						MeshSwitchNode->Options[SelectorIndex + 1] = MeshNode;
 					}
 				}
 
-				ImageNode->SetMesh(MeshSwitchNode);
-				ImageNode->SetProjector(ProjectorTempData.NodeProjectorParameterPtr);
-				ImageNode->SetImage(ProjectorTempData.NodeImagePtr);
+				ImageNode->Mesh = MeshSwitchNode;
+				ImageNode->Projector = ProjectorTempData.NodeProjectorParameterPtr;
+				ImageNode->Image = ProjectorTempData.NodeImagePtr;
 
 				TextureSize = ProjectorTempData.TextureSize;
-				ImageNode->SetImageSize( FUintVector2(TextureSize, TextureSize) );
+				ImageNode->ImageSize = FUintVector2(TextureSize, TextureSize);
 
 				ImageNodes.Add(ImageNode);
 				ImageNodes_ProjectorTempData.Add(ProjectorTempData);
@@ -212,17 +212,18 @@ mu::NodeImagePtr GenerateMutableSourceGroupProjector(const int32 NodeLOD, const
 	mu::Ptr<mu::NodeColourConstant> ZeroColorNode = new mu::NodeColourConstant();
 	ZeroColorNode->Value = FVector4f(0.f, 0.f, 0.f, 1.0f);
 
-	mu::NodeImagePlainColourPtr ZeroPlainColourNode = new mu::NodeImagePlainColour;
-	ZeroPlainColourNode->SetSize(TextureSize, TextureSize);
-	ZeroPlainColourNode->SetColour(ZeroColorNode);
+	mu::Ptr<mu::NodeImagePlainColour> ZeroPlainColourNode = new mu::NodeImagePlainColour;
+	ZeroPlainColourNode->SizeX = TextureSize;
+	ZeroPlainColourNode->SizeY = TextureSize;
+	ZeroPlainColourNode->Colour = ZeroColorNode;
 
-	mu::NodeImageSwizzlePtr ZeroChannelNode = new mu::NodeImageSwizzle;
-	ZeroChannelNode->SetFormat(mu::EImageFormat::L_UByte);
-	ZeroChannelNode->SetSource(0, ZeroPlainColourNode);
-	ZeroChannelNode->SetSourceChannel(0, 2); // Just take a zeroed channel for the base alpha
+	mu::Ptr<mu::NodeImageSwizzle> ZeroChannelNode = new mu::NodeImageSwizzle;
+	ZeroChannelNode->SetFormat( mu::EImageFormat::L_UByte );
+	ZeroChannelNode->Sources[0] = ZeroPlainColourNode;
+	ZeroChannelNode->SourceChannels[0] = 2; // Just take a zeroed channel for the base alpha
 
-	mu::NodeScalarConstantPtr OneConstantNode = new mu::NodeScalarConstant;
-	OneConstantNode->SetValue(1.f);
+	mu::Ptr<mu::NodeScalarConstant> OneConstantNode = new mu::NodeScalarConstant;
+	OneConstantNode->Value = 1.f;
 
 	mu::NodeImagePtr ResultAlpha = ZeroChannelNode;
 	mu::NodeImagePtr ResultImage = ZeroPlainColourNode;
@@ -243,15 +244,16 @@ mu::NodeImagePtr GenerateMutableSourceGroupProjector(const int32 NodeLOD, const
 				TextureSize = NewTextureSize;
 
 				ZeroPlainColourNode = new mu::NodeImagePlainColour;
-				ZeroPlainColourNode->SetSize(TextureSize, TextureSize);
-				ZeroPlainColourNode->SetColour(ZeroColorNode);
+				ZeroPlainColourNode->SizeX = TextureSize;
+				ZeroPlainColourNode->SizeY = TextureSize;
+				ZeroPlainColourNode->Colour = ZeroColorNode;
 			}
 		}
 
-		mu::NodeImageSwizzlePtr ImageNodesAlphaChannelNode = new mu::NodeImageSwizzle;
-		ImageNodesAlphaChannelNode->SetFormat(mu::EImageFormat::L_UByte);
-		ImageNodesAlphaChannelNode->SetSource(0, ImageNodes[i]);
-		ImageNodesAlphaChannelNode->SetSourceChannel(0, 3);
+		mu::Ptr<mu::NodeImageSwizzle> ImageNodesAlphaChannelNode = new mu::NodeImageSwizzle;
+		ImageNodesAlphaChannelNode->SetFormat( mu::EImageFormat::L_UByte );
+		ImageNodesAlphaChannelNode->Sources[0] = ImageNodes[i];
+		ImageNodesAlphaChannelNode->SourceChannels[0] = 3;
 
 		mu::Ptr<mu::NodeColourFromScalars> ColourFromScalars = new mu::NodeColourFromScalars;
 		ColourFromScalars->X = ImageNodes_ProjectorTempData[i].NodeOpacityParameter;
@@ -259,59 +261,59 @@ mu::NodeImagePtr GenerateMutableSourceGroupProjector(const int32 NodeLOD, const
 		ColourFromScalars->Z = ImageNodes_ProjectorTempData[i].NodeOpacityParameter;
 		ColourFromScalars->W = OneConstantNode;
 
-		mu::NodeImageLayerColourPtr OpacityMultiLayerNode = new mu::NodeImageLayerColour;
-		OpacityMultiLayerNode->SetBlendType(mu::EBlendType::BT_MULTIPLY);
-		OpacityMultiLayerNode->SetColour(ColourFromScalars);
-		OpacityMultiLayerNode->SetBase(ImageNodesAlphaChannelNode);
+		mu::Ptr<mu::NodeImageLayerColour> OpacityMultiLayerNode = new mu::NodeImageLayerColour;
+		OpacityMultiLayerNode->Type = mu::EBlendType::BT_MULTIPLY;
+		OpacityMultiLayerNode->Colour = ColourFromScalars;
+		OpacityMultiLayerNode->Base = ImageNodesAlphaChannelNode;
 		//OpacityMultiLayerNode->SetMask(OneChannelNode); // No mask needed
 
-		mu::NodeImageSwizzlePtr MultiplySwizzleNode = new mu::NodeImageSwizzle;
-		MultiplySwizzleNode->SetFormat(mu::EImageFormat::L_UByte);
-		MultiplySwizzleNode->SetSource(0, OpacityMultiLayerNode);
-		MultiplySwizzleNode->SetSourceChannel(0, 0);
+		mu::Ptr<mu::NodeImageSwizzle> MultiplySwizzleNode = new mu::NodeImageSwizzle;
+		MultiplySwizzleNode->SetFormat( mu::EImageFormat::L_UByte );
+		MultiplySwizzleNode->Sources[0] = OpacityMultiLayerNode;
+		MultiplySwizzleNode->SourceChannels[0] = 0;
 
-		mu::NodeImageMultiLayerPtr BaseAlphaMultiLayerNode = new mu::NodeImageMultiLayer;
-		BaseAlphaMultiLayerNode->SetRange(ImageNodes_ProjectorTempData[i].NodeRange);
-		BaseAlphaMultiLayerNode->SetBlendType(mu::EBlendType::BT_LIGHTEN);
-		BaseAlphaMultiLayerNode->SetBase(ResultAlpha);
-		BaseAlphaMultiLayerNode->SetBlended(MultiplySwizzleNode);
+		mu::Ptr<mu::NodeImageMultiLayer> BaseAlphaMultiLayerNode = new mu::NodeImageMultiLayer;
+		BaseAlphaMultiLayerNode->Range = ImageNodes_ProjectorTempData[i].NodeRange;
+		BaseAlphaMultiLayerNode->Type = mu::EBlendType::BT_LIGHTEN;
+		BaseAlphaMultiLayerNode->Base = ResultAlpha;
+		BaseAlphaMultiLayerNode->Blended = MultiplySwizzleNode;
 		//BaseAlphaMultiLayerNode->SetMask(MultiplySwizzleNode); // No mask needed
 		ResultAlpha = BaseAlphaMultiLayerNode;
 
-		mu::NodeImageMultiLayerPtr BaseMultiLayerNode = new mu::NodeImageMultiLayer;
-		BaseMultiLayerNode->SetRange(ImageNodes_ProjectorTempData[i].NodeRange);
-		BaseMultiLayerNode->SetBlendType(mu::EBlendType::BT_BLEND);
-		BaseMultiLayerNode->SetBase(ResultImage);
-		BaseMultiLayerNode->SetBlended(ImageNodes[i]);
-		BaseMultiLayerNode->SetMask(MultiplySwizzleNode);
+		mu::Ptr<mu::NodeImageMultiLayer> BaseMultiLayerNode = new mu::NodeImageMultiLayer;
+		BaseMultiLayerNode->Range = ImageNodes_ProjectorTempData[i].NodeRange;
+		BaseMultiLayerNode->Type = mu::EBlendType::BT_BLEND;
+		BaseMultiLayerNode->Base = ResultImage;
+		BaseMultiLayerNode->Blended = ImageNodes[i];
+		BaseMultiLayerNode->Mask = MultiplySwizzleNode;
 		ResultImage = BaseMultiLayerNode;
 	}
 
-	mu::NodeImageSwizzlePtr SwizzleNodeR = new mu::NodeImageSwizzle;
+	mu::Ptr<mu::NodeImageSwizzle> SwizzleNodeR = new mu::NodeImageSwizzle;
 	SwizzleNodeR->SetFormat(mu::EImageFormat::L_UByte);
-	SwizzleNodeR->SetSource(0, ResultImage);
-	SwizzleNodeR->SetSourceChannel(0, 0);
+	SwizzleNodeR->Sources[0] = ResultImage;
+	SwizzleNodeR->SourceChannels[0] = 0;
 
-	mu::NodeImageSwizzlePtr SwizzleNodeG = new mu::NodeImageSwizzle;
+	mu::Ptr<mu::NodeImageSwizzle> SwizzleNodeG = new mu::NodeImageSwizzle;
 	SwizzleNodeG->SetFormat(mu::EImageFormat::L_UByte);
-	SwizzleNodeG->SetSource(0, ResultImage);
-	SwizzleNodeG->SetSourceChannel(0, 1);
+	SwizzleNodeG->Sources[0] = ResultImage;
+	SwizzleNodeG->SourceChannels[0] = 1;
 
-	mu::NodeImageSwizzlePtr SwizzleNodeB = new mu::NodeImageSwizzle;
+	mu::Ptr<mu::NodeImageSwizzle> SwizzleNodeB = new mu::NodeImageSwizzle;
 	SwizzleNodeB->SetFormat(mu::EImageFormat::L_UByte);
-	SwizzleNodeB->SetSource(0, ResultImage);
-	SwizzleNodeB->SetSourceChannel(0, 2);
+	SwizzleNodeB->Sources[0] = ResultImage;
+	SwizzleNodeB->SourceChannels[0] = 2;
 
-	mu::NodeImageSwizzlePtr FinalSwizzleNode = new mu::NodeImageSwizzle;
+	mu::Ptr<mu::NodeImageSwizzle> FinalSwizzleNode = new mu::NodeImageSwizzle;
 	FinalSwizzleNode->SetFormat(mu::EImageFormat::RGBA_UByte);
-	FinalSwizzleNode->SetSource(0, SwizzleNodeR);
-	FinalSwizzleNode->SetSourceChannel(0, 0);
-	FinalSwizzleNode->SetSource(1, SwizzleNodeG);
-	FinalSwizzleNode->SetSourceChannel(1, 0);
-	FinalSwizzleNode->SetSource(2, SwizzleNodeB);
-	FinalSwizzleNode->SetSourceChannel(2, 0);
-	FinalSwizzleNode->SetSource(3, ResultAlpha);
-	FinalSwizzleNode->SetSourceChannel(3, 0);
+	FinalSwizzleNode->Sources[0] = SwizzleNodeR;
+	FinalSwizzleNode->SourceChannels[0] = 0;
+	FinalSwizzleNode->Sources[1] = SwizzleNodeG;
+	FinalSwizzleNode->SourceChannels[1] = 0;
+	FinalSwizzleNode->Sources[2] = SwizzleNodeB;
+	FinalSwizzleNode->SourceChannels[2] = 0;
+	FinalSwizzleNode->Sources[3] = ResultAlpha;
+	FinalSwizzleNode->SourceChannels[3] = 0;
 
 	return FinalSwizzleNode;
 }
@@ -360,7 +362,7 @@ TOptional<FGroupProjectorTempData> GenerateMutableGroupProjector(UCustomizableOb
 
 		GroupProjectorTempData.CustomizableObjectNodeGroupProjectorParameter = ProjParamNode;
 
-		mu::NodeScalarParameterPtr NodeScalarParam;
+		mu::Ptr<mu::NodeScalarParameter> NodeScalarParam;
 		if (mu::Ptr<mu::NodeScalarParameter>* FindResult = GenerationContext.GeneratedScalarParameters.Find(NumLayersParamUid.ToString()))
 		{
 			NodeScalarParam = *FindResult;
@@ -369,9 +371,9 @@ TOptional<FGroupProjectorTempData> GenerateMutableGroupProjector(UCustomizableOb
 		{
 			NodeScalarParam = new mu::NodeScalarParameter;
 			FString NodeScalarParamName = ProjParamNode->ParameterName + NUM_LAYERS_PARAMETER_POSTFIX;
-			NodeScalarParam->SetName(NodeScalarParamName);
-			NodeScalarParam->SetUid(NumLayersParamUid.ToString());
-			NodeScalarParam->SetDefaultValue(0.f);
+			NodeScalarParam->Name = NodeScalarParamName;
+			NodeScalarParam->UID = NumLayersParamUid.ToString();
+			NodeScalarParam->DefaultValue =0.f;
 
 			GenerationContext.ParameterUIDataMap.Add(NodeScalarParamName,
 		FMutableParameterData(ProjParamNode->ParamUIMetadata, EMutableParameterType::Int));
@@ -380,12 +382,12 @@ TOptional<FGroupProjectorTempData> GenerateMutableGroupProjector(UCustomizableOb
 		}
 
 		mu::Ptr<mu::NodeRangeFromScalar> NodeRangeFromScalar = new mu::NodeRangeFromScalar;
-		NodeRangeFromScalar->SetSize(NodeScalarParam);
+		NodeRangeFromScalar->Size = NodeScalarParam;
 		GroupProjectorTempData.NodeRange = NodeRangeFromScalar;
 		GroupProjectorTempData.NodeProjectorParameterPtr->SetRangeCount(1);
 		GroupProjectorTempData.NodeProjectorParameterPtr->SetRange(0, NodeRangeFromScalar);
 		
-		mu::NodeScalarEnumParameterPtr PoseEnumParameterNode;
+		mu::Ptr<mu::NodeScalarEnumParameter> PoseEnumParameterNode;
 		if (mu::Ptr<mu::NodeScalarEnumParameter>* FindResult = GenerationContext.GeneratedEnumParameters.Find(SelectedPoseParamUid.ToString()))
 		{
 			PoseEnumParameterNode = *FindResult;
@@ -394,10 +396,10 @@ TOptional<FGroupProjectorTempData> GenerateMutableGroupProjector(UCustomizableOb
 		{
 			PoseEnumParameterNode = new mu::NodeScalarEnumParameter;
 			FString PoseNodeEnumParamName = ProjParamNode->ParameterName + POSE_PARAMETER_POSTFIX;
-			PoseEnumParameterNode->SetName(PoseNodeEnumParamName);
-			PoseEnumParameterNode->SetUid(SelectedPoseParamUid.ToString());
-			PoseEnumParameterNode->SetValueCount(ProjParamNode->OptionPoses.Num() + 1);
-			PoseEnumParameterNode->SetDefaultValueIndex(0);
+			PoseEnumParameterNode->Name = PoseNodeEnumParamName;
+			PoseEnumParameterNode->UID = SelectedPoseParamUid.ToString();
+			PoseEnumParameterNode->Options.SetNum( ProjParamNode->OptionPoses.Num() + 1 );
+			PoseEnumParameterNode->DefaultValue = 0;
 			GroupProjectorTempData.PoseOptionsParameter = PoseEnumParameterNode;
 
 			GenerationContext.GeneratedEnumParameters.Add(SelectedPoseParamUid.ToString(), PoseEnumParameterNode);
@@ -406,7 +408,7 @@ TOptional<FGroupProjectorTempData> GenerateMutableGroupProjector(UCustomizableOb
 		FMutableParameterData(ProjParamNode->ParamUIMetadata, EMutableParameterType::Int));
 		}
 		
-		mu::NodeScalarParameterPtr OpacityParameterNode;
+		mu::Ptr<mu::NodeScalarParameter> OpacityParameterNode;
 		if (mu::Ptr<mu::NodeScalarParameter>* FindResult = GenerationContext.GeneratedScalarParameters.Find(OpacityParamUid.ToString()))
 		{
 			OpacityParameterNode = *FindResult;
@@ -415,11 +417,11 @@ TOptional<FGroupProjectorTempData> GenerateMutableGroupProjector(UCustomizableOb
 		{
 			OpacityParameterNode = new mu::NodeScalarParameter;
 			FString OpacityParameterNodeName = ProjParamNode->ParameterName + OPACITY_PARAMETER_POSTFIX;
-			OpacityParameterNode->SetName(OpacityParameterNodeName);
-			OpacityParameterNode->SetUid(OpacityParamUid.ToString());
-			OpacityParameterNode->SetDefaultValue(0.75f);
-			OpacityParameterNode->SetRangeCount(1);
-			OpacityParameterNode->SetRange(0, NodeRangeFromScalar);
+			OpacityParameterNode->Name = OpacityParameterNodeName;
+			OpacityParameterNode->UID = OpacityParamUid.ToString();
+			OpacityParameterNode->DefaultValue = 0.75f;
+			OpacityParameterNode->Ranges.SetNum(1);
+			OpacityParameterNode->Ranges[0] = NodeRangeFromScalar;
 
 			GenerationContext.GeneratedScalarParameters.Add(OpacityParamUid.ToString(), OpacityParameterNode);
 					
@@ -444,11 +446,13 @@ TOptional<FGroupProjectorTempData> GenerateMutableGroupProjector(UCustomizableOb
 		FMutableComponentInfo* CurrentComponentInfo = GenerationContext.GetCurrentComponentInfo();
 		check(CurrentComponentInfo);
 		
-		PoseEnumParameterNode->SetValue(0, 0.f, "Default pose");
+		PoseEnumParameterNode->Options[0].Value = 0.f;
+		PoseEnumParameterNode->Options[0].Name = "Default pose";
 
 		for (int32 PoseIndex = 0; PoseIndex < ProjParamNode->OptionPoses.Num(); ++PoseIndex)
 		{
-			PoseEnumParameterNode->SetValue(PoseIndex + 1, (float)PoseIndex + 1.f, ProjParamNode->OptionPoses[PoseIndex].PoseName);
+			PoseEnumParameterNode->Options[PoseIndex + 1].Value = (float)PoseIndex + 1.f;
+			PoseEnumParameterNode->Options[PoseIndex + 1].Name = ProjParamNode->OptionPoses[PoseIndex].PoseName;
 
 			TArray<FString> ArrayBoneName;
 			TArray<FTransform> ArrayTransform;
@@ -467,7 +471,7 @@ TOptional<FGroupProjectorTempData> GenerateMutableGroupProjector(UCustomizableOb
 				GroupProjectorTempData.PoseBoneDataArray[PoseIndex].ArrayBoneName, GroupProjectorTempData.PoseBoneDataArray[PoseIndex].ArrayTransform);
 		}
 	
-		mu::NodeScalarEnumParameterPtr EnumParameterNode;
+		mu::Ptr<mu::NodeScalarEnumParameter> EnumParameterNode;
 		if (mu::Ptr<mu::NodeScalarEnumParameter>* FindResult = GenerationContext.GeneratedEnumParameters.Find(SelectedImageParamUid.ToString()))
 		{
 			EnumParameterNode = *FindResult;
@@ -477,12 +481,12 @@ TOptional<FGroupProjectorTempData> GenerateMutableGroupProjector(UCustomizableOb
 			EnumParameterNode = new mu::NodeScalarEnumParameter;
 
 			FString NodeEnumParamName = ProjParamNode->ParameterName + IMAGE_PARAMETER_POSTFIX;
-			EnumParameterNode->SetName(NodeEnumParamName);
-			EnumParameterNode->SetUid(SelectedImageParamUid.ToString());
-			EnumParameterNode->SetValueCount(ArrayOptionTexture.Num());
-			EnumParameterNode->SetDefaultValueIndex(0);
-			EnumParameterNode->SetRangeCount(1);
-			EnumParameterNode->SetRange(0, NodeRangeFromScalar);
+			EnumParameterNode->Name = NodeEnumParamName;
+			EnumParameterNode->UID = SelectedImageParamUid.ToString();
+			EnumParameterNode->Options.SetNum(ArrayOptionTexture.Num());
+			EnumParameterNode->DefaultValue = 0;
+			EnumParameterNode->Ranges.SetNum(1);
+			EnumParameterNode->Ranges[0] = NodeRangeFromScalar;
 
 			GenerationContext.GeneratedEnumParameters.Add(SelectedImageParamUid.ToString(), EnumParameterNode);
 					
@@ -494,7 +498,8 @@ TOptional<FGroupProjectorTempData> GenerateMutableGroupProjector(UCustomizableOb
 
 			for (int OptionTextureIndex = 0; OptionTextureIndex < ArrayOptionTexture.Num(); ++OptionTextureIndex)
 			{
-				EnumParameterNode->SetValue(OptionTextureIndex, (float)OptionTextureIndex, ArrayOptionTexture[OptionTextureIndex].OptionName);
+				EnumParameterNode->Options[OptionTextureIndex].Value = (float)OptionTextureIndex;
+				EnumParameterNode->Options[OptionTextureIndex].Name = ArrayOptionTexture[OptionTextureIndex].OptionName;
 
 				FMutableParamUIMetadata OptionMetadata = ParameterUIData.ParamUIMetadata;
 				OptionMetadata.UIThumbnail = ArrayOptionTexture[OptionTextureIndex].OptionTexture;
@@ -506,9 +511,9 @@ TOptional<FGroupProjectorTempData> GenerateMutableGroupProjector(UCustomizableOb
 			GenerationContext.ParameterUIDataMap.Add(NodeEnumParamName, ParameterUIData);
 		}
 
-		mu::NodeImageSwitchPtr SwitchNode = new mu::NodeImageSwitch;
-		SwitchNode->SetParameter(EnumParameterNode);
-		SwitchNode->SetOptionCount(ArrayOptionTexture.Num());
+		mu::Ptr<mu::NodeImageSwitch> SwitchNode = new mu::NodeImageSwitch;
+		SwitchNode->Parameter = EnumParameterNode;
+		SwitchNode->Options.SetNum(ArrayOptionTexture.Num());
 
 		for (int32 SelectorIndex = 0; SelectorIndex < ArrayOptionTexture.Num(); ++SelectorIndex)
 		{
@@ -520,7 +525,7 @@ TOptional<FGroupProjectorTempData> GenerateMutableGroupProjector(UCustomizableOb
 				ImageNode->SetValue(ImageConstant);
 
 				const uint32 MipsToSkip = ComputeLODBiasForTexture(GenerationContext, *Texture, ProjParamNode->ReferenceTexture);
-				SwitchNode->SetOption(SelectorIndex, ResizeTextureByNumMips(ImageNode, MipsToSkip));
+				SwitchNode->Options[SelectorIndex] = ResizeTextureByNumMips(ImageNode, MipsToSkip);
 
 				// Calculate the number of mips to tag as high res for this image.
 				if (ProjParamNode->ReferenceTexture)
diff --git a/Engine/Plugins/Mutable/Source/CustomizableObjectEditor/Private/MuCOE/GenerateMutableSource/GenerateMutableSourceImage.cpp b/Engine/Plugins/Mutable/Source/CustomizableObjectEditor/Private/MuCOE/GenerateMutableSource/GenerateMutableSourceImage.cpp
--- a/Engine/Plugins/Mutable/Source/CustomizableObjectEditor/Private/MuCOE/GenerateMutableSource/GenerateMutableSourceImage.cpp
+++ b/Engine/Plugins/Mutable/Source/CustomizableObjectEditor/Private/MuCOE/GenerateMutableSource/GenerateMutableSourceImage.cpp
@@ -63,11 +63,12 @@ mu::Ptr<mu::NodeImage> ResizeTextureByNumMips(const mu::Ptr<mu::NodeImage>& Imag
 {
 	if (MipsToSkip > 0)
 	{
-		mu::NodeImageResizePtr ImageResize = new mu::NodeImageResize();
-		ImageResize->SetBase(ImageConstant.get());
-		ImageResize->SetRelative(true);
-		const float factor = FMath::Pow(0.5f, MipsToSkip);
-		ImageResize->SetSize(factor, factor);
+		mu::Ptr<mu::NodeImageResize> ImageResize = new mu::NodeImageResize();
+		ImageResize->Base = ImageConstant;
+		ImageResize->bRelative = true;
+		const float Factor = FMath::Pow(0.5f, MipsToSkip);
+		ImageResize->SizeX = Factor;
+		ImageResize->SizeY = Factor;
 		return ImageResize;
 	}
 
@@ -75,7 +76,7 @@ mu::Ptr<mu::NodeImage> ResizeTextureByNumMips(const mu::Ptr<mu::NodeImage>& Imag
 }
 
 
-mu::NodeImagePtr GenerateMutableSourceImage(const UEdGraphPin* Pin, FMutableGraphGenerationContext& GenerationContext, int32 ReferenceTextureSize)
+mu::Ptr<mu::NodeImage> GenerateMutableSourceImage(const UEdGraphPin* Pin, FMutableGraphGenerationContext& GenerationContext, int32 ReferenceTextureSize)
 {
 	check(Pin)
 	RETURN_ON_CYCLE(*Pin, GenerationContext)
@@ -99,7 +100,7 @@ mu::NodeImagePtr GenerateMutableSourceImage(const UEdGraphPin* Pin, FMutableGrap
 
 	bool bDoNotAddToGeneratedCache = false;
 
-	mu::NodeImagePtr Result;
+	mu::Ptr<mu::NodeImage> Result;
 
 	const FCompilationOptions& CompilationOptions = GenerationContext.CompilationContext.Options;
 
@@ -110,11 +111,11 @@ mu::NodeImagePtr GenerateMutableSourceImage(const UEdGraphPin* Pin, FMutableGrap
 		{
 			// Check the specific image cache
 			FGeneratedImageKey ImageKey = FGeneratedImageKey(Pin);
-			mu::NodeImagePtr ImageNode;
+			mu::Ptr<mu::NodeImage> ImageNode;
 			
 			const uint32 MipsToSkip = ComputeLODBiasForTexture(GenerationContext, *BaseTexture, nullptr, ReferenceTextureSize);
 
-			if (mu::NodeImagePtr* Cached = GenerationContext.GeneratedImages.Find(ImageKey))
+			if (mu::Ptr<mu::NodeImage>* Cached = GenerationContext.GeneratedImages.Find(ImageKey))
 			{
 				ImageNode = *Cached;
 			}
@@ -152,12 +153,12 @@ mu::NodeImagePtr GenerateMutableSourceImage(const UEdGraphPin* Pin, FMutableGrap
 	{
 		mu::Ptr<mu::NodeImageParameter> TextureNode = new mu::NodeImageParameter();
 
-		TextureNode->SetName(TypedNodeParam->ParameterName);
-		TextureNode->SetUid(GenerationContext.GetNodeIdUnique(Node).ToString());
+		TextureNode->Name = TypedNodeParam->ParameterName;
+		TextureNode->UID = GenerationContext.GetNodeIdUnique(Node).ToString();
 
 		if (TypedNodeParam->DefaultValue)
 		{
-			TextureNode->SetDefaultValue(FName(TypedNodeParam->DefaultValue->GetPathName()));			
+			TextureNode->DefaultValue = FName(TypedNodeParam->DefaultValue->GetPathName());
 		}
 
 		GenerationContext.ParameterUIDataMap.Add(TypedNodeParam->ParameterName, FMutableParameterData(
@@ -165,8 +166,8 @@ mu::NodeImagePtr GenerateMutableSourceImage(const UEdGraphPin* Pin, FMutableGrap
 			EMutableParameterType::Texture));
 
 		// Force the same format that the default texture if any.
-		mu::NodeImageFormatPtr FormatNode = new mu::NodeImageFormat();
-		// \TODO: Take if from default?
+		mu::Ptr<mu::NodeImageFormat> FormatNode = new mu::NodeImageFormat();
+		// \TODO: Take it from default?
 		//if (TypedNodeParam->DefaultValue)
 		//{
 		//	FormatNode->SetFormat(mu::EImageFormat::RGBA_UByte);
@@ -174,13 +175,13 @@ mu::NodeImagePtr GenerateMutableSourceImage(const UEdGraphPin* Pin, FMutableGrap
 		//else
 		{
 			// Force an "easy format" on the texture.
-			FormatNode->SetFormat(mu::EImageFormat::RGBA_UByte);
+			FormatNode->Format = mu::EImageFormat::RGBA_UByte;
 		}
-		FormatNode->SetSource(TextureNode);
+		FormatNode->Source = TextureNode;
 
-		mu::NodeImageResizePtr ResizeNode = new mu::NodeImageResize();
-		ResizeNode->SetBase(FormatNode);
-		ResizeNode->SetRelative(false);
+		mu::Ptr<mu::NodeImageResize> ResizeNode = new mu::NodeImageResize();
+		ResizeNode->Base = FormatNode;
+		ResizeNode->bRelative = false;
 
 		FUintVector2 TextureSize(TypedNodeParam->TextureSizeX, TypedNodeParam->TextureSizeY);
 
@@ -207,7 +208,8 @@ mu::NodeImagePtr GenerateMutableSourceImage(const UEdGraphPin* Pin, FMutableGrap
 			}
 		}
 
-		ResizeNode->SetSize(TextureSize.X, TextureSize.Y);
+		ResizeNode->SizeX = TextureSize.X;
+		ResizeNode->SizeY = TextureSize.Y;
 
 		Result = ResizeNode;
 	}
@@ -239,19 +241,19 @@ mu::NodeImagePtr GenerateMutableSourceImage(const UEdGraphPin* Pin, FMutableGrap
 		}
 		else
 		{
-			Result = mu::NodeImagePtr();
+			Result = mu::Ptr<mu::NodeImage>();
 		}
 	}
 
 	else if (const UCustomizableObjectNodeTextureInterpolate* TypedNodeInterp = Cast<UCustomizableObjectNodeTextureInterpolate>(Node))
 	{
-		mu::NodeImageInterpolatePtr ImageNode = new mu::NodeImageInterpolate();
+		mu::Ptr<mu::NodeImageInterpolate> ImageNode = new mu::NodeImageInterpolate();
 		Result = ImageNode;
 
 		if (const UEdGraphPin* ConnectedPin = FollowInputPin(*TypedNodeInterp->FactorPin()))
 		{
-			mu::NodeScalarPtr FactorNode = GenerateMutableSourceFloat(ConnectedPin, GenerationContext);
-			ImageNode->SetFactor(FactorNode);
+			mu::Ptr<mu::NodeScalar> FactorNode = GenerateMutableSourceFloat(ConnectedPin, GenerationContext);
+			ImageNode->Factor = FactorNode;
 		}
 
 		int currentTarget = 0;
@@ -259,12 +261,12 @@ mu::NodeImagePtr GenerateMutableSourceImage(const UEdGraphPin* Pin, FMutableGrap
 		{
 			if (const UEdGraphPin* ConnectedPin = FollowInputPin(*TypedNodeInterp->Targets(LayerIndex)))
 			{
-				mu::NodeImagePtr TargetNode = GenerateMutableSourceImage(ConnectedPin, GenerationContext, ReferenceTextureSize);
+				mu::Ptr<mu::NodeImage> TargetNode = GenerateMutableSourceImage(ConnectedPin, GenerationContext, ReferenceTextureSize);
 
 				if (TargetNode)
 				{
-					ImageNode->SetTargetCount(currentTarget + 1);
-					ImageNode->SetTarget(currentTarget, TargetNode);
+					ImageNode->Targets.SetNum(currentTarget + 1);
+					ImageNode->Targets[currentTarget] = TargetNode;
 					++currentTarget;
 				}
 			}
@@ -282,7 +284,7 @@ mu::NodeImagePtr GenerateMutableSourceImage(const UEdGraphPin* Pin, FMutableGrap
 		{
 			if (const UEdGraphPin* OtherPin = FollowInputPin(*TypedNodeLayer->LayerPin(LayerIndex)))
 			{
-				mu::NodeImagePtr MaskNode = nullptr;
+				mu::Ptr<mu::NodeImage> MaskNode = nullptr;
 				if (const UEdGraphPin* ConnectedPin = FollowInputPin(*TypedNodeLayer->MaskPin(LayerIndex)))
 				{
 					MaskNode = GenerateMutableSourceImage(ConnectedPin, GenerationContext, ReferenceTextureSize);
@@ -313,25 +315,25 @@ mu::NodeImagePtr GenerateMutableSourceImage(const UEdGraphPin* Pin, FMutableGrap
 				
 				if (OtherPin->PinType.PinCategory == Schema->PC_Image)
 				{
-					mu::NodeImagePtr BlendNode = GenerateMutableSourceImage(OtherPin, GenerationContext, ReferenceTextureSize);
+					mu::Ptr<mu::NodeImage> BlendNode = GenerateMutableSourceImage(OtherPin, GenerationContext, ReferenceTextureSize);
 
-					mu::NodeImageLayerPtr LayerNode = new mu::NodeImageLayer;
-					LayerNode->SetBlendType(Type);
-					LayerNode->SetBase(Result);
-					LayerNode->SetBlended(BlendNode);
-					LayerNode->SetMask(MaskNode);
+					mu::Ptr<mu::NodeImageLayer> LayerNode = new mu::NodeImageLayer;
+					LayerNode->Type = Type;
+					LayerNode->Base = Result;
+					LayerNode->Blended = BlendNode;
+					LayerNode->Mask = MaskNode;
 					Result = LayerNode;
 				}
 
 				else if (OtherPin->PinType.PinCategory == Schema->PC_Color)
 				{
-					mu::NodeColourPtr ColorNode = GenerateMutableSourceColor(OtherPin, GenerationContext);
+					mu::Ptr<mu::NodeColour> ColorNode = GenerateMutableSourceColor(OtherPin, GenerationContext);
 
-					mu::NodeImageLayerColourPtr LayerNode = new mu::NodeImageLayerColour;
-					LayerNode->SetBlendType(Type);
-					LayerNode->SetBase(Result);
-					LayerNode->SetColour(ColorNode);
-					LayerNode->SetMask(MaskNode);
+					mu::Ptr<mu::NodeImageLayerColour> LayerNode = new mu::NodeImageLayerColour;
+					LayerNode->Type = Type;
+					LayerNode->Base = Result;
+					LayerNode->Colour = ColorNode;
+					LayerNode->Mask = MaskNode;
 					Result = LayerNode;
 				}
 
@@ -350,7 +352,7 @@ mu::NodeImagePtr GenerateMutableSourceImage(const UEdGraphPin* Pin, FMutableGrap
 			// Check Switch Parameter arity preconditions.
 			if (const UEdGraphPin* ConnectedPin = FollowInputPin(*SwitchParameter))
 			{
-				mu::NodeScalarPtr SwitchParam = GenerateMutableSourceFloat(ConnectedPin, GenerationContext);
+				mu::Ptr<mu::NodeScalar> SwitchParam = GenerateMutableSourceFloat(ConnectedPin, GenerationContext);
 				// Switch Param not generated
 				if (!SwitchParam)
 				{
@@ -371,21 +373,21 @@ mu::NodeImagePtr GenerateMutableSourceImage(const UEdGraphPin* Pin, FMutableGrap
 				const int32 NumSwitchOptions = TypedNodeTextureSwitch->GetNumElements();
 
 				mu::NodeScalarEnumParameter* EnumParameter = static_cast<mu::NodeScalarEnumParameter*>(SwitchParam.get());
-				if (NumSwitchOptions != EnumParameter->GetValueCount())
+				if (NumSwitchOptions != EnumParameter->Options.Num())
 				{
 					const FText Message = LOCTEXT("MismatchedSwitch", "Switch enum and switch node have different number of options. Please refresh the switch node to make sure the outcomes are labeled properly.");
 					GenerationContext.Log(Message, Node);
 				}
 
-				mu::NodeImageSwitchPtr SwitchNode = new mu::NodeImageSwitch;
-				SwitchNode->SetParameter(SwitchParam);
-				SwitchNode->SetOptionCount(NumSwitchOptions);
+				mu::Ptr<mu::NodeImageSwitch> SwitchNode = new mu::NodeImageSwitch;
+				SwitchNode->Parameter = SwitchParam;
+				SwitchNode->Options.SetNum(NumSwitchOptions);
 
 				for (int SelectorIndex = 0; SelectorIndex < NumSwitchOptions; ++SelectorIndex)
 				{
 					if (const UEdGraphPin* TexturePin = FollowInputPin(*TypedNodeTextureSwitch->GetElementPin(SelectorIndex)))
 					{
-						SwitchNode->SetOption(SelectorIndex, GenerateMutableSourceImage(TexturePin, GenerationContext, ReferenceTextureSize));
+						SwitchNode->Options[SelectorIndex] = GenerateMutableSourceImage(TexturePin, GenerationContext, ReferenceTextureSize);
 					}
 					else
 					{
@@ -409,15 +411,15 @@ mu::NodeImagePtr GenerateMutableSourceImage(const UEdGraphPin* Pin, FMutableGrap
 	else if (const UCustomizableObjectNodeTextureVariation* TypedNodeImageVar = Cast<const UCustomizableObjectNodeTextureVariation>(Node))
 	{
 		// UCustomizableObjectNodePassThroughTextureVariation nodes are also handled here
-		mu::NodeImageVariationPtr TextureNode = new mu::NodeImageVariation();
+		mu::Ptr<mu::NodeImageVariation> TextureNode = new mu::NodeImageVariation();
 		Result = TextureNode;
 
 		if (const UEdGraphPin* ConnectedPin = FollowInputPin(*TypedNodeImageVar->DefaultPin()))
 		{
-			mu::NodeImagePtr ChildNode = GenerateMutableSourceImage(ConnectedPin, GenerationContext, ReferenceTextureSize);
+			mu::Ptr<mu::NodeImage> ChildNode = GenerateMutableSourceImage(ConnectedPin, GenerationContext, ReferenceTextureSize);
 			if (ChildNode)
 			{
-				TextureNode->SetDefaultImage(ChildNode.get());
+				TextureNode->DefaultImage = ChildNode;
 			}
 			else
 			{
@@ -425,7 +427,7 @@ mu::NodeImagePtr GenerateMutableSourceImage(const UEdGraphPin* Pin, FMutableGrap
 			}
 		}
 
-		TextureNode->SetVariationCount(TypedNodeImageVar->Variations.Num());
+		TextureNode->Variations.SetNum(TypedNodeImageVar->Variations.Num());
 		for (int32 VariationIndex = 0; VariationIndex < TypedNodeImageVar->Variations.Num(); ++VariationIndex)
 		{
 			const UEdGraphPin* VariationPin = TypedNodeImageVar->VariationPin(VariationIndex);
@@ -434,35 +436,36 @@ mu::NodeImagePtr GenerateMutableSourceImage(const UEdGraphPin* Pin, FMutableGrap
 				continue;
 			}
 
-			TextureNode->SetVariationTag(VariationIndex, TypedNodeImageVar->Variations[VariationIndex].Tag);
+			TextureNode->Variations[VariationIndex].Tag = TypedNodeImageVar->Variations[VariationIndex].Tag;
 			if (const UEdGraphPin* ConnectedPin = FollowInputPin(*VariationPin))
 			{
-				mu::NodeImagePtr ChildNode = GenerateMutableSourceImage(ConnectedPin, GenerationContext, ReferenceTextureSize);
-				TextureNode->SetVariationImage(VariationIndex, ChildNode.get());
+				mu::Ptr<mu::NodeImage> ChildNode = GenerateMutableSourceImage(ConnectedPin, GenerationContext, ReferenceTextureSize);
+				TextureNode->Variations[VariationIndex].Image = ChildNode;
 			}
 		}
 	}
 
 	else if (const UCustomizableObjectNodeTextureFromColor* TypedNodeFromColor = Cast<UCustomizableObjectNodeTextureFromColor>(Node))
 	{
-		mu::NodeColourPtr color;
+		mu::Ptr<mu::NodeColour> color;
 
 		if (const UEdGraphPin* ConnectedPin = FollowInputPin(*TypedNodeFromColor->ColorPin()))
 		{
 			color = GenerateMutableSourceColor(ConnectedPin, GenerationContext);
 		}
 
-		mu::NodeImagePlainColourPtr ImageFromColour = new mu::NodeImagePlainColour;
+		mu::Ptr<mu::NodeImagePlainColour> ImageFromColour = new mu::NodeImagePlainColour;
 		Result = ImageFromColour;
 
 		if (color)
 		{
-			ImageFromColour->SetColour(color);
+			ImageFromColour->Colour = color;
 		}
 
 		if (ReferenceTextureSize > 0)
 		{
-			ImageFromColour->SetSize(ReferenceTextureSize, ReferenceTextureSize);
+			ImageFromColour->SizeX = ReferenceTextureSize;
+			ImageFromColour->SizeY = ReferenceTextureSize;
 		}
 	}
 
@@ -490,23 +493,24 @@ mu::NodeImagePtr GenerateMutableSourceImage(const UEdGraphPin* Pin, FMutableGrap
 			Color->W = GenerateMutableSourceFloat(APin, GenerationContext);
 		}
 
-		mu::NodeImagePlainColourPtr ImageFromColour = new mu::NodeImagePlainColour;
+		mu::Ptr<mu::NodeImagePlainColour> ImageFromColour = new mu::NodeImagePlainColour;
 		Result = ImageFromColour;
 
 		if (Color)
 		{
-			ImageFromColour->SetColour(Color);
+			ImageFromColour->Colour = Color;
 		}
 
 		if (ReferenceTextureSize > 0)
 		{
-			ImageFromColour->SetSize(ReferenceTextureSize, ReferenceTextureSize);
+			ImageFromColour->SizeX = ReferenceTextureSize;
+			ImageFromColour->SizeY = ReferenceTextureSize;
 		}
 	}
 
 	else if (const UCustomizableObjectNodeTextureFromChannels* TypedNodeFrom = Cast<UCustomizableObjectNodeTextureFromChannels>(Node))
 	{
-		mu::NodeImagePtr RNode, GNode, BNode, ANode;
+		mu::Ptr<mu::NodeImage> RNode, GNode, BNode, ANode;
 		bool RGB = false;
 
 		if (const UEdGraphPin* ConnectedPin = FollowInputPin(*TypedNodeFrom->RPin()))
@@ -529,73 +533,73 @@ mu::NodeImagePtr GenerateMutableSourceImage(const UEdGraphPin* Pin, FMutableGrap
 			ANode = GenerateMutableSourceImage(ConnectedPin, GenerationContext, ReferenceTextureSize);
 		}
 
-		mu::NodeImageSwizzlePtr SwizzleNode = new mu::NodeImageSwizzle;
+		mu::Ptr<mu::NodeImageSwizzle> SwizzleNode = new mu::NodeImageSwizzle;
 		Result = SwizzleNode;
 
 		if (RGB && ANode)
 		{
 			SwizzleNode->SetFormat(mu::EImageFormat::RGBA_UByte);
-			SwizzleNode->SetSource(0, RNode);
-			SwizzleNode->SetSourceChannel(0, 0);
-			SwizzleNode->SetSource(1, GNode);
-			SwizzleNode->SetSourceChannel(1, 0);
-			SwizzleNode->SetSource(2, BNode);
-			SwizzleNode->SetSourceChannel(2, 0);
-			SwizzleNode->SetSource(3, ANode);
-			SwizzleNode->SetSourceChannel(3, 0);
+			SwizzleNode->Sources[0] = RNode;
+			SwizzleNode->SourceChannels[0] = 0;
+			SwizzleNode->Sources[1] = GNode;
+			SwizzleNode->SourceChannels[1] = 0;
+			SwizzleNode->Sources[2] = BNode;
+			SwizzleNode->SourceChannels[2] = 0;
+			SwizzleNode->Sources[3] = ANode;
+			SwizzleNode->SourceChannels[3] = 0;
 		}
 		else if (RGB)
 		{
 			SwizzleNode->SetFormat(mu::EImageFormat::RGB_UByte);
-			SwizzleNode->SetSource(0, RNode);
-			SwizzleNode->SetSourceChannel(0, 0);
-			SwizzleNode->SetSource(1, GNode);
-			SwizzleNode->SetSourceChannel(1, 0);
-			SwizzleNode->SetSource(2, BNode);
-			SwizzleNode->SetSourceChannel(2, 0);
+			SwizzleNode->Sources[0] = RNode;
+			SwizzleNode->SourceChannels[0] = 0;
+			SwizzleNode->Sources[1] = GNode;
+			SwizzleNode->SourceChannels[1] = 0;
+			SwizzleNode->Sources[2] = BNode;
+			SwizzleNode->SourceChannels[2] = 0;
 		}
 		else if (RNode)
 		{
 			SwizzleNode->SetFormat(mu::EImageFormat::L_UByte);
-			SwizzleNode->SetSource(0, RNode);
-			SwizzleNode->SetSourceChannel(0, 0);
+			SwizzleNode->Sources[0] = RNode;
+			SwizzleNode->SourceChannels[0] = 0;
 		}
 		else if (ANode)
 		{
 			SwizzleNode->SetFormat(mu::EImageFormat::L_UByte);
-			SwizzleNode->SetSource(0, ANode);
-			SwizzleNode->SetSourceChannel(0, 0);
+			SwizzleNode->Sources[0] = ANode;
+			SwizzleNode->SourceChannels[0] = 0;
 		}
 	}
 
 	else if (const UCustomizableObjectNodeTextureToChannels* TypedNodeTo = Cast<UCustomizableObjectNodeTextureToChannels>(Node))
 	{
-		mu::NodeImagePtr BaseNode;
+		mu::Ptr<mu::NodeImage> BaseNode;
 		if (const UEdGraphPin* ConnectedPin = FollowInputPin(*TypedNodeTo->InputPin()))
 		{
 			BaseNode = GenerateMutableSourceImage(ConnectedPin, GenerationContext, ReferenceTextureSize);
 		}
 
-		mu::NodeImageSwizzlePtr SwizzleNode = new mu::NodeImageSwizzle;
+		mu::Ptr<mu::NodeImageSwizzle> SwizzleNode = new mu::NodeImageSwizzle;
 		Result = SwizzleNode;
 		SwizzleNode->SetFormat(mu::EImageFormat::L_UByte);
-		SwizzleNode->SetSource(0, BaseNode);
+		SwizzleNode->Sources[0] = BaseNode;
 
 		if (Pin == TypedNodeTo->RPin())
 		{
-			SwizzleNode->SetSourceChannel(0, 0);
+			SwizzleNode->SourceChannels[0] = 0;
 		}
 		else if (Pin == TypedNodeTo->GPin())
 		{
-			SwizzleNode->SetSourceChannel(0, 1);
+			SwizzleNode->SourceChannels[0] = 1;
 		}
 		else if (Pin == TypedNodeTo->BPin())
 		{
-			SwizzleNode->SetSourceChannel(0, 2);
+			SwizzleNode->SourceChannels[0] = 2;
 		}
 		else if (Pin == TypedNodeTo->APin())
 		{
-			SwizzleNode->SetSourceChannel(0, 3);
+			SwizzleNode->SourceChannels[0] = 3;
 		}
 		else
 		{
@@ -614,7 +618,7 @@ mu::NodeImagePtr GenerateMutableSourceImage(const UEdGraphPin* Pin, FMutableGrap
 			Result = nullptr;
 		}
 
-		ImageNode->SetLayout(TypedNodeProject->Layout);
+		ImageNode->Layout = TypedNodeProject->Layout;
 
 		// Calculate the max TextureSize allowed using the ReferenceTextureSize and the Reference texture from the node
 		int32 MaxReferenceTextureSizeInGame = ReferenceTextureSize;
@@ -639,11 +643,12 @@ mu::NodeImagePtr GenerateMutableSourceImage(const UEdGraphPin* Pin, FMutableGrap
 			TextureSize.Y = TextureSize.Y >> (NumMips - MaxNumMips);
 		}
 
-		ImageNode->SetImageSize(TextureSize);
+		ImageNode->ImageSize = TextureSize;
 
-		ImageNode->SetEnableSeamCorrection(TypedNodeProject->bEnableTextureSeamCorrection);
-		ImageNode->SetAngleFadeChannels(TypedNodeProject->bEnableAngleFadeOutForRGB, TypedNodeProject->bEnableAngleFadeOutForAlpha);
-		ImageNode->SetSamplingMethod(Invoke(
+		ImageNode->bEnableTextureSeamCorrection = TypedNodeProject->bEnableTextureSeamCorrection;
+		ImageNode->bIsRGBFadingEnabled = TypedNodeProject->bEnableAngleFadeOutForRGB;
+		ImageNode->bIsAlphaFadingEnabled = TypedNodeProject->bEnableAngleFadeOutForAlpha;
+		ImageNode->SamplingMethod = Invoke(
 			[](ETextureProjectSamplingMethod Method) -> mu::ESamplingMethod 
 			{
 				switch(Method)
@@ -653,10 +658,9 @@ mu::NodeImagePtr GenerateMutableSourceImage(const UEdGraphPin* Pin, FMutableGrap
 				default: { check(false); return mu::ESamplingMethod::Point; }
 				}
 			}, 
-			TypedNodeProject->SamplingMethod)
-		);
+			TypedNodeProject->SamplingMethod);
 
-		ImageNode->SetMinFilterMethod(Invoke(
+		ImageNode->MinFilterMethod = Invoke(
 			[](ETextureProjectMinFilterMethod Method) -> mu::EMinFilterMethod 
 			{
 				switch(Method)
@@ -666,33 +670,32 @@ mu::NodeImagePtr GenerateMutableSourceImage(const UEdGraphPin* Pin, FMutableGrap
 				default: { check(false); return mu::EMinFilterMethod::None; }
 				}
 			}, 
-			TypedNodeProject->MinFilterMethod)
-		);
+			TypedNodeProject->MinFilterMethod);
 
 		if (const UEdGraphPin* ConnectedPin = FollowInputPin(*TypedNodeProject->AngleFadeStartPin()))
 		{
-			mu::NodeScalarPtr ScalarNode = GenerateMutableSourceFloat(ConnectedPin, GenerationContext);
-			ImageNode->SetAngleFadeStart(ScalarNode);
+			mu::Ptr<mu::NodeScalar> ScalarNode = GenerateMutableSourceFloat(ConnectedPin, GenerationContext);
+			ImageNode->AngleFadeStart = ScalarNode;
 		}
 		else
 		{
-			mu::NodeScalarConstantPtr ScalarNode = new mu::NodeScalarConstant;
+			mu::Ptr<mu::NodeScalarConstant> ScalarNode = new mu::NodeScalarConstant;
 			float value = FCString::Atof(*TypedNodeProject->AngleFadeStartPin()->DefaultValue);
-			ScalarNode->SetValue(value);
-			ImageNode->SetAngleFadeStart(ScalarNode);
+			ScalarNode->Value = value;
+			ImageNode->AngleFadeStart = ScalarNode;
 		}
 
 		if (const UEdGraphPin* ConnectedPin = FollowInputPin(*TypedNodeProject->AngleFadeEndPin()))
 		{
-			mu::NodeScalarPtr ScalarNode = GenerateMutableSourceFloat(ConnectedPin, GenerationContext);
-			ImageNode->SetAngleFadeEnd(ScalarNode);
+			mu::Ptr<mu::NodeScalar> ScalarNode = GenerateMutableSourceFloat(ConnectedPin, GenerationContext);
+			ImageNode->AngleFadeEnd = ScalarNode;
 		}
 		else
 		{
-			mu::NodeScalarConstantPtr ScalarNode = new mu::NodeScalarConstant;
+			mu::Ptr<mu::NodeScalarConstant> ScalarNode = new mu::NodeScalarConstant;
 			float value = FCString::Atof(*TypedNodeProject->AngleFadeEndPin()->DefaultValue);
-			ScalarNode->SetValue(value);
-			ImageNode->SetAngleFadeEnd(ScalarNode);
+			ScalarNode->Value = value;
+			ImageNode->AngleFadeEnd = ScalarNode;
 		}
 
 		if (const UEdGraphPin* ConnectedPin = FollowInputPin(*TypedNodeProject->MeshPin()))
@@ -702,21 +705,21 @@ mu::NodeImagePtr GenerateMutableSourceImage(const UEdGraphPin* Pin, FMutableGrap
 			FLayoutGenerationFlags LayoutGenerationFlags;
 			LayoutGenerationFlags.TexturePinModes.Init(EPinMode::Mutable, TEXSTREAM_MAX_NUM_UVCHANNELS);
 			GenerationContext.LayoutGenerationFlags.Push(LayoutGenerationFlags);
-			mu::NodeMeshPtr MeshNode = GenerateMutableSourceMesh(ConnectedPin, GenerationContext, DummyMeshData, 0, false, false);
+			mu::Ptr<mu::NodeMesh> MeshNode = GenerateMutableSourceMesh(ConnectedPin, GenerationContext, DummyMeshData, 0, false, false);
 			GenerationContext.LayoutGenerationFlags.Pop();
-			ImageNode->SetMesh(MeshNode);
+			ImageNode->Mesh = MeshNode;
 		}
 
 		if (const UEdGraphPin* ConnectedPin = FollowInputPin(*TypedNodeProject->MeshMaskPin()))
 		{
-			mu::NodeImagePtr MeshMaskNode = GenerateMutableSourceImage(ConnectedPin, GenerationContext, ReferenceTextureSize);
-			ImageNode->SetTargetMask(MeshMaskNode);
+			mu::Ptr<mu::NodeImage> MeshMaskNode = GenerateMutableSourceImage(ConnectedPin, GenerationContext, ReferenceTextureSize);
+			ImageNode->Mask = MeshMaskNode;
 		}
 
 		if (const UEdGraphPin* ConnectedPin = FollowInputPin(*TypedNodeProject->ProjectorPin()))
 		{
-			mu::NodeProjectorPtr ProjectorNode = GenerateMutableSourceProjector(ConnectedPin, GenerationContext);
-			ImageNode->SetProjector(ProjectorNode);
+			mu::Ptr<mu::NodeProjector> ProjectorNode = GenerateMutableSourceProjector(ConnectedPin, GenerationContext);
+			ImageNode->Projector = ProjectorNode;
 		}
 
 		int TexIndex = -1;// TypedNodeProject->OutputPins.Find((UEdGraphPin*)Pin);
@@ -733,108 +736,108 @@ mu::NodeImagePtr GenerateMutableSourceImage(const UEdGraphPin* Pin, FMutableGrap
 
 		if (const UEdGraphPin* ConnectedPin = FollowInputPin(*TypedNodeProject->TexturePins(TexIndex)))
 		{
-			mu::NodeImagePtr SourceNode = GenerateMutableSourceImage(ConnectedPin, GenerationContext, FMath::Max(TextureSize.X, TextureSize.Y));
-			ImageNode->SetImage(SourceNode);
+			mu::Ptr<mu::NodeImage> SourceNode = GenerateMutableSourceImage(ConnectedPin, GenerationContext, FMath::Max(TextureSize.X, TextureSize.Y));
+			ImageNode->Image = SourceNode;
 		}
 	}
 
 	else if (const UCustomizableObjectNodeTextureBinarise* TypedNodeTexBin = Cast<UCustomizableObjectNodeTextureBinarise>(Node))
 	{
-		mu::NodeImageBinarisePtr BinariseNode = new mu::NodeImageBinarise();
+		mu::Ptr<mu::NodeImageBinarise> BinariseNode = new mu::NodeImageBinarise();
 		Result = BinariseNode;
 
 		if (const UEdGraphPin* ConnectedPin = FollowInputPin(*TypedNodeTexBin->GetBaseImagePin()))
 		{
-			mu::NodeImagePtr BaseImageNode = GenerateMutableSourceImage(ConnectedPin, GenerationContext, ReferenceTextureSize);
-			BinariseNode->SetBase(BaseImageNode);
+			mu::Ptr<mu::NodeImage> BaseImageNode = GenerateMutableSourceImage(ConnectedPin, GenerationContext, ReferenceTextureSize);
+			BinariseNode->Base = BaseImageNode;
 		}
 
 		if (const UEdGraphPin* ConnectedPin = FollowInputPin(*TypedNodeTexBin->GetThresholdPin()))
 		{
-			mu::NodeScalarPtr ThresholdNode = GenerateMutableSourceFloat(ConnectedPin, GenerationContext);
-			BinariseNode->SetThreshold(ThresholdNode);
+			mu::Ptr<mu::NodeScalar> ThresholdNode = GenerateMutableSourceFloat(ConnectedPin, GenerationContext);
+			BinariseNode->Threshold = ThresholdNode;
 		}
 	}
 
 	else if (const UCustomizableObjectNodeTextureInvert* TypedNodeTexInv = Cast<UCustomizableObjectNodeTextureInvert>(Node))
 	{
-		mu::NodeImageInvertPtr InvertNode = new mu::NodeImageInvert();
+		mu::Ptr<mu::NodeImageInvert> InvertNode = new mu::NodeImageInvert();
 		Result = InvertNode;
 
 		if (const UEdGraphPin* ConnectedPin = FollowInputPin(*TypedNodeTexInv->GetBaseImagePin()))
 		{
-			mu::NodeImagePtr BaseImageNode = GenerateMutableSourceImage(ConnectedPin, GenerationContext, ReferenceTextureSize);
-			InvertNode->SetBase(BaseImageNode);
+			mu::Ptr<mu::NodeImage> BaseImageNode = GenerateMutableSourceImage(ConnectedPin, GenerationContext, ReferenceTextureSize);
+			InvertNode->Base = BaseImageNode;
 		}
 	}
 
 	else if (const UCustomizableObjectNodeTextureColourMap* TypedNodeColourMap = Cast<UCustomizableObjectNodeTextureColourMap>(Node))
 	{
-		mu::NodeImageColourMapPtr ColourMapNode = new mu::NodeImageColourMap();
+		mu::Ptr<mu::NodeImageColourMap> ColourMapNode = new mu::NodeImageColourMap();
 
 		Result = ColourMapNode;
 
 		if (const UEdGraphPin* ConnectedPin = FollowInputPin(*TypedNodeColourMap->GetMapPin()))
 		{
-			mu::NodeImagePtr GradientImageNode = GenerateMutableSourceImage(ConnectedPin, GenerationContext, ReferenceTextureSize);
-			ColourMapNode->SetMap( GradientImageNode ); 
+			mu::Ptr<mu::NodeImage> GradientImageNode = GenerateMutableSourceImage(ConnectedPin, GenerationContext, ReferenceTextureSize);
+			ColourMapNode->Map = GradientImageNode;
 		}
 
 		if (const UEdGraphPin* ConnectedPin = FollowInputPin(*TypedNodeColourMap->GetMaskPin()))
 		{
-			mu::NodeImagePtr GradientImageNode = GenerateMutableSourceImage(ConnectedPin, GenerationContext, ReferenceTextureSize);
-			ColourMapNode->SetMask( GradientImageNode ); 
+			mu::Ptr<mu::NodeImage> GradientImageNode = GenerateMutableSourceImage(ConnectedPin, GenerationContext, ReferenceTextureSize);
+			ColourMapNode->Mask = GradientImageNode;
 		}
 
 		if (const UEdGraphPin* ConnectedPin = FollowInputPin(*TypedNodeColourMap->GetBasePin()))
 		{
-			mu::NodeImagePtr SourceImageNode = GenerateMutableSourceImage(ConnectedPin, GenerationContext, ReferenceTextureSize);
-			ColourMapNode->SetBase( SourceImageNode );
+			mu::Ptr<mu::NodeImage> SourceImageNode = GenerateMutableSourceImage(ConnectedPin, GenerationContext, ReferenceTextureSize);
+			ColourMapNode->Base = SourceImageNode;
 		}
 	}
 
 	else if (const UCustomizableObjectNodeTextureTransform* TypedNodeTransform = Cast<UCustomizableObjectNodeTextureTransform>(Node))
 	{
-		mu::NodeImageTransformPtr TransformNode = new mu::NodeImageTransform();
+		mu::Ptr<mu::NodeImageTransform> TransformNode = new mu::NodeImageTransform();
 		Result = TransformNode;
 		
 		if ( UEdGraphPin* BaseImagePin = FollowInputPin(*TypedNodeTransform->GetBaseImagePin()) )
 		{
-			mu::NodeImagePtr ImageNode = GenerateMutableSourceImage( BaseImagePin, GenerationContext, ReferenceTextureSize);
-			TransformNode->SetBase( ImageNode ); 
+			mu::Ptr<mu::NodeImage> ImageNode = GenerateMutableSourceImage( BaseImagePin, GenerationContext, ReferenceTextureSize);
+			TransformNode->Base = ImageNode;
 		}
 
 		if ( UEdGraphPin* OffsetXPin = FollowInputPin(*TypedNodeTransform->GetOffsetXPin()) )
 		{
-			mu::NodeScalarPtr OffsetXNode = GenerateMutableSourceFloat(OffsetXPin, GenerationContext);
-			TransformNode->SetOffsetX( OffsetXNode ); 
+			mu::Ptr<mu::NodeScalar> OffsetXNode = GenerateMutableSourceFloat(OffsetXPin, GenerationContext);
+			TransformNode->OffsetX = OffsetXNode;
 		}
 
 		if ( UEdGraphPin* OffsetYPin = FollowInputPin(*TypedNodeTransform->GetOffsetYPin()) )
 		{
-			mu::NodeScalarPtr OffsetYNode = GenerateMutableSourceFloat(OffsetYPin, GenerationContext);
-			TransformNode->SetOffsetY( OffsetYNode ); 
+			mu::Ptr<mu::NodeScalar> OffsetYNode = GenerateMutableSourceFloat(OffsetYPin, GenerationContext);
+			TransformNode->OffsetY = OffsetYNode;
 		}
 
 		if ( UEdGraphPin* ScaleXPin = FollowInputPin(*TypedNodeTransform->GetScaleXPin()) )
 		{
-			mu::NodeScalarPtr ScaleXNode = GenerateMutableSourceFloat(ScaleXPin, GenerationContext);
-			TransformNode->SetScaleX( ScaleXNode ); 
+			mu::Ptr<mu::NodeScalar> ScaleXNode = GenerateMutableSourceFloat(ScaleXPin, GenerationContext);
+			TransformNode->ScaleX = ScaleXNode;
 		}
 
 		if ( UEdGraphPin* ScaleYPin = FollowInputPin(*TypedNodeTransform->GetScaleYPin()) )
 		{
-			mu::NodeScalarPtr ScaleYNode = GenerateMutableSourceFloat(ScaleYPin, GenerationContext);
-			TransformNode->SetScaleY( ScaleYNode ); 
+			mu::Ptr<mu::NodeScalar> ScaleYNode = GenerateMutableSourceFloat(ScaleYPin, GenerationContext);
+			TransformNode->ScaleY = ScaleYNode; 
 		}
 
 		if ( UEdGraphPin* RotationPin = FollowInputPin(*TypedNodeTransform->GetRotationPin()) )
 		{
-			mu::NodeScalarPtr RotationNode = GenerateMutableSourceFloat(RotationPin, GenerationContext);
-			TransformNode->SetRotation( RotationNode ); 
+			mu::Ptr<mu::NodeScalar> RotationNode = GenerateMutableSourceFloat(RotationPin, GenerationContext);
+			TransformNode->Rotation = RotationNode; 
 		}
 
-		TransformNode->SetAddressMode(Invoke([&]() 
+		TransformNode->AddressMode = Invoke([&]() 
 		{
 			switch (TypedNodeTransform->AddressMode)
 			{
@@ -843,7 +846,7 @@ mu::NodeImagePtr GenerateMutableSourceImage(const UEdGraphPin* Pin, FMutableGrap
 			case ETextureTransformAddressMode::ClampToBlack: return mu::EAddressMode::ClampToBlack;
 			default: { check(false); return mu::EAddressMode::None; }
 			}
-		}));
+		});
 
 		FUintVector2 TextureSize(TypedNodeTransform->TextureSizeX, TypedNodeTransform->TextureSizeY);
 
@@ -868,9 +871,9 @@ mu::NodeImagePtr GenerateMutableSourceImage(const UEdGraphPin* Pin, FMutableGrap
 			TextureSize.Y = TextureSize.Y >> (NumMips - MaxNumMips);
 		}
 
-		TransformNode->SetKeepAspectRatio(TypedNodeTransform->bKeepAspectRatio);
-		TransformNode->SetSizeX(TextureSize.X);
-		TransformNode->SetSizeY(TextureSize.Y);
+		TransformNode->bKeepAspectRatio = TypedNodeTransform->bKeepAspectRatio;
+		TransformNode->SizeX = TextureSize.X;
+		TransformNode->SizeY = TextureSize.Y;
 	}
 
 	else if (const UCustomizableObjectNodeTextureSaturate* TypedNodeSaturate = Cast<UCustomizableObjectNodeTextureSaturate>(Node))
@@ -880,14 +883,14 @@ mu::NodeImagePtr GenerateMutableSourceImage(const UEdGraphPin* Pin, FMutableGrap
 	
 		if ( UEdGraphPin* BaseImagePin = FollowInputPin(*TypedNodeSaturate->GetBaseImagePin()) )
 		{
-			mu::NodeImagePtr ImageNode = GenerateMutableSourceImage(BaseImagePin, GenerationContext, ReferenceTextureSize);
-			SaturateNode->SetSource(ImageNode); 
+			mu::Ptr<mu::NodeImage> ImageNode = GenerateMutableSourceImage(BaseImagePin, GenerationContext, ReferenceTextureSize);
+			SaturateNode->Source = ImageNode;
 		}
 
 		if ( UEdGraphPin* FactorPin = FollowInputPin(*TypedNodeSaturate->GetFactorPin()))
 		{
-			mu::NodeScalarPtr FactorNode = GenerateMutableSourceFloat(FactorPin, GenerationContext);
-			SaturateNode->SetFactor(FactorNode); 
+			mu::Ptr<mu::NodeScalar> FactorNode = GenerateMutableSourceFloat(FactorPin, GenerationContext);
+			SaturateNode->Factor = FactorNode;
 		}
 	}
 
@@ -916,7 +919,7 @@ mu::NodeImagePtr GenerateMutableSourceImage(const UEdGraphPin* Pin, FMutableGrap
 			// Check Switch Parameter arity preconditions.
 			if (const UEdGraphPin* ConnectedPin = FollowInputPin(*SwitchParameter))
 			{
-				mu::NodeScalarPtr SwitchParam = GenerateMutableSourceFloat(ConnectedPin, GenerationContext);
+				mu::Ptr<mu::NodeScalar> SwitchParam = GenerateMutableSourceFloat(ConnectedPin, GenerationContext);
 				// Switch Param not generated
 				if (!SwitchParam)
 				{
@@ -937,7 +940,7 @@ mu::NodeImagePtr GenerateMutableSourceImage(const UEdGraphPin* Pin, FMutableGrap
 				const int32 NumSwitchOptions = TypedNodePassThroughTextureSwitch->GetNumElements();
 
 				mu::NodeScalarEnumParameter* EnumParameter = static_cast<mu::NodeScalarEnumParameter*>(SwitchParam.get());
-				if (NumSwitchOptions != EnumParameter->GetValueCount())
+				if (NumSwitchOptions != EnumParameter->Options.Num())
 				{
 					const FText Message = LOCTEXT("MismatchedSwitch", "Switch enum and switch node have different number of options. Please refresh the switch node to make sure the outcomes are labeled properly.");
 					GenerationContext.Log(Message, Node);
@@ -945,15 +948,15 @@ mu::NodeImagePtr GenerateMutableSourceImage(const UEdGraphPin* Pin, FMutableGrap
 
 				// TODO: Implement Mutable core pass-through switch nodes
 				mu::Ptr<mu::NodeImageSwitch> SwitchNode = new mu::NodeImageSwitch;
-				SwitchNode->SetParameter(SwitchParam);
-				SwitchNode->SetOptionCount(NumSwitchOptions);
+				SwitchNode->Parameter = SwitchParam;
+				SwitchNode->Options.SetNum(NumSwitchOptions);
 
 				for (int32 SelectorIndex = 0; SelectorIndex < NumSwitchOptions; ++SelectorIndex)
 				{
 					if (const UEdGraphPin* TexturePin = FollowInputPin(*TypedNodePassThroughTextureSwitch->GetElementPin(SelectorIndex)))
 					{
 						mu::Ptr<mu::NodeImage> PassThroughImage = GenerateMutableSourceImage(TexturePin, GenerationContext, ReferenceTextureSize);
-						SwitchNode->SetOption(SelectorIndex, PassThroughImage);
+						SwitchNode->Options[SelectorIndex] = PassThroughImage;
 					}
 					else
 					{
@@ -977,11 +980,12 @@ mu::NodeImagePtr GenerateMutableSourceImage(const UEdGraphPin* Pin, FMutableGrap
 	// If the node is a plain colour node, generate an image out of it
 	else if (Pin->PinType.PinCategory == Schema->PC_Color)
 	{
-		mu::NodeColourPtr ColorNode = GenerateMutableSourceColor(Pin, GenerationContext);
+		mu::Ptr<mu::NodeColour> ColorNode = GenerateMutableSourceColor(Pin, GenerationContext);
 
-		mu::NodeImagePlainColourPtr ImageNode = new mu::NodeImagePlainColour;
-		ImageNode->SetSize(16, 16);
-		ImageNode->SetColour(ColorNode);
+		mu::Ptr<mu::NodeImagePlainColour> ImageNode = new mu::NodeImagePlainColour;
+		ImageNode->SizeX = 16;
+		ImageNode->SizeY = 16;
+		ImageNode->Colour = ColorNode;
 		Result = ImageNode;
 	}
 
@@ -989,11 +993,12 @@ mu::NodeImagePtr GenerateMutableSourceImage(const UEdGraphPin* Pin, FMutableGrap
 	{
 		if (Pin->PinType.PinCategory == Schema->PC_Color)
 		{
-			mu::NodeColourPtr ColorNode = GenerateMutableSourceColor(Pin, GenerationContext);
-			mu::NodeImagePlainColourPtr ImageNode = new mu::NodeImagePlainColour;
+			mu::Ptr<mu::NodeColour> ColorNode = GenerateMutableSourceColor(Pin, GenerationContext);
+			mu::Ptr<mu::NodeImagePlainColour> ImageNode = new mu::NodeImagePlainColour;
 
-			ImageNode->SetSize(16, 16);
-			ImageNode->SetColour(ColorNode);
+			ImageNode->SizeX = 16;
+			ImageNode->SizeY = 16;
+			ImageNode->Colour = ColorNode;
 
 			Result = ImageNode;
 		}
