[mutable] Initial boilerplate for the implementation of Mesh Parameters.

diff --git a/Engine/Plugins/Mutable/Source/CustomizableObjectEditor/Private/MuCOE/GenerateMutableSource/GenerateMutableSourceMesh.cpp b/Engine/Plugins/Mutable/Source/CustomizableObjectEditor/Private/MuCOE/GenerateMutableSource/GenerateMutableSourceMesh.cpp
--- a/Engine/Plugins/Mutable/Source/CustomizableObjectEditor/Private/MuCOE/GenerateMutableSource/GenerateMutableSourceMesh.cpp
+++ b/Engine/Plugins/Mutable/Source/CustomizableObjectEditor/Private/MuCOE/GenerateMutableSource/GenerateMutableSourceMesh.cpp
@@ -34,6 +34,7 @@
 #include "MuCOE/Nodes/CustomizableObjectNodeMeshReshapeCommon.h"
 #include "MuCOE/Nodes/CustomizableObjectNodeMeshSwitch.h"
 #include "MuCOE/Nodes/CustomizableObjectNodeMeshVariation.h"
+#include "MuCOE/Nodes/CustomizableObjectNodeMeshParameter.h"
 #include "MuCOE/Nodes/CustomizableObjectNodeSkeletalMesh.h"
 #include "MuCOE/Nodes/CustomizableObjectNodeStaticMesh.h"
 #include "MuCOE/Nodes/CustomizableObjectNodeTable.h"
@@ -45,6 +46,7 @@
 #include "MuT/NodeMeshSwitch.h"
 #include "MuT/NodeMeshTable.h"
 #include "MuT/NodeMeshVariation.h"
+#include "MuT/NodeMeshParameter.h"
 #include "MuR/Types.h"
 #include "MuR/OpMeshFormat.h"
 #include "PhysicsEngine/PhysicsAsset.h"
@@ -2997,8 +2999,8 @@ void GenerateMorphFactor(const UCustomizableObjectNode* Node, const UEdGraphPin&
 
 		if (validStaticFactor)
 		{
-			mu::NodeScalarPtr FactorNode = GenerateMutableSourceFloat(ConnectedPin, GenerationContext);
-			MeshNode->SetFactor(FactorNode);
+			mu::Ptr<mu::NodeScalar> FactorNode = GenerateMutableSourceFloat(ConnectedPin, GenerationContext);
+			MeshNode->Factor = FactorNode;
 		}
 	}
 }
@@ -3469,12 +3471,12 @@ mu::NodeMeshPtr GenerateMorphMesh(const UEdGraphPin* Pin,
 	// Base
 	if (MorphIndex == TypedNodeMorphs.Num() - 1)
 	{
-		Result->SetBase(SourceNode);
+		Result->Base = SourceNode;
 	}
 	else
 	{
 		mu::NodeMeshPtr NextMorph = GenerateMorphMesh(Pin, TypedNodeMorphs, MorphIndex + 1, SourceNode, GenerationContext, bOnlyConnectedLOD, TableColumnName); // TODO FutureGMT change to a for. This recursion can be problematic with the production cache
-		Result->SetBase(NextMorph);
+		Result->Base = NextMorph;
 	}
 	
 	// Target
@@ -3524,21 +3526,21 @@ mu::NodeMeshPtr GenerateMorphMesh(const UEdGraphPin* Pin,
 		{
 			bSuccess = true;
 
-			mu::NodeMeshTablePtr MorphedSourceMeshNodeTable = new mu::NodeMeshTable;
+			mu::Ptr<mu::NodeMeshTable> MorphedSourceMeshNodeTable = new mu::NodeMeshTable;
 			MorphedSourceMeshNodeTable->Table = Table;
-			MorphedSourceMeshNodeTable->SetColumn(ColumnName);
-			MorphedSourceMeshNodeTable->SetParameterName(TypedNodeTable->ParameterName);
-			MorphedSourceMeshNodeTable->SetMessageContext(MorphNode);
+			MorphedSourceMeshNodeTable->ColumnName = ColumnName;
+			MorphedSourceMeshNodeTable->ParameterName = TypedNodeTable->ParameterName;
+			MorphedSourceMeshNodeTable->SetMessageContext( MorphNode );
 			MorphedSourceMeshNodeTable->SourceDataDescriptor.SourceId = TableId;
 			MorphedSourceMeshNodeTable->SourceDataDescriptor.OptionalMaxLODSize = 0;
 
-			mu::NodeMeshMakeMorphPtr Morph = new mu::NodeMeshMakeMorph;
-			Morph->SetBase(BaseSourceMesh.get());
-			Morph->SetTarget(MorphedSourceMeshNodeTable.get());
-			Morph->SetOnlyPositionAndNormal(true);
+			mu::Ptr<mu::NodeMeshMakeMorph> Morph = new mu::NodeMeshMakeMorph;
+			Morph->Base = BaseSourceMesh;
+			Morph->Target = MorphedSourceMeshNodeTable;
+			Morph->bOnlyPositionAndNormal = true;
 			Morph->SetMessageContext(MorphNode);
 
-			Result->SetMorph(Morph);
+			Result->Morph = Morph;
 		}
 	}
 	else
@@ -3549,22 +3551,22 @@ mu::NodeMeshPtr GenerateMorphMesh(const UEdGraphPin* Pin,
 		{
 			bSuccess = true;
 
-			mu::NodeMeshConstantPtr MorphedSourceMeshNode = new mu::NodeMeshConstant;
-			MorphedSourceMeshNode->SetValue(MorphedSourceMesh);
+			mu::Ptr<mu::NodeMeshConstant> MorphedSourceMeshNode = new mu::NodeMeshConstant;
+			MorphedSourceMeshNode->Value = MorphedSourceMesh;
 			MorphedSourceMeshNode->SetMessageContext(MorphNode);
 
-			mu::NodeMeshMakeMorphPtr Morph = new mu::NodeMeshMakeMorph;
-			Morph->SetBase(BaseSourceMesh.get());
-			Morph->SetTarget(MorphedSourceMeshNode.get());
-			Morph->SetOnlyPositionAndNormal(true);
+			mu::Ptr<mu::NodeMeshMakeMorph> Morph = new mu::NodeMeshMakeMorph;
+			Morph->Base = BaseSourceMesh;
+			Morph->Target = MorphedSourceMeshNode;
+			Morph->bOnlyPositionAndNormal = true;
 			Morph->SetMessageContext(MorphNode);
 
-			Result->SetMorph(Morph);
+			Result->Morph = Morph;
 
 			if (UCustomizableObjectNodeMeshMorph* TypedMorphNode = Cast<UCustomizableObjectNodeMeshMorph>(TypedNodeMorphs[MorphIndex].OwningNode))
 			{
-				Result->SetReshapeSkeleton(TypedMorphNode->bReshapeSkeleton);
-				Result->SetReshapePhysicsVolumes(TypedMorphNode->bReshapePhysicsVolumes);
+				Result->bReshapeSkeleton = TypedMorphNode->bReshapeSkeleton;
+				Result->bReshapePhysicsVolumes = TypedMorphNode->bReshapePhysicsVolumes;
 				{
 					const UEdGraphPin* ConnectedPin = FollowInputPin(*TypedMorphNode->MeshPin());
 					const UEdGraphPin* SourceMeshPin = ConnectedPin ? FindMeshBaseSource(*ConnectedPin, false) : nullptr;
@@ -3582,7 +3584,7 @@ mu::NodeMeshPtr GenerateMorphMesh(const UEdGraphPin* Pin,
 
 						for (const FName BoneName : BonesToDeform)
 						{
-							Result->AddBoneToDeform(GenerationContext.CompilationContext.GetBoneUnique(BoneName));
+							Result->BonesToDeform.Add(GenerationContext.CompilationContext.GetBoneUnique(BoneName));
 						}
 					}
 
@@ -3598,7 +3600,7 @@ mu::NodeMeshPtr GenerateMorphMesh(const UEdGraphPin* Pin,
 	
 						for (const FName& PhysicsBoneName : PhysicsToDeform)
 						{
-							Result->AddPhysicsBodyToDeform(GenerationContext.CompilationContext.GetBoneUnique(PhysicsBoneName));
+							Result->PhysicsToDeform.Add(GenerationContext.CompilationContext.GetBoneUnique(PhysicsBoneName));
 						}
 					}
 						
@@ -3655,7 +3657,7 @@ mu::Ptr<mu::NodeMesh> GenerateMutableSourceMesh(const UEdGraphPin* Pin,
 	
 	if (const UCustomizableObjectNodeSkeletalMesh* TypedNodeSkel = Cast<UCustomizableObjectNodeSkeletalMesh>(Node))
 	{
-		mu::NodeMeshConstantPtr MeshNode = new mu::NodeMeshConstant();
+		mu::Ptr<mu::NodeMeshConstant> MeshNode = new mu::NodeMeshConstant();
 		Result = MeshNode;
 
 		if (TypedNodeSkel->SkeletalMesh.IsNull())
@@ -3739,7 +3741,7 @@ mu::Ptr<mu::NodeMesh> GenerateMutableSourceMesh(const UEdGraphPin* Pin,
 																LODIndex, SectionIndex, MeshUniqueTags, SurfaceMetadataId, GenerationContext, TypedNodeSkel, nullptr, bIsPassthrough);
 			if (MutableMesh)
 			{
-				MeshNode->SetValue(MutableMesh);
+				MeshNode->Value = MutableMesh;
 
 				if (GenerationContext.CompilationContext.Options.bSkinWeightProfilesEnabled && LODIndex>=0)
 				{
@@ -3810,7 +3812,7 @@ mu::Ptr<mu::NodeMesh> GenerateMutableSourceMesh(const UEdGraphPin* Pin,
 				{
 					NumLayouts = ImportedModel->LODModels[LODIndexLayout].NumTexCoords;
 				}
-				MeshNode->SetLayoutCount(NumLayouts);
+				MeshNode->Layouts.SetNum(NumLayouts);
 				
 				const FLayoutGenerationFlags& LayoutFlags = GenerationContext.LayoutGenerationFlags.Last();
 
@@ -3820,12 +3822,12 @@ mu::Ptr<mu::NodeMesh> GenerateMutableSourceMesh(const UEdGraphPin* Pin,
 					if (!LayoutFlags.TexturePinModes.IsValidIndex(LayoutIndex) ||
 						LayoutFlags.TexturePinModes[LayoutIndex] != EPinMode::Mutable)
 					{
-						MeshNode->SetLayout(LayoutIndex, CreateDefaultLayout());
+						MeshNode->Layouts[LayoutIndex] = CreateDefaultLayout();
 					
 						// Keep packing strategy if possible, Overlay can be valid with EPinMode == Passthrough	
 						if (Layouts.IsValidIndex(LayoutIndex) && Layouts[LayoutIndex])
 						{
-							MeshNode->GetLayout(LayoutIndex)->Strategy = ConvertLayoutStrategy(Layouts[LayoutIndex]->PackingStrategy);
+							MeshNode->Layouts[LayoutIndex]->Strategy = ConvertLayoutStrategy(Layouts[LayoutIndex]->PackingStrategy);
 						}
 
 						// Ignore layout
@@ -3839,7 +3841,7 @@ mu::Ptr<mu::NodeMesh> GenerateMutableSourceMesh(const UEdGraphPin* Pin,
 						mu::Ptr<mu::NodeLayout> LayoutNode = CreateMutableLayoutNode(GenerationContext, Layout, bLinkedToExtendMaterial, bWasEmpty); // TODO PERE: Figure out
 						LayoutNode->SetMessageContext(Node);
 
-						MeshNode->SetLayout(LayoutIndex, LayoutNode);
+						MeshNode->Layouts[LayoutIndex] = LayoutNode;
 					}
 				}
 			}
@@ -3874,7 +3876,7 @@ mu::Ptr<mu::NodeMesh> GenerateMutableSourceMesh(const UEdGraphPin* Pin,
 			return {};
 		}
 
-		mu::NodeMeshConstantPtr MeshNode = new mu::NodeMeshConstant();
+		mu::Ptr<mu::NodeMeshConstant> MeshNode = new mu::NodeMeshConstant();
 		Result = MeshNode;
 
 		if (TypedNodeStatic->StaticMesh)
@@ -3907,10 +3909,10 @@ mu::Ptr<mu::NodeMesh> GenerateMutableSourceMesh(const UEdGraphPin* Pin,
 														  FString(), 0, GenerationContext, TypedNodeStatic, nullptr, bIsPassthrough);
 			if (MutableMesh)
 			{
-				MeshNode->SetValue(MutableMesh);
+				MeshNode->Value = MutableMesh;
 
 				// Layouts
-				MeshNode->SetLayoutCount(1);
+				MeshNode->Layouts.SetNum(1);
 
 				mu::Ptr<mu::NodeLayout> LayoutNode;
 
@@ -3933,7 +3935,7 @@ mu::Ptr<mu::NodeMesh> GenerateMutableSourceMesh(const UEdGraphPin* Pin,
 					LayoutNode = CreateDefaultLayout();
 				}
 
-				MeshNode->SetLayout(0, LayoutNode);
+				MeshNode->Layouts[0] = LayoutNode;
 				LayoutNode->SetMessageContext(Node);  // We need it here because we create multiple nodes.
 
 
@@ -4048,7 +4050,7 @@ mu::Ptr<mu::NodeMesh> GenerateMutableSourceMesh(const UEdGraphPin* Pin,
 			// Check Switch Parameter arity preconditions.
 			if (const UEdGraphPin* EnumPin = FollowInputPin(*SwitchParameter))
 			{
-				mu::NodeScalarPtr SwitchParam = GenerateMutableSourceFloat(EnumPin, GenerationContext);
+				mu::Ptr<mu::NodeScalar> SwitchParam = GenerateMutableSourceFloat(EnumPin, GenerationContext);
 
 				// Switch Param not generated
 				if (!SwitchParam)
@@ -4074,15 +4076,15 @@ mu::Ptr<mu::NodeMesh> GenerateMutableSourceMesh(const UEdGraphPin* Pin,
 				const int32 NumSwitchOptions = TypedNodeMeshSwitch->GetNumElements();
 
 				mu::NodeScalarEnumParameter* EnumParameter = static_cast<mu::NodeScalarEnumParameter*>(SwitchParam.get());
-				if (NumSwitchOptions != EnumParameter->GetValueCount())
+				if (NumSwitchOptions != EnumParameter->Options.Num())
 				{
 					const FText Message = LOCTEXT("MismatchedSwitch", "Switch enum and switch node have different number of options. Please refresh the switch node to make sure the outcomes are labeled properly.");
 					GenerationContext.Log(Message, Node);
 				}
 
-				mu::NodeMeshSwitchPtr SwitchNode = new mu::NodeMeshSwitch;
-				SwitchNode->SetParameter(SwitchParam);
-				SwitchNode->SetOptionCount(NumSwitchOptions);
+				mu::Ptr<mu::NodeMeshSwitch> SwitchNode = new mu::NodeMeshSwitch;
+				SwitchNode->Parameter = SwitchParam;
+				SwitchNode->Options.SetNum(NumSwitchOptions);
 
 				for (int SelectorIndex = 0; SelectorIndex < NumSwitchOptions; ++SelectorIndex)
 				{
@@ -4090,7 +4092,7 @@ mu::Ptr<mu::NodeMesh> GenerateMutableSourceMesh(const UEdGraphPin* Pin,
 					{
 						FMutableGraphMeshGenerationData ChildMeshData;
 						Result = GenerateMutableSourceMesh(ConnectedPin, GenerationContext, ChildMeshData, SurfaceMetadataId, false, bOnlyConnectedLOD);
-						SwitchNode->SetOption(SelectorIndex, Result);
+						SwitchNode->Options[SelectorIndex] = Result;
 						MeshData.Combine(ChildMeshData);
 					}
 				}
@@ -4108,7 +4110,7 @@ mu::Ptr<mu::NodeMesh> GenerateMutableSourceMesh(const UEdGraphPin* Pin,
 
 	else if (const UCustomizableObjectNodeMeshVariation* TypedNodeMeshVar = Cast<const UCustomizableObjectNodeMeshVariation>(Node))
 	{
-		mu::NodeMeshVariationPtr MeshNode = new mu::NodeMeshVariation();
+		mu::Ptr<mu::NodeMeshVariation> MeshNode = new mu::NodeMeshVariation();
 		Result = MeshNode;
 
 		if (const UEdGraphPin* ConnectedPin = FollowInputPin(*TypedNodeMeshVar->DefaultPin()))
@@ -4117,7 +4119,7 @@ mu::Ptr<mu::NodeMesh> GenerateMutableSourceMesh(const UEdGraphPin* Pin,
 			mu::NodeMeshPtr ChildNode = GenerateMutableSourceMesh(ConnectedPin, GenerationContext, ChildMeshData, SurfaceMetadataId, false, bOnlyConnectedLOD);
 			if (ChildNode)
 			{
-				MeshNode->SetDefaultMesh(ChildNode.get());
+				MeshNode->DefaultMesh = ChildNode;
 				MeshData.Combine(ChildMeshData);
 			}
 			else
@@ -4127,18 +4129,18 @@ mu::Ptr<mu::NodeMesh> GenerateMutableSourceMesh(const UEdGraphPin* Pin,
 		}
 
 		const int32 NumVariations = TypedNodeMeshVar->GetNumVariations();
-		MeshNode->SetVariationCount(NumVariations);
+		MeshNode->Variations.SetNum(NumVariations);
 		for (int VariationIndex = 0; VariationIndex < NumVariations; ++VariationIndex)
 		{
 			const UEdGraphPin* VariationPin = TypedNodeMeshVar->VariationPin(VariationIndex);
 			if (!VariationPin) continue;
 
-			MeshNode->SetVariationTag(VariationIndex, StringCast<ANSICHAR>(*TypedNodeMeshVar->GetVariation(VariationIndex).Tag).Get());
+			MeshNode->Variations[VariationIndex].Tag = StringCast<ANSICHAR>(*TypedNodeMeshVar->GetVariation(VariationIndex).Tag).Get();
 			if (const UEdGraphPin* ConnectedPin = FollowInputPin(*VariationPin))
 			{
 				FMutableGraphMeshGenerationData VariationMeshData;
 				mu::NodeMeshPtr ChildNode = GenerateMutableSourceMesh(ConnectedPin, GenerationContext, VariationMeshData, SurfaceMetadataId, false, bOnlyConnectedLOD);
-				MeshNode->SetVariationMesh(VariationIndex, ChildNode.get());
+				MeshNode->Variations[VariationIndex].Mesh = ChildNode;
 				MeshData.Combine(VariationMeshData);
 			}
 		}
@@ -4155,7 +4157,7 @@ mu::Ptr<mu::NodeMesh> GenerateMutableSourceMesh(const UEdGraphPin* Pin,
 			mu::Ptr<mu::NodeMesh> ChildNode = GenerateMutableSourceMesh(ConnectedPin, GenerationContext, ChildMeshData, SurfaceMetadataId, false, bOnlyConnectedLOD);
 			if (ChildNode)
 			{
-				MeshNode->SetBaseMesh(ChildNode.get());
+				MeshNode->BaseMesh = ChildNode;
 				MeshData.Combine(ChildMeshData);
 			}
 			else
@@ -4170,11 +4172,11 @@ mu::Ptr<mu::NodeMesh> GenerateMutableSourceMesh(const UEdGraphPin* Pin,
 	
 		{
 
-			MeshNode->SetReshapeVertices(TypedNodeReshape->bReshapeVertices);
-			MeshNode->SetRecomputeNormals(TypedNodeReshape->bRecomputeNormals);
-			MeshNode->SetApplyLaplacian(TypedNodeReshape->bApplyLaplacianSmoothing);
-			MeshNode->SetReshapeSkeleton(TypedNodeReshape->bReshapeSkeleton);
-			MeshNode->SetReshapePhysicsVolumes(TypedNodeReshape->bReshapePhysicsVolumes);
+			MeshNode->bReshapeVertices = TypedNodeReshape->bReshapeVertices;
+			MeshNode->bRecomputeNormals = TypedNodeReshape->bRecomputeNormals;
+			MeshNode->bApplyLaplacian = TypedNodeReshape->bApplyLaplacianSmoothing;
+			MeshNode->bReshapeSkeleton = TypedNodeReshape->bReshapeSkeleton;
+			MeshNode->bReshapePhysicsVolumes = TypedNodeReshape->bReshapePhysicsVolumes;
 
 			EMeshReshapeVertexColorChannelUsage ChannelUsages[4] =
 			{
@@ -4222,11 +4224,10 @@ mu::Ptr<mu::NodeMesh> GenerateMutableSourceMesh(const UEdGraphPin* Pin,
 				};
 			};
 
-			MeshNode->SetColorUsages(
-				ConvertColorUsage(ChannelUsages[0]),
-				ConvertColorUsage(ChannelUsages[1]),
-				ConvertColorUsage(ChannelUsages[2]),
-				ConvertColorUsage(ChannelUsages[3]));
+			MeshNode->ColorRChannelUsage = ConvertColorUsage(ChannelUsages[0]);
+			MeshNode->ColorGChannelUsage = ConvertColorUsage(ChannelUsages[1]);
+			MeshNode->ColorBChannelUsage = ConvertColorUsage(ChannelUsages[2]);
+			MeshNode->ColorAChannelUsage = ConvertColorUsage(ChannelUsages[3]);
 				
 			const UEdGraphPin* ConnectedPin = FollowInputPin(*TypedNodeReshape->BaseMeshPin());
 			const UEdGraphPin* SourceMeshPin = ConnectedPin ? FindMeshBaseSource(*ConnectedPin, false) : nullptr;
@@ -4244,7 +4245,7 @@ mu::Ptr<mu::NodeMesh> GenerateMutableSourceMesh(const UEdGraphPin* Pin,
 				
 				for (const FName& BoneName : BonesToDeform)
 				{
-					MeshNode->AddBoneToDeform(GenerationContext.CompilationContext.GetBoneUnique(BoneName));
+					MeshNode->BonesToDeform.Add(GenerationContext.CompilationContext.GetBoneUnique(BoneName));
 				}
 			}
 
@@ -4259,7 +4260,7 @@ mu::Ptr<mu::NodeMesh> GenerateMutableSourceMesh(const UEdGraphPin* Pin,
 
 				for (const FName& PhysicsBoneName : PhysicsToDeform)
 				{
-					MeshNode->AddPhysicsBodyToDeform(GenerationContext.CompilationContext.GetBoneUnique(PhysicsBoneName));
+					MeshNode->PhysicsToDeform.Add(GenerationContext.CompilationContext.GetBoneUnique(PhysicsBoneName));
 				}	
 			}
 			
@@ -4285,7 +4286,7 @@ mu::Ptr<mu::NodeMesh> GenerateMutableSourceMesh(const UEdGraphPin* Pin,
 			if (ChildNode)
 			{
 				BaseShapeTriangleCount = ChildMeshData.MaxNumTriangles == ChildMeshData.MinNumTriangles ? ChildMeshData.MaxNumTriangles : -1;
-				MeshNode->SetBaseShape(ChildNode.get());	
+				MeshNode->BaseShape = ChildNode;
 			}
 			else
 			{
@@ -4302,7 +4303,7 @@ mu::Ptr<mu::NodeMesh> GenerateMutableSourceMesh(const UEdGraphPin* Pin,
 			if (ChildNode)
 			{
 				TargetShapeTriangleCount = ChildMeshData.MaxNumTriangles == ChildMeshData.MinNumTriangles ? ChildMeshData.MaxNumTriangles : -1;
-				MeshNode->SetTargetShape(ChildNode.get());
+				MeshNode->TargetShape = ChildNode;
 			}
 			else
 			{
@@ -4346,7 +4347,7 @@ mu::Ptr<mu::NodeMesh> GenerateMutableSourceMesh(const UEdGraphPin* Pin,
 					TArray<FName> ArrayBoneName;
 					TArray<FTransform> ArrayTransform;
 					UCustomizableObjectNodeAnimationPose::StaticRetrievePoseInformation(TypedNode->PoseAsset, GenerationContext.GetCurrentComponentInfo()->RefSkeletalMesh, ArrayBoneName, ArrayTransform);
-					mu::NodeMeshApplyPosePtr NodeMeshApplyPose = CreateNodeMeshApplyPose(GenerationContext, InputMeshNode, ArrayBoneName, ArrayTransform);
+					mu::Ptr<mu::NodeMeshApplyPose> NodeMeshApplyPose = CreateNodeMeshApplyPose(GenerationContext, InputMeshNode, ArrayBoneName, ArrayTransform);
 
 					if (NodeMeshApplyPose)
 					{
@@ -4364,11 +4365,11 @@ mu::Ptr<mu::NodeMesh> GenerateMutableSourceMesh(const UEdGraphPin* Pin,
 				{
 					if (const UCustomizableObjectNodeTable* TypedNodeTable = Cast<UCustomizableObjectNodeTable>(TablePosePin->GetOwningNode()))
 					{
-						mu::NodeMeshApplyPosePtr NodeMeshApplyPose = new mu::NodeMeshApplyPose();
+						mu::Ptr<mu::NodeMeshApplyPose> NodeMeshApplyPose = new mu::NodeMeshApplyPose();
 						mu::Ptr<mu::NodeMesh> MeshTableNode = GenerateMutableSourceMesh(TablePosePin, GenerationContext, MeshData, SurfaceMetadataId, false, bOnlyConnectedLOD);
 
-						NodeMeshApplyPose->SetBase(InputMeshNode);
-						NodeMeshApplyPose->SetPose(MeshTableNode);
+						NodeMeshApplyPose->Base = InputMeshNode;
+						NodeMeshApplyPose->Pose = MeshTableNode;
 
 						Result = NodeMeshApplyPose;
 					}
@@ -4393,7 +4394,7 @@ mu::Ptr<mu::NodeMesh> GenerateMutableSourceMesh(const UEdGraphPin* Pin,
 		const FString TableName = TypedNodeTable->Table ? GetNameSafe(TypedNodeTable->Table).ToLower() : GetNameSafe(TypedNodeTable->Structure).ToLower();
 		const uint32 TableId = CityHash32(reinterpret_cast<const char*>(*TableName), TableName.Len() * sizeof(FString::ElementType));
 
-		mu::NodeMeshConstantPtr EmptyNode = new mu::NodeMeshConstant();
+		mu::Ptr<mu::NodeMeshConstant> EmptyNode = new mu::NodeMeshConstant();
 		Result = EmptyNode;
 		bool bSuccess = true;
 
@@ -4435,7 +4436,7 @@ mu::Ptr<mu::NodeMesh> GenerateMutableSourceMesh(const UEdGraphPin* Pin,
 
 				if (Table)
 				{
-					mu::NodeMeshTablePtr MeshTableNode = new mu::NodeMeshTable();
+					mu::Ptr<mu::NodeMeshTable> MeshTableNode = new mu::NodeMeshTable();
 					
 					int32 LODIndexConnected = -1; // LOD which the pin is connected to
 					int32 SectionIndexConnected = -1;
@@ -4478,10 +4479,10 @@ mu::Ptr<mu::NodeMesh> GenerateMutableSourceMesh(const UEdGraphPin* Pin,
 						Result = MeshTableNode;
 
 						MeshTableNode->Table = Table;
-						MeshTableNode->SetColumn(MutableColumnName);
-						MeshTableNode->SetParameterName(TypedNodeTable->ParameterName);
-						MeshTableNode->SetNoneOption(TypedNodeTable->bAddNoneOption);
-						MeshTableNode->SetDefaultRowName(TypedNodeTable->DefaultRowName.ToString());
+						MeshTableNode->ColumnName = MutableColumnName;
+						MeshTableNode->ParameterName = TypedNodeTable->ParameterName;
+						MeshTableNode->bNoneOption = TypedNodeTable->bAddNoneOption;
+						MeshTableNode->DefaultRowName = TypedNodeTable->DefaultRowName.ToString();
 						MeshTableNode->SourceDataDescriptor.SourceId = TableId;
 						MeshTableNode->SourceDataDescriptor.OptionalMaxLODSize = 0;
 
@@ -4513,7 +4514,7 @@ mu::Ptr<mu::NodeMesh> GenerateMutableSourceMesh(const UEdGraphPin* Pin,
 							}
 
 							TArray<UCustomizableObjectLayout*> Layouts = TypedNodeTable->GetLayouts(Pin);
-							MeshTableNode->SetLayoutCount(Layouts.Num());
+							MeshTableNode->Layouts.SetNum(Layouts.Num());
 
 							if (Layouts.Num())
 							{
@@ -4525,12 +4526,12 @@ mu::Ptr<mu::NodeMesh> GenerateMutableSourceMesh(const UEdGraphPin* Pin,
 									if (!LayoutFlags.TexturePinModes.IsValidIndex(LayoutIndex) ||
 										LayoutFlags.TexturePinModes[LayoutIndex] != EPinMode::Mutable)
 									{
-										MeshTableNode->SetLayout(LayoutIndex, CreateDefaultLayout());
+										MeshTableNode->Layouts[LayoutIndex] = CreateDefaultLayout();
 
 										// Keep packing strategy if possible, Overlay can be valid with EPinMode == Passthrough	
 										if (Layouts[LayoutIndex])
 										{
-											MeshTableNode->GetLayout(LayoutIndex)->Strategy = ConvertLayoutStrategy(Layouts[LayoutIndex]->PackingStrategy);
+											MeshTableNode->Layouts[LayoutIndex]->Strategy = ConvertLayoutStrategy(Layouts[LayoutIndex]->PackingStrategy);
 										}
 
 										// Ignore layouts
@@ -4547,7 +4548,7 @@ mu::Ptr<mu::NodeMesh> GenerateMutableSourceMesh(const UEdGraphPin* Pin,
 										GenerationContext.Log(FText::FromString(msg), Node, EMessageSeverity::Warning);
 									}
 
-									MeshTableNode->SetLayout(LayoutIndex, LayoutNode);
+									MeshTableNode->Layouts[LayoutIndex] = LayoutNode;
 								}
 							}
 
@@ -4571,7 +4572,36 @@ mu::Ptr<mu::NodeMesh> GenerateMutableSourceMesh(const UEdGraphPin* Pin,
 			GenerationContext.Log(LOCTEXT("ImageTableError", "Couldn't find the data table of the node."), Node);
 		}
 	}
-	
+
+	else if (const UCustomizableObjectNodeMeshParameter* TypedNodeParam = Cast<UCustomizableObjectNodeMeshParameter>(Node))
+	{
+		mu::Ptr<mu::NodeMeshParameter> ParameterNode = new mu::NodeMeshParameter();
+
+		ParameterNode->Name = TypedNodeParam->ParameterName;
+		ParameterNode->UID = GenerationContext.GetNodeIdUnique(Node).ToString();
+
+		if (TypedNodeParam->DefaultValue)
+		{
+			ParameterNode->DefaultValue = FName(TypedNodeParam->DefaultValue->GetPathName());
+		}
+
+		GenerationContext.ParameterUIDataMap.Add(TypedNodeParam->ParameterName, FMutableParameterData(
+			TypedNodeParam->ParamUIMetadata,
+			EMutableParameterType::Texture));
+
+		// For constants we also have:
+		// animbp instance?
+		// gameplay tags?
+		// layouts?
+		// morphs?
+
+		// Force the format?
+		//mu::NodeMeshFormatPtr FormatNode = new mu::NodeMeshFormat();
+		//FormatNode->SetSource(ParameterNode);
+
+		Result = ParameterNode;
+	}
+
 	else
 	{
 		GenerationContext.Log(LOCTEXT("UnimplementedMeshNode", "Mesh node type not implemented yet."), Node);
@@ -4593,4 +4623,5 @@ mu::Ptr<mu::NodeMesh> GenerateMutableSourceMesh(const UEdGraphPin* Pin,
 
 	return Result;
 }
+
 #undef LOCTEXT_NAMESPACE
diff --git a/Engine/Plugins/Mutable/Source/CustomizableObjectEditor/Private/MuCOE/GenerateMutableSource/GenerateMutableSourceModifier.cpp b/Engine/Plugins/Mutable/Source/CustomizableObjectEditor/Private/MuCOE/GenerateMutableSource/GenerateMutableSourceModifier.cpp
--- a/Engine/Plugins/Mutable/Source/CustomizableObjectEditor/Private/MuCOE/GenerateMutableSource/GenerateMutableSourceModifier.cpp
+++ b/Engine/Plugins/Mutable/Source/CustomizableObjectEditor/Private/MuCOE/GenerateMutableSource/GenerateMutableSourceModifier.cpp
@@ -172,10 +172,10 @@ mu::Ptr<mu::NodeModifier> GenerateMutableSourceModifier(const UEdGraphPin * Pin,
 
 			if (FMatrix Matrix = TypedNodeClipMesh->Transform.ToMatrixWithScale(); Matrix != FMatrix::Identity)
 			{
-				mu::NodeMeshTransformPtr TransformMesh = new mu::NodeMeshTransform();
-				TransformMesh->SetSource(ClipMesh.get());
+				mu::Ptr<mu::NodeMeshTransform> TransformMesh = new mu::NodeMeshTransform();
+				TransformMesh->Source = ClipMesh;
 
-				TransformMesh->SetTransform(FMatrix44f(Matrix));
+				TransformMesh->Transform = FMatrix44f(Matrix);
 				ClipMesh = TransformMesh;
 			}
 
@@ -302,9 +302,9 @@ mu::Ptr<mu::NodeModifier> GenerateMutableSourceModifier(const UEdGraphPin * Pin,
 				mu::Ptr<mu::NodeMesh> MeshPtr = AddMeshNode;
 
 				mu::Ptr<mu::NodeMeshFormat> MeshFormat = new mu::NodeMeshFormat();
-				SetSurfaceFormat(GenerationContext.CompilationContext, MeshFormat->GetVertexBuffers(), MeshFormat->GetIndexBuffers(), MeshData);
+				SetSurfaceFormat(GenerationContext.CompilationContext, MeshFormat->VertexBuffers, MeshFormat->IndexBuffers, MeshData);
 
-				MeshFormat->SetSource(MeshPtr.get());
+				MeshFormat->Source = MeshPtr;
 
 				SurfNode->LODs[LODIndex].MeshAdd = MeshFormat;
 			}
@@ -573,7 +573,7 @@ mu::Ptr<mu::NodeModifier> GenerateMutableSourceModifier(const UEdGraphPin * Pin,
 
 			if (validStaticFactor)
 			{
-				mu::NodeScalarPtr FactorNode = GenerateMutableSourceFloat(ConnectedPin, GenerationContext);
+				mu::Ptr<mu::NodeScalar> FactorNode = GenerateMutableSourceFloat(ConnectedPin, GenerationContext);
 				SurfNode->MorphFactor = FactorNode;
 			}
 		}
@@ -609,10 +609,10 @@ mu::Ptr<mu::NodeModifier> GenerateMutableSourceModifier(const UEdGraphPin * Pin,
 
 			if (FMatrix Matrix = TypedNodeTransformMesh->BoundingMeshTransform.ToMatrixWithScale(); Matrix != FMatrix::Identity)
 			{
-				mu::NodeMeshTransformPtr TransformMesh = new mu::NodeMeshTransform();
-				TransformMesh->SetSource(BoundingMesh.get());
+				mu::Ptr<mu::NodeMeshTransform> TransformMesh = new mu::NodeMeshTransform();
+				TransformMesh->Source = BoundingMesh;
 
-				TransformMesh->SetTransform(FMatrix44f(Matrix));
+				TransformMesh->Transform = FMatrix44f(Matrix);
 				BoundingMesh = TransformMesh;
 			}
 
diff --git a/Engine/Plugins/Mutable/Source/CustomizableObjectEditor/Private/MuCOE/GenerateMutableSource/GenerateMutableSourceProjector.cpp b/Engine/Plugins/Mutable/Source/CustomizableObjectEditor/Private/MuCOE/GenerateMutableSource/GenerateMutableSourceProjector.cpp
--- a/Engine/Plugins/Mutable/Source/CustomizableObjectEditor/Private/MuCOE/GenerateMutableSource/GenerateMutableSourceProjector.cpp
+++ b/Engine/Plugins/Mutable/Source/CustomizableObjectEditor/Private/MuCOE/GenerateMutableSource/GenerateMutableSourceProjector.cpp
@@ -11,7 +11,7 @@
 #define LOCTEXT_NAMESPACE "CustomizableObjectEditor"
 
 
-mu::NodeProjectorPtr GenerateMutableSourceProjector(const UEdGraphPin* Pin, FMutableGraphGenerationContext& GenerationContext)
+mu::Ptr<mu::NodeProjector> GenerateMutableSourceProjector(const UEdGraphPin* Pin, FMutableGraphGenerationContext& GenerationContext)
 {
 	check(Pin)
 	RETURN_ON_CYCLE(*Pin, GenerationContext)
@@ -26,11 +26,11 @@ mu::NodeProjectorPtr GenerateMutableSourceProjector(const UEdGraphPin* Pin, FMut
 		return static_cast<mu::NodeProjector*>(Generated->Node.get());
 	}
 
-	mu::NodeProjectorPtr Result;
+	mu::Ptr<mu::NodeProjector> Result;
 	
 	if (const UCustomizableObjectNodeProjectorConstant* TypedNodeConst = Cast<UCustomizableObjectNodeProjectorConstant>(Node))
 	{
-		mu::NodeProjectorConstantPtr ProjectorNode = new mu::NodeProjectorConstant();
+		mu::Ptr<mu::NodeProjectorConstant> ProjectorNode = new mu::NodeProjectorConstant();
 		Result = ProjectorNode;
 
 		const mu::EProjectorType ProjectorType = ProjectorUtils::GetEquivalentProjectorType(TypedNodeConst->Value.ProjectionType);
@@ -80,7 +80,7 @@ mu::NodeProjectorPtr GenerateMutableSourceProjector(const UEdGraphPin* Pin, FMut
 
 	else if (const UCustomizableObjectNodeProjectorParameter* TypedNodeParam = Cast<UCustomizableObjectNodeProjectorParameter>(Node))
 	{
-		mu::NodeProjectorParameterPtr ProjectorNode = new mu::NodeProjectorParameter();
+		mu::Ptr<mu::NodeProjectorParameter> ProjectorNode = new mu::NodeProjectorParameter();
 		Result = ProjectorNode;
 
 		ProjectorNode->SetName(TypedNodeParam->ParameterName);
diff --git a/Engine/Plugins/Mutable/Source/CustomizableObjectEditor/Private/MuCOE/GenerateMutableSource/GenerateMutableSourceProjector.h b/Engine/Plugins/Mutable/Source/CustomizableObjectEditor/Private/MuCOE/GenerateMutableSource/GenerateMutableSourceProjector.h
--- a/Engine/Plugins/Mutable/Source/CustomizableObjectEditor/Private/MuCOE/GenerateMutableSource/GenerateMutableSourceProjector.h
+++ b/Engine/Plugins/Mutable/Source/CustomizableObjectEditor/Private/MuCOE/GenerateMutableSource/GenerateMutableSourceProjector.h
@@ -8,4 +8,4 @@ class UEdGraphPin;
 struct FMutableGraphGenerationContext;
 
 /** Convert a CustomizableObject Source Graph into a mutable source graph. */
-mu::NodeProjectorPtr GenerateMutableSourceProjector(const UEdGraphPin* Pin, FMutableGraphGenerationContext& GenerationContext);
+mu::Ptr<mu::NodeProjector> GenerateMutableSourceProjector(const UEdGraphPin* Pin, FMutableGraphGenerationContext& GenerationContext);
diff --git a/Engine/Plugins/Mutable/Source/CustomizableObjectEditor/Private/MuCOE/GenerateMutableSource/GenerateMutableSourceSurface.cpp b/Engine/Plugins/Mutable/Source/CustomizableObjectEditor/Private/MuCOE/GenerateMutableSource/GenerateMutableSourceSurface.cpp
--- a/Engine/Plugins/Mutable/Source/CustomizableObjectEditor/Private/MuCOE/GenerateMutableSource/GenerateMutableSourceSurface.cpp
+++ b/Engine/Plugins/Mutable/Source/CustomizableObjectEditor/Private/MuCOE/GenerateMutableSource/GenerateMutableSourceSurface.cpp
@@ -475,9 +475,10 @@ mu::Ptr<mu::NodeSurface> GenerateMutableSourceSurface(const UEdGraphPin * Pin, F
 								});
 
 
-								mu::NodeImageMipmapPtr MipmapImage = new mu::NodeImageMipmap();
-								MipmapImage->SetSource(LastImage.get());
-								MipmapImage->SetMipmapGenerationSettings(MipGenerationFilterType, mu::EAddressMode::None);
+								mu::Ptr<mu::NodeImageMipmap> MipmapImage = new mu::NodeImageMipmap();
+								MipmapImage->Source = LastImage;
+								MipmapImage->Settings.FilterType = MipGenerationFilterType;
+								MipmapImage->Settings.AddressMode = mu::EAddressMode::None;
 
 								MipmapImage->SetMessageContext(Node);
 								LastImage = MipmapImage;
@@ -486,9 +487,9 @@ mu::Ptr<mu::NodeSurface> GenerateMutableSourceSurface(const UEdGraphPin * Pin, F
 							// Apply composite image. This needs to be computed after mipmaps generation. 	
 							if (ReferenceTexture && ReferenceTexture->GetCompositeTexture() && ReferenceTexture->CompositeTextureMode != CTM_Disabled)
 							{
-								mu::NodeImageNormalCompositePtr CompositedImage = new mu::NodeImageNormalComposite();
-								CompositedImage->SetBase(LastImage.get());
-								CompositedImage->SetPower(ReferenceTexture->CompositePower);
+								mu::Ptr<mu::NodeImageNormalComposite> CompositedImage = new mu::NodeImageNormalComposite();
+								CompositedImage->Base = LastImage;
+								CompositedImage->Power = ReferenceTexture->CompositePower;
 
 								mu::ECompositeImageMode CompositeImageMode = [CompositeTextureMode = ReferenceTexture->CompositeTextureMode]() -> mu::ECompositeImageMode
 								{
@@ -503,7 +504,7 @@ mu::Ptr<mu::NodeSurface> GenerateMutableSourceSurface(const UEdGraphPin * Pin, F
 									}
 								}();
 
-								CompositedImage->SetMode(CompositeImageMode);
+								CompositedImage->Mode = CompositeImageMode;
 
 								mu::Ptr<mu::NodeImageConstant> CompositeNormalImage = new mu::NodeImageConstant();
 
@@ -516,12 +517,13 @@ mu::Ptr<mu::NodeSurface> GenerateMutableSourceSurface(const UEdGraphPin * Pin, F
 									TSharedPtr<mu::FImage> ImageConstant = GenerateImageConstant(ReferenceCompositeNormalTexture, GenerationContext, false);
 									CompositeNormalImage->SetValue(ImageConstant);
 
-									mu::NodeImageMipmapPtr NormalCompositeMipmapImage = new mu::NodeImageMipmap();
+									mu::Ptr<mu::NodeImageMipmap> NormalCompositeMipmapImage = new mu::NodeImageMipmap();
 									const uint32 MipsToSkip = ComputeLODBiasForTexture(GenerationContext, *ReferenceCompositeNormalTexture, ReferenceTexture);
-									NormalCompositeMipmapImage->SetSource(ResizeTextureByNumMips(CompositeNormalImage, MipsToSkip));
-									NormalCompositeMipmapImage->SetMipmapGenerationSettings(mu::EMipmapFilterType::SimpleAverage, mu::EAddressMode::None);
+									NormalCompositeMipmapImage->Source = ResizeTextureByNumMips(CompositeNormalImage, MipsToSkip);
+									NormalCompositeMipmapImage->Settings.FilterType = mu::EMipmapFilterType::SimpleAverage;
+									NormalCompositeMipmapImage->Settings.AddressMode = mu::EAddressMode::None;
 
-									CompositedImage->SetNormal(NormalCompositeMipmapImage);
+									CompositedImage->Normal = NormalCompositeMipmapImage;
 
 									CompositeNormalImage->SourceDataDescriptor.OptionalMaxLODSize = 0;
 									if (GenerationContext.CompilationContext.Options.TargetPlatform)
@@ -542,9 +544,9 @@ mu::Ptr<mu::NodeSurface> GenerateMutableSourceSurface(const UEdGraphPin * Pin, F
 							}
 
 							mu::Ptr<mu::NodeImage> FormatSource = LastImage;
-							mu::NodeImageFormatPtr FormatImage = new mu::NodeImageFormat();
-							FormatImage->SetSource(LastImage.get());
-							FormatImage->SetFormat(mu::EImageFormat::RGBA_UByte);
+							mu::Ptr<mu::NodeImageFormat> FormatImage = new mu::NodeImageFormat();
+							FormatImage->Source = LastImage.get();
+							FormatImage->Format = mu::EImageFormat::RGBA_UByte;
 							FormatImage->SetMessageContext(Node);
 							LastImage = FormatImage;
 
@@ -614,20 +616,19 @@ mu::Ptr<mu::NodeSurface> GenerateMutableSourceSurface(const UEdGraphPin * Pin, F
 
 											if (bIsNormalMapFormat)
 											{
-
 												// Insert a channel swizzle
 												mu::Ptr<mu::NodeImageSwizzle> Swizzle = new mu::NodeImageSwizzle;
 												Swizzle->SetFormat( mu::EImageFormat::RGBA_UByte );
-												Swizzle->SetSource(0, FormatSource);
-												Swizzle->SetSource(1, FormatSource);
-												Swizzle->SetSource(2, FormatSource);
-												Swizzle->SetSource(3, FormatSource);
-												Swizzle->SetSourceChannel(0, 0);
-												Swizzle->SetSourceChannel(1, 0);
-												Swizzle->SetSourceChannel(2, 0);
-												Swizzle->SetSourceChannel(3, 1);
-
-												FormatImage->SetSource( Swizzle.get() );
+												Swizzle->Sources[0] = FormatSource;
+												Swizzle->Sources[1] = FormatSource;
+												Swizzle->Sources[2] = FormatSource;
+												Swizzle->Sources[3] = FormatSource;
+												Swizzle->SourceChannels[0] = 0;
+												Swizzle->SourceChannels[1] = 0;
+												Swizzle->SourceChannels[2] = 0;
+												Swizzle->SourceChannels[3] = 1;
+
+												FormatImage->Source = Swizzle;
 											}
 										}
 									}
@@ -650,7 +651,8 @@ mu::Ptr<mu::NodeSurface> GenerateMutableSourceSurface(const UEdGraphPin * Pin, F
 										MutableFormatIfAlpha = mu::EImageFormat::RGBA_UByte;
 									}
 
-									FormatImage->SetFormat(MutableFormat, MutableFormatIfAlpha);
+									FormatImage->Format = MutableFormat;
+									FormatImage->FormatIfAlpha = MutableFormatIfAlpha;
 								}
 							}
 
@@ -751,7 +753,7 @@ mu::Ptr<mu::NodeSurface> GenerateMutableSourceSurface(const UEdGraphPin * Pin, F
 			{
 				if (const UEdGraphPin* ConnectedPin = FollowInputPin(*ScalarPin))
 				{
-					mu::NodeScalarPtr ScalarNode = GenerateMutableSourceFloat(ConnectedPin, GenerationContext);
+					mu::Ptr<mu::NodeScalar> ScalarNode = GenerateMutableSourceFloat(ConnectedPin, GenerationContext);
 
 					// Encoding material layer in mutable name
 					if (const int32 LayerIndex = TypedNodeMat->GetParameterLayerIndex(EMaterialParameterType::Scalar, ScalarIndex); LayerIndex != INDEX_NONE)
@@ -787,7 +789,7 @@ mu::Ptr<mu::NodeSurface> GenerateMutableSourceSurface(const UEdGraphPin * Pin, F
 					MaterialName += MaterialPin->PinId.ToString();
 
 					GenerationContext.CurrentMaterialTableParameterId = MaterialName;
-					mu::NodeScalarPtr ScalarNode = GenerateMutableSourceFloat(ConnectedPin, GenerationContext);
+					mu::Ptr<mu::NodeScalar> ScalarNode = GenerateMutableSourceFloat(ConnectedPin, GenerationContext);
 
 					SurfNode->Scalars[MaterialIndex].Scalar = ScalarNode;
 					SurfNode->Scalars[MaterialIndex].Name = MaterialName;
@@ -795,8 +797,8 @@ mu::Ptr<mu::NodeSurface> GenerateMutableSourceSurface(const UEdGraphPin * Pin, F
 			}
 			else
 			{
-				mu::NodeScalarConstantPtr ScalarNode = new mu::NodeScalarConstant();
-				ScalarNode->SetValue(GenerationContext.CurrentReferencedMaterialIndex);
+				mu::Ptr<mu::NodeScalarConstant> ScalarNode = new mu::NodeScalarConstant();
+				ScalarNode->Value = GenerationContext.CurrentReferencedMaterialIndex;
 
 				SurfNode->Scalars[MaterialIndex].Scalar = ScalarNode;
 				SurfNode->Scalars[MaterialIndex].Name = MaterialName;
@@ -838,10 +840,11 @@ mu::Ptr<mu::NodeSurface> GenerateMutableSourceSurface(const UEdGraphPin * Pin, F
 
 					if (!ProjectorInfo || !ProjectorInfo->bIsAlternateResolutionResized)
 					{
-						mu::NodeImageResizePtr NodeImageResize = new mu::NodeImageResize;
-						NodeImageResize->SetRelative(true);
-						NodeImageResize->SetSize(*ResolutionFactor, *ResolutionFactor);
-						NodeImageResize->SetBase(SurfNode2->Images[ImageIndex].Image);
+						mu::Ptr<mu::NodeImageResize> NodeImageResize = new mu::NodeImageResize;
+						NodeImageResize->bRelative = true;
+						NodeImageResize->SizeX = *ResolutionFactor;
+						NodeImageResize->SizeY = *ResolutionFactor;
+						NodeImageResize->Base = SurfNode2->Images[ImageIndex].Image;
 
 						SurfNode2->Images[ImageIndex].Image = NodeImageResize;
 
@@ -939,7 +942,7 @@ mu::Ptr<mu::NodeSurface> GenerateMutableSourceSurface(const UEdGraphPin * Pin, F
 			// Check Switch Parameter arity preconditions.
 			if (const UEdGraphPin* EnumPin = FollowInputPin(*SwitchParameter))
 			{
-				mu::NodeScalarPtr SwitchParam = GenerateMutableSourceFloat(EnumPin, GenerationContext);
+				mu::Ptr<mu::NodeScalar> SwitchParam = GenerateMutableSourceFloat(EnumPin, GenerationContext);
 
 				// Switch Param not generated
 				if (!SwitchParam)
@@ -965,7 +968,7 @@ mu::Ptr<mu::NodeSurface> GenerateMutableSourceSurface(const UEdGraphPin * Pin, F
 				const int32 NumSwitchOptions = TypedNodeSwitch->GetNumElements();
 
 				mu::NodeScalarEnumParameter* EnumParameter = static_cast<mu::NodeScalarEnumParameter*>(SwitchParam.get());
-				if (NumSwitchOptions != EnumParameter->GetValueCount())
+				if (NumSwitchOptions != EnumParameter->Options.Num())
 				{
 					const FText Message = LOCTEXT("MismatchedSwitch", "Switch enum and switch node have different number of options. Please refresh the switch node to make sure the outcomes are labeled properly.");
 					GenerationContext.Log(Message, Node);
diff --git a/Engine/Plugins/Mutable/Source/CustomizableObjectEditor/Private/MuCOE/GraphTraversal.cpp b/Engine/Plugins/Mutable/Source/CustomizableObjectEditor/Private/MuCOE/GraphTraversal.cpp
--- a/Engine/Plugins/Mutable/Source/CustomizableObjectEditor/Private/MuCOE/GraphTraversal.cpp
+++ b/Engine/Plugins/Mutable/Source/CustomizableObjectEditor/Private/MuCOE/GraphTraversal.cpp
@@ -18,6 +18,7 @@
 #include "MuCOE/Nodes/CustomizableObjectNodeMeshReshape.h"
 #include "MuCOE/Nodes/CustomizableObjectNodeMeshSwitch.h"
 #include "MuCOE/Nodes/CustomizableObjectNodeMeshVariation.h"
+#include "MuCOE/Nodes/CustomizableObjectNodeMeshParameter.h"
 #include "MuCOE/Nodes/CustomizableObjectNodeObject.h"
 #include "MuCOE/Nodes/CustomizableObjectNodeSkeletalMesh.h"
 #include "MuCOE/Nodes/CustomizableObjectNodeStaticMesh.h"
@@ -492,7 +493,7 @@ UCustomizableObjectNodeObject* GraphTraversal::GetFullGraphRootNode(const UCusto
 const UEdGraphPin* FindMeshBaseSource(const UEdGraphPin& Pin, const bool bOnlyLookForStaticMesh)
 {
 	check(Pin.Direction == EGPD_Output);
-	check(Pin.PinType.PinCategory == UEdGraphSchema_CustomizableObject::PC_Mesh 
+	check(Pin.PinType.PinCategory == UEdGraphSchema_CustomizableObject::PC_Mesh
 		||
 		Pin.PinType.PinCategory == UEdGraphSchema_CustomizableObject::PC_PassThroughMesh
 		||
@@ -500,7 +501,7 @@ const UEdGraphPin* FindMeshBaseSource(const UEdGraphPin& Pin, const bool bOnlyLo
 		||
 		Pin.PinType.PinCategory == UEdGraphSchema_CustomizableObject::PC_Modifier
 	);
-	
+
 	const UEdGraphNode* Node = Pin.GetOwningNode();
 	check(Node);
 
@@ -589,7 +590,7 @@ const UEdGraphPin* FindMeshBaseSource(const UEdGraphPin& Pin, const bool bOnlyLo
 			return FindMeshBaseSource(*ConnectedPin, bOnlyLookForStaticMesh);
 		}
 	}
-	
+
 	else if (const UCustomizableObjectNodeMeshMorphStackDefinition* TypedNodeMorphStackDef = Cast<UCustomizableObjectNodeMeshMorphStackDefinition>(Node))
 	{
 		if (const UEdGraphPin* ConnectedPin = FollowInputPin(*TypedNodeMorphStackDef->GetMeshPin()))
@@ -613,7 +614,7 @@ const UEdGraphPin* FindMeshBaseSource(const UEdGraphPin& Pin, const bool bOnlyLo
 			return FindMeshBaseSource(*ConnectedPin, bOnlyLookForStaticMesh);
 		}
 	}
-	
+
 	else if (Cast<UCustomizableObjectNodeTable>(Node))
 	{
 		if (!bOnlyLookForStaticMesh)
@@ -630,6 +631,14 @@ const UEdGraphPin* FindMeshBaseSource(const UEdGraphPin& Pin, const bool bOnlyLo
 		}
 	}
 
+	else if (Cast<UCustomizableObjectNodeMeshParameter>(Node))
+	{
+		if (!bOnlyLookForStaticMesh)
+		{
+			return &Pin;
+		}
+	}
+
 	else
 	{
 		unimplemented(); // Case missing.
diff --git a/Engine/Plugins/Mutable/Source/CustomizableObjectEditor/Private/MuCOE/Nodes/CustomizableObjectNodeExtensionDataSwitch.cpp b/Engine/Plugins/Mutable/Source/CustomizableObjectEditor/Private/MuCOE/Nodes/CustomizableObjectNodeExtensionDataSwitch.cpp
--- a/Engine/Plugins/Mutable/Source/CustomizableObjectEditor/Private/MuCOE/Nodes/CustomizableObjectNodeExtensionDataSwitch.cpp
+++ b/Engine/Plugins/Mutable/Source/CustomizableObjectEditor/Private/MuCOE/Nodes/CustomizableObjectNodeExtensionDataSwitch.cpp
@@ -31,7 +31,7 @@ mu::Ptr<mu::NodeExtensionData> UCustomizableObjectNodeExtensionDataSwitch::Gener
 	const int32 NumParameters = FollowInputPinArray(*SwitchParameter()).Num();
 
 	const UEdGraphPin* EnumPin = FollowInputPin(*SwitchParameter());
-	mu::NodeScalarPtr SwitchParam = EnumPin ? GenerateMutableSourceFloat(EnumPin, InCompilerInterface.GenerationContext) : nullptr;
+	mu::Ptr<mu::NodeScalar> SwitchParam = EnumPin ? GenerateMutableSourceFloat(EnumPin, InCompilerInterface.GenerationContext) : nullptr;
 
 	const int32 NumSwitchOptions = GetNumElements();
 
diff --git a/Engine/Plugins/Mutable/Source/CustomizableObjectEditor/Private/MuCOE/Nodes/CustomizableObjectNodeMeshParameter.cpp b/Engine/Plugins/Mutable/Source/CustomizableObjectEditor/Private/MuCOE/Nodes/CustomizableObjectNodeMeshParameter.cpp
new file mode 100644
--- /dev/null
+++ b/Engine/Plugins/Mutable/Source/CustomizableObjectEditor/Private/MuCOE/Nodes/CustomizableObjectNodeMeshParameter.cpp
@@ -0,0 +1,70 @@
+// Copyright Epic Games, Inc. All Rights Reserved.
+
+#include "MuCOE/Nodes/CustomizableObjectNodeMeshParameter.h"
+
+#include "MuCO/CustomizableObjectCustomVersion.h"
+#include "MuCOE/EdGraphSchema_CustomizableObject.h"
+
+class UCustomizableObjectNodeRemapPins;
+
+#define LOCTEXT_NAMESPACE "CustomizableObjectEditor"
+
+
+void UCustomizableObjectNodeMeshParameter::AllocateDefaultPins(UCustomizableObjectNodeRemapPins* RemapPins)
+{
+	const FName Type = UEdGraphSchema_CustomizableObject::PC_Mesh;
+	const FName PinName = UEdGraphSchema_CustomizableObject::GetPinCategoryName(Type);
+	const FText PinFriendlyName = UEdGraphSchema_CustomizableObject::GetPinCategoryFriendlyName(Type);
+	
+	UEdGraphPin* ValuePin = CustomCreatePin(EGPD_Output, Type, PinName);
+	ValuePin->bDefaultValueIsIgnored = true;
+	ValuePin->PinFriendlyName = PinFriendlyName;
+}
+
+
+bool UCustomizableObjectNodeMeshParameter::IsExperimental() const
+{
+	return true;
+}
+
+
+FText UCustomizableObjectNodeMeshParameter::GetNodeTitle(ENodeTitleType::Type TitleType) const
+{
+	if (TitleType == ENodeTitleType::ListView)
+	{
+		return LOCTEXT("Mesh_Parameter", "Mesh Parameter");
+	}
+	else if (TitleType == ENodeTitleType::EditableTitle)
+	{
+		return FText::Format(LOCTEXT("Mesh_Parameter_EditableTitle", "{0}"), FText::FromString(ParameterName));
+	}
+	else
+	{
+		return FText::Format(LOCTEXT("Mesh_Parameter_Title", "{0}\nMesh Parameter"), FText::FromString(ParameterName));
+	}
+}
+
+
+FLinearColor UCustomizableObjectNodeMeshParameter::GetNodeTitleColor() const
+{
+	const UEdGraphSchema_CustomizableObject* Schema = GetDefault<UEdGraphSchema_CustomizableObject>();
+	return Schema->GetPinTypeColor(Schema->PC_Mesh);
+}
+
+
+FText UCustomizableObjectNodeMeshParameter::GetTooltipText() const
+{
+	return LOCTEXT("Mesh_Parameter_Tooltip", "Expose a runtime modifiable Mesh parameter from the Customizable Object.");
+}
+
+
+void UCustomizableObjectNodeMeshParameter::OnRenameNode(const FString& NewName)
+{
+	if (!NewName.IsEmpty())
+	{
+		ParameterName = NewName;
+	}
+}
+
+#undef LOCTEXT_NAMESPACE
+
diff --git a/Engine/Plugins/Mutable/Source/CustomizableObjectEditor/Private/MuCOE/Nodes/CustomizableObjectNodeMeshParameter.h b/Engine/Plugins/Mutable/Source/CustomizableObjectEditor/Private/MuCOE/Nodes/CustomizableObjectNodeMeshParameter.h
new file mode 100644
--- /dev/null
+++ b/Engine/Plugins/Mutable/Source/CustomizableObjectEditor/Private/MuCOE/Nodes/CustomizableObjectNodeMeshParameter.h
@@ -0,0 +1,47 @@
+// Copyright Epic Games, Inc. All Rights Reserved.
+
+#pragma once
+
+#include "MuCOE/Nodes/CustomizableObjectNode.h"
+
+#include "CustomizableObjectNodeMeshParameter.generated.h"
+
+namespace ENodeTitleType { enum Type : int; }
+
+class UCustomizableObjectNodeRemapPins;
+class UObject;
+class UMesh2D;
+
+
+UCLASS()
+class CUSTOMIZABLEOBJECTEDITOR_API UCustomizableObjectNodeMeshParameter : public UCustomizableObjectNode
+{
+public:
+	GENERATED_BODY()
+
+	/** Default value of the parameter. */
+	UPROPERTY(EditAnywhere, Category=CustomizableObject)
+	TObjectPtr<USkeletalMesh> DefaultValue;
+
+	UPROPERTY(EditAnywhere, Category = CustomizableObject)
+	int32 DefaultValueSection = 0;
+
+	UPROPERTY(EditAnywhere, Category=CustomizableObject)
+	FString ParameterName = "Mesh Param";
+
+	UPROPERTY(EditAnywhere, Category = UI, meta = (DisplayName = "Parameter UI Metadata"))
+	FMutableParamUIMetadata ParamUIMetadata;
+
+	// EdGraphNode interface
+	virtual FText GetNodeTitle(ENodeTitleType::Type TitleType) const override;
+	virtual FLinearColor GetNodeTitleColor() const override;
+	virtual FText GetTooltipText() const override;
+	virtual void OnRenameNode(const FString& NewName) override;
+	virtual bool GetCanRenameNode() const override { return true; }
+
+	// UCustomizableObjectNode interface
+	virtual void AllocateDefaultPins(UCustomizableObjectNodeRemapPins* RemapPins) override;
+	virtual bool IsAffectedByLOD() const override { return false; }
+	virtual bool IsExperimental() const override;
+};
+
diff --git a/Engine/Plugins/Mutable/Source/CustomizableObjectEditor/Private/MuCOE/Nodes/CustomizableObjectNodeObjectDetails.cpp b/Engine/Plugins/Mutable/Source/CustomizableObjectEditor/Private/MuCOE/Nodes/CustomizableObjectNodeObjectDetails.cpp
--- a/Engine/Plugins/Mutable/Source/CustomizableObjectEditor/Private/MuCOE/Nodes/CustomizableObjectNodeObjectDetails.cpp
+++ b/Engine/Plugins/Mutable/Source/CustomizableObjectEditor/Private/MuCOE/Nodes/CustomizableObjectNodeObjectDetails.cpp
@@ -18,6 +18,7 @@
 #include "MuCOE/Nodes/CustomizableObjectNodeGroupProjectorParameter.h"
 #include "MuCOE/Nodes/CustomizableObjectNodeTable.h"
 #include "MuCOE/Nodes/CustomizableObjectNodeTextureParameter.h"
+#include "MuCOE/Nodes/CustomizableObjectNodeMeshParameter.h"
 #include "MuCOE/SCustomizableObjectNodeObjectRTMorphTargetOverride.h"
 #include "MuCOE/UnrealEditorPortabilityHelpers.h"
 #include "MuCOE/Widgets/CustomizableObjectLODReductionSettings.h"
@@ -457,6 +458,17 @@ void FCustomizableObjectNodeObjectDetails::FillParameterNamesArray()
 			}
 		}
 
+		TArray<UCustomizableObjectNodeMeshParameter*> MeshParameterNodes;
+		Source->GetNodesOfClass(MeshParameterNodes);
+
+		for (const UCustomizableObjectNodeMeshParameter* MeshParameterNode : MeshParameterNodes)
+		{
+			if (MeshParameterNode)
+			{
+				BaseObjectNode->ParameterNames.Add(MeshParameterNode->ParameterName);
+			}
+		}
+
 		TArray<UCustomizableObjectNodeTable*> TableNodes;
 		Source->GetNodesOfClass(TableNodes);
 
diff --git a/Engine/Plugins/Mutable/Source/CustomizableObjectEditor/Private/MuCOE/SMutableGraphViewer.cpp b/Engine/Plugins/Mutable/Source/CustomizableObjectEditor/Private/MuCOE/SMutableGraphViewer.cpp
--- a/Engine/Plugins/Mutable/Source/CustomizableObjectEditor/Private/MuCOE/SMutableGraphViewer.cpp
+++ b/Engine/Plugins/Mutable/Source/CustomizableObjectEditor/Private/MuCOE/SMutableGraphViewer.cpp
@@ -59,28 +59,6 @@
 #include "MuT/NodeScalarSwitch.h"
 #include "MuT/NodeScalarTable.h"
 #include "MuT/NodeObjectNew.h"
-
-#include "MuT/NodeObjectGroupPrivate.h"
-#include "MuT/NodeImageFormatPrivate.h"
-#include "MuT/NodeMeshFormatPrivate.h"
-#include "MuT/NodeMeshConstantPrivate.h"
-#include "MuT/NodeImageSwitchPrivate.h"
-#include "MuT/NodeImageLayerColourPrivate.h"
-#include "MuT/NodeImageLayerPrivate.h"
-#include "MuT/NodeImageMipmapPrivate.h"
-#include "MuT/NodeImageResizePrivate.h"
-#include "MuT/NodeImageInterpolatePrivate.h"
-#include "MuT/NodeImagePlainColourPrivate.h"
-#include "MuT/NodeImageProjectPrivate.h"
-#include "MuT/NodeMeshMorphPrivate.h"
-#include "MuT/NodeScalarEnumParameterPrivate.h"
-#include "MuT/NodeScalarParameterPrivate.h"
-#include "MuT/NodeMeshMakeMorphPrivate.h"
-#include "MuT/NodeScalarCurvePrivate.h"
-#include "MuT/NodeProjectorPrivate.h"
-#include "MuT/NodeImageInvertPrivate.h"
-#include "MuT/NodeImageSwizzlePrivate.h"
-#include "MuT/NodeImageMultiLayerPrivate.h"
 #include "MuT/NodeModifierMeshTransformInMesh.h"
 #include "Widgets/MutableExpanderArrow.h"
 
@@ -122,6 +100,7 @@ static const char* MutableNodeNames[] =
 	"MeshVariation",
 	"MeshReshape",
 	"MeshClipDeform",
+	"MeshParameter",
 
 	"Image",
 	"ImageConstant",
@@ -227,7 +206,7 @@ static const char* MutableNodeNames[] =
 	"MatrixParameter",
 };
 
-static_assert(sizeof(MutableNodeNames) == SIZE_T(mu::Node::EType::Count) * sizeof(char*));
+static_assert(UE_ARRAY_COUNT(MutableNodeNames) == SIZE_T(mu::Node::EType::Count));
 
 
 class SMutableGraphTreeRow : public STableRow<TSharedPtr<FMutableGraphTreeElement>>
@@ -465,10 +444,9 @@ void SMutableGraphViewer::GetChildrenForInfo(TSharedPtr<FMutableGraphTreeElement
 	else if (ParentNode->GetType() == mu::NodeObjectGroup::GetStaticType())
 	{
 		mu::NodeObjectGroup* ObjectGroup = StaticCast<mu::NodeObjectGroup*>(ParentNode);
-		mu::NodeObjectGroup::Private* Private = ObjectGroup->GetPrivate();
-		for (int32 l = 0; l < Private->m_children.Num(); ++l)
+		for (int32 l = 0; l < ObjectGroup->Children.Num(); ++l)
 		{
-			AddChildFunc(Private->m_children[l].get(), TEXT("CHILD"));
+			AddChildFunc(ObjectGroup->Children[l].get(), TEXT("CHILD"));
 		}
 	}
 
@@ -598,25 +576,22 @@ void SMutableGraphViewer::GetChildrenForInfo(TSharedPtr<FMutableGraphTreeElement
 	else if (ParentNode->GetType() == mu::NodeMeshConstant::GetStaticType())
 	{
 		mu::NodeMeshConstant* MeshConstantVar = StaticCast<mu::NodeMeshConstant*>(ParentNode);
-		mu::NodeMeshConstant::Private* Private = MeshConstantVar->GetPrivate();
-		for (int32 LayoutIndex = 0; LayoutIndex < Private->Layouts.Num(); LayoutIndex++)
+		for (int32 LayoutIndex = 0; LayoutIndex < MeshConstantVar->Layouts.Num(); LayoutIndex++)
 		{
-			AddChildFunc(Private->Layouts[LayoutIndex].get(), FString::Printf(TEXT("LAYOUT [%d]"), LayoutIndex));
+			AddChildFunc(MeshConstantVar->Layouts[LayoutIndex].get(), FString::Printf(TEXT("LAYOUT [%d]"), LayoutIndex));
 		}
 	}
 
 	else if (ParentNode->GetType() == mu::NodeImageFormat::GetStaticType())
 	{
 		mu::NodeImageFormat* ImageFormatVar = StaticCast<mu::NodeImageFormat*>(ParentNode);
-		mu::NodeImageFormat::Private* Private = ImageFormatVar->GetPrivate();
-		AddChildFunc(Private->m_source.get(), FString::Printf(TEXT("SOURCE IMAGE")));
+		AddChildFunc(ImageFormatVar->Source.get(), FString::Printf(TEXT("SOURCE IMAGE")));
 	}
 
 	else if (ParentNode->GetType() == mu::NodeMeshFormat::GetStaticType())
 	{
 		mu::NodeMeshFormat* MeshFormatVar = StaticCast<mu::NodeMeshFormat*>(ParentNode);
-		mu::NodeMeshFormat::Private* Private = MeshFormatVar->GetPrivate();
-		AddChildFunc(Private->Source.get(), FString::Printf(TEXT("SOURCE MESH")));
+		AddChildFunc(MeshFormatVar->Source.get(), FString::Printf(TEXT("SOURCE MESH")));
 	}
 
 	else if (ParentNode->GetType() == mu::NodeModifierMeshClipMorphPlane::GetStaticType())
@@ -653,73 +628,65 @@ void SMutableGraphViewer::GetChildrenForInfo(TSharedPtr<FMutableGraphTreeElement
 	else if (ParentNode->GetType() == mu::NodeImageSwitch::GetStaticType())
 	{
 		mu::NodeImageSwitch* ImageSwitchVar = StaticCast<mu::NodeImageSwitch*>(ParentNode);
-		mu::NodeImageSwitch::Private* Private = ImageSwitchVar->GetPrivate();
-		AddChildFunc(Private->m_pParameter.get(), FString::Printf(TEXT("PARAM")));
-		for (int32 OptionIndex = 0; OptionIndex < Private->Options.Num(); OptionIndex++)
+		AddChildFunc(ImageSwitchVar->Parameter.get(), FString::Printf(TEXT("PARAM")));
+		for (int32 OptionIndex = 0; OptionIndex < ImageSwitchVar->Options.Num(); OptionIndex++)
 		{
-			AddChildFunc(Private->Options[OptionIndex].get(), FString::Printf(TEXT("OPTION [%d]"), OptionIndex));
+			AddChildFunc(ImageSwitchVar->Options[OptionIndex].get(), FString::Printf(TEXT("OPTION [%d]"), OptionIndex));
 		}
 	}
 
 	else if (ParentNode->GetType() == mu::NodeImageMipmap::GetStaticType())
 	{
 		mu::NodeImageMipmap* ImageMipMapVar = StaticCast<mu::NodeImageMipmap*>(ParentNode);
-		mu::NodeImageMipmap::Private* Private = ImageMipMapVar->GetPrivate();
-		AddChildFunc(Private->m_pSource.get(), FString::Printf(TEXT("SOURCE")));
-		AddChildFunc(Private->m_pFactor.get(), FString::Printf(TEXT("FACTOR")));
+		AddChildFunc(ImageMipMapVar->Source.get(), FString::Printf(TEXT("SOURCE")));
+		AddChildFunc(ImageMipMapVar->Factor.get(), FString::Printf(TEXT("FACTOR")));
 	}
 
 	else if (ParentNode->GetType() == mu::NodeImageLayer::GetStaticType())
 	{
 		mu::NodeImageLayer* ImageLayerVar = StaticCast<mu::NodeImageLayer*>(ParentNode);
-		mu::NodeImageLayer::Private* Private = ImageLayerVar->GetPrivate();
-		AddChildFunc(Private->m_pBase.get(), FString::Printf(TEXT("BASE")));
-		AddChildFunc(Private->m_pMask.get(), FString::Printf(TEXT("MASK")));
-		AddChildFunc(Private->m_pBlended.get(), FString::Printf(TEXT("BLEND")));
+		AddChildFunc(ImageLayerVar->Base.get(), FString::Printf(TEXT("BASE")));
+		AddChildFunc(ImageLayerVar->Mask.get(), FString::Printf(TEXT("MASK")));
+		AddChildFunc(ImageLayerVar->Blended.get(), FString::Printf(TEXT("BLEND")));
 	}
 	
 	else if (ParentNode->GetType() == mu::NodeImageLayerColour::GetStaticType())
 	{
 		mu::NodeImageLayerColour* ImageLayerColourVar = StaticCast<mu::NodeImageLayerColour*>(ParentNode);
-		mu::NodeImageLayerColour::Private* Private = ImageLayerColourVar->GetPrivate();
-		AddChildFunc(Private->m_pBase.get(), FString::Printf(TEXT("BASE")));
-		AddChildFunc(Private->m_pMask.get(), FString::Printf(TEXT("MASK")));
-		AddChildFunc(Private->m_pColour.get(), FString::Printf(TEXT("COLOR")));
+		AddChildFunc(ImageLayerColourVar->Base.get(), FString::Printf(TEXT("BASE")));
+		AddChildFunc(ImageLayerColourVar->Mask.get(), FString::Printf(TEXT("MASK")));
+		AddChildFunc(ImageLayerColourVar->Colour.get(), FString::Printf(TEXT("COLOR")));
 	}
 
 	else if (ParentNode->GetType() == mu::NodeImageResize::GetStaticType())
 	{
 		mu::NodeImageResize* ImageResizeVar = StaticCast<mu::NodeImageResize*>(ParentNode);
-		mu::NodeImageResize::Private* Private = ImageResizeVar->GetPrivate();
-		AddChildFunc(Private->m_pBase.get(), FString::Printf(TEXT("BASE")));
+		AddChildFunc(ImageResizeVar->Base.get(), FString::Printf(TEXT("BASE")));
 	}
 
 	else if (ParentNode->GetType() == mu::NodeMeshMorph::GetStaticType())
 	{
 		mu::NodeMeshMorph* MeshMorphVar = StaticCast<mu::NodeMeshMorph*>(ParentNode);
-		mu::NodeMeshMorph::Private* Private = MeshMorphVar->GetPrivate();
-		AddChildFunc(Private->Base.get(), FString::Printf(TEXT("BASE")));
-		AddChildFunc(Private->Morph.get(), FString::Printf(TEXT("MORPH")));
-		AddChildFunc(Private->Factor.get(), FString::Printf(TEXT("FACTOR")));
+		AddChildFunc(MeshMorphVar->Base.get(), FString::Printf(TEXT("BASE")));
+		AddChildFunc(MeshMorphVar->Morph.get(), FString::Printf(TEXT("MORPH")));
+		AddChildFunc(MeshMorphVar->Factor.get(), FString::Printf(TEXT("FACTOR")));
 	}
 
 	else if (ParentNode->GetType() == mu::NodeImageProject::GetStaticType())
 	{
 		mu::NodeImageProject* ImageProjectVar = StaticCast<mu::NodeImageProject*>(ParentNode);
-		mu::NodeImageProject::Private* Private = ImageProjectVar->GetPrivate();
-		AddChildFunc(Private->m_pProjector.get(), FString::Printf(TEXT("PROJECTOR")));
-		AddChildFunc(Private->m_pMesh.get(), FString::Printf(TEXT("MESH")));
-		AddChildFunc(Private->m_pImage.get(), FString::Printf(TEXT("IMAGE")));
-		AddChildFunc(Private->m_pMask.get(), FString::Printf(TEXT("MASK")));
-		AddChildFunc(Private->m_pAngleFadeStart.get(), FString::Printf(TEXT("FADE START ANGLE")));
-		AddChildFunc(Private->m_pAngleFadeEnd.get(), FString::Printf(TEXT("FADE END ANGLE")));
+		AddChildFunc(ImageProjectVar->Projector.get(), FString::Printf(TEXT("PROJECTOR")));
+		AddChildFunc(ImageProjectVar->Mesh.get(), FString::Printf(TEXT("MESH")));
+		AddChildFunc(ImageProjectVar->Image.get(), FString::Printf(TEXT("IMAGE")));
+		AddChildFunc(ImageProjectVar->Mask.get(), FString::Printf(TEXT("MASK")));
+		AddChildFunc(ImageProjectVar->AngleFadeStart.get(), FString::Printf(TEXT("FADE START ANGLE")));
+		AddChildFunc(ImageProjectVar->AngleFadeEnd.get(), FString::Printf(TEXT("FADE END ANGLE")));
 	}
 
 	else if (ParentNode->GetType() == mu::NodeImagePlainColour::GetStaticType())
 	{
 		mu::NodeImagePlainColour* ImagePlainColourVar = StaticCast<mu::NodeImagePlainColour*>(ParentNode);
-        mu::NodeImagePlainColour::Private* Private = ImagePlainColourVar->GetPrivate();
-		AddChildFunc(Private->m_pColour.get(), FString::Printf(TEXT("COLOR")));
+		AddChildFunc(ImagePlainColourVar->Colour.get(), FString::Printf(TEXT("COLOR")));
 	}
 
 	else if (ParentNode->GetType() == mu::NodeLayout::GetStaticType())
@@ -730,10 +697,9 @@ void SMutableGraphViewer::GetChildrenForInfo(TSharedPtr<FMutableGraphTreeElement
 	else if (ParentNode->GetType() == mu::NodeScalarEnumParameter::GetStaticType())
 	{
 		mu::NodeScalarEnumParameter* ScalarEnumParameterVar = StaticCast<mu::NodeScalarEnumParameter*>(ParentNode);
-		mu::NodeScalarEnumParameter::Private* Private = ScalarEnumParameterVar->GetPrivate();
-		for (int32 RangeIndex = 0; RangeIndex < Private->Ranges.Num(); RangeIndex++)
+		for (int32 RangeIndex = 0; RangeIndex < ScalarEnumParameterVar->Ranges.Num(); RangeIndex++)
 		{
-			AddChildFunc(Private->Ranges[RangeIndex].get(), FString::Printf(TEXT("RANGE [%d]"), RangeIndex));
+			AddChildFunc(ScalarEnumParameterVar->Ranges[RangeIndex].get(), FString::Printf(TEXT("RANGE [%d]"), RangeIndex));
 		}
 	}
 
@@ -754,11 +720,10 @@ void SMutableGraphViewer::GetChildrenForInfo(TSharedPtr<FMutableGraphTreeElement
 	else if (ParentNode->GetType() == mu::NodeImageInterpolate::GetStaticType())
 	{
 		mu::NodeImageInterpolate* ImageInterpolateVar = StaticCast<mu::NodeImageInterpolate*>(ParentNode);
-		mu::NodeImageInterpolate::Private* Private = ImageInterpolateVar->GetPrivate();
-		AddChildFunc(Private->m_pFactor.get(), FString::Printf(TEXT("FACTOR")));
-		for (int32 TargetIndex = 0; TargetIndex < Private->m_targets.Num(); TargetIndex++)
+		AddChildFunc(ImageInterpolateVar->Factor.get(), FString::Printf(TEXT("FACTOR")));
+		for (int32 TargetIndex = 0; TargetIndex < ImageInterpolateVar->Targets.Num(); TargetIndex++)
 		{
-			AddChildFunc(Private->m_targets[TargetIndex].get(), FString::Printf(TEXT("TARGET [%d]"), TargetIndex));
+			AddChildFunc(ImageInterpolateVar->Targets[TargetIndex].get(), FString::Printf(TEXT("TARGET [%d]"), TargetIndex));
 		}
 	}
 	
@@ -770,10 +735,9 @@ void SMutableGraphViewer::GetChildrenForInfo(TSharedPtr<FMutableGraphTreeElement
 	else if (ParentNode->GetType() == mu::NodeScalarParameter::GetStaticType())
 	{
 		mu::NodeScalarParameter* ScalarParameterVar = StaticCast<mu::NodeScalarParameter*>(ParentNode);
-		mu::NodeScalarParameter::Private* Private = ScalarParameterVar->GetPrivate();
-		for (int32 RangeIndex = 0; RangeIndex < Private->Ranges.Num(); RangeIndex++)
+		for (int32 RangeIndex = 0; RangeIndex < ScalarParameterVar->Ranges.Num(); RangeIndex++)
 		{
-			AddChildFunc(Private->Ranges[RangeIndex].get(), FString::Printf(TEXT("RANGE [%d]"), RangeIndex));
+			AddChildFunc(ScalarParameterVar->Ranges[RangeIndex].get(), FString::Printf(TEXT("RANGE [%d]"), RangeIndex));
 		}
 	}
 
@@ -807,18 +771,16 @@ void SMutableGraphViewer::GetChildrenForInfo(TSharedPtr<FMutableGraphTreeElement
 	else if (ParentNode->GetType() == mu::NodeMeshMakeMorph::GetStaticType())
 	{
 		mu::NodeMeshMakeMorph* MeshMakeMorphVar = StaticCast<mu::NodeMeshMakeMorph*>(ParentNode);
-		mu::NodeMeshMakeMorph::Private* Private = MeshMakeMorphVar->GetPrivate();
-		AddChildFunc(Private->m_pBase.get(), FString::Printf(TEXT("BASE")));
-		AddChildFunc(Private->m_pTarget.get(), FString::Printf(TEXT("TARGET")));
+		AddChildFunc(MeshMakeMorphVar->Base.get(), FString::Printf(TEXT("BASE")));
+		AddChildFunc(MeshMakeMorphVar->Target.get(), FString::Printf(TEXT("TARGET")));
 	}
 	
 	else if (ParentNode->GetType() == mu::NodeProjectorParameter::GetStaticType())
 	{
 		mu::NodeProjectorParameter* ProjectorParameterVar = StaticCast<mu::NodeProjectorParameter*>(ParentNode);
-		mu::NodeProjectorParameter::Private* Private = ProjectorParameterVar->GetPrivate();
-		for (int32 RangeIndex = 0; RangeIndex < Private->Ranges.Num(); RangeIndex++)
+		for (int32 RangeIndex = 0; RangeIndex < ProjectorParameterVar->Ranges.Num(); RangeIndex++)
 		{
-			AddChildFunc(Private->Ranges[RangeIndex].get(), FString::Printf(TEXT("RANGE [%d]"), RangeIndex));
+			AddChildFunc(ProjectorParameterVar->Ranges[RangeIndex].get(), FString::Printf(TEXT("RANGE [%d]"), RangeIndex));
 		}
 	}
 
@@ -840,28 +802,25 @@ void SMutableGraphViewer::GetChildrenForInfo(TSharedPtr<FMutableGraphTreeElement
 	else if (ParentNode->GetType() == mu::NodeImageSwizzle::GetStaticType())
 	{
 		mu::NodeImageSwizzle* ImageSwizzleVar = StaticCast<mu::NodeImageSwizzle*>(ParentNode);
-		mu::NodeImageSwizzle::Private* Private = ImageSwizzleVar->GetPrivate();
-		for (int32 SourceIndex = 0; SourceIndex < Private->m_sources.Num(); ++SourceIndex)
+		for (int32 SourceIndex = 0; SourceIndex < ImageSwizzleVar->Sources.Num(); ++SourceIndex)
 		{
-			AddChildFunc(Private->m_sources[SourceIndex].get(), FString::Printf(TEXT("SOURCE [%d]"), SourceIndex));
+			AddChildFunc(ImageSwizzleVar->Sources[SourceIndex].get(), FString::Printf(TEXT("SOURCE [%d]"), SourceIndex));
 		}
 	}
 	
 	else if (ParentNode->GetType() == mu::NodeImageInvert::GetStaticType())
 	{
 		mu::NodeImageInvert* ImageInvertVar = StaticCast<mu::NodeImageInvert*>(ParentNode);
-		mu::NodeImageInvert::Private* Private = ImageInvertVar->GetPrivate();
-		AddChildFunc(Private->m_pBase.get(), TEXT("BASE"));
+		AddChildFunc(ImageInvertVar->Base.get(), TEXT("BASE"));
 	}
 
 	else if (ParentNode->GetType() == mu::NodeImageMultiLayer::GetStaticType())
 	{
 		mu::NodeImageMultiLayer* ImageMultilayerVar = StaticCast<mu::NodeImageMultiLayer*>(ParentNode);
-		mu::NodeImageMultiLayer::Private* Private = ImageMultilayerVar->GetPrivate();
-		AddChildFunc(Private->m_pBase.get(), FString::Printf(TEXT("BASE")));
-		AddChildFunc(Private->m_pMask.get(), FString::Printf(TEXT("MASK")));
-		AddChildFunc(Private->m_pBlended.get(), FString::Printf(TEXT("BLEND")));
-		AddChildFunc(Private->m_pRange.get(), FString::Printf(TEXT("RANGE")));
+		AddChildFunc(ImageMultilayerVar->Base.get(), FString::Printf(TEXT("BASE")));
+		AddChildFunc(ImageMultilayerVar->Mask.get(), FString::Printf(TEXT("MASK")));
+		AddChildFunc(ImageMultilayerVar->Blended.get(), FString::Printf(TEXT("BLEND")));
+		AddChildFunc(ImageMultilayerVar->Range.get(), FString::Printf(TEXT("RANGE")));
 	}
 	
 	else if (ParentNode->GetType() == mu::NodeImageTable::GetStaticType())
diff --git a/Engine/Plugins/Mutable/Source/MutableRuntime/Internal/MuR/Operations.h b/Engine/Plugins/Mutable/Source/MutableRuntime/Internal/MuR/Operations.h
--- a/Engine/Plugins/Mutable/Source/MutableRuntime/Internal/MuR/Operations.h
+++ b/Engine/Plugins/Mutable/Source/MutableRuntime/Internal/MuR/Operations.h
@@ -46,8 +46,9 @@ namespace mu
         SC_PARAMETER,
         CO_PARAMETER,
         PR_PARAMETER,
-        IM_PARAMETER,
-        ST_PARAMETER,
+		IM_PARAMETER,
+		ME_PARAMETER,
+		ST_PARAMETER,
     	MA_PARAMETER,
 
 		//! A referenced, but opaque engine resource
diff --git a/Engine/Plugins/Mutable/Source/MutableRuntime/Private/MuR/CodeVisitor.cpp b/Engine/Plugins/Mutable/Source/MutableRuntime/Private/MuR/CodeVisitor.cpp
--- a/Engine/Plugins/Mutable/Source/MutableRuntime/Private/MuR/CodeVisitor.cpp
+++ b/Engine/Plugins/Mutable/Source/MutableRuntime/Private/MuR/CodeVisitor.cpp
@@ -55,6 +55,7 @@ namespace mu
 				case EOpType::CO_PARAMETER:
 				case EOpType::PR_PARAMETER:
 				case EOpType::IM_PARAMETER:
+				case EOpType::ME_PARAMETER:
 				case EOpType::ST_PARAMETER:
 					CurrentParams[Program.GetOpArgs<OP::ParameterArgs>(at).variable]++;
 					break;
diff --git a/Engine/Plugins/Mutable/Source/MutableRuntime/Private/MuR/Operations.cpp b/Engine/Plugins/Mutable/Source/MutableRuntime/Private/MuR/Operations.cpp
--- a/Engine/Plugins/Mutable/Source/MutableRuntime/Private/MuR/Operations.cpp
+++ b/Engine/Plugins/Mutable/Source/MutableRuntime/Private/MuR/Operations.cpp
@@ -35,6 +35,7 @@ namespace mu
 		{ EDataType::Color,			},	// CO_PARAMETER
 		{ EDataType::Projector,		},	// PR_PARAMETER
 		{ EDataType::Image,			},	// IM_PARAMETER
+		{ EDataType::Mesh,			},	// ME_PARAMETER
 		{ EDataType::String,		},	// ST_PARAMETER
 		{ EDataType::Matrix,		},	// MA_PARAMETER
 
@@ -172,8 +173,9 @@ namespace mu
         case EOpType::SC_PARAMETER:
         case EOpType::CO_PARAMETER:
         case EOpType::PR_PARAMETER:
-        case EOpType::IM_PARAMETER:
-        case EOpType::MA_PARAMETER:
+		case EOpType::IM_PARAMETER:
+		case EOpType::ME_PARAMETER:
+		case EOpType::MA_PARAMETER:
 		case EOpType::IM_REFERENCE:
 		case EOpType::ME_REFERENCE:
 			break;
diff --git a/Engine/Plugins/Mutable/Source/MutableRuntime/Private/MuR/System.cpp b/Engine/Plugins/Mutable/Source/MutableRuntime/Private/MuR/System.cpp
--- a/Engine/Plugins/Mutable/Source/MutableRuntime/Private/MuR/System.cpp
+++ b/Engine/Plugins/Mutable/Source/MutableRuntime/Private/MuR/System.cpp
@@ -716,8 +716,9 @@ namespace mu
             case EOpType::SC_PARAMETER:
             case EOpType::CO_PARAMETER:
             case EOpType::PR_PARAMETER:
-            case EOpType::IM_PARAMETER:
-            case EOpType::MA_PARAMETER:
+			case EOpType::IM_PARAMETER:
+			case EOpType::ME_PARAMETER:
+			case EOpType::MA_PARAMETER:
             {
 				OP::ParameterArgs args = Program.GetOpArgs<OP::ParameterArgs>(at);
 				OP::ADDRESS ParamIndex = args.variable;
diff --git a/Engine/Plugins/Mutable/Source/MutableTools/Private/MuT/CodeGenerator.cpp b/Engine/Plugins/Mutable/Source/MutableTools/Private/MuT/CodeGenerator.cpp
--- a/Engine/Plugins/Mutable/Source/MutableTools/Private/MuT/CodeGenerator.cpp
+++ b/Engine/Plugins/Mutable/Source/MutableTools/Private/MuT/CodeGenerator.cpp
@@ -63,19 +63,17 @@
 #include "MuT/NodeImage.h"
 #include "MuT/NodeImageConstant.h"
 #include "MuT/NodeImageFormat.h"
-#include "MuT/NodeImageFormatPrivate.h"
 #include "MuT/NodeImageMipmap.h"
-#include "MuT/NodeImageMipmapPrivate.h"
-#include "MuT/NodeImageSwizzlePrivate.h"
+#include "MuT/NodeImageSwizzle.h"
 #include "MuT/NodeMatrixConstant.h"
 #include "MuT/NodeMesh.h"
 #include "MuT/NodeMeshClipMorphPlane.h"
 #include "MuT/NodeMeshClipWithMesh.h"
-#include "MuT/NodeMeshConstantPrivate.h"
 #include "MuT/NodeMeshFormat.h"
 #include "MuT/NodeMeshFragment.h"
 #include "MuT/NodeMeshMorph.h"
 #include "MuT/NodeMeshReshape.h"
+#include "MuT/NodeMeshConstant.h"
 #include "MuT/NodeModifier.h"
 #include "MuT/NodeModifierMeshClipDeform.h"
 #include "MuT/NodeModifierMeshClipMorphPlane.h"
@@ -84,9 +82,7 @@
 #include "MuT/NodeModifierMeshTransformInMesh.h"
 #include "MuT/NodeModifierSurfaceEdit.h"
 #include "MuT/NodeObject.h"
-#include "MuT/NodeObjectGroupPrivate.h"
 #include "MuT/NodeObjectNew.h"
-#include "MuT/NodePrivate.h"
 #include "MuT/NodeRange.h"
 #include "MuT/NodeRangeFromScalar.h"
 #include "MuT/NodeScalar.h"
@@ -312,10 +308,10 @@ namespace mu
 			const NodeRangeFromScalar* FromScalar = static_cast<const NodeRangeFromScalar*>(Untyped.get());
 
 			Result = FRangeGenerationResult();
-			Result.rangeName = FromScalar->GetName();
+			Result.rangeName = FromScalar->Name;
 
 			FScalarGenerationResult ChildResult;
-			GenerateScalar(ChildResult, Options, FromScalar->GetSize());
+			GenerateScalar(ChildResult, Options, FromScalar->Size);
 			Result.sizeOp = ChildResult.op;
 		}
 		else
@@ -340,9 +336,9 @@ namespace mu
 		{
 			ParamName = CacheKey.Table->GetName();
 		}
-		Result->SetName(ParamName);
+		Result->Name = ParamName;
 
-		Result->SetDefaultValueIndex(0);
+		Result->DefaultValue = 0;
 
 		int32 CurrentRow = 0;
 		int32 RowCount = CacheKey.Table->GetPrivate()->Rows.Num();
@@ -350,12 +346,13 @@ namespace mu
 
 		if (bAddNoneOption)
 		{
-			Result->SetValueCount(RowCount);
-			Result->SetValue(CurrentRow, -1, "None");
+			Result->Options.SetNum(RowCount);
+			Result->Options[CurrentRow].Value = -1;
+			Result->Options[CurrentRow].Name = "None";
 		}
 		else
 		{
-			Result->SetValueCount(RowCount-1);
+			Result->Options.SetNum(RowCount-1);
 		}
 
 		// Add the possible values
@@ -383,19 +380,20 @@ namespace mu
 					ValueName = CacheKey.Table->GetPrivate()->Rows[RowIndex].Values[NameCol].String;
 				}
 
-				Result->SetValue(CurrentRow, RowIndex, ValueName);
+				Result->Options[CurrentRow].Value = RowIndex;
+				Result->Options[CurrentRow].Name = ValueName;
 
 
 				// Set the first row or the selected row as the default one.
 				if(RowIndex == StarCount || ValueName == DefaultRowName)
 				{
-					Result->SetDefaultValueIndex(RowIndex);
+					Result->DefaultValue = RowIndex;
 				}
 
 				// Set the selected row as default (if exists)
 				if (ValueName == DefaultRowName)
 				{
-					Result->SetDefaultValueIndex(RowIndex);
+					Result->DefaultValue = RowIndex;
 				}
 
 				++CurrentRow;
@@ -1264,26 +1262,26 @@ namespace mu
 							{
 								NodeImageMipmap* tm = static_cast<NodeImageMipmap*>(pImageNode.get());
 								if (!mipmapNode) mipmapNode = tm;
-								pImageNode = tm->GetSource();
+								pImageNode = tm->Source;
 							}
 							else if (pImageNode->GetType() == NodeImageFormat::GetStaticType())
 							{
 								NodeImageFormat* tf = static_cast<NodeImageFormat*>(pImageNode.get());
 								if (!formatNode) formatNode = tf;
-								pImageNode = tf->GetSource();
+								pImageNode = tf->Source;
 							}
 							else if (pImageNode->GetType() == NodeImageSwizzle::GetStaticType())
 							{
 								NodeImageSwizzle* ts = static_cast<NodeImageSwizzle*>(pImageNode.get());
 
-								if (!ts->GetPrivate()->m_sources.IsEmpty())
+								if (!ts->Sources.IsEmpty())
 								{
-									NodeImage* Source = ts->GetSource(0).get();
+									NodeImage* Source = ts->Sources[0].get();
 
 									bool bAllSourcesAreTheSame = true;
-									for (int32 SourceIndex = 1; SourceIndex < ts->GetPrivate()->m_sources.Num(); ++SourceIndex)
+									for (int32 SourceIndex = 1; SourceIndex < ts->Sources.Num(); ++SourceIndex)
 									{
-										bAllSourcesAreTheSame = bAllSourcesAreTheSame && (Source == ts->GetSource(SourceIndex));
+										bAllSourcesAreTheSame = bAllSourcesAreTheSame && (Source == ts->Sources[SourceIndex]);
 									}
 
 									if (!swizzleNode && bAllSourcesAreTheSame)
@@ -1351,15 +1349,15 @@ namespace mu
 								if (swizzleNode)
 								{
 									Ptr<ASTOpImageSwizzle> fop = new ASTOpImageSwizzle();
-									fop->Format = swizzleNode->GetPrivate()->m_format;
+									fop->Format = swizzleNode->NewFormat;
 									fop->Sources[0] = imageAd;
 									fop->Sources[1] = imageAd;
 									fop->Sources[2] = imageAd;
 									fop->Sources[3] = imageAd;
-									fop->SourceChannels[0] = swizzleNode->GetPrivate()->m_sourceChannels[0];
-									fop->SourceChannels[1] = swizzleNode->GetPrivate()->m_sourceChannels[1];
-									fop->SourceChannels[2] = swizzleNode->GetPrivate()->m_sourceChannels[2];
-									fop->SourceChannels[3] = swizzleNode->GetPrivate()->m_sourceChannels[3];
+									fop->SourceChannels[0] = swizzleNode->SourceChannels[0];
+									fop->SourceChannels[1] = swizzleNode->SourceChannels[1];
+									fop->SourceChannels[2] = swizzleNode->SourceChannels[2];
+									fop->SourceChannels[3] = swizzleNode->SourceChannels[3];
 									check(fop->Format != EImageFormat::None);
 									imageAd = fop;
 								}
@@ -1371,16 +1369,16 @@ namespace mu
 									op->Source = imageAd;
 									op->BlockLevels = 0;
 
-									op->AddressMode = mipmapNode->GetPrivate()->m_settings.AddressMode;
-									op->FilterType = mipmapNode->GetPrivate()->m_settings.FilterType;
+									op->AddressMode = mipmapNode->Settings.AddressMode;
+									op->FilterType = mipmapNode->Settings.FilterType;
 									imageAd = op;
 								}
 
 								if (formatNode)
 								{
 									Ptr<ASTOpImagePixelFormat> fop = new ASTOpImagePixelFormat();
-									fop->Format = formatNode->GetPrivate()->m_format;
-									fop->FormatIfAlpha = formatNode->GetPrivate()->m_formatIfAlpha;
+									fop->Format = formatNode->Format;
+									fop->FormatIfAlpha = formatNode->FormatIfAlpha;
 									fop->Source = imageAd;
 									check(fop->Format != EImageFormat::None);
 									imageAd = fop;
@@ -1421,10 +1419,10 @@ namespace mu
 									FLayoutBlockDesc LayoutBlockDesc;
 									if (formatNode)
 									{
-										LayoutBlockDesc.FinalFormat = formatNode->GetPrivate()->m_formatIfAlpha;
+										LayoutBlockDesc.FinalFormat = formatNode->FormatIfAlpha;
 										if (LayoutBlockDesc.FinalFormat == EImageFormat::None)
 										{
-											LayoutBlockDesc.FinalFormat = formatNode->GetPrivate()->m_format;
+											LayoutBlockDesc.FinalFormat = formatNode->Format;
 										}
 									}
 
@@ -1523,12 +1521,12 @@ namespace mu
 									if (swizzleNode)
 									{
 										Ptr<ASTOpImageSwizzle> fop = new ASTOpImageSwizzle();
-										fop->Format = swizzleNode->GetPrivate()->m_format;
+										fop->Format = swizzleNode->NewFormat;
 
-										for (int32 ChannelIndex = 0; ChannelIndex < swizzleNode->GetPrivate()->m_sourceChannels.Num(); ++ChannelIndex)
+										for (int32 ChannelIndex = 0; ChannelIndex < swizzleNode->SourceChannels.Num(); ++ChannelIndex)
 										{
 											fop->Sources[ChannelIndex] = imageAd;
-											fop->SourceChannels[ChannelIndex] = swizzleNode->GetPrivate()->m_sourceChannels[ChannelIndex];
+											fop->SourceChannels[ChannelIndex] = swizzleNode->SourceChannels[ChannelIndex];
 										}
 										check(fop->Format != EImageFormat::None);
 										imageAd = fop;
@@ -1570,8 +1568,8 @@ namespace mu
 
 										if (mipmapNode)
 										{
-											mop->AddressMode = mipmapNode->GetPrivate()->m_settings.AddressMode;
-											mop->FilterType = mipmapNode->GetPrivate()->m_settings.FilterType;
+											mop->AddressMode = mipmapNode->Settings.AddressMode;
+											mop->FilterType = mipmapNode->Settings.FilterType;
 										}
 
 										imageAd = mop;
@@ -1580,8 +1578,8 @@ namespace mu
 									if (formatNode)
 									{
 										Ptr<ASTOpImagePixelFormat> fop = new ASTOpImagePixelFormat();
-										fop->Format = formatNode->GetPrivate()->m_format;
-										fop->FormatIfAlpha = formatNode->GetPrivate()->m_formatIfAlpha;
+										fop->Format = formatNode->Format;
+										fop->FormatIfAlpha = formatNode->FormatIfAlpha;
 										fop->Source = imageAd;
 										check(fop->Format != EImageFormat::None);
 										imageAd = fop;
@@ -1634,7 +1632,7 @@ namespace mu
 					{
 						// MUTABLE_CPUPROFILER_SCOPE(SurfaceScalar);
 
-						if (NodeScalarPtr pScalarNode = SurfaceNode->Scalars[t].Scalar)
+						if (Ptr<NodeScalar> ScalarNode = SurfaceNode->Scalars[t].Scalar)
 						{
 							Ptr<ASTOpInstanceAdd> op = new ASTOpInstanceAdd();
 							op->type = EOpType::IN_ADDSCALAR;
@@ -1642,7 +1640,7 @@ namespace mu
 
 							// Scalar
 							FScalarGenerationResult ScalarResult;
-							GenerateScalar(ScalarResult, Options, pScalarNode);
+							GenerateScalar(ScalarResult, Options, ScalarNode);
 							op->value = ScalarResult.op;
 
 							// Name
@@ -1656,14 +1654,14 @@ namespace mu
 					//------------------------------------------------------------------------
 					for (int32 t = 0; t < SurfaceNode->Strings.Num(); ++t)
 					{
-						if (NodeStringPtr pStringNode = SurfaceNode->Strings[t].String)
+						if (Ptr<NodeString> StringNode = SurfaceNode->Strings[t].String)
 						{
 							Ptr<ASTOpInstanceAdd> op = new ASTOpInstanceAdd();
 							op->type = EOpType::IN_ADDSTRING;
 							op->instance = LastSurfOp;
 
 							FStringGenerationResult StringResult;
-							GenerateString(StringResult, Options, pStringNode);
+							GenerateString(StringResult, Options, StringNode);
 							op->value = StringResult.op;
 
 							// Name
@@ -1997,16 +1995,14 @@ namespace mu
     }
 
 
-	void CodeGenerator::GenerateObject_Group(const FObjectGenerationOptions& Options, FObjectGenerationResult& OutResult, const NodeObjectGroup* InNode)
+	void CodeGenerator::GenerateObject_Group(const FObjectGenerationOptions& Options, FObjectGenerationResult& OutResult, const NodeObjectGroup* Node)
 	{
-		const NodeObjectGroup::Private& node = *InNode->GetPrivate();
-
 		TArray<FString> UsedNames;
 
         // Parse the child objects first, which will accumulate operations in the patching lists
-        for ( int32 ChildIndex=0; ChildIndex <node.m_children.Num(); ++ChildIndex)
+        for ( int32 ChildIndex=0; ChildIndex < Node->Children.Num(); ++ChildIndex)
         {
-            if ( const NodeObject* pChildNode = node.m_children[ChildIndex].get() )
+            if ( const NodeObject* pChildNode = Node->Children[ChildIndex].get() )
             {
 				// Look for the child condition in the first pass
                 Ptr<ASTOp> ConditionOp;
@@ -2034,7 +2030,7 @@ namespace mu
 				if (UsedNames.Contains(ChildName))
 				{
 					FString Msg = FString::Printf(TEXT("Object group has more than one children with the same name [%s]."), *ChildName );
-					ErrorLog->Add(Msg, ELMT_WARNING, InNode->GetMessageContext());
+					ErrorLog->Add(Msg, ELMT_WARNING, Node->GetMessageContext());
 				}
 				else
 				{
@@ -2054,7 +2050,7 @@ namespace mu
 
         // Create a constant node
         Ptr<NodeBoolConstant> pNode = new NodeBoolConstant();
-        pNode->SetValue(Value);
+        pNode->Value = Value;
 
 		FBoolGenerationResult ChildResult;
 		FGenericGenerationOptions Options;
@@ -2422,7 +2418,7 @@ namespace mu
 				TSharedPtr< TSharedPtr<FMesh> > ResolveMorphedMesh = MakeShared<TSharedPtr<FMesh>>();
 				UE::Tasks::FTask TargetMeshTask;
 				{
-					FMesh* OriginalMesh = OriginalMeshNode->GetValue().Get();
+					FMesh* OriginalMesh = OriginalMeshNode->Value.Get();
 					check(OriginalMesh->IsReference());
 					uint32 OriginalMeshID = OriginalMesh->GetReferencedMesh();
 
@@ -2445,7 +2441,7 @@ namespace mu
 				else
 				{
 					Ptr<NodeScalarConstant> auxNode = new NodeScalarConstant();
-					auxNode->SetValue(1.0f);
+					auxNode->Value = 1.0f;
 
 					FScalarGenerationResult ChildResult;
 					GenerateScalar(ChildResult, StaticOptions, auxNode);
@@ -2586,7 +2582,7 @@ namespace mu
 				TSharedPtr< TSharedPtr<FMesh> > ResolveOriginalMesh = MakeShared<TSharedPtr<FMesh>>();
 				UE::Tasks::FTask TargetMeshTask;
 				{
-					FMesh* OriginalMesh = OriginalMeshNode->GetValue().Get();
+					FMesh* OriginalMesh = OriginalMeshNode->Value.Get();
 					check(OriginalMesh->IsReference());
 					uint32 OriginalMeshID = OriginalMesh->GetReferencedMesh();
 
diff --git a/Engine/Plugins/Mutable/Source/MutableTools/Private/MuT/CodeGenerator.h b/Engine/Plugins/Mutable/Source/MutableTools/Private/MuT/CodeGenerator.h
--- a/Engine/Plugins/Mutable/Source/MutableTools/Private/MuT/CodeGenerator.h
+++ b/Engine/Plugins/Mutable/Source/MutableTools/Private/MuT/CodeGenerator.h
@@ -87,6 +87,7 @@ namespace mu
 	class NodeMeshSwitch;
 	class NodeMeshTransform;
 	class NodeMeshVariation;
+	class NodeMeshParameter;
 	class NodeRange;
 	class NodeScalarArithmeticOperation;
 	class NodeScalarConstant;
@@ -657,6 +658,7 @@ namespace mu
 		FMeshTask GenerateMesh_Table(const FMeshGenerationStaticOptions& StaticOptions, FMeshOptionsTask, const NodeMeshTable*);
 		FMeshTask GenerateMesh_Reshape(const FMeshGenerationStaticOptions& StaticOptions, FMeshOptionsTask, const NodeMeshReshape*);
 		FMeshTask GenerateMesh_ClipDeform(const FMeshGenerationStaticOptions& StaticOptions, FMeshOptionsTask, const NodeMeshClipDeform*);
+		FMeshTask GenerateMesh_Parameter(const FMeshGenerationStaticOptions& StaticOptions, FMeshOptionsTask, const NodeMeshParameter*);
 
 		//-----------------------------------------------------------------------------------------
 		void PrepareMeshForLayout(const FGeneratedLayout&, TSharedPtr<FMesh>,
@@ -1022,14 +1024,14 @@ namespace mu
 
         // Create the switch node
         Ptr<NODE_SWITCH_TYPE> SwitchNode = new NODE_SWITCH_TYPE;
-		SwitchNode->SetParameter( Variable );
-		SwitchNode->SetOptionCount(NumRows);
+		SwitchNode->Parameter = Variable;
+		SwitchNode->Options.SetNum(NumRows);
 
 		for (int32 RowIndex = 0; RowIndex < NumRows; ++RowIndex)
         {
             check(RowIndex <= 0xFFFF);
             Ptr<NODE_TYPE> BranchNode = GenerateOption( node, ColIndex, RowIndex, ErrorLog.Get() );
-			SwitchNode->SetOption(RowIndex, BranchNode);
+			SwitchNode->Options[RowIndex] = BranchNode;
         }
 
 		OutResult = SwitchNode;
diff --git a/Engine/Plugins/Mutable/Source/MutableTools/Private/MuT/CodeGenerator_Bool.cpp b/Engine/Plugins/Mutable/Source/MutableTools/Private/MuT/CodeGenerator_Bool.cpp
--- a/Engine/Plugins/Mutable/Source/MutableTools/Private/MuT/CodeGenerator_Bool.cpp
+++ b/Engine/Plugins/Mutable/Source/MutableTools/Private/MuT/CodeGenerator_Bool.cpp
@@ -15,7 +15,6 @@
 #include "MuT/ASTOpBoolNot.h"
 #include "MuT/CodeGenerator.h"
 #include "MuT/NodeBool.h"
-#include "MuT/NodeBoolPrivate.h"
 #include "MuT/NodeRange.h"
 
 
@@ -78,18 +77,14 @@ namespace mu
 	}
 
 
-	void CodeGenerator::GenerateBool_Constant(FBoolGenerationResult& result, const FGenericGenerationOptions& Options, const Ptr<const NodeBoolConstant>& Typed)
+	void CodeGenerator::GenerateBool_Constant(FBoolGenerationResult& result, const FGenericGenerationOptions& Options, const Ptr<const NodeBoolConstant>& Node)
 	{
-		const NodeBoolConstant::Private& node = *Typed->GetPrivate();
-
-		result.op = new ASTOpConstantBool(node.Value);
+		result.op = new ASTOpConstantBool(Node->Value);
 	}
 
 
-	void CodeGenerator::GenerateBool_Parameter(FBoolGenerationResult& result, const FGenericGenerationOptions& Options, const Ptr<const NodeBoolParameter>& Typed)
+	void CodeGenerator::GenerateBool_Parameter(FBoolGenerationResult& result, const FGenericGenerationOptions& Options, const Ptr<const NodeBoolParameter>& Node)
 	{
-		const NodeBoolParameter::Private& node = *Typed->GetPrivate();
-
 		Ptr<ASTOpParameter> op;
 
 		// This can run in any thread.
@@ -97,21 +92,21 @@ namespace mu
 		{
 			UE::TUniqueLock Lock(FirstPass.ParameterNodes.Mutex);
 
-			Found = FirstPass.ParameterNodes.Map.Find(node.m_pNode);
+			Found = FirstPass.ParameterNodes.Map.Find(Node);
 			if (!Found)
 			{
 				FParameterDesc param;
-				param.Name = node.Name;
-				bool bParseOk = FGuid::Parse(node.UID, param.UID);
+				param.Name = Node->Name;
+				bool bParseOk = FGuid::Parse(Node->UID, param.UID);
 				check(bParseOk);
 				param.Type = EParameterType::Bool;
-				param.DefaultValue.Set<FParamBoolType>(node.DefaultValue);
+				param.DefaultValue.Set<FParamBoolType>(Node->DefaultValue);
 
 				op = new ASTOpParameter();
 				op->Type = EOpType::BO_PARAMETER;
 				op->Parameter = param;
 
-				FirstPass.ParameterNodes.Map.Add(node.m_pNode, op);
+				FirstPass.ParameterNodes.Map.Add(Node, op);
 			}
 			else
 			{
@@ -122,10 +117,10 @@ namespace mu
 		if (!Found)
 		{
 			// Generate the code for the ranges
-			for (int32 a = 0; a < node.Ranges.Num(); ++a)
+			for (int32 a = 0; a < Node->Ranges.Num(); ++a)
 			{
 				FRangeGenerationResult RangeResult;
-				GenerateRange(RangeResult, Options, node.Ranges[a]);
+				GenerateRange(RangeResult, Options, Node->Ranges[a]);
 				op->Ranges.Emplace(op.get(), RangeResult.sizeOp, RangeResult.rangeName, RangeResult.rangeUID);
 			}
 		}
@@ -134,31 +129,27 @@ namespace mu
 	}
 
 
-	void CodeGenerator::GenerateBool_Not(FBoolGenerationResult& result, const FGenericGenerationOptions& Options, const Ptr<const NodeBoolNot>& Typed)
+	void CodeGenerator::GenerateBool_Not(FBoolGenerationResult& result, const FGenericGenerationOptions& Options, const Ptr<const NodeBoolNot>& Node)
 	{
-		const NodeBoolNot::Private& node = *Typed->GetPrivate();
-
 		Ptr<ASTOpBoolNot> op = new ASTOpBoolNot();
 
 		FBoolGenerationResult ChildResult;
-		GenerateBool(ChildResult, Options, node.m_pSource);
+		GenerateBool(ChildResult, Options, Node->Source);
 		op->A = ChildResult.op;
 
 		result.op = op;
 	}
 
 
-	void CodeGenerator::GenerateBool_And(FBoolGenerationResult& result, const FGenericGenerationOptions& Options, const Ptr<const NodeBoolAnd>& Typed)
+	void CodeGenerator::GenerateBool_And(FBoolGenerationResult& result, const FGenericGenerationOptions& Options, const Ptr<const NodeBoolAnd>& Node)
 	{
-		const NodeBoolAnd::Private& node = *Typed->GetPrivate();
-
 		Ptr<ASTOpBoolAnd> op = new ASTOpBoolAnd;
 
 		FBoolGenerationResult ChildResult;
-		GenerateBool(ChildResult, Options, node.m_pA);
+		GenerateBool(ChildResult, Options, Node->A);
 		op->A = ChildResult.op;
 
-		GenerateBool(ChildResult, Options, node.m_pB);
+		GenerateBool(ChildResult, Options, Node->B);
 		op->B = ChildResult.op;
 
 		result.op = op;
diff --git a/Engine/Plugins/Mutable/Source/MutableTools/Private/MuT/CodeGenerator_Color.cpp b/Engine/Plugins/Mutable/Source/MutableTools/Private/MuT/CodeGenerator_Color.cpp
--- a/Engine/Plugins/Mutable/Source/MutableTools/Private/MuT/CodeGenerator_Color.cpp
+++ b/Engine/Plugins/Mutable/Source/MutableTools/Private/MuT/CodeGenerator_Color.cpp
@@ -348,9 +348,9 @@ namespace mu
 		else
 		{
 			// Set a constant 0.5 value
-			NodeScalarConstantPtr pNode = new NodeScalarConstant();
-			pNode->SetValue(0.5f);
-			GenerateScalar(ChildResult, Options, pNode);
+			Ptr<NodeScalarConstant> Node = new NodeScalarConstant();
+			Node->Value = 0.5f;
+			GenerateScalar(ChildResult, Options, Node);
 			op->X = ChildResult.op;
 		}
 
@@ -364,9 +364,9 @@ namespace mu
 		else
 		{
 			// Set a constant 0.5 value
-			NodeScalarConstantPtr pNode = new NodeScalarConstant();
-			pNode->SetValue(0.5f);
-			GenerateScalar(ChildResult, Options, pNode);
+			Ptr<NodeScalarConstant> Node = new NodeScalarConstant();
+			Node->Value = 0.5f;
+			GenerateScalar(ChildResult, Options, Node);
 			op->Y = ChildResult.op;
 		}
 
@@ -394,9 +394,9 @@ namespace mu
 		}
 		else
 		{
-			NodeScalarConstantPtr pNode = new NodeScalarConstant();
-			pNode->SetValue(1.0f);
-			GenerateScalar(ChildResult, Options, pNode);
+			Ptr<NodeScalarConstant> Node = new NodeScalarConstant();
+			Node->Value  = 1.0f;
+			GenerateScalar(ChildResult, Options, Node);
 			op->V[0] = ChildResult.op;
 		}
 
@@ -408,9 +408,9 @@ namespace mu
 		}
 		else
 		{
-			NodeScalarConstantPtr pNode = new NodeScalarConstant();
-			pNode->SetValue(1.0f);
-			GenerateScalar(ChildResult, Options, pNode);
+			Ptr<NodeScalarConstant> Node = new NodeScalarConstant();
+			Node->Value = 1.0f;
+			GenerateScalar(ChildResult, Options, Node);
 			op->V[1] = ChildResult.op;
 		}
 
@@ -422,9 +422,9 @@ namespace mu
 		}
 		else
 		{
-			NodeScalarConstantPtr pNode = new NodeScalarConstant();
-			pNode->SetValue(1.0f);
-			GenerateScalar(ChildResult, Options, pNode);
+			Ptr<NodeScalarConstant> Node = new NodeScalarConstant();
+			Node->Value = 1.0f;
+			GenerateScalar(ChildResult, Options, Node);
 			op->V[2] = ChildResult.op;
 		}
 
@@ -436,9 +436,9 @@ namespace mu
 		}
 		else
 		{
-			NodeScalarConstantPtr pNode = new NodeScalarConstant();
-			pNode->SetValue(1.0f);
-			GenerateScalar(ChildResult, Options, pNode);
+			Ptr<NodeScalarConstant> Node = new NodeScalarConstant();
+			Node->Value = 1.0f;
+			GenerateScalar(ChildResult, Options, Node);
 			op->V[3] = ChildResult.op;
 		}
 
diff --git a/Engine/Plugins/Mutable/Source/MutableTools/Private/MuT/CodeGenerator_FirstPass.cpp b/Engine/Plugins/Mutable/Source/MutableTools/Private/MuT/CodeGenerator_FirstPass.cpp
--- a/Engine/Plugins/Mutable/Source/MutableTools/Private/MuT/CodeGenerator_FirstPass.cpp
+++ b/Engine/Plugins/Mutable/Source/MutableTools/Private/MuT/CodeGenerator_FirstPass.cpp
@@ -24,14 +24,12 @@
 #include "MuT/NodeModifierMeshClipWithUVMask.h"
 #include "MuT/NodeModifierSurfaceEdit.h"
 #include "MuT/NodeObject.h"
-#include "MuT/NodeObjectGroupPrivate.h"
 #include "MuT/NodeObjectNew.h"
-#include "MuT/NodePrivate.h"
 #include "MuT/NodeSurface.h"
 #include "MuT/NodeSurfaceNew.h"
 #include "MuT/NodeSurfaceVariation.h"
 #include "MuT/NodeSurfaceSwitch.h"
-#include "MuT/NodeScalarEnumParameterPrivate.h"
+#include "MuT/NodeScalarEnumParameter.h"
 
 namespace mu
 {
@@ -611,30 +609,29 @@ namespace mu
 
 
 	//---------------------------------------------------------------------------------------------
-	void FirstPassGenerator::Generate_ObjectGroup(const NodeObjectGroup* InNode)
+	void FirstPassGenerator::Generate_ObjectGroup(const NodeObjectGroup* Node)
 	{
-		const NodeObjectGroup::Private* Private = InNode->GetPrivate();
-
 		// Prepare the enumeration parameter if necessary
         Ptr<ASTOpParameter> enumOp;
-        if (Private->Type==NodeObjectGroup::CS_ALWAYS_ONE ||
-			Private->Type==NodeObjectGroup::CS_ONE_OR_NONE )
+        if (Node->Type==NodeObjectGroup::CS_ALWAYS_ONE ||
+			Node->Type==NodeObjectGroup::CS_ONE_OR_NONE )
         {
 			UE::TUniqueLock Lock(ParameterNodes.Mutex);
-			Ptr<ASTOpParameter>* Found = ParameterNodes.Map.Find(InNode);
+			Ptr<ASTOpParameter>* Found = ParameterNodes.Map.Find(Node);
 
 			if (!Found)
 			{
 				Ptr<ASTOpParameter> op = new ASTOpParameter();
 				op->Type = EOpType::NU_PARAMETER;
+				op->Parameter.DefaultValue.Set<FParamIntType>(Node->DefaultValue);
 
-				op->Parameter.Name = Private->Name;
-				bool bParseOk = FGuid::Parse(Private->Uid, op->Parameter.UID);
+				op->Parameter.Name = Node->Name;
+				bool bParseOk = FGuid::Parse(Node->Uid, op->Parameter.UID);
 				check(bParseOk);
 				op->Parameter.Type = EParameterType::Int;
-				op->Parameter.DefaultValue.Set<FParamIntType>(Private->DefaultValue);
+				op->Parameter.DefaultValue.Set<FParamIntType>(Node->DefaultValue);
 
-				if (Private->Type == NodeObjectGroup::CS_ONE_OR_NONE)
+				if (Node->Type == NodeObjectGroup::CS_ONE_OR_NONE)
 				{
 					FParameterDesc::FIntValueDesc nullValue;
 					nullValue.Value = -1;
@@ -642,7 +639,7 @@ namespace mu
 					op->Parameter.PossibleValues.Add(nullValue);
 				}
 
-				ParameterNodes.Map.Add(InNode, op);
+				ParameterNodes.Map.Add(Node, op);
 
 				enumOp = op;
 			}
@@ -654,21 +651,20 @@ namespace mu
 
 
         // Parse the child objects
-		for ( int32 t=0; t< Private->m_children.Num(); ++t )
+		for ( int32 t=0; t< Node->Children.Num(); ++t )
         {
-            if ( mu::Ptr<const NodeObject> ChildNode = Private->m_children[t] )
+            if ( mu::Ptr<const NodeObject> ChildNode = Node->Children[t] )
             {
                 // Overwrite the implicit condition
                 Ptr<ASTOp> paramOp = 0;
-                switch (Private->Type )
+                switch (Node->Type )
                 {
                     case NodeObjectGroup::CS_TOGGLE_EACH:
                     {           
 						if (ChildNode->GetType() == NodeObjectGroup::GetStaticType())
 						{
-							FString Msg = FString::Printf(TEXT("The Group Node [%s] has type Toggle and its direct child is a Group node, which is not allowed. Change the type or add a Child Object node in between them."),
-								*Private->Name);
-							ErrorLog->Add(Msg, ELMT_ERROR, InNode->GetMessageContext());
+							FString Msg = FString::Printf(TEXT("The Group Node [%s] has type Toggle and its direct child is a Group node, which is not allowed. Change the type or add a Child Object node in between them."), *Node->Name);
+							ErrorLog->Add(Msg, ELMT_ERROR, Node->GetMessageContext());
 						}
 						else
 						{
