[mutable] Initial boilerplate for the implementation of Mesh Parameters.

diff --git a/Engine/Plugins/Mutable/Source/MutableTools/Private/MuT/CodeGenerator_Image.cpp b/Engine/Plugins/Mutable/Source/MutableTools/Private/MuT/CodeGenerator_Image.cpp
--- a/Engine/Plugins/Mutable/Source/MutableTools/Private/MuT/CodeGenerator_Image.cpp
+++ b/Engine/Plugins/Mutable/Source/MutableTools/Private/MuT/CodeGenerator_Image.cpp
@@ -55,49 +55,28 @@
 #include "MuT/NodeColourConstant.h"
 #include "MuT/NodeImage.h"
 #include "MuT/NodeImageBinarise.h"
-#include "MuT/NodeImageBinarisePrivate.h"
 #include "MuT/NodeImageColourMap.h"
-#include "MuT/NodeImageColourMapPrivate.h"
 #include "MuT/NodeImageConditional.h"
-#include "MuT/NodeImageConditionalPrivate.h"
 #include "MuT/NodeImageConstant.h"
 #include "MuT/NodeImageFormat.h"
-#include "MuT/NodeImageFormatPrivate.h"
 #include "MuT/NodeImageInterpolate.h"
-#include "MuT/NodeImageInterpolatePrivate.h"
 #include "MuT/NodeImageInvert.h"
-#include "MuT/NodeImageInvertPrivate.h"
 #include "MuT/NodeImageLayer.h"
 #include "MuT/NodeImageLayerColour.h"
-#include "MuT/NodeImageLayerColourPrivate.h"
-#include "MuT/NodeImageLayerPrivate.h"
 #include "MuT/NodeImageLuminance.h"
-#include "MuT/NodeImageLuminancePrivate.h"
 #include "MuT/NodeImageMipmap.h"
-#include "MuT/NodeImageMipmapPrivate.h"
 #include "MuT/NodeImageMultiLayer.h"
-#include "MuT/NodeImageMultiLayerPrivate.h"
 #include "MuT/NodeImageNormalComposite.h"
-#include "MuT/NodeImageNormalCompositePrivate.h"
 #include "MuT/NodeImageParameter.h"
-#include "MuT/NodeImageParameterPrivate.h"
 #include "MuT/NodeImagePlainColour.h"
-#include "MuT/NodeImagePlainColourPrivate.h"
 #include "MuT/NodeImageProject.h"
-#include "MuT/NodeImageProjectPrivate.h"
 #include "MuT/NodeImageResize.h"
-#include "MuT/NodeImageResizePrivate.h"
 #include "MuT/NodeImageSaturate.h"
-#include "MuT/NodeImageSaturatePrivate.h"
 #include "MuT/NodeImageSwitch.h"
-#include "MuT/NodeImageSwitchPrivate.h"
 #include "MuT/NodeImageSwizzle.h"
-#include "MuT/NodeImageSwizzlePrivate.h"
 #include "MuT/NodeImageTable.h"
 #include "MuT/NodeImageTransform.h"
-#include "MuT/NodeImageTransformPrivate.h"
 #include "MuT/NodeImageVariation.h"
-#include "MuT/NodeImageVariationPrivate.h"
 #include "MuT/NodeScalar.h"
 #include "MuT/NodeScalarConstant.h"
 #include "MuT/Table.h"
@@ -107,8 +86,6 @@
 namespace mu
 {
 
-
-	//-------------------------------------------------------------------------------------------------
 	void CodeGenerator::GenerateImage(const FImageGenerationOptions& Options, FImageGenerationResult& Result, const NodeImagePtrConst& Untyped)
 	{
 		if (!Untyped)
@@ -168,8 +145,7 @@ namespace mu
 	}
 
 
-    //---------------------------------------------------------------------------------------------
-    void CodeGenerator::GenerateImage_Constant(const FImageGenerationOptions& Options, FImageGenerationResult& Result, const NodeImageConstant* InNode)
+   void CodeGenerator::GenerateImage_Constant(const FImageGenerationOptions& Options, FImageGenerationResult& Result, const NodeImageConstant* InNode)
     {
 		const NodeImageConstant& node = *InNode;
 		
@@ -274,29 +250,26 @@ namespace mu
     }
 
 
-    //---------------------------------------------------------------------------------------------
-	void CodeGenerator::GenerateImage_Parameter(const FImageGenerationOptions& Options, FImageGenerationResult& Result, const NodeImageParameter* InNode )
+	void CodeGenerator::GenerateImage_Parameter(const FImageGenerationOptions& Options, FImageGenerationResult& Result, const NodeImageParameter* Node )
     {
-		const NodeImageParameter::Private& node = *InNode->GetPrivate();
-
         Ptr<ASTOpParameter> op;
 
 		Ptr<ASTOpParameter>* Found = nullptr;
 		{
 			UE::TUniqueLock Lock(FirstPass.ParameterNodes.Mutex);
-			Found = FirstPass.ParameterNodes.Map.Find(node.m_pNode);
+			Found = FirstPass.ParameterNodes.Map.Find(Node);
 			if (!Found)
 			{
 				op = new ASTOpParameter();
 				op->Type = EOpType::IM_PARAMETER;
 
-				op->Parameter.Name = node.Name;
-				bool bParseOk = FGuid::Parse(node.UID, op->Parameter.UID);
+				op->Parameter.Name = Node->Name;
+				bool bParseOk = FGuid::Parse(Node->UID, op->Parameter.UID);
 				check(bParseOk);
 				op->Parameter.Type = EParameterType::Image;
-				op->Parameter.DefaultValue.Set<FParamImageType>(node.DefaultValue);
+				op->Parameter.DefaultValue.Set<FParamImageType>(Node->DefaultValue);
 
-				FirstPass.ParameterNodes.Map.Add(node.m_pNode, op);
+				FirstPass.ParameterNodes.Map.Add(Node, op);
 			}
 			else
 			{
@@ -307,10 +280,10 @@ namespace mu
 		if (!Found)
 		{
 			// Generate the code for the ranges
-			for (int32 a = 0; a < node.Ranges.Num(); ++a)
+			for (int32 a = 0; a < Node->Ranges.Num(); ++a)
 			{
 				FRangeGenerationResult rangeResult;
-				GenerateRange(rangeResult, Options, node.Ranges[a]);
+				GenerateRange(rangeResult, Options, Node->Ranges[a]);
 				op->Ranges.Emplace(op.get(), rangeResult.sizeOp, rangeResult.rangeName, rangeResult.rangeUID);
 			}
 		}
@@ -319,29 +292,26 @@ namespace mu
     }
 
 
-    //---------------------------------------------------------------------------------------------
-	void CodeGenerator::GenerateImage_Layer(const FImageGenerationOptions& Options, FImageGenerationResult& Result, const NodeImageLayer* InNode)
+	void CodeGenerator::GenerateImage_Layer(const FImageGenerationOptions& Options, FImageGenerationResult& Result, const NodeImageLayer* Node)
 	{
-		const NodeImageLayer::Private& node = *InNode->GetPrivate();
-		
 		MUTABLE_CPUPROFILER_SCOPE(NodeImageLayer);
 
         Ptr<ASTOpImageLayer> op = new ASTOpImageLayer();
 
-        op->blendType = node.Type;
+        op->blendType = Node->Type;
 
         // Base image
         Ptr<ASTOp> base;
-        if ( node.m_pBase )
+        if (Node->Base )
         {
 			FImageGenerationResult BaseResult;
-			GenerateImage(Options, BaseResult, node.m_pBase);
+			GenerateImage(Options, BaseResult, Node->Base);
 			base = BaseResult.op;
         }
         else
         {
             // This argument is required
-            base = GenerateMissingImageCode(TEXT("Image Layer base"), EImageFormat::RGB_UByte, InNode->GetMessageContext(), Options);
+            base = GenerateMissingImageCode(TEXT("Image Layer base"), EImageFormat::RGB_UByte, Node->GetMessageContext(), Options);
         }
 
 		UE::Math::TIntVector2<int32> TargetSize = Options.RectSize;
@@ -360,12 +330,12 @@ namespace mu
 
         // Mask of the effect
         Ptr<ASTOp> mask;
-        if ( node.m_pMask )
+        if (Node->Mask )
         {
 			FImageGenerationResult MaskResult;
 			FImageGenerationOptions MaskOptions(Options);
 			MaskOptions.RectSize = TargetSize;
-			GenerateImage(MaskOptions, MaskResult, node.m_pMask);
+			GenerateImage(MaskOptions, MaskResult, Node->Mask);
 			mask = MaskResult.op;
 
             mask = GenerateImageFormat( mask, EImageFormat::L_UByte );
@@ -375,12 +345,12 @@ namespace mu
 
         // Image to apply
         Ptr<ASTOp> blended = 0;
-        if ( node.m_pBlended )
+        if (Node->Blended )
         {
 			FImageGenerationResult BlendedResult;
 			FImageGenerationOptions BlendOptions(Options);
 			BlendOptions.RectSize = TargetSize;
-			GenerateImage(BlendOptions, BlendedResult, node.m_pBlended);
+			GenerateImage(BlendOptions, BlendedResult, Node->Blended);
 			blended = BlendedResult.op;
         }
         else
@@ -396,28 +366,25 @@ namespace mu
     }
 
 
-    //---------------------------------------------------------------------------------------------
-	void CodeGenerator::GenerateImage_LayerColour(const FImageGenerationOptions& Options, FImageGenerationResult& Result, const NodeImageLayerColour* InNode)
-	{
-		const NodeImageLayerColour::Private& node = *InNode->GetPrivate();
-		
+	void CodeGenerator::GenerateImage_LayerColour(const FImageGenerationOptions& Options, FImageGenerationResult& Result, const NodeImageLayerColour* Node)
+	{		
 		MUTABLE_CPUPROFILER_SCOPE(NodeImageLayerColour);
 
         Ptr<ASTOpImageLayerColor> op = new ASTOpImageLayerColor();
-		op->blendType = node.Type;
+		op->blendType = Node->Type;
 
         // Base image
         Ptr<ASTOp> base;
-        if ( node.m_pBase )
+        if (Node->Base )
         {
 			FImageGenerationResult BaseResult;
-			GenerateImage(Options, BaseResult, node.m_pBase);
+			GenerateImage(Options, BaseResult, Node->Base);
 			base = BaseResult.op;
         }
         else
         {
             // This argument is required
-            base = GenerateMissingImageCode(TEXT("Layer base image"), EImageFormat::RGB_UByte, InNode->GetMessageContext(), Options );
+            base = GenerateMissingImageCode(TEXT("Layer base image"), EImageFormat::RGB_UByte, Node->GetMessageContext(), Options );
         }
 
 		UE::Math::TIntVector2<int32> TargetSize = Options.RectSize;
@@ -435,10 +402,10 @@ namespace mu
 
         // Mask of the effect
         Ptr<ASTOp> mask;
-        if ( node.m_pMask )
+        if (Node->Mask )
         {
 			FImageGenerationResult MaskResult;
-			GenerateImage(Options, MaskResult, node.m_pMask);
+			GenerateImage(Options, MaskResult, Node->Mask);
 			mask = MaskResult.op;
 			
 			mask = GenerateImageFormat( mask, EImageFormat::L_UByte );
@@ -448,16 +415,16 @@ namespace mu
 
         // Colour to apply
         Ptr<ASTOp> colour = 0;
-        if ( node.m_pColour )
+        if (Node->Colour )
         {
 			FColorGenerationResult ColorResult;
-            GenerateColor(ColorResult, Options, node.m_pColour);
+            GenerateColor(ColorResult, Options, Node->Colour);
 			colour = ColorResult.op;
         }
         else
         {
             // This argument is required
-            colour = GenerateMissingColourCode(TEXT("Layer colour"), InNode->GetMessageContext() );
+            colour = GenerateMissingColourCode(TEXT("Layer colour"), Node->GetMessageContext() );
         }
         op->color = colour;
 
@@ -465,29 +432,26 @@ namespace mu
     }
 
 
-    //---------------------------------------------------------------------------------------------
-	void CodeGenerator::GenerateImage_MultiLayer(const FImageGenerationOptions& Options, FImageGenerationResult& Result, const NodeImageMultiLayer* InNode)
+	void CodeGenerator::GenerateImage_MultiLayer(const FImageGenerationOptions& Options, FImageGenerationResult& Result, const NodeImageMultiLayer* Node)
 	{
-		const NodeImageMultiLayer::Private& node = *InNode->GetPrivate();
-
 		MUTABLE_CPUPROFILER_SCOPE(NodeImageMultiLayer);
 
         Ptr<ASTOpImageMultiLayer> op = new ASTOpImageMultiLayer();
 
-		op->blendType = node.Type;
+		op->blendType = Node->Type;
 
         // Base image
         Ptr<ASTOp> base;
-        if ( node.m_pBase )
+        if (Node->Base )
         {
 			FImageGenerationResult BaseResult;
-			GenerateImage(Options, BaseResult, node.m_pBase);
+			GenerateImage(Options, BaseResult, Node->Base);
 			base = BaseResult.op;
 		}
         else
         {
             // This argument is required
-            base = GenerateMissingImageCode(TEXT("Image MultiLayer base"), EImageFormat::RGB_UByte, InNode->GetMessageContext(), Options);
+            base = GenerateMissingImageCode(TEXT("Image MultiLayer base"), EImageFormat::RGB_UByte, Node->GetMessageContext(), Options);
         }
 
 		FImageDesc BaseDesc = base->GetImageDesc(true);
@@ -506,10 +470,10 @@ namespace mu
 
         // Mask of the effect
         Ptr<ASTOp> mask;
-        if ( node.m_pMask )
+        if (Node->Mask )
         {
 			FImageGenerationResult MaskResult;
-			GenerateImage(Options, MaskResult, node.m_pMask);
+			GenerateImage(Options, MaskResult, Node->Mask);
 			mask = MaskResult.op;
 
 			mask = GenerateImageFormat( mask, EImageFormat::L_UByte );
@@ -519,10 +483,10 @@ namespace mu
 
         // Image to apply
         Ptr<ASTOp> blended;
-        if (node.m_pBlended)
+        if (Node->Blended)
         {
 			FImageGenerationResult MaskResult;
-			GenerateImage(Options, MaskResult, node.m_pBlended);
+			GenerateImage(Options, MaskResult, Node->Blended);
 			blended = MaskResult.op;
         }
         else
@@ -535,10 +499,10 @@ namespace mu
         op->blend = blended;
 
         // Range of iteration
-        if ( node.m_pRange )
+        if (Node->Range )
         {
             FRangeGenerationResult rangeResult;
-            GenerateRange( rangeResult, Options, node.m_pRange );
+            GenerateRange( rangeResult, Options, Node->Range );
 
             op->range.rangeSize = rangeResult.sizeOp;
             op->range.rangeName = rangeResult.rangeName;
@@ -549,30 +513,27 @@ namespace mu
     }
 
 
-	//---------------------------------------------------------------------------------------------
-	void CodeGenerator::GenerateImage_NormalComposite(const FImageGenerationOptions& Options, FImageGenerationResult& Result, const NodeImageNormalComposite* InNode)
+	void CodeGenerator::GenerateImage_NormalComposite(const FImageGenerationOptions& Options, FImageGenerationResult& Result, const NodeImageNormalComposite* Node)
 	{
-		const NodeImageNormalComposite::Private& node = *InNode->GetPrivate();
-
 		MUTABLE_CPUPROFILER_SCOPE(NodeImageNormalComposite);
 
         Ptr<ASTOpImageNormalComposite> op = new ASTOpImageNormalComposite();
 
-		op->Mode = node.m_mode; 
-		op->Power = node.m_power;
+		op->Mode = Node->Mode;
+		op->Power = Node->Power;
 
         // Base image
         Ptr<ASTOp> base;
-        if ( node.m_pBase )
+        if (Node->Base )
         {
 			FImageGenerationResult BaseResult;
-			GenerateImage(Options, BaseResult, node.m_pBase);
+			GenerateImage(Options, BaseResult, Node->Base);
 			base = BaseResult.op;
 		}
         else
         {
             // This argument is required
-            base = GenerateMissingImageCode(TEXT("Image Composite Base"), EImageFormat::RGB_UByte, InNode->GetMessageContext(), Options);
+            base = GenerateMissingImageCode(TEXT("Image Composite Base"), EImageFormat::RGB_UByte, Node->GetMessageContext(), Options);
         }
 
 		FImageDesc BaseDesc = base->GetImageDesc(true);
@@ -591,10 +552,10 @@ namespace mu
         op->Base = base;
 
         Ptr<ASTOp> normal;
-        if ( node.m_pNormal )
+        if (Node->Normal )
         {
 			FImageGenerationResult BaseResult;
-			GenerateImage(Options, BaseResult, node.m_pNormal);
+			GenerateImage(Options, BaseResult, Node->Normal);
 			normal = BaseResult.op;
 
             normal = GenerateImageFormat( normal, EImageFormat::RGB_UByte );
@@ -602,7 +563,7 @@ namespace mu
 		else
 		{
             // This argument is required
-            normal = GenerateMissingImageCode(TEXT("Image Composite Normal"), EImageFormat::RGB_UByte, InNode->GetMessageContext(), Options);
+            normal = GenerateMissingImageCode(TEXT("Image Composite Normal"), EImageFormat::RGB_UByte, Node->GetMessageContext(), Options);
 		}
 
 		normal = GenerateImageSize(normal, TargetSize);
@@ -612,51 +573,50 @@ namespace mu
 		Result.op = op;
     }
 
-	void CodeGenerator::GenerateImage_Transform(const FImageGenerationOptions& Options, FImageGenerationResult& Result, const NodeImageTransform* InNode)
-    {
-		const NodeImageTransform::Private& Node = *InNode->GetPrivate();
 
+	void CodeGenerator::GenerateImage_Transform(const FImageGenerationOptions& Options, FImageGenerationResult& Result, const NodeImageTransform* Node)
+    {
         MUTABLE_CPUPROFILER_SCOPE(NodeImageTransform);
 
         Ptr<ASTOpImageTransform> Op = new ASTOpImageTransform();
 
 		Ptr<ASTOp> OffsetX;
-		if (Node.m_pOffsetX)
+		if (Node->OffsetX)
 		{
 			FScalarGenerationResult ScalarResult;
-			GenerateScalar(ScalarResult, Options, Node.m_pOffsetX);
+			GenerateScalar(ScalarResult, Options, Node->OffsetX);
 			OffsetX = ScalarResult.op;
 		}
 
 		Ptr<ASTOp> OffsetY;
-		if (Node.m_pOffsetY)
+		if (Node->OffsetY)
 		{
 			FScalarGenerationResult ScalarResult;
-			GenerateScalar(ScalarResult, Options, Node.m_pOffsetY);
+			GenerateScalar(ScalarResult, Options, Node->OffsetY);
 			OffsetY = ScalarResult.op;
 		}
 	
 		Ptr<ASTOp> ScaleX;
-		if (Node.m_pScaleX)
+		if (Node->ScaleX)
 		{
 			FScalarGenerationResult ScalarResult;
-			GenerateScalar(ScalarResult, Options, Node.m_pScaleX);
+			GenerateScalar(ScalarResult, Options, Node->ScaleX);
 			ScaleX = ScalarResult.op;
 		}
 	
 		Ptr<ASTOp> ScaleY;
-		if (Node.m_pScaleY)
+		if (Node->ScaleY)
 		{
 			FScalarGenerationResult ScalarResult;
-			GenerateScalar(ScalarResult, Options, Node.m_pScaleY);
+			GenerateScalar(ScalarResult, Options, Node->ScaleY);
 			ScaleY = ScalarResult.op;
 		}
 
 		Ptr<ASTOp> Rotation;
-		if (Node.m_pRotation)
+		if (Node->Rotation)
 		{
 			FScalarGenerationResult ScalarResult;
-			GenerateScalar(ScalarResult, Options, Node.m_pRotation);
+			GenerateScalar(ScalarResult, Options, Node->Rotation);
 			Rotation = ScalarResult.op;
 		}
 
@@ -666,10 +626,10 @@ namespace mu
  		Op->ScaleX = ScaleX ? ScaleX : ScaleY;
 		Op->ScaleY = ScaleY ? ScaleY : ScaleX;
 		Op->Rotation = Rotation; 
-		Op->AddressMode = Node.AddressMode;
-		Op->SizeX = Node.SizeX;
-		Op->SizeY = Node.SizeY;
-		Op->bKeepAspectRatio = Node.bKeepAspectRatio;
+		Op->AddressMode = Node->AddressMode;
+		Op->SizeX = Node->SizeX;
+		Op->SizeY = Node->SizeY;
+		Op->bKeepAspectRatio = Node->bKeepAspectRatio;
 
 		// Base image
         Ptr<ASTOp> Base;
@@ -679,16 +639,16 @@ namespace mu
 		NewOptions.LayoutBlockId = -1;
 		NewOptions.RectSize = {};
 
-        if (Node.m_pBase)
+        if (Node->Base)
         {
 			FImageGenerationResult BaseResult;
-			GenerateImage(NewOptions, BaseResult, Node.m_pBase);
+			GenerateImage(NewOptions, BaseResult, Node->Base);
 			Base = BaseResult.op;
 		}
         else
         {
             // This argument is required
-            Base = GenerateMissingImageCode(TEXT("Image Transform Base"), EImageFormat::RGB_UByte, InNode->GetMessageContext(), NewOptions);
+            Base = GenerateMissingImageCode(TEXT("Image Transform Base"), EImageFormat::RGB_UByte, Node->GetMessageContext(), NewOptions);
         }
 		
 		FImageDesc BaseDesc = Base->GetImageDesc(true);
@@ -700,17 +660,15 @@ namespace mu
         Result.op = Op; 
     }
 
-    //---------------------------------------------------------------------------------------------
-	void CodeGenerator::GenerateImage_Interpolate(const FImageGenerationOptions& Options, FImageGenerationResult& Result, const NodeImageInterpolate* InNode)
-	{
-		const NodeImageInterpolate::Private& node = *InNode->GetPrivate();
 
+	void CodeGenerator::GenerateImage_Interpolate(const FImageGenerationOptions& Options, FImageGenerationResult& Result, const NodeImageInterpolate* Node)
+	{
 		MUTABLE_CPUPROFILER_SCOPE(NodeImageInterpolate);
 
         Ptr<ASTOpImageInterpolate> op = new ASTOpImageInterpolate();
 
         // Factor
-        if ( NodeScalar* Factor = node.m_pFactor.get() )
+        if ( NodeScalar* Factor = Node->Factor.get() )
         {
 			FScalarGenerationResult ParamResult;
 			GenerateScalar(ParamResult, Options, Factor);
@@ -719,7 +677,7 @@ namespace mu
         else
         {
             // This argument is required
-            op->Factor = GenerateMissingScalarCode(TEXT("Interpolation factor"), 0.5f, InNode->GetMessageContext() );
+            op->Factor = GenerateMissingScalarCode(TEXT("Interpolation factor"), 0.5f, Node->GetMessageContext() );
         }
 
         // Target images
@@ -728,15 +686,15 @@ namespace mu
 		UE::Math::TIntVector2<int32> FinalRectSize = Options.RectSize;
 
         for ( int32 TargetIndex=0
-            ; TargetIndex < node.m_targets.Num() && numTargets<MUTABLE_OP_MAX_INTERPOLATE_COUNT
+            ; TargetIndex < Node->Targets.Num() && numTargets<MUTABLE_OP_MAX_INTERPOLATE_COUNT
             ; ++TargetIndex)
         {
-            if ( node.m_targets[TargetIndex] )
+            if (Node->Targets[TargetIndex] )
             {
 				FImageGenerationOptions ChildOptions = Options;
 				ChildOptions.RectSize = FinalRectSize;
 				FImageGenerationResult BaseResult;
-				GenerateImage(ChildOptions, BaseResult, node.m_targets[TargetIndex]);
+				GenerateImage(ChildOptions, BaseResult, Node->Targets[TargetIndex]);
 				Ptr<ASTOp> target = BaseResult.op;
 
 				if (FinalRectSize[0] == 0)
@@ -757,7 +715,7 @@ namespace mu
         // At least one target is required
         if (!op->Targets[0])
         {
-            Ptr<ASTOp> target = GenerateMissingImageCode(TEXT("First interpolation image"), EImageFormat::RGB_UByte, InNode->GetMessageContext(), Options);
+            Ptr<ASTOp> target = GenerateMissingImageCode(TEXT("First interpolation image"), EImageFormat::RGB_UByte, Node->GetMessageContext(), Options);
             target = GenerateImageSize( target, Options.RectSize);
             op->Targets[0] = target;
         }
@@ -766,16 +724,11 @@ namespace mu
     }
 
 
-	//---------------------------------------------------------------------------------------------
-	void CodeGenerator::GenerateImage_Swizzle(const FImageGenerationOptions& Options, FImageGenerationResult& Result, const NodeImageSwizzle* InNode)
+	void CodeGenerator::GenerateImage_Swizzle(const FImageGenerationOptions& Options, FImageGenerationResult& Result, const NodeImageSwizzle* Node)
 	{
-		const NodeImageSwizzle::Private& node = *InNode->GetPrivate();
-
-		//MUTABLE_CPUPROFILER_SCOPE(NodeImageSwizzle);
-
-		if (node.m_format == EImageFormat::None)
+		if (Node->NewFormat == EImageFormat::None)
 		{
-			Result.op = GenerateMissingImageCode(TEXT("Make Texture channel."), EImageFormat::L_UByte, InNode->GetMessageContext(), Options);
+			Result.op = GenerateMissingImageCode(TEXT("Make Texture channel."), EImageFormat::L_UByte, Node->GetMessageContext(), Options);
 			return;
 		}
 
@@ -786,12 +739,12 @@ namespace mu
 		// Format
 		EImageFormat CompressedFormat = EImageFormat::None;
 
-		switch (node.m_format)
+		switch (Node->NewFormat)
 		{
         case EImageFormat::BC1:
         case EImageFormat::ASTC_4x4_RGB_LDR:
-            CompressedFormat = node.m_format;
-            SwizzleOp->Format = (node.m_sources.IsValidIndex(3) && node.m_sources[3]) ? EImageFormat::RGBA_UByte : EImageFormat::RGB_UByte;
+            CompressedFormat = Node->NewFormat;
+            SwizzleOp->Format = (Node->Sources.IsValidIndex(3) && Node->Sources[3]) ? EImageFormat::RGBA_UByte : EImageFormat::RGB_UByte;
 			break;
 
 		case EImageFormat::BC2:
@@ -799,45 +752,45 @@ namespace mu
 		case EImageFormat::BC6:
         case EImageFormat::BC7:
         case EImageFormat::ASTC_4x4_RGBA_LDR:
-            CompressedFormat = node.m_format;
+            CompressedFormat = Node->NewFormat;
             SwizzleOp->Format = EImageFormat::RGBA_UByte;
 			break;
 
 		case EImageFormat::BC4:
-			CompressedFormat = node.m_format;
+			CompressedFormat = Node->NewFormat;
             SwizzleOp->Format = EImageFormat::L_UByte;
 			break;
 
 		case EImageFormat::BC5:
         case EImageFormat::ASTC_4x4_RG_LDR:
-            CompressedFormat = node.m_format;
+            CompressedFormat = Node->NewFormat;
 			// TODO: Should be RG
             SwizzleOp->Format = EImageFormat::RGB_UByte;
 			break;
 
 		default:
-            SwizzleOp->Format = node.m_format;
+            SwizzleOp->Format = Node->NewFormat;
 			break;
 
 		}
 
-		check(node.m_format != EImageFormat::None);
+		check(Node->NewFormat != EImageFormat::None);
 
 		// Source images and channels
-		check(node.m_sources.Num() == node.m_sourceChannels.Num());
+		check(Node->Sources.Num() == Node->SourceChannels.Num());
 
 		// First source, for reference in the size
         Ptr<ASTOp> FirstValid;
 		FImageDesc FirstValidDesc;
 		int32 FirstValidSourceIndex = -1;
 
-		check(MUTABLE_OP_MAX_SWIZZLE_CHANNELS >= node.m_sources.Num());
-		for (int32 SourceIndex = 0; SourceIndex < node.m_sources.Num(); ++SourceIndex)
+		check(MUTABLE_OP_MAX_SWIZZLE_CHANNELS >= Node->Sources.Num());
+		for (int32 SourceIndex = 0; SourceIndex < Node->Sources.Num(); ++SourceIndex)
 		{
-			if (node.m_sources[SourceIndex])
+			if (Node->Sources[SourceIndex])
 			{
 				FImageGenerationResult BaseResult;
-				GenerateImage(Options, BaseResult, node.m_sources[SourceIndex]);
+				GenerateImage(Options, BaseResult, Node->Sources[SourceIndex]);
                 Ptr<ASTOp> Source = BaseResult.op;
 
 				Source = GenerateImageUncompressed(Source);
@@ -845,7 +798,7 @@ namespace mu
 				if (!Source)
 				{
 					// TODO: Warn?
-					Source = GenerateMissingImageCode(TEXT("Swizzle channel"), EImageFormat::L_UByte, InNode->GetMessageContext(), Options);
+					Source = GenerateMissingImageCode(TEXT("Swizzle channel"), EImageFormat::L_UByte, Node->GetMessageContext(), Options);
 				}
 
                 Ptr<ASTOp> SizedSource;
@@ -862,13 +815,13 @@ namespace mu
 				}
 
                 SwizzleOp->Sources[SourceIndex] = SizedSource;
-                SwizzleOp->SourceChannels[SourceIndex] = (uint8)node.m_sourceChannels[SourceIndex];
+                SwizzleOp->SourceChannels[SourceIndex] = (uint8)Node->SourceChannels[SourceIndex];
 			}
 		}
 
         if (FirstValidSourceIndex < 0)
 		{
-            Ptr<ASTOp> Source = GenerateMissingImageCode(TEXT("First swizzle image"), EImageFormat::RGBA_UByte, InNode->GetMessageContext(), Options);
+            Ptr<ASTOp> Source = GenerateMissingImageCode(TEXT("First swizzle image"), EImageFormat::RGBA_UByte, Node->GetMessageContext(), Options);
             SwizzleOp->Sources[0] = Source;
 
 		}
@@ -887,28 +840,25 @@ namespace mu
 	}
 
 
-	//---------------------------------------------------------------------------------------------
-	void CodeGenerator::GenerateImage_Format(const FImageGenerationOptions& Options, FImageGenerationResult& Result, const NodeImageFormat* InNode)
+	void CodeGenerator::GenerateImage_Format(const FImageGenerationOptions& Options, FImageGenerationResult& Result, const NodeImageFormat* Node)
 	{
-		const NodeImageFormat::Private& node = *InNode->GetPrivate();
-
 		MUTABLE_CPUPROFILER_SCOPE(NodeImageFormat);
 
-        check(node.m_format != EImageFormat::None);
+        check(Node->Format != EImageFormat::None);
 
         Ptr<ASTOpImagePixelFormat> fop = new ASTOpImagePixelFormat();
-        fop->Format = node.m_format;
-        fop->FormatIfAlpha = node.m_formatIfAlpha;
+        fop->Format = Node->Format;
+        fop->FormatIfAlpha = Node->FormatIfAlpha;
 
 		// Source is required
-		if (!node.m_source)
+		if (!Node->Source)
 		{
-            fop->Source = GenerateMissingImageCode(TEXT("Source image for format."), EImageFormat::RGBA_UByte, InNode->GetMessageContext(), Options);
+            fop->Source = GenerateMissingImageCode(TEXT("Source image for format."), EImageFormat::RGBA_UByte, Node->GetMessageContext(), Options);
 		}
 		else
 		{
 			FImageGenerationResult BaseResult;
-			GenerateImage(Options, BaseResult, node.m_source);
+			GenerateImage(Options, BaseResult, Node->Source);
 			fop->Source = BaseResult.op;
 		}
 
@@ -916,25 +866,22 @@ namespace mu
 	}
 
 
-    //---------------------------------------------------------------------------------------------
-	void CodeGenerator::GenerateImage_Saturate(const FImageGenerationOptions& Options, FImageGenerationResult& Result, const NodeImageSaturate* InNode)
+	void CodeGenerator::GenerateImage_Saturate(const FImageGenerationOptions& Options, FImageGenerationResult& Result, const NodeImageSaturate* Node)
 	{
-		const NodeImageSaturate::Private& node = *InNode->GetPrivate();
-
 		Ptr<ASTOpImageSaturate> op = new ASTOpImageSaturate();
 
         // Source image
         Ptr<ASTOp> Base;
-        if ( node.m_pSource )
+        if (Node->Source )
         {
 			FImageGenerationResult BaseResult;
-			GenerateImage(Options, BaseResult, node.m_pSource);
+			GenerateImage(Options, BaseResult, Node->Source);
 			Base = BaseResult.op;
 		}
         else
         {
             // This argument is required
-			Base = GenerateMissingImageCode(TEXT("Saturate image"), EImageFormat::RGB_UByte, InNode->GetMessageContext(), Options);
+			Base = GenerateMissingImageCode(TEXT("Saturate image"), EImageFormat::RGB_UByte, Node->GetMessageContext(), Options);
         }
 		
 		Base = GenerateImageFormat(Base, GetRGBOrRGBAFormat(Base->GetImageDesc().m_format));
@@ -943,7 +890,7 @@ namespace mu
 
 
         // Factor
-        if ( NodeScalar* Factor = node.m_pFactor.get() )
+        if ( NodeScalar* Factor = Node->Factor.get() )
         {
 			FScalarGenerationResult ParamResult;
 			GenerateScalar(ParamResult, Options, Factor);
@@ -952,32 +899,29 @@ namespace mu
         else
         {
             // This argument is required
-            op->Factor = GenerateMissingScalarCode(TEXT("Saturation factor"), 0.5f, InNode->GetMessageContext() );
+            op->Factor = GenerateMissingScalarCode(TEXT("Saturation factor"), 0.5f, Node->GetMessageContext() );
         }
 
         Result.op = op;
     }
 
 
-    //---------------------------------------------------------------------------------------------
-	void CodeGenerator::GenerateImage_Luminance(const FImageGenerationOptions& Options, FImageGenerationResult& Result, const NodeImageLuminance* InNode)
+	void CodeGenerator::GenerateImage_Luminance(const FImageGenerationOptions& Options, FImageGenerationResult& Result, const NodeImageLuminance* Node)
 	{
-		const NodeImageLuminance::Private& node = *InNode->GetPrivate();
-
 		Ptr<ASTOpImageLuminance> op = new ASTOpImageLuminance();
 
         // Source image
         Ptr<ASTOp> base;
-        if ( node.m_pSource )
+        if (Node->Source )
         {
 			FImageGenerationResult BaseResult;
-			GenerateImage(Options, BaseResult, node.m_pSource);
+			GenerateImage(Options, BaseResult, Node->Source);
 			base = BaseResult.op;
         }
         else
         {
             // This argument is required
-            base = GenerateMissingImageCode(TEXT("Image luminance"), EImageFormat::RGB_UByte, InNode->GetMessageContext(), Options);
+            base = GenerateMissingImageCode(TEXT("Image luminance"), EImageFormat::RGB_UByte, Node->GetMessageContext(), Options);
         }
         base = GenerateImageFormat( base, EImageFormat::RGB_UByte );
         op->Base = base;
@@ -986,35 +930,32 @@ namespace mu
     }
 
 
-    //---------------------------------------------------------------------------------------------
-	void CodeGenerator::GenerateImage_ColourMap(const FImageGenerationOptions& Options, FImageGenerationResult& Result, const NodeImageColourMap* InNode)
+	void CodeGenerator::GenerateImage_ColourMap(const FImageGenerationOptions& Options, FImageGenerationResult& Result, const NodeImageColourMap* Node)
 	{
-		const NodeImageColourMap::Private& node = *InNode->GetPrivate();
-
 		Ptr<ASTOpImageColorMap> op = new ASTOpImageColorMap();
 
         // Base image
         Ptr<ASTOp> base ;
-        if ( node.m_pBase )
+        if (Node->Base )
         {
 			FImageGenerationResult BaseResult;
-			GenerateImage(Options, BaseResult, node.m_pBase);
+			GenerateImage(Options, BaseResult, Node->Base);
 			base = BaseResult.op;
 		}
         else
         {
             // This argument is required
-            base = GenerateMissingImageCode(TEXT("Colourmap base image"), EImageFormat::RGB_UByte, InNode->GetMessageContext(), Options);
+            base = GenerateMissingImageCode(TEXT("Colourmap base image"), EImageFormat::RGB_UByte, Node->GetMessageContext(), Options);
         }
         base = GenerateImageSize( base, Options.RectSize);
         op->Base = base;
 
         // Mask of the effect
         Ptr<ASTOp> mask;
-        if ( node.m_pMask )
+        if (Node->Mask )
         {
 			FImageGenerationResult BaseResult;
-			GenerateImage(Options, BaseResult, node.m_pMask);
+			GenerateImage(Options, BaseResult, Node->Mask);
 			mask = BaseResult.op;
         }
         else
@@ -1029,7 +970,7 @@ namespace mu
 
         // Map image
         Ptr<ASTOp> MapImageOp;
-        if ( node.m_pMap )
+        if (Node->Map )
         {
 			FImageGenerationOptions MapOptions = Options;
 			MapOptions.ImageLayoutStrategy = CompilerOptions::TextureLayoutStrategy::None;
@@ -1038,14 +979,14 @@ namespace mu
 			MapOptions.RectSize = {};
 
 			FImageGenerationResult BaseResult;
-			GenerateImage(MapOptions, BaseResult, node.m_pMap);
+			GenerateImage(MapOptions, BaseResult, Node->Map);
 			MapImageOp = BaseResult.op;
 			MapImageOp = GenerateImageFormat(MapImageOp, EImageFormat::RGB_UByte);
 		}
         else
         {
             // This argument is required
-			MapImageOp = GenerateMissingImageCode(TEXT("Map image"), EImageFormat::RGB_UByte, InNode->GetMessageContext(), Options);
+			MapImageOp = GenerateMissingImageCode(TEXT("Map image"), EImageFormat::RGB_UByte, Node->GetMessageContext(), Options);
         }
         op->Map = MapImageOp;
 
@@ -1053,25 +994,22 @@ namespace mu
     }
 
 
-    //---------------------------------------------------------------------------------------------
-	void CodeGenerator::GenerateImage_Binarise(const FImageGenerationOptions& Options, FImageGenerationResult& Result, const NodeImageBinarise* InNode)
+	void CodeGenerator::GenerateImage_Binarise(const FImageGenerationOptions& Options, FImageGenerationResult& Result, const NodeImageBinarise* Node)
 	{
-		const NodeImageBinarise::Private& node = *InNode->GetPrivate();
-
 		Ptr<ASTOpImageBinarize> op = new ASTOpImageBinarize();
 
         // A image
         Ptr<ASTOp> BaseOp;
-        if ( node.m_pBase )
+        if (Node->Base )
         {
 			FImageGenerationResult BaseResult;
-			GenerateImage(Options, BaseResult, node.m_pBase);
+			GenerateImage(Options, BaseResult, Node->Base);
 			BaseOp = BaseResult.op;
 		}
         else
         {
             // This argument is required
-			BaseOp = GenerateMissingImageCode(TEXT("Image Binarise Base"), EImageFormat::RGB_UByte, InNode->GetMessageContext(), Options);
+			BaseOp = GenerateMissingImageCode(TEXT("Image Binarise Base"), EImageFormat::RGB_UByte, Node->GetMessageContext(), Options);
         }
 		BaseOp = GenerateImageFormat(BaseOp, EImageFormat::RGB_UByte );
 		BaseOp = GenerateImageSize(BaseOp, Options.RectSize);
@@ -1079,7 +1017,7 @@ namespace mu
 
         // Threshold
         Ptr<ASTOp> ThresholdOp;
-        if ( NodeScalar* ThresholdNode = node.m_pThreshold.get() )
+        if ( NodeScalar* ThresholdNode = Node->Threshold.get() )
         {
 			FScalarGenerationResult ParamResult;
 			GenerateScalar(ParamResult, Options, ThresholdNode);
@@ -1088,7 +1026,7 @@ namespace mu
         else
         {
             // This argument is required
-			ThresholdOp = GenerateMissingScalarCode(TEXT("Image Binarise Threshold"), 0.5f, InNode->GetMessageContext() );
+			ThresholdOp = GenerateMissingScalarCode(TEXT("Image Binarise Threshold"), 0.5f, Node->GetMessageContext() );
         }
         op->Threshold = ThresholdOp;
 
@@ -1096,50 +1034,47 @@ namespace mu
     }
 
 
-    //---------------------------------------------------------------------------------------------
-	void CodeGenerator::GenerateImage_Resize(const FImageGenerationOptions& Options, FImageGenerationResult& Result, const NodeImageResize* InNode)
+	void CodeGenerator::GenerateImage_Resize(const FImageGenerationOptions& Options, FImageGenerationResult& Result, const NodeImageResize* Node)
 	{
-		const NodeImageResize::Private& node = *InNode->GetPrivate();
-
 		MUTABLE_CPUPROFILER_SCOPE(NodeImageResize);
 
         Ptr<ASTOp> at;
 
         // Source image
         Ptr<ASTOp> base;
-        if ( node.m_pBase )
+        if (Node->Base )
         {
 			FImageGenerationOptions NewOptions = Options;
 			
-			if (node.m_relative)
+			if (Node->bRelative)
 			{
-				NewOptions.RectSize[0] = FMath::RoundToInt32(NewOptions.RectSize[0] / node.m_sizeX);
-				NewOptions.RectSize[1] = FMath::RoundToInt32(NewOptions.RectSize[1] / node.m_sizeY);
+				NewOptions.RectSize[0] = FMath::RoundToInt32(NewOptions.RectSize[0] / Node->SizeX);
+				NewOptions.RectSize[1] = FMath::RoundToInt32(NewOptions.RectSize[1] / Node->SizeY);
 			}
 
 			FImageGenerationResult BaseResult;
-			GenerateImage(NewOptions, BaseResult, node.m_pBase);
+			GenerateImage(NewOptions, BaseResult, Node->Base);
 			base = BaseResult.op;
         }
         else
         {
             // This argument is required
-            base = GenerateMissingImageCode(TEXT("Image resize base"), EImageFormat::RGB_UByte, InNode->GetMessageContext(), Options);
+            base = GenerateMissingImageCode(TEXT("Image resize base"), EImageFormat::RGB_UByte, Node->GetMessageContext(), Options);
         }
 
         // Size
-        if ( node.m_relative )
+        if (Node->bRelative )
         {
             Ptr<ASTOpImageResizeRel> NewOp = new ASTOpImageResizeRel();
-			NewOp->Factor[0] = node.m_sizeX;
-			NewOp->Factor[1] = node.m_sizeY;
+			NewOp->Factor[0] = Node->SizeX;
+			NewOp->Factor[1] = Node->SizeY;
 			NewOp->Source = base;
             at = NewOp;
         }
         else
         {
 			// Apply the layout block to the rect size
-			UE::Math::TIntVector2<int32> FinalImageSize = { int32(node.m_sizeX), int32(node.m_sizeY) };
+			UE::Math::TIntVector2<int32> FinalImageSize = { int32(Node->SizeX), int32(Node->SizeY) };
 			if (Options.LayoutToApply)
 			{
 				int32 BlockIndex = Options.LayoutToApply->FindBlock(Options.LayoutBlockId);
@@ -1170,23 +1105,20 @@ namespace mu
     }
 
 
-    //---------------------------------------------------------------------------------------------
-	void CodeGenerator::GenerateImage_PlainColour(const FImageGenerationOptions& Options, FImageGenerationResult& Result, const NodeImagePlainColour* InNode)
+	void CodeGenerator::GenerateImage_PlainColour(const FImageGenerationOptions& Options, FImageGenerationResult& Result, const NodeImagePlainColour* Node)
 	{
-		const NodeImagePlainColour::Private& node = *InNode->GetPrivate();
-
 		// Source colour
         Ptr<ASTOp> base;
-        if ( node.m_pColour )
+        if (Node->Colour )
         {
 			FColorGenerationResult ColorResult;
-            GenerateColor( ColorResult, Options, node.m_pColour.get() );
+            GenerateColor( ColorResult, Options, Node->Colour.get() );
 			base = ColorResult.op;
         }
         else
         {
             // This argument is required
-            base = GenerateMissingColourCode(TEXT("Image plain colour base"), InNode->GetMessageContext() );
+            base = GenerateMissingColourCode(TEXT("Image plain colour base"), Node->GetMessageContext() );
         }
 
 		UE::Math::TIntVector2<int32> FinalImageSize = { 0, 0 };
@@ -1197,7 +1129,7 @@ namespace mu
 		}
 		else
 		{
-			FinalImageSize = { node.m_sizeX, node.m_sizeY };
+			FinalImageSize = { Node->SizeX, Node->SizeY };
 
 			// Apply the layout block to the rect size
 			if (Options.LayoutToApply)
@@ -1223,7 +1155,7 @@ namespace mu
 
         Ptr<ASTOpImagePlainColor> op = new ASTOpImagePlainColor();
         op->Color = base;
-		op->Format = node.Format;
+		op->Format = Node->Format;
 		op->Size[0] = FinalImageSize[0];
 		op->Size[1] = FinalImageSize[1];
 		op->LODs = 1;
@@ -1233,16 +1165,14 @@ namespace mu
 
 
     //---------------------------------------------------------------------------------------------
-	void CodeGenerator::GenerateImage_Switch(const FImageGenerationOptions& Options, FImageGenerationResult& Result, const NodeImageSwitch* InNode)
+	void CodeGenerator::GenerateImage_Switch(const FImageGenerationOptions& Options, FImageGenerationResult& Result, const NodeImageSwitch* Node)
 	{
-		const NodeImageSwitch::Private& node = *InNode->GetPrivate();
-
 		MUTABLE_CPUPROFILER_SCOPE(NodeImageSwitch);
 
-        if (node.Options.Num() == 0)
+        if (Node->Options.Num() == 0)
 		{
 			// No options in the switch!
-			Result.op = GenerateMissingImageCode(TEXT("Switch option"), EImageFormat::RGBA_UByte, InNode->GetMessageContext(), Options);
+			Result.op = GenerateMissingImageCode(TEXT("Switch option"), EImageFormat::RGBA_UByte, Node->GetMessageContext(), Options);
 			return;
 		}
 
@@ -1250,33 +1180,33 @@ namespace mu
         op->type = EOpType::IM_SWITCH;
 
 		// Variable value
-		if ( node.m_pParameter )
+		if (Node->Parameter )
 		{
 			FScalarGenerationResult ParamResult;
-			GenerateScalar(ParamResult, Options, node.m_pParameter.get());
+			GenerateScalar(ParamResult, Options, Node->Parameter.get());
 			op->variable = ParamResult.op;
 		}
 		else
 		{
 			// This argument is required
-            op->variable = GenerateMissingScalarCode(TEXT("Switch variable"), 0.0f, InNode->GetMessageContext() );
+            op->variable = GenerateMissingScalarCode(TEXT("Switch variable"), 0.0f, Node->GetMessageContext() );
 		}
 
 		// Options
-        for ( int32 OptionIndex=0; OptionIndex < node.Options.Num(); ++OptionIndex)
+        for ( int32 OptionIndex=0; OptionIndex < Node->Options.Num(); ++OptionIndex)
         {
             Ptr<ASTOp> branch;
 
-            if (node.Options[OptionIndex])
+            if (Node->Options[OptionIndex])
             {
 				FImageGenerationResult BaseResult;
-				GenerateImage(Options, BaseResult, node.Options[OptionIndex]);
+				GenerateImage(Options, BaseResult, Node->Options[OptionIndex]);
 				branch = BaseResult.op;
 			}
             else
             {
                 // This argument is required
-                branch = GenerateMissingImageCode(TEXT("Switch option"), EImageFormat::RGBA_UByte, InNode->GetMessageContext(), Options);
+                branch = GenerateMissingImageCode(TEXT("Switch option"), EImageFormat::RGBA_UByte, Node->GetMessageContext(), Options);
             }
 
             op->cases.Emplace(int16(OptionIndex),op,branch);
@@ -1288,45 +1218,39 @@ namespace mu
     }
 
 
-    //---------------------------------------------------------------------------------------------
-	void CodeGenerator::GenerateImage_Conditional(const FImageGenerationOptions& Options, FImageGenerationResult& Result, const NodeImageConditional* InNode)
+	void CodeGenerator::GenerateImage_Conditional(const FImageGenerationOptions& Options, FImageGenerationResult& Result, const NodeImageConditional* Node)
 	{
-		const NodeImageConditional::Private& node = *InNode->GetPrivate();
-
 		Ptr<ASTOpConditional> op = new ASTOpConditional();
         op->type = EOpType::IM_CONDITIONAL;
 
         // Condition
-        if ( node.Parameter )
+        if (Node->Parameter )
         {
 			FBoolGenerationResult ParamResult;
-			GenerateBool(ParamResult, Options, node.Parameter.get());
+			GenerateBool(ParamResult, Options, Node->Parameter.get());
 			op->condition = ParamResult.op;
 		}
         else
         {
             // This argument is required
-            op->condition = GenerateMissingBoolCode(TEXT("Conditional condition"), true, InNode->GetMessageContext() );
+            op->condition = GenerateMissingBoolCode(TEXT("Conditional condition"), true, Node->GetMessageContext() );
         }
 
         // Options
 		FImageGenerationResult YesResult;
-		GenerateImage(Options, YesResult, node.m_true);
+		GenerateImage(Options, YesResult, Node->True);
 		op->yes = YesResult.op;
 
 		FImageGenerationResult NoResult;
-		GenerateImage(Options, NoResult, node.m_false);
+		GenerateImage(Options, NoResult, Node->False);
 		op->no = NoResult.op;
 		
         Result.op = op;
     }
 
 
-    //---------------------------------------------------------------------------------------------
-	void CodeGenerator::GenerateImage_Project(const FImageGenerationOptions& Options, FImageGenerationResult& Result, const NodeImageProject* InNode)
+	void CodeGenerator::GenerateImage_Project(const FImageGenerationOptions& Options, FImageGenerationResult& Result, const NodeImageProject* Node)
 	{
-		const NodeImageProject::Private& node = *InNode->GetPrivate();
-
 		MUTABLE_CPUPROFILER_SCOPE(NodeImageProject);
 
         // Mesh project operation
@@ -1337,15 +1261,15 @@ namespace mu
 
         // Projector
         FProjectorGenerationResult projectorResult;
-        if ( node.m_pProjector )
+        if (Node->Projector )
         {
-            GenerateProjector( projectorResult, Options, node.m_pProjector );
-            //projectorAt = Generate( node.m_pProjector.get() );
+            GenerateProjector( projectorResult, Options, Node->Projector );
+            //projectorAt = Generate( Node->Projector.get() );
         }
         else
         {
             // This argument is required
-            GenerateMissingProjectorCode( projectorResult, InNode->GetMessageContext() );
+            GenerateMissingProjectorCode( projectorResult, Node->GetMessageContext() );
         }
 
 		ProjectOp->Projector = projectorResult.op;
@@ -1358,7 +1282,7 @@ namespace mu
 		uint64 GeneratedLayoutBlockId = FLayoutBlock::InvalidBlockId;
 
         // Mesh
-        if ( node.m_pMesh )
+        if (Node->Mesh )
         {
 			// TODO: This will probably Result in a duplicated mesh subgraph, with the original mesh but new layout block ids.
 			// See if it can be optimized and try to reuse the existing layout block ids instead of generating new ones.
@@ -1369,7 +1293,7 @@ namespace mu
 			MeshOptions.bLayouts = true;			// We need the layout that we will use to render
 			MeshOptions.bNormalizeUVs = true;		// We need normalized UVs for the projection
 
-			FMeshTask MeshTask = GenerateMesh(MeshStaticOptions, UE::Tasks::MakeCompletedTask<FMeshGenerationDynamicOptions>(MeshOptions), node.m_pMesh );
+			FMeshTask MeshTask = GenerateMesh(MeshStaticOptions, UE::Tasks::MakeCompletedTask<FMeshGenerationDynamicOptions>(MeshOptions), Node->Mesh );
 			// This forces a wait to sync here. When images are also generated in tasks, it can be turned into a prerequisite instead.
 			if (WaitCallback.IsSet())
 			{
@@ -1381,7 +1305,7 @@ namespace mu
 			FMeshGenerationResult MeshResult = MeshTask.GetResult();
 
 			// Match the block id of the block we are generating with the id that resulted in the generated mesh			
-			TSharedPtr<const FLayout> Layout = MeshResult.GeneratedLayouts.IsValidIndex(node.m_layout) ? MeshResult.GeneratedLayouts[node.m_layout].Layout : nullptr;
+			TSharedPtr<const FLayout> Layout = MeshResult.GeneratedLayouts.IsValidIndex(Node->Layout) ? MeshResult.GeneratedLayouts[Node->Layout].Layout : nullptr;
 			if (Layout && Layout->Blocks.IsValidIndex(LayoutBlockIndex))
 			{
 				GeneratedLayoutBlockId = Layout->Blocks[LayoutBlockIndex].Id;
@@ -1393,7 +1317,7 @@ namespace mu
 			}
 			else
 			{
-				ErrorLog->Add("Layout or block index error.", ELMT_ERROR, InNode->GetMessageContext());
+				ErrorLog->Add("Layout or block index error.", ELMT_ERROR, Node->GetMessageContext());
 			}
 
 			// TODO: 
@@ -1417,7 +1341,7 @@ namespace mu
                 cop->Type = EOpType::ME_CONSTANT;
 				cop->SourceDataDescriptor = CurrentMeshToProjectOp->GetSourceDataDescriptor();
 				TSharedPtr<FMesh> FormatMeshResult = MakeShared<FMesh>();
-				CreateMeshOptimisedForWrappingProjection(FormatMeshResult.Get(), node.m_layout);
+				CreateMeshOptimisedForWrappingProjection(FormatMeshResult.Get(), Node->Layout);
 
                 cop->SetValue(FormatMeshResult, CompilerOptions->OptimisationOptions.DiskCacheContext);
 
@@ -1436,7 +1360,7 @@ namespace mu
                 {
                     Ptr<ASTOpMeshExtractLayoutBlocks> eop = new ASTOpMeshExtractLayoutBlocks();
                     eop->Source = CurrentMeshToProjectOp;
-                    eop->LayoutIndex = node.m_layout;
+                    eop->LayoutIndex = Node->Layout;
 
                     eop->Blocks.Add(GeneratedLayoutBlockId);
 
@@ -1449,7 +1373,7 @@ namespace mu
 				cop->SourceDataDescriptor = CurrentMeshToProjectOp->GetSourceDataDescriptor();
 
 				TSharedPtr<FMesh> FormatMeshResult = MakeShared<FMesh>();
-                CreateMeshOptimisedForProjection(FormatMeshResult.Get(), node.m_layout);
+                CreateMeshOptimisedForProjection(FormatMeshResult.Get(), Node->Layout);
 
                 cop->SetValue(FormatMeshResult, CompilerOptions->OptimisationOptions.DiskCacheContext);
 
@@ -1473,7 +1397,7 @@ namespace mu
             cop->Type = EOpType::ME_CONSTANT;
             cop->SetValue(TempMesh, CompilerOptions->OptimisationOptions.DiskCacheContext);
 			ProjectOp->Mesh = cop;
-            ErrorLog->Add( "Projector mesh not set.", ELMT_ERROR, InNode->GetMessageContext() );
+            ErrorLog->Add( "Projector mesh not set.", ELMT_ERROR, Node->GetMessageContext() );
         }
 
 
@@ -1496,15 +1420,15 @@ namespace mu
 
 		if (RasterImageSize.X == 0)
 		{
-			RasterImageSize = UE::Math::TIntVector2<int32>(node.m_imageSize);
+			RasterImageSize = UE::Math::TIntVector2<int32>(Node->ImageSize);
 			bApplyLayoutToSize = true;
 		}
 
 		// Target mask
-		if (node.m_pMask)
+		if (Node->Mask)
 		{
 			FImageGenerationResult MaskResult;
-			GenerateImage(Options, MaskResult, node.m_pMask);
+			GenerateImage(Options, MaskResult, Node->Mask);
 			Ptr<ASTOp> mask = MaskResult.op;
 
 			mask = GenerateImageFormat(mask, EImageFormat::L_UByte);
@@ -1520,7 +1444,7 @@ namespace mu
 		}
 
         // Image
-        if ( node.m_pImage )
+        if (Node->Image )
         {
             // Generate
 			FImageGenerationOptions NewOptions = Options;
@@ -1530,7 +1454,7 @@ namespace mu
 			NewOptions.RectSize = { 0,0 };
 
 			FImageGenerationResult ImageResult;
-			GenerateImage(NewOptions, ImageResult, node.m_pImage);
+			GenerateImage(NewOptions, ImageResult, Node->Image);
 			ImageRasterOp->image = ImageResult.op;
 
 			FImageDesc desc = ImageRasterOp->image->GetImageDesc();
@@ -1546,7 +1470,7 @@ namespace mu
         else
         {
             // This argument is required
-			ImageRasterOp->image = GenerateMissingImageCode(TEXT("Projector image"), EImageFormat::RGB_UByte, InNode->GetMessageContext(), Options);
+			ImageRasterOp->image = GenerateMissingImageCode(TEXT("Projector image"), EImageFormat::RGB_UByte, Node->GetMessageContext(), Options);
         }
 
 		if (RasterImageSize.X == 0)
@@ -1579,26 +1503,38 @@ namespace mu
 		ImageRasterOp->SizeX = RasterImageSize[0];
 		ImageRasterOp->SizeY = RasterImageSize[1];
 		ImageRasterOp->BlockId = GeneratedLayoutBlockId;
-		ImageRasterOp->LayoutIndex = node.m_layout;
+		ImageRasterOp->LayoutIndex = Node->Layout;
 
-		ImageRasterOp->bIsRGBFadingEnabled = node.bIsRGBFadingEnabled;
-		ImageRasterOp->bIsAlphaFadingEnabled = node.bIsAlphaFadingEnabled;
-		ImageRasterOp->SamplingMethod = node.SamplingMethod;
-		ImageRasterOp->MinFilterMethod = node.MinFilterMethod;
+		ImageRasterOp->bIsRGBFadingEnabled = Node->bIsRGBFadingEnabled;
+		ImageRasterOp->bIsAlphaFadingEnabled = Node->bIsAlphaFadingEnabled;
+		ImageRasterOp->SamplingMethod = Node->SamplingMethod;
+		ImageRasterOp->MinFilterMethod = Node->MinFilterMethod;
 
 		// Fading angles are optional, and stored in a colour. If one exists, we generate both.
-		if (node.m_pAngleFadeStart || node.m_pAngleFadeEnd)
+		if (Node->AngleFadeStart || Node->AngleFadeEnd)
 		{
 			Ptr<NodeScalarConstant> pDefaultFade = new NodeScalarConstant();
-			pDefaultFade->SetValue(180.0f);
+			pDefaultFade->Value = 180.0f;
 
 			Ptr<NodeColourFromScalars> pPropsNode = new NodeColourFromScalars();
 
-			if (node.m_pAngleFadeStart) pPropsNode->X = node.m_pAngleFadeStart;
-			else pPropsNode->X = pDefaultFade;
+			if (Node->AngleFadeStart)
+			{
+				pPropsNode->X = Node->AngleFadeStart;
+			}
+			else
+			{
+				pPropsNode->X = pDefaultFade;
+			}
 
-			if (node.m_pAngleFadeEnd) pPropsNode->Y = node.m_pAngleFadeEnd;
-			else pPropsNode->Y = pDefaultFade;
+			if (Node->AngleFadeEnd)
+			{
+				pPropsNode->Y = Node->AngleFadeEnd;
+			}
+			else
+			{
+				pPropsNode->Y = pDefaultFade;
+			}
 
 			FColorGenerationResult ParamResult;
 			GenerateColor(ParamResult, Options, pPropsNode);
@@ -1607,7 +1543,7 @@ namespace mu
 
         // Seam correction operations
         //------------------------------
-		if (node.bEnableTextureSeamCorrection)
+		if (Node->bEnableTextureSeamCorrection)
 		{
 			Ptr<ASTOpImageRasterMesh> MaskRasterOp = new ASTOpImageRasterMesh();
 			MaskRasterOp->mesh = ImageRasterOp->mesh.child();
@@ -1650,11 +1586,8 @@ namespace mu
     }
 
 
-    //---------------------------------------------------------------------------------------------
-	void CodeGenerator::GenerateImage_Mipmap(const FImageGenerationOptions& Options, FImageGenerationResult& Result, const NodeImageMipmap* InNode)
+	void CodeGenerator::GenerateImage_Mipmap(const FImageGenerationOptions& Options, FImageGenerationResult& Result, const NodeImageMipmap* Node)
 	{
-		const NodeImageMipmap::Private& node = *InNode->GetPrivate();
-
 		MUTABLE_CPUPROFILER_SCOPE(NodeImageMipmap);
 
         Ptr<ASTOp> res;
@@ -1666,17 +1599,17 @@ namespace mu
 
         // Source image
         Ptr<ASTOp> base;
-        if ( node.m_pSource )
+        if (Node->Source )
         {
             MUTABLE_CPUPROFILER_SCOPE(Base);
 			FImageGenerationResult BaseResult;
-			GenerateImage(Options, BaseResult, node.m_pSource);
+			GenerateImage(Options, BaseResult, Node->Source);
 			base = BaseResult.op;
 		}
         else
         {
             // This argument is required
-            base = GenerateMissingImageCode(TEXT("Mipmap image"), EImageFormat::RGB_UByte, InNode->GetMessageContext(), Options);
+            base = GenerateMissingImageCode(TEXT("Mipmap image"), EImageFormat::RGB_UByte, Node->GetMessageContext(), Options);
         }
 
         op->Source = base;
@@ -1705,8 +1638,8 @@ namespace mu
             op->BlockLevels = 0;
         }
 
-		op->AddressMode = node.m_settings.AddressMode;
-		op->FilterType = node.m_settings.FilterType;
+		op->AddressMode = Node->Settings.AddressMode;
+		op->FilterType = Node->Settings.FilterType;
 
         res = op;
 
@@ -1714,24 +1647,21 @@ namespace mu
     }
 
 
-	//---------------------------------------------------------------------------------------------
-	void CodeGenerator::GenerateImage_Invert(const FImageGenerationOptions& Options, FImageGenerationResult& Result, const NodeImageInvert* InNode)
+	void CodeGenerator::GenerateImage_Invert(const FImageGenerationOptions& Options, FImageGenerationResult& Result, const NodeImageInvert* Node)
 	{
-		const NodeImageInvert::Private& node = *InNode->GetPrivate();
-
 		Ptr<ASTOpImageInvert> op = new ASTOpImageInvert();
 
 		Ptr<ASTOp> a;
-		if (node.m_pBase)
+		if (Node->Base)
 		{
 			FImageGenerationResult BaseResult;
-			GenerateImage(Options, BaseResult, node.m_pBase);
+			GenerateImage(Options, BaseResult, Node->Base);
 			a = BaseResult.op;
 		}
 		else
 		{
 			// This argument is required
-			a = GenerateMissingImageCode(TEXT("Image Invert Color"), EImageFormat::RGB_UByte, InNode->GetMessageContext(), Options);
+			a = GenerateMissingImageCode(TEXT("Image Invert Color"), EImageFormat::RGB_UByte, Node->GetMessageContext(), Options);
 		}
 		a = GenerateImageFormat(a, EImageFormat::RGB_UByte);
 		a = GenerateImageSize(a, Options.RectSize);
@@ -1741,26 +1671,23 @@ namespace mu
 	}
 
 
-    //---------------------------------------------------------------------------------------------
-	void CodeGenerator::GenerateImage_Variation(const FImageGenerationOptions& Options, FImageGenerationResult& Result, const NodeImageVariation* InNode)
+	void CodeGenerator::GenerateImage_Variation(const FImageGenerationOptions& Options, FImageGenerationResult& Result, const NodeImageVariation* Node)
 	{
-		const NodeImageVariation::Private& node = *InNode->GetPrivate();
-
 		Ptr<ASTOp> currentOp;
 
         // Default case
-        if ( node.m_defaultImage )
+        if (Node->DefaultImage )
         {
             FImageGenerationResult BranchResults;
-			GenerateImage(Options, BranchResults, node.m_defaultImage);
+			GenerateImage(Options, BranchResults, Node->DefaultImage);
 			currentOp = BranchResults.op;
         }
         
         // Process variations in reverse order, since conditionals are built bottom-up.
-        for ( int32 t = int32( node.m_variations.Num() ) - 1; t >= 0; --t )
+        for ( int32 t = int32(Node->Variations.Num() ) - 1; t >= 0; --t )
         {
             int32 tagIndex = -1;
-            const FString& tag = node.m_variations[t].m_tag;
+            const FString& tag = Node->Variations[t].Tag;
             for ( int32 i = 0; i < int32( FirstPass.Tags.Num() ); ++i )
             {
                 if ( FirstPass.Tags[i].Tag == tag )
@@ -1773,21 +1700,21 @@ namespace mu
             {
 				FString Msg = FString::Printf(TEXT("Unknown tag found in image variation [%s]."), *tag );
 
-                ErrorLog->Add( Msg, ELMT_WARNING, InNode->GetMessageContext() );
+                ErrorLog->Add( Msg, ELMT_WARNING, Node->GetMessageContext() );
                 continue;
             }
 
             Ptr<ASTOp> variationOp;
-            if ( node.m_variations[t].m_image )
+            if (Node->Variations[t].Image )
             {
 				FImageGenerationResult VariationResult;
-				GenerateImage(Options, VariationResult, node.m_variations[t].m_image);
+				GenerateImage(Options, VariationResult, Node->Variations[t].Image);
 				variationOp = VariationResult.op;
             }
             else
             {
                 // This argument is required
-                variationOp = GenerateMissingImageCode(TEXT("Variation option"), EImageFormat::RGBA_UByte, InNode->GetMessageContext(), Options);
+                variationOp = GenerateMissingImageCode(TEXT("Variation option"), EImageFormat::RGBA_UByte, Node->GetMessageContext(), Options);
             }
 
 
@@ -1834,7 +1761,6 @@ namespace mu
     }
 
 	
-	//---------------------------------------------------------------------------------------------
 	void CodeGenerator::GenerateImage_Table(const FImageGenerationOptions& Options, FImageGenerationResult& Result, const NodeImageTable* InNode)
 	{
 		const NodeImageTable& node = *InNode;
@@ -2004,14 +1930,13 @@ namespace mu
 	}
 
 
-    //---------------------------------------------------------------------------------------------
     Ptr<ASTOp> CodeGenerator::GeneratePlainImageCode( const FVector4f& InColor, const FImageGenerationOptions& Options )
     {
 		Ptr<NodeColourConstant> ConstantColor = new NodeColourConstant();
 		ConstantColor->Value = InColor;
 
         Ptr<NodeImagePlainColour> PlainNode = new NodeImagePlainColour();
-		PlainNode->SetColour(ConstantColor);
+		PlainNode->Colour = ConstantColor;
 
 		FImageGenerationResult TempResult;
 		GenerateImage(Options, TempResult, PlainNode);
@@ -2021,7 +1946,6 @@ namespace mu
     }
 
 
-    //---------------------------------------------------------------------------------------------
     Ptr<ASTOp> CodeGenerator::GenerateImageFormat( Ptr<ASTOp> Op, EImageFormat InFormat)
     {
         Ptr<ASTOp> Result = Op;
@@ -2039,7 +1963,6 @@ namespace mu
     }
 
 
-    //---------------------------------------------------------------------------------------------
     Ptr<ASTOp> CodeGenerator::GenerateImageUncompressed( Ptr<ASTOp> at )
     {
         Ptr<ASTOp> Result = at;
@@ -2063,7 +1986,6 @@ namespace mu
     }
 
 
-    //---------------------------------------------------------------------------------------------
     Ptr<ASTOp> CodeGenerator::GenerateImageSize( Ptr<ASTOp> at, UE::Math::TIntVector2<int32> size )
     {
         Ptr<ASTOp> Result = at;
@@ -2083,5 +2005,4 @@ namespace mu
         return Result;
     }
 
- 
 }
diff --git a/Engine/Plugins/Mutable/Source/MutableTools/Private/MuT/CodeGenerator_Mesh.cpp b/Engine/Plugins/Mutable/Source/MutableTools/Private/MuT/CodeGenerator_Mesh.cpp
--- a/Engine/Plugins/Mutable/Source/MutableTools/Private/MuT/CodeGenerator_Mesh.cpp
+++ b/Engine/Plugins/Mutable/Source/MutableTools/Private/MuT/CodeGenerator_Mesh.cpp
@@ -41,28 +41,20 @@
 #include "MuT/NodeLayout.h"
 #include "MuT/NodeMesh.h"
 #include "MuT/NodeMeshApplyPose.h"
-#include "MuT/NodeMeshApplyPosePrivate.h"
 #include "MuT/NodeMeshClipDeform.h"
 #include "MuT/NodeMeshClipMorphPlane.h"
 #include "MuT/NodeMeshClipWithMesh.h"
 #include "MuT/NodeMeshConstant.h"
-#include "MuT/NodeMeshConstantPrivate.h"
 #include "MuT/NodeMeshFormat.h"
-#include "MuT/NodeMeshFormatPrivate.h"
 #include "MuT/NodeMeshFragment.h"
 #include "MuT/NodeMeshMakeMorph.h"
-#include "MuT/NodeMeshMakeMorphPrivate.h"
 #include "MuT/NodeMeshMorph.h"
-#include "MuT/NodeMeshMorphPrivate.h"
 #include "MuT/NodeMeshReshape.h"
-#include "MuT/NodeMeshReshapePrivate.h"
 #include "MuT/NodeMeshSwitch.h"
-#include "MuT/NodeMeshSwitchPrivate.h"
 #include "MuT/NodeMeshTable.h"
 #include "MuT/NodeMeshTransform.h"
-#include "MuT/NodeMeshTransformPrivate.h"
 #include "MuT/NodeMeshVariation.h"
-#include "MuT/NodeMeshVariationPrivate.h"
+#include "MuT/NodeMeshParameter.h"
 #include "MuT/NodeScalar.h"
 #include "MuT/Table.h"
 #include "MuT/TablePrivate.h"
@@ -613,6 +605,7 @@ namespace mu
 		case Node::EType::MeshTable: Result = GenerateMesh_Table(StaticOptions, InOptions, static_cast<const NodeMeshTable*>(Node)); break;
 		case Node::EType::MeshReshape: Result = GenerateMesh_Reshape(StaticOptions, InOptions, static_cast<const NodeMeshReshape*>(Node)); break;
 		case Node::EType::MeshClipDeform: Result = GenerateMesh_ClipDeform(StaticOptions, InOptions, static_cast<const NodeMeshClipDeform*>(Node)); break;
+		case Node::EType::MeshParameter: Result = GenerateMesh_Parameter(StaticOptions, InOptions, static_cast<const NodeMeshParameter*>(Node)); break;
 		default: check(false);
 		}
 
@@ -620,42 +613,40 @@ namespace mu
     }
 
 
-	FMeshTask CodeGenerator::GenerateMesh_Morph( const FMeshGenerationStaticOptions& StaticOptions, FMeshOptionsTask InOptions, const NodeMeshMorph* InMorphNode )
+	FMeshTask CodeGenerator::GenerateMesh_Morph( const FMeshGenerationStaticOptions& StaticOptions, FMeshOptionsTask InOptions, const NodeMeshMorph* Node )
     {
-        NodeMeshMorph::Private& node = *InMorphNode->GetPrivate();
-
 		Ptr<ASTOp> FactorOp;
 
         // Factor
-        if ( node.Factor )
+        if (Node->Factor )
         {
 			FScalarGenerationResult ParamResult;
-			GenerateScalar(ParamResult, StaticOptions, node.Factor.get());
+			GenerateScalar(ParamResult, StaticOptions, Node->Factor.get());
 			FactorOp = ParamResult.op;
         }
         else
         {
             // This argument is required
-			FactorOp = GenerateMissingScalarCode(TEXT("Morph factor"), 0.5f, InMorphNode->GetMessageContext());
+			FactorOp = GenerateMissingScalarCode(TEXT("Morph factor"), 0.5f, Node->GetMessageContext());
         }
 
 		TArray<UE::Tasks::FTask, TInlineAllocator<4>> Requisites;
 
         // Base
 		FMeshTask BaseTask;
-        if ( node.Base )
+        if (Node->Base )
         {
-			BaseTask = GenerateMesh(StaticOptions, InOptions, node.Base );
+			BaseTask = GenerateMesh(StaticOptions, InOptions, Node->Base );
 			Requisites.Add(BaseTask);
         }
         else
         {
             // This argument is required
-            ErrorLog->Add( "Mesh morph base node is not set.", ELMT_ERROR, InMorphNode->GetMessageContext());
+            ErrorLog->Add( "Mesh morph base node is not set.", ELMT_ERROR, Node->GetMessageContext());
 			return UE::Tasks::MakeCompletedTask<FMeshGenerationResult>();
         }        
 
-		if (!node.Morph)
+		if (!Node->Morph)
 		{
 			return BaseTask;
 		}
@@ -680,11 +671,11 @@ namespace mu
 				UE::Tasks::Prerequisites( BaseTask, InOptions )
 				);
 
-			TargetTask = GenerateMesh(TargetStaticOptions, TargetOptionsTask, node.Morph);
+			TargetTask = GenerateMesh(TargetStaticOptions, TargetOptionsTask, Node->Morph);
         }
 
 		return UE::Tasks::Launch(TEXT("MutableMorph"),
-			[BaseTask, TargetTask, node, FactorOp]() mutable
+			[BaseTask, TargetTask, Node, FactorOp]() mutable
 			{
 				FMeshGenerationResult BaseResult = BaseTask.GetResult();
 				FMeshGenerationResult TargetResult = TargetTask.GetResult();
@@ -694,7 +685,7 @@ namespace mu
 				OpMorph->Base = BaseTask.GetResult().MeshOp;
 				OpMorph->Target = TargetTask.GetResult().MeshOp;
 
-				const bool bReshapeEnabled = node.bReshapeSkeleton || node.bReshapePhysicsVolumes;
+				const bool bReshapeEnabled = Node->bReshapeSkeleton || Node->bReshapePhysicsVolumes;
 
 				Ptr<ASTOpMeshMorphReshape> OpMorphReshape;
 				if (bReshapeEnabled)
@@ -707,10 +698,10 @@ namespace mu
 					OpBind->bReshapeVertices = false;
 					OpBind->bApplyLaplacian = false;
 					OpBind->bRecomputeNormals = false;
-					OpBind->bReshapeSkeleton = node.bReshapeSkeleton;
-					OpBind->BonesToDeform = node.BonesToDeform;
-					OpBind->bReshapePhysicsVolumes = node.bReshapePhysicsVolumes;
-					OpBind->PhysicsToDeform = node.PhysicsToDeform;
+					OpBind->bReshapeSkeleton = Node->bReshapeSkeleton;
+					OpBind->BonesToDeform = Node->BonesToDeform;
+					OpBind->bReshapePhysicsVolumes = Node->bReshapePhysicsVolumes;
+					OpBind->PhysicsToDeform = Node->PhysicsToDeform;
 					OpBind->BindingMethod = static_cast<uint32>(EShapeBindingMethod::ReshapeClosestProject);
 
 					OpBind->Mesh = BaseResult.MeshOp;
@@ -749,21 +740,19 @@ namespace mu
 	}
 
 
-	FMeshTask CodeGenerator::GenerateMesh_MakeMorph(const FMeshGenerationStaticOptions& StaticOptions, FMeshOptionsTask InOptions, const NodeMeshMakeMorph* InMakeMorphNode )
+	FMeshTask CodeGenerator::GenerateMesh_MakeMorph(const FMeshGenerationStaticOptions& StaticOptions, FMeshOptionsTask InOptions, const NodeMeshMakeMorph* Node )
     {
-        NodeMeshMakeMorph::Private& node = *InMakeMorphNode->GetPrivate();
-
-		if (!node.m_pBase)
+		if (!Node->Base)
 		{
 			// This argument is required
-			ErrorLog->Add("Mesh make morph base node is not set.", ELMT_ERROR, InMakeMorphNode->GetMessageContext());
+			ErrorLog->Add("Mesh make morph base node is not set.", ELMT_ERROR, Node->GetMessageContext());
 			return UE::Tasks::MakeCompletedTask<FMeshGenerationResult>();
 		}
 
-		if (!node.m_pTarget)
+		if (!Node->Target)
 		{
 			// This argument is required
-			ErrorLog->Add("Mesh make morph target node is not set.", ELMT_ERROR, InMakeMorphNode->GetMessageContext());
+			ErrorLog->Add("Mesh make morph target node is not set.", ELMT_ERROR, Node->GetMessageContext());
 			return UE::Tasks::MakeCompletedTask<FMeshGenerationResult>();
 		}
 
@@ -781,7 +770,7 @@ namespace mu
 			);
 
 			FMeshGenerationStaticOptions BaseStaticOptions = StaticOptions;
-			BaseTask = GenerateMesh(BaseStaticOptions, BaseOptionsTask, node.m_pBase );
+			BaseTask = GenerateMesh(BaseStaticOptions, BaseOptionsTask, Node->Base );
         }
 
         // Target
@@ -801,11 +790,11 @@ namespace mu
 				InOptions
 			);
 
-			TargetTask = GenerateMesh(TargetStaticOptions, TargetOptionsTask, node.m_pTarget );
+			TargetTask = GenerateMesh(TargetStaticOptions, TargetOptionsTask, Node->Target );
         }
 
 		return UE::Tasks::Launch(TEXT("MutableMakeMorph"),
-			[BaseTask, TargetTask, node]() mutable
+			[BaseTask, TargetTask, Node]() mutable
 			{
 				Ptr<ASTOpMeshDifference> op = new ASTOpMeshDifference();
 
@@ -813,7 +802,7 @@ namespace mu
 				op->bIgnoreTextureCoords = true;
 
 				// UE only has position and normal morph data, optimize for this case if indicated. 
-				if (node.bOnlyPositionAndNormal)
+				if (Node->bOnlyPositionAndNormal)
 				{
 					op->Channels = { {static_cast<uint8>(EMeshBufferSemantic::Position), 0}, {static_cast<uint8>(EMeshBufferSemantic::Normal), 0} };
 				}
@@ -892,11 +881,9 @@ namespace mu
 	}
 
 
-	FMeshTask CodeGenerator::GenerateMesh_Switch(const FMeshGenerationStaticOptions& StaticOptions, FMeshOptionsTask InOptions, const NodeMeshSwitch* SwitchNode )
+	FMeshTask CodeGenerator::GenerateMesh_Switch(const FMeshGenerationStaticOptions& StaticOptions, FMeshOptionsTask InOptions, const NodeMeshSwitch* Node )
     {
-        NodeMeshSwitch::Private& node = *SwitchNode->GetPrivate();
-
-        if (node.Options.Num() == 0)
+        if (Node->Options.Num() == 0)
         {
             // No options in the switch!
             // TODO
@@ -904,16 +891,16 @@ namespace mu
 		}
 
 		Ptr<ASTOp> Variable;
-        if ( node.m_pParameter )
+        if (Node->Parameter )
         {
 			FScalarGenerationResult ParamResult;
-			GenerateScalar(ParamResult, StaticOptions, node.m_pParameter.get());
+			GenerateScalar(ParamResult, StaticOptions, Node->Parameter.get());
 			Variable = ParamResult.op;
 		}
         else
         {
             // This argument is required
-			Variable = GenerateMissingScalarCode(TEXT("Switch variable"), 0.0f, SwitchNode->GetMessageContext());
+			Variable = GenerateMissingScalarCode(TEXT("Switch variable"), 0.0f, Node->GetMessageContext());
         }
 
 		Ptr<ASTOpSwitch> SwitchOp = new ASTOpSwitch();
@@ -925,13 +912,13 @@ namespace mu
 		FMeshOptionsTask CurrentTargetOptions = InOptions;
 		TArray<FMeshTask> OptionTasks;
 		TArray<int16> OptionIndices;
-		OptionTasks.Reserve(node.Options.Num());
-		OptionIndices.Reserve(node.Options.Num());
-		for ( int32 OptionIndex=0; OptionIndex<node.Options.Num(); ++OptionIndex)
+		OptionTasks.Reserve(Node->Options.Num());
+		OptionIndices.Reserve(Node->Options.Num());
+		for ( int32 OptionIndex=0; OptionIndex< Node->Options.Num(); ++OptionIndex)
         {
-            if ( node.Options[OptionIndex] )
+            if (Node->Options[OptionIndex] )
             {								
-				FMeshTask BranchTask = GenerateMesh(StaticOptions, CurrentTargetOptions, node.Options[OptionIndex] );
+				FMeshTask BranchTask = GenerateMesh(StaticOptions, CurrentTargetOptions, Node->Options[OptionIndex] );
 				OptionTasks.Emplace(BranchTask);
 				OptionIndices.Emplace(int16(OptionIndex));
 
@@ -999,15 +986,10 @@ namespace mu
 				}
 
 				Ptr<NodeMeshConstant> CellNode = new NodeMeshConstant();
-				CellNode->SetValue(pMesh);
+				CellNode->Value = pMesh;
 
 				// TODO Take into account layout strategy
-				int32 numLayouts = node.Layouts.Num();
-				CellNode->SetLayoutCount(numLayouts);
-				for (int32 i = 0; i < numLayouts; ++i)
-				{
-					CellNode->SetLayout(i, node.Layouts[i]);
-				}
+				CellNode->Layouts = node.Layouts;
 
 				CellNode->SetMessageContext(Cell.ErrorContext);
 				CellNode->SourceDataDescriptor = node.SourceDataDescriptor;
@@ -1024,19 +1006,17 @@ namespace mu
 	}
 
 
-	FMeshTask CodeGenerator::GenerateMesh_Variation(const FMeshGenerationStaticOptions& StaticOptions, FMeshOptionsTask InOptions, const NodeMeshVariation* VariationNode )
+	FMeshTask CodeGenerator::GenerateMesh_Variation(const FMeshGenerationStaticOptions& StaticOptions, FMeshOptionsTask InOptions, const NodeMeshVariation* Node )
     {
-        NodeMeshVariation::Private& node = *VariationNode->GetPrivate();
-
 		bool bFirstValidConnectionFound = false;
 		FMeshOptionsTask VariationOptionsTask = InOptions;
 
 		FMeshTask DefaultTask = UE::Tasks::MakeCompletedTask<FMeshGenerationResult>();
 
         // Default case
-        if ( node.m_defaultMesh )
+        if (Node->DefaultMesh )
         {
-			DefaultTask = GenerateMesh(StaticOptions, InOptions, node.m_defaultMesh );
+			DefaultTask = GenerateMesh(StaticOptions, InOptions, Node->DefaultMesh );
 			bFirstValidConnectionFound = true;
 
 			VariationOptionsTask = UE::Tasks::Launch(TEXT("MutableMeshVariationDefaultOptions"),
@@ -1053,13 +1033,13 @@ namespace mu
         // Process variations in reverse order, since conditionals are built bottom-up.
 		TArray<FMeshTask> ReverseVariations;
 		TArray<int32> ReverseVariationsIndices;
-		ReverseVariations.Reserve(node.m_variations.Num());
-		ReverseVariationsIndices.Reserve(node.m_variations.Num());
+		ReverseVariations.Reserve(Node->Variations.Num());
+		ReverseVariationsIndices.Reserve(Node->Variations.Num());
 
-        for ( int32 VariationIndex = node.m_variations.Num()-1; VariationIndex >= 0; --VariationIndex)
+        for ( int32 VariationIndex = Node->Variations.Num()-1; VariationIndex >= 0; --VariationIndex)
         {
             int32 tagIndex = -1;
-            const FString& tag = node.m_variations[VariationIndex].m_tag;
+            const FString& tag = Node->Variations[VariationIndex].Tag;
             for ( int32 i = 0; i < FirstPass.Tags.Num(); ++i )
             {
                 if ( FirstPass.Tags[i].Tag==tag)
@@ -1073,16 +1053,16 @@ namespace mu
                 ErrorLog->Add( 
 					FString::Printf(TEXT("Unknown tag found in mesh variation [%s]."), *tag),
 					ELMT_WARNING,
-					VariationNode->GetMessageContext(),
+					Node->GetMessageContext(),
 					ELMSB_UNKNOWN_TAG
 				);
                 continue;
             }
 
             Ptr<ASTOp> variationMeshOp;
-            if ( node.m_variations[VariationIndex].m_mesh )
+            if ( Node->Variations[VariationIndex].Mesh )
             {         
-				FMeshTask BranchTask = GenerateMesh(StaticOptions, VariationOptionsTask, node.m_variations[VariationIndex].m_mesh );
+				FMeshTask BranchTask = GenerateMesh(StaticOptions, VariationOptionsTask, Node->Variations[VariationIndex].Mesh );
 
 				ReverseVariations.Add(BranchTask);
 
@@ -1116,7 +1096,7 @@ namespace mu
 		Requisites.Append(ReverseVariations);
 
 		return UE::Tasks::Launch(TEXT("MutableMeshVariation"),
-			[this, node, DefaultTask, ReverseVariations, ReverseVariationsIndices]() mutable
+			[this, Node, DefaultTask, ReverseVariations, ReverseVariationsIndices]() mutable
 			{
 				FMeshGenerationResult Result;
 				Ptr<ASTOp> CurrentMeshOp;
@@ -1124,7 +1104,7 @@ namespace mu
 				bool bFirstValidConnectionFound = false;
 
 				// Default case
-				if (node.m_defaultMesh)
+				if (Node->DefaultMesh)
 				{
 					Result = DefaultTask.GetResult();
 					CurrentMeshOp = Result.MeshOp;
@@ -1155,13 +1135,11 @@ namespace mu
     }
 
 
-	FMeshTask CodeGenerator::GenerateMesh_Constant(const FMeshGenerationStaticOptions& StaticOptions, FMeshOptionsTask InOptions, const NodeMeshConstant* InNode )
+	FMeshTask CodeGenerator::GenerateMesh_Constant(const FMeshGenerationStaticOptions& StaticOptions, FMeshOptionsTask InOptions, const NodeMeshConstant* Node )
     {
 		MUTABLE_CPUPROFILER_SCOPE(GenerateMesh_Constant);
 
-        NodeMeshConstant::Private& Node = *InNode->GetPrivate();
-
-		TSharedPtr<FMesh> Mesh = Node.Value;
+		TSharedPtr<FMesh> Mesh = Node->Value;
 
 		// True passthrough?
 		if (Mesh && Mesh->IsReference() && !Mesh->IsForceLoad())
@@ -1182,7 +1160,7 @@ namespace mu
 
         Ptr<ASTOpConstantResource> ConstantOp = new ASTOpConstantResource();
 		ConstantOp->Type = EOpType::ME_CONSTANT;
-		ConstantOp->SourceDataDescriptor = InNode->SourceDataDescriptor;
+		ConstantOp->SourceDataDescriptor = Node->SourceDataDescriptor;
 
 		FMeshGenerationResult Result;
 		Result.BaseMeshOp = ConstantOp;
@@ -1196,7 +1174,7 @@ namespace mu
 			EmptyMesh->MeshIDPrefix = ConstantOp->GetValueHash();
 
 			// Log an error message
-			ErrorLog->Add("Constant mesh not set.", ELMT_WARNING, InNode->GetMessageContext());
+			ErrorLog->Add("Constant mesh not set.", ELMT_WARNING, Node->GetMessageContext());
 
 			return UE::Tasks::MakeCompletedTask<FMeshGenerationResult>(Result);
 		}
@@ -1222,7 +1200,7 @@ namespace mu
 		}
 
 		// This task does the necessary processing of the mesh when it is available
-		Ptr<const NodeMeshConstant> NodeCopy = InNode;
+		Ptr<const NodeMeshConstant> NodeCopy = Node;
 		FMeshTask MeshProcessingTask = GenerateMeshConstantPipe.Launch( TEXT("MutableGenerateMeshConstant"),
 			[ReferenceMesh=Mesh,ResolveMesh,Result,ConstantOp, StaticOptions, InOptions, NodeCopy,
 			// This is read-only
@@ -1359,12 +1337,10 @@ namespace mu
 					{
 						if (!bIsOverridingLayouts)
 						{
-							NodeMeshConstant::Private& Node = *NodeCopy->GetPrivate();
-
 							// Apply whatever transform is necessary for every layout
-							for (int32 LayoutIndex = 0; LayoutIndex < Node.Layouts.Num(); ++LayoutIndex)
+							for (int32 LayoutIndex = 0; LayoutIndex < NodeCopy->Layouts.Num(); ++LayoutIndex)
 							{
-								Ptr<NodeLayout> LayoutNode = Node.Layouts[LayoutIndex];
+								Ptr<NodeLayout> LayoutNode = NodeCopy->Layouts[LayoutIndex];
 								if (!LayoutNode)
 								{
 									continue;
@@ -1425,26 +1401,24 @@ namespace mu
 		this->GetModifiersFor(StaticOptions.ComponentId, StaticOptions.ActiveTags, bModifiersForBeforeOperations, Modifiers);
 
 		// This task does the necessary processing of the mesh when it is available
-		FMeshTask ApplyModifiersTask = ApplyMeshModifiers(Modifiers, StaticOptions, InOptions, MeshProcessingTask, INDEX_NONE, InNode->GetMessageContext(), InNode);
+		FMeshTask ApplyModifiersTask = ApplyMeshModifiers(Modifiers, StaticOptions, InOptions, MeshProcessingTask, INDEX_NONE, Node->GetMessageContext(), Node);
 
 		return ApplyModifiersTask;
     }
 
 
-	FMeshTask CodeGenerator::GenerateMesh_Format(const FMeshGenerationStaticOptions& StaticOptions, FMeshOptionsTask InOptions, const NodeMeshFormat* format )
+	FMeshTask CodeGenerator::GenerateMesh_Format(const FMeshGenerationStaticOptions& StaticOptions, FMeshOptionsTask InOptions, const NodeMeshFormat* Node )
     {
-        NodeMeshFormat::Private& node = *format->GetPrivate();
-
-		if (!node.Source)
+		if (!Node->Source)
 		{
 			// Put something there
 			return GenerateMesh(StaticOptions, InOptions, new NodeMeshConstant());
 		}
 
- 		FMeshTask BaseTask = GenerateMesh(StaticOptions, InOptions, node.Source);
+ 		FMeshTask BaseTask = GenerateMesh(StaticOptions, InOptions, Node->Source);
 
 		return UE::Tasks::Launch(TEXT("MutableMeshFormat"),
-			[node, BaseTask, CompilerOptions=CompilerOptions]() mutable
+			[Node, BaseTask, CompilerOptions=CompilerOptions]() mutable
 			{
 				Ptr<ASTOpMeshFormat> op = new ASTOpMeshFormat();
 
@@ -1454,19 +1428,19 @@ namespace mu
 
 				TSharedPtr<FMesh> FormatMesh = MakeShared<FMesh>();
 
-				if (node.VertexBuffers.GetBufferCount())
+				if (Node->VertexBuffers.GetBufferCount())
 				{
 					op->Flags |= OP::MeshFormatArgs::Vertex;
-					FormatMesh->VertexBuffers = node.VertexBuffers;
+					FormatMesh->VertexBuffers = Node->VertexBuffers;
 				}
 
-				if (node.IndexBuffers.GetBufferCount())
+				if (Node->IndexBuffers.GetBufferCount())
 				{
 					op->Flags |= OP::MeshFormatArgs::Index;
-					FormatMesh->IndexBuffers = node.IndexBuffers;
+					FormatMesh->IndexBuffers = Node->IndexBuffers;
 				}
 
-				if (node.bOptimizeBuffers)
+				if (Node->bOptimizeBuffers)
 				{
 					op->Flags |= OP::MeshFormatArgs::OptimizeBuffers;
 				}
@@ -1492,20 +1466,18 @@ namespace mu
     }
 
 
-	FMeshTask CodeGenerator::GenerateMesh_Transform(const FMeshGenerationStaticOptions& StaticOptions, FMeshOptionsTask InOptions, const NodeMeshTransform* TransformNode )
+	FMeshTask CodeGenerator::GenerateMesh_Transform(const FMeshGenerationStaticOptions& StaticOptions, FMeshOptionsTask InOptions, const NodeMeshTransform* Node )
     {
-        const NodeMeshTransform::Private& node = *TransformNode->GetPrivate();
-
-		if (!node.Source)
+		if (!Node->Source)
 		{
 			// This argument is required
-			ErrorLog->Add("Mesh transform base node is not set.", ELMT_ERROR, TransformNode->GetMessageContext());
+			ErrorLog->Add("Mesh transform base node is not set.", ELMT_ERROR, Node->GetMessageContext());
 			return UE::Tasks::MakeCompletedTask<FMeshGenerationResult>();
 		}
 
-		FMeshTask SourceTask = GenerateMesh(StaticOptions, InOptions, node.Source);
+		FMeshTask SourceTask = GenerateMesh(StaticOptions, InOptions, Node->Source);
 
-		Ptr<const NodeMeshTransform> NodeCopy = TransformNode;
+		Ptr<const NodeMeshTransform> NodeCopy = Node;
 		return UE::Tasks::Launch(TEXT("MutableMeshTransform"),
 			[NodeCopy, SourceTask, CompilerOptions = CompilerOptions]() mutable
 			{
@@ -1513,7 +1485,7 @@ namespace mu
 
 				Ptr<ASTOpMeshTransform> TransformOp = new ASTOpMeshTransform();
 				TransformOp->source = SourceResult.MeshOp;
-				TransformOp->matrix = NodeCopy->GetPrivate()->Transform;
+				TransformOp->matrix = NodeCopy->Transform;
 
 				FMeshGenerationResult Result = SourceResult;
 				Result.MeshOp = TransformOp;
@@ -1704,23 +1676,21 @@ namespace mu
 	}
 
 	
-	FMeshTask CodeGenerator::GenerateMesh_ApplyPose(const FMeshGenerationStaticOptions& StaticOptions, FMeshOptionsTask InOptions, const NodeMeshApplyPose* PoseNode )
+	FMeshTask CodeGenerator::GenerateMesh_ApplyPose(const FMeshGenerationStaticOptions& StaticOptions, FMeshOptionsTask InOptions, const NodeMeshApplyPose* Node )
     {
-        const NodeMeshApplyPose::Private& node = *PoseNode->GetPrivate();
-
-		if (!node.m_pBase)
+		if (!Node->Base)
 		{
 			// This argument is required
-			ErrorLog->Add("Mesh apply-pose base node is not set.", ELMT_ERROR, PoseNode->GetMessageContext());
+			ErrorLog->Add("Mesh apply-pose base node is not set.", ELMT_ERROR, Node->GetMessageContext());
 			return UE::Tasks::MakeCompletedTask<FMeshGenerationResult>();
 		}
 
 		// Base Mesh
-		FMeshTask BaseTask = GenerateMesh(StaticOptions, InOptions, node.m_pBase);
+		FMeshTask BaseTask = GenerateMesh(StaticOptions, InOptions, Node->Base);
 
-		if (!node.m_pPose)
+		if (!Node->Pose)
 		{
-			ErrorLog->Add("Mesh apply-pose pose node is not set.", ELMT_ERROR, PoseNode->GetMessageContext());
+			ErrorLog->Add("Mesh apply-pose pose node is not set.", ELMT_ERROR, Node->GetMessageContext());
 			return BaseTask;
 		}
 
@@ -1736,7 +1706,7 @@ namespace mu
 			},
 			InOptions
 		);
-		FMeshTask PoseTask = GenerateMesh(StaticPoseOptions, PoseOptionsTask, node.m_pPose);
+		FMeshTask PoseTask = GenerateMesh(StaticPoseOptions, PoseOptionsTask, Node->Pose);
 
 		return UE::Tasks::Launch(TEXT("MutableMeshApplyPose"),
 			[BaseTask, PoseTask]() mutable
@@ -1758,21 +1728,19 @@ namespace mu
     }
 
 
-	FMeshTask CodeGenerator::GenerateMesh_Reshape(const FMeshGenerationStaticOptions& StaticOptions, FMeshOptionsTask InOptions, const NodeMeshReshape* Reshape)
+	FMeshTask CodeGenerator::GenerateMesh_Reshape(const FMeshGenerationStaticOptions& StaticOptions, FMeshOptionsTask InOptions, const NodeMeshReshape* Node)
 	{
-		const NodeMeshReshape::Private& Node = *Reshape->GetPrivate();
-
-		if (!Node.BaseMesh)
+		if (!Node->BaseMesh)
 		{
 			// This argument is required
-			ErrorLog->Add("Mesh reshape base node is not set.", ELMT_ERROR, Reshape->GetMessageContext());
+			ErrorLog->Add("Mesh reshape base node is not set.", ELMT_ERROR, Node->GetMessageContext());
 			return UE::Tasks::MakeCompletedTask<FMeshGenerationResult>();
 		}
 
 		// Base Mesh
-		FMeshTask BaseTask = GenerateMesh(StaticOptions, InOptions, Node.BaseMesh);
+		FMeshTask BaseTask = GenerateMesh(StaticOptions, InOptions, Node->BaseMesh);
 
-		if (!Node.BaseShape || !Node.TargetShape)
+		if (!Node->BaseShape || !Node->TargetShape)
 		{
 			return BaseTask;
 		}
@@ -1790,11 +1758,11 @@ namespace mu
 			},
 			InOptions
 		);
-		FMeshTask BaseShapeTask = GenerateMesh(StaticShapeOptions, ShapeOptionsTask, Node.BaseShape);
-		FMeshTask TargetShapeTask = GenerateMesh(StaticShapeOptions, ShapeOptionsTask, Node.TargetShape);
+		FMeshTask BaseShapeTask = GenerateMesh(StaticShapeOptions, ShapeOptionsTask, Node->BaseShape);
+		FMeshTask TargetShapeTask = GenerateMesh(StaticShapeOptions, ShapeOptionsTask, Node->TargetShape);
 
 		return UE::Tasks::Launch(TEXT("MutableMeshApplyPose"),
-			[BaseTask, BaseShapeTask, TargetShapeTask, &Node]() mutable
+			[BaseTask, BaseShapeTask, TargetShapeTask, Node]() mutable
 			{
 				FMeshGenerationResult BaseResult = BaseTask.GetResult();
 				FMeshGenerationResult BaseShapeResult = BaseShapeTask.GetResult();
@@ -1803,19 +1771,19 @@ namespace mu
 				Ptr<ASTOpMeshBindShape> OpBind = new ASTOpMeshBindShape();
 				Ptr<ASTOpMeshApplyShape> OpApply = new ASTOpMeshApplyShape();
 
-				OpBind->bReshapeSkeleton = Node.bReshapeSkeleton;
-				OpBind->BonesToDeform = Node.BonesToDeform;
-				OpBind->bReshapePhysicsVolumes = Node.bReshapePhysicsVolumes;
-				OpBind->PhysicsToDeform = Node.PhysicsToDeform;
-				OpBind->bReshapeVertices = Node.bReshapeVertices;
-				OpBind->bRecomputeNormals = Node.bRecomputeNormals;
-				OpBind->bApplyLaplacian = Node.bApplyLaplacian;
+				OpBind->bReshapeSkeleton = Node->bReshapeSkeleton;
+				OpBind->BonesToDeform = Node->BonesToDeform;
+				OpBind->bReshapePhysicsVolumes = Node->bReshapePhysicsVolumes;
+				OpBind->PhysicsToDeform = Node->PhysicsToDeform;
+				OpBind->bReshapeVertices = Node->bReshapeVertices;
+				OpBind->bRecomputeNormals = Node->bRecomputeNormals;
+				OpBind->bApplyLaplacian = Node->bApplyLaplacian;
 				OpBind->BindingMethod = static_cast<uint32>(EShapeBindingMethod::ReshapeClosestProject);
 
-				OpBind->RChannelUsage = Node.ColorRChannelUsage;
-				OpBind->GChannelUsage = Node.ColorGChannelUsage;
-				OpBind->BChannelUsage = Node.ColorBChannelUsage;
-				OpBind->AChannelUsage = Node.ColorAChannelUsage;
+				OpBind->RChannelUsage = Node->ColorRChannelUsage;
+				OpBind->GChannelUsage = Node->ColorGChannelUsage;
+				OpBind->BChannelUsage = Node->ColorBChannelUsage;
+				OpBind->AChannelUsage = Node->ColorAChannelUsage;
 
 				OpApply->bReshapeVertices = OpBind->bReshapeVertices;
 				OpApply->bRecomputeNormals = OpBind->bRecomputeNormals;
@@ -1839,4 +1807,52 @@ namespace mu
 		);
 	}
 
+
+	FMeshTask CodeGenerator::GenerateMesh_Parameter(const FMeshGenerationStaticOptions& StaticOptions, FMeshOptionsTask InOptions, const NodeMeshParameter* Node)
+	{
+		FMeshGenerationResult Result;
+		return UE::Tasks::MakeCompletedTask<FMeshGenerationResult>(Result);
+
+		// Add concurrency protection before enabling.
+		//Ptr<ASTOpParameter>* Found = FirstPass.ParameterNodes.Find(InNode);
+		//if (Found)
+		//{
+		//	check((*Found)->Type==EOpType::ME_PARAMETER);
+		//	FMeshGenerationResult Result;
+		//	Result.BaseMeshOp = *Found;
+		//	Result.MeshOp = *Found;
+		//	// Rest of the parameters?
+		//	//Result.LayoutOps ...
+		//	// ... maybe a special cache for these parameters?
+		//	return UE::Tasks::MakeCompletedTask<FMeshGenerationResult>(Result);
+		//}
+
+
+		//// LOD ? 
+		////int32 LODIndex = StaticOptions.LODIndex;
+
+		//// To task?
+		////FMeshGenerationDynamicOptions = InOptions.GetResult();
+
+		//Ptr<ASTOpParameter> Op = new ASTOpParameter();
+		//Op->Type = EOpType::ME_PARAMETER;
+
+		//Op->Parameter.Name = InNode->Name;
+		//bool bParseOk = FGuid::Parse(InNode->UID, Op->Parameter.UID);
+		////check(bParseOk);
+		//Op->Parameter.Type = EParameterType::Image;
+		//Op->Parameter.DefaultValue.Set<FParamImageType>(InNode->DefaultValue);
+
+		//// Generate the code for the ranges
+		//for (int32 a = 0; a < InNode->Ranges.Num(); ++a)
+		//{
+		//	FRangeGenerationResult rangeResult;
+		//	GenerateRange(rangeResult, StaticOptions, InNode->Ranges[a]);
+		//	Op->Ranges.Emplace(Op.get(), rangeResult.sizeOp, rangeResult.rangeName, rangeResult.rangeUID);
+		//}
+
+		//FirstPass.ParameterNodes.Add(InNode, Op);
+
+	}
+
 }
diff --git a/Engine/Plugins/Mutable/Source/MutableTools/Private/MuT/CodeGenerator_Projector.cpp b/Engine/Plugins/Mutable/Source/MutableTools/Private/MuT/CodeGenerator_Projector.cpp
--- a/Engine/Plugins/Mutable/Source/MutableTools/Private/MuT/CodeGenerator_Projector.cpp
+++ b/Engine/Plugins/Mutable/Source/MutableTools/Private/MuT/CodeGenerator_Projector.cpp
@@ -15,38 +15,13 @@
 #include "MuT/CodeGenerator.h"
 #include "MuT/ErrorLog.h"
 #include "MuT/NodeProjector.h"
-#include "MuT/NodeProjectorPrivate.h"
 #include "MuT/NodeRange.h"
-#include "map"
 
 
 namespace mu
 {
 
-	//-------------------------------------------------------------------------------------------------
-	FProjector ProjectorFromNode(const NodeProjectorConstant::Private& node)
-	{
-		FProjector p;
-		p.type = node.Type;
-		p.position[0] = node.m_position[0];
-		p.position[1] = node.m_position[1];
-		p.position[2] = node.m_position[2];
-		p.direction[0] = node.m_direction[0];
-		p.direction[1] = node.m_direction[1];
-		p.direction[2] = node.m_direction[2];
-		p.up[0] = node.m_up[0];
-		p.up[1] = node.m_up[1];
-		p.up[2] = node.m_up[2];
-		p.scale[0] = node.m_scale[0];
-		p.scale[1] = node.m_scale[1];
-		p.scale[2] = node.m_scale[2];
-		p.projectionAngle = node.m_projectionAngle;
-
-		return p;
-	}
 
-
-	//-------------------------------------------------------------------------------------------------
 	void CodeGenerator::GenerateProjector(FProjectorGenerationResult& OutResult, const FGenericGenerationOptions& Options, const Ptr<const NodeProjector>& Untyped)
 	{
 		if (!Untyped)
@@ -94,46 +69,55 @@ namespace mu
 
 	//-------------------------------------------------------------------------------------------------
 	void CodeGenerator::GenerateProjector_Constant(FProjectorGenerationResult& result, const FGenericGenerationOptions& Options,
-		const Ptr<const NodeProjectorConstant>& constant)
+		const Ptr<const NodeProjectorConstant>& Node)
 	{
-		const NodeProjectorConstant::Private& node = *constant->GetPrivate();
-
 		Ptr<ASTOpConstantProjector> op = new ASTOpConstantProjector();
-		op->value = ProjectorFromNode(node);
+		FProjector p;
+		p.type = Node->Type;
+		p.position = Node->Position;
+		p.direction = Node->Direction;
+		p.up = Node->Up;
+		p.scale = Node->Scale;
+		p.projectionAngle = Node->ProjectionAngle;
+		op->value = p;
 
 		result.op = op;
 		result.type = op->value.type;
 	}
 
 
-	//-------------------------------------------------------------------------------------------------
 	void CodeGenerator::GenerateProjector_Parameter(FProjectorGenerationResult& OutResult, const FGenericGenerationOptions& Options,
-		const Ptr<const NodeProjectorParameter>& ParamNode)
+		const Ptr<const NodeProjectorParameter>& Node)
 	{
-		const NodeProjectorParameter::Private& node = *ParamNode->GetPrivate();
-
 		Ptr<ASTOpParameter> op;
 
 		Ptr<ASTOpParameter>* Found = nullptr;
 		{
 			UE::TUniqueLock Lock(FirstPass.ParameterNodes.Mutex);
-			Found = FirstPass.ParameterNodes.Map.Find(node.m_pNode);
+			Found = FirstPass.ParameterNodes.Map.Find(Node);
 			if (!Found)
 			{
 				FParameterDesc param;
-				param.Name = node.Name;
-				bool bParseOk = FGuid::Parse(node.UID, param.UID);
+				param.Name = Node->Name;
+				bool bParseOk = FGuid::Parse(Node->UID, param.UID);
 				check(bParseOk);
 				param.Type = EParameterType::Projector;
 
-				FProjector p = ProjectorFromNode(node);
+				FProjector p;
+				p.type = Node->Type;
+				p.position = Node->Position;
+				p.direction = Node->Direction;
+				p.up = Node->Up;
+				p.scale = Node->Scale;
+				p.projectionAngle = Node->ProjectionAngle;
+
 				param.DefaultValue.Set<FParamProjectorType>(p);
 
 				op = new ASTOpParameter();
 				op->Type = EOpType::PR_PARAMETER;
 				op->Parameter = param;
 
-				FirstPass.ParameterNodes.Map.Add(node.m_pNode, op);
+				FirstPass.ParameterNodes.Map.Add(Node, op);
 			}
 			else
 			{
@@ -144,10 +128,10 @@ namespace mu
 		if (!Found)
 		{
 			// Generate the code for the ranges
-			for (int32 a = 0; a < node.Ranges.Num(); ++a)
+			for (int32 a = 0; a < Node->Ranges.Num(); ++a)
 			{
 				FRangeGenerationResult rangeResult;
-				GenerateRange(rangeResult, Options, node.Ranges[a]);
+				GenerateRange(rangeResult, Options, Node->Ranges[a]);
 				op->Ranges.Emplace(op.get(), rangeResult.sizeOp, rangeResult.rangeName, rangeResult.rangeUID);
 			}
 		}
@@ -168,20 +152,10 @@ namespace mu
 		FProjector p;
 		p.type = EProjectorType::Planar;
 
-		p.direction[0] = 1;
-		p.direction[1] = 0;
-		p.direction[2] = 0;
-
-		p.up[0] = 0;
-		p.up[1] = 0;
-		p.up[2] = 1;
-
-		p.position[0] = 0.0f;
-		p.position[1] = 0.0f;
-		p.position[2] = 0.0f;
-
-		p.scale[0] = 1.0f;
-		p.scale[1] = 1.0f;
+		p.direction = { 1, 0 ,0 };
+		p.up = { 0, 0, 1 };
+		p.position = { 0.0f, 0.0f, 0.0f };
+		p.scale = { 1.0f, 1.0f, 1.0f };
 
 		Ptr<ASTOpConstantProjector> op = new ASTOpConstantProjector();
 		op->value = p;
