[mutable] Initial boilerplate for the implementation of Mesh Parameters.

diff --git a/Engine/Plugins/Mutable/Source/MutableTools/Private/MuT/CodeGenerator_Scalar.cpp b/Engine/Plugins/Mutable/Source/MutableTools/Private/MuT/CodeGenerator_Scalar.cpp
--- a/Engine/Plugins/Mutable/Source/MutableTools/Private/MuT/CodeGenerator_Scalar.cpp
+++ b/Engine/Plugins/Mutable/Source/MutableTools/Private/MuT/CodeGenerator_Scalar.cpp
@@ -26,19 +26,13 @@
 #include "MuT/NodeRange.h"
 #include "MuT/NodeScalar.h"
 #include "MuT/NodeScalarArithmeticOperation.h"
-#include "MuT/NodeScalarArithmeticOperationPrivate.h"
 #include "MuT/NodeScalarConstant.h"
-#include "MuT/NodeScalarConstantPrivate.h"
 #include "MuT/NodeScalarCurve.h"
-#include "MuT/NodeScalarCurvePrivate.h"
 #include "MuT/NodeScalarEnumParameter.h"
-#include "MuT/NodeScalarEnumParameterPrivate.h"
 #include "MuT/NodeScalarParameter.h"
-#include "MuT/NodeScalarParameterPrivate.h"
 #include "MuT/NodeScalarSwitch.h"
 #include "MuT/NodeScalarTable.h"
 #include "MuT/NodeScalarVariation.h"
-#include "MuT/NodeScalarVariationPrivate.h"
 #include "MuT/Table.h"
 #include "MuT/TablePrivate.h"
 
@@ -123,43 +117,36 @@ namespace mu
 	}
 
 
-	//-------------------------------------------------------------------------------------------------
-	void CodeGenerator::GenerateScalar_Constant(FScalarGenerationResult& result, const FGenericGenerationOptions& Options, const Ptr<const NodeScalarConstant>& Typed)
+	void CodeGenerator::GenerateScalar_Constant(FScalarGenerationResult& result, const FGenericGenerationOptions& Options, const Ptr<const NodeScalarConstant>& Node)
 	{
-		const NodeScalarConstant::Private& node = *Typed->GetPrivate();
-
 		Ptr<ASTOpConstantScalar> op = new ASTOpConstantScalar;
-		op->Value = node.Value;
-
+		op->Value = Node->Value;
 		result.op = op;
 	}
 
 
-	//-------------------------------------------------------------------------------------------------
-	void CodeGenerator::GenerateScalar_Parameter(FScalarGenerationResult& result, const FGenericGenerationOptions& Options, const Ptr<const NodeScalarParameter>& Typed)
+	void CodeGenerator::GenerateScalar_Parameter(FScalarGenerationResult& result, const FGenericGenerationOptions& Options, const Ptr<const NodeScalarParameter>& Node)
 	{
-		const NodeScalarParameter::Private& node = *Typed->GetPrivate();
-
 		Ptr<ASTOpParameter> op;
 
 		Ptr<ASTOpParameter>* Found = nullptr;		
 		{
 			UE::TUniqueLock Lock(FirstPass.ParameterNodes.Mutex);
-			Found = FirstPass.ParameterNodes.Map.Find(node.m_pNode);
+			Found = FirstPass.ParameterNodes.Map.Find(Node);
 			if (!Found)
 			{
 				FParameterDesc param;
-				param.Name = node.Name;
-				bool bParseOk = FGuid::Parse(node.UID, param.UID);
+				param.Name = Node->Name;
+				bool bParseOk = FGuid::Parse(Node->UID, param.UID);
 				check(bParseOk);
 				param.Type = EParameterType::Float;
-				param.DefaultValue.Set<FParamFloatType>(node.DefaultValue);
+				param.DefaultValue.Set<FParamFloatType>(Node->DefaultValue);
 
 				op = new ASTOpParameter();
 				op->Type = EOpType::SC_PARAMETER;
 				op->Parameter = param;
 
-				FirstPass.ParameterNodes.Map.Add(node.m_pNode, op);
+				FirstPass.ParameterNodes.Map.Add(Node, op);
 			}
 			else
 			{
@@ -170,10 +157,10 @@ namespace mu
 		if (!Found)
 		{
 			// Generate the code for the ranges
-			for (int32 a = 0; a < node.Ranges.Num(); ++a)
+			for (int32 a = 0; a < Node->Ranges.Num(); ++a)
 			{
 				FRangeGenerationResult rangeResult;
-				GenerateRange(rangeResult, Options, node.Ranges[a]);
+				GenerateRange(rangeResult, Options, Node->Ranges[a]);
 				op->Ranges.Emplace(op.get(), rangeResult.sizeOp, rangeResult.rangeName, rangeResult.rangeUID);
 			}
 		}
@@ -182,41 +169,38 @@ namespace mu
 	}
 
 
-	//-------------------------------------------------------------------------------------------------
-	void CodeGenerator::GenerateScalar_EnumParameter(FScalarGenerationResult& result, const FGenericGenerationOptions& Options, const Ptr<const NodeScalarEnumParameter>& Typed)
+	void CodeGenerator::GenerateScalar_EnumParameter(FScalarGenerationResult& result, const FGenericGenerationOptions& Options, const Ptr<const NodeScalarEnumParameter>& Node)
 	{
-		const NodeScalarEnumParameter::Private& node = *Typed->GetPrivate();
-
 		Ptr<ASTOpParameter> op;
 
 		Ptr<ASTOpParameter>* Found = nullptr;
 		{
 			UE::TUniqueLock Lock(FirstPass.ParameterNodes.Mutex);
-			Found = FirstPass.ParameterNodes.Map.Find(node.m_pNode);
+			Found = FirstPass.ParameterNodes.Map.Find(Node);
 			if (!Found)
 			{
 				FParameterDesc param;
-				param.Name = node.Name;
-				if (!node.UID.IsEmpty())
+				param.Name = Node->Name;
+				if (!Node->UID.IsEmpty())
 				{
-					bool bParseOk = FGuid::Parse(node.UID, param.UID);
+					bool bParseOk = FGuid::Parse(Node->UID, param.UID);
 					check(bParseOk);
 				}
 				param.Type = EParameterType::Int;
-				param.DefaultValue.Set<FParamIntType>(node.DefaultValue);
+				param.DefaultValue.Set<FParamIntType>(Node->DefaultValue);
 
-				param.PossibleValues.SetNum(node.Options.Num());
-				for (int32 i = 0; i < node.Options.Num(); ++i)
+				param.PossibleValues.SetNum(Node->Options.Num());
+				for (int32 i = 0; i < Node->Options.Num(); ++i)
 				{
-					param.PossibleValues[i].Value = (int16_t)node.Options[i].value;
-					param.PossibleValues[i].Name = node.Options[i].name;
+					param.PossibleValues[i].Value = (int16)Node->Options[i].Value;
+					param.PossibleValues[i].Name = Node->Options[i].Name;
 				}
 
 				op = new ASTOpParameter();
 				op->Type = EOpType::NU_PARAMETER;
 				op->Parameter = param;
 
-				FirstPass.ParameterNodes.Map.Add(node.m_pNode, op);
+				FirstPass.ParameterNodes.Map.Add(Node, op);
 			}
 			else
 			{
@@ -227,11 +211,11 @@ namespace mu
 		if (!Found)
 		{
 			// Generate the code for the ranges
-			for (int32 a = 0; a < node.Ranges.Num(); ++a)
+			for (int32 RangeIndex = 0; RangeIndex < Node->Ranges.Num(); ++RangeIndex)
 			{
-				FRangeGenerationResult rangeResult;
-				GenerateRange(rangeResult, Options, node.Ranges[a]);
-				op->Ranges.Emplace(op.get(), rangeResult.sizeOp, rangeResult.rangeName, rangeResult.rangeUID);
+				FRangeGenerationResult RangeResult;
+				GenerateRange(RangeResult, Options, Node->Ranges[RangeIndex]);
+				op->Ranges.Emplace(op.get(), RangeResult.sizeOp, RangeResult.rangeName, RangeResult.rangeUID);
 			}
 		}
 
@@ -239,7 +223,6 @@ namespace mu
 	}
 
 
-	//-------------------------------------------------------------------------------------------------
 	void CodeGenerator::GenerateScalar_Switch(FScalarGenerationResult& result, const FGenericGenerationOptions& Options, const Ptr<const NodeScalarSwitch>& Typed)
 	{
 		const NodeScalarSwitch& node = *Typed;
@@ -292,28 +275,25 @@ namespace mu
 	}
 
 
-	//-------------------------------------------------------------------------------------------------
-	void CodeGenerator::GenerateScalar_Variation(FScalarGenerationResult& result, const FGenericGenerationOptions& Options, const Ptr<const NodeScalarVariation>& Typed)
+	void CodeGenerator::GenerateScalar_Variation(FScalarGenerationResult& result, const FGenericGenerationOptions& Options, const Ptr<const NodeScalarVariation>& Node)
 	{
-		const NodeScalarVariation::Private& node = *Typed->GetPrivate();
-
 		Ptr<ASTOp> op;
 
 		// Default case
-		if (node.m_defaultScalar)
+		if (Node->DefaultScalar)
 		{
 			FMeshGenerationResult branchResults;
 
 			FScalarGenerationResult ChildResult;
-			GenerateScalar(ChildResult, Options, node.m_defaultScalar);
+			GenerateScalar(ChildResult, Options, Node->DefaultScalar);
 			op = ChildResult.op;
 		}
 
 		// Process variations in reverse order, since conditionals are built bottom-up.
-		for (int32 t = node.m_variations.Num() - 1; t >= 0; --t)
+		for (int32 t = Node->Variations.Num() - 1; t >= 0; --t)
 		{
 			int32 tagIndex = -1;
-			const FString& tag = node.m_variations[t].m_tag;
+			const FString& tag = Node->Variations[t].Tag;
 			for (int32 i = 0; i < FirstPass.Tags.Num(); ++i)
 			{
 				if (FirstPass.Tags[i].Tag == tag)
@@ -326,21 +306,21 @@ namespace mu
 			{
 				FString Msg = FString::Printf(TEXT("Unknown tag found in image variation [%s]."), *tag);
 
-				ErrorLog->Add(Msg, ELMT_WARNING, Typed->GetMessageContext());
+				ErrorLog->Add(Msg, ELMT_WARNING, Node->GetMessageContext());
 				continue;
 			}
 
 			Ptr<ASTOp> variationOp;
-			if (node.m_variations[t].m_scalar)
+			if (Node->Variations[t].Scalar)
 			{
 				FScalarGenerationResult ChildResult;
-				GenerateScalar(ChildResult, Options, node.m_variations[t].m_scalar);
+				GenerateScalar(ChildResult, Options, Node->Variations[t].Scalar);
 				variationOp = ChildResult.op;
 			}
 			else
 			{
 				// This argument is required
-				variationOp = GenerateMissingScalarCode(TEXT("Variation option"), 0.0f, Typed->GetMessageContext());
+				variationOp = GenerateMissingScalarCode(TEXT("Variation option"), 0.0f, Node->GetMessageContext());
 			}
 
 
@@ -357,7 +337,6 @@ namespace mu
 	}
 
 
-	//-------------------------------------------------------------------------------------------------
 	void CodeGenerator::GenerateScalar_Curve(FScalarGenerationResult& result, const FGenericGenerationOptions& Options, const Ptr<const NodeScalarCurve>& Typed)
 	{
 		Ptr<ASTOpScalarCurve> op = new ASTOpScalarCurve();
@@ -379,14 +358,11 @@ namespace mu
 	}
 
 
-	//-------------------------------------------------------------------------------------------------
-	void CodeGenerator::GenerateScalar_Arithmetic(FScalarGenerationResult& result, const FGenericGenerationOptions& Options, const Ptr<const NodeScalarArithmeticOperation>& Typed)
+	void CodeGenerator::GenerateScalar_Arithmetic(FScalarGenerationResult& result, const FGenericGenerationOptions& Options, const Ptr<const NodeScalarArithmeticOperation>& Node)
 	{
-		const NodeScalarArithmeticOperation::Private& node = *Typed->GetPrivate();
-
 		Ptr<ASTOpScalarArithmetic> op = new ASTOpScalarArithmetic;
 
-		switch (node.m_operation)
+		switch (Node->Operation)
 		{
 		case NodeScalarArithmeticOperation::AO_ADD: op->Operation = OP::ArithmeticArgs::ADD; break;
 		case NodeScalarArithmeticOperation::AO_SUBTRACT: op->Operation = OP::ArithmeticArgs::SUBTRACT; break;
@@ -399,7 +375,7 @@ namespace mu
 		}
 
 		// A
-		if (NodeScalar* pA = node.m_pA.get())
+		if (NodeScalar* pA = Node->A.get())
 		{
 			FScalarGenerationResult ParamResult;
 			GenerateScalar(ParamResult, Options, pA);
@@ -408,11 +384,11 @@ namespace mu
 		}
 		else
 		{
-			op->A = GenerateMissingScalarCode( TEXT("ScalarArithmetic A"), 1.0f, Typed->GetMessageContext() );
+			op->A = GenerateMissingScalarCode( TEXT("ScalarArithmetic A"), 1.0f, Node->GetMessageContext() );
 		}
 
 		// B
-		if (NodeScalar* pB = node.m_pB.get())
+		if (NodeScalar* pB = Node->B.get())
 		{
 			FScalarGenerationResult ParamResult;
 			GenerateScalar(ParamResult, Options, pB);
@@ -421,27 +397,23 @@ namespace mu
 		}
 		else
 		{
-			op->B =  GenerateMissingScalarCode( TEXT("ScalarArithmetic B"), 1.0f, Typed->GetMessageContext() );
+			op->B =  GenerateMissingScalarCode( TEXT("ScalarArithmetic B"), 1.0f, Node->GetMessageContext() );
 		}
 
 		result.op = op;
 	}
 
 
-	//-------------------------------------------------------------------------------------------------
-	void CodeGenerator::GenerateScalar_Table(FScalarGenerationResult& result, const FGenericGenerationOptions& Options, const Ptr<const NodeScalarTable>& Typed)
+	void CodeGenerator::GenerateScalar_Table(FScalarGenerationResult& result, const FGenericGenerationOptions& Options, const Ptr<const NodeScalarTable>& Node)
 	{
-		const NodeScalarTable& node = *Typed;
-
-		Ptr<ASTOp> Op = GenerateTableSwitch<NodeScalarTable, ETableColumnType::Scalar, EOpType::SC_SWITCH>(node,
-			[this,&Options](const NodeScalarTable& node, int32 colIndex, int32 row, FErrorLog*)
+		Ptr<ASTOp> Op = GenerateTableSwitch<NodeScalarTable, ETableColumnType::Scalar, EOpType::SC_SWITCH>(*Node,
+			[this,&Options](const NodeScalarTable& Node, int32 colIndex, int32 row, FErrorLog*)
 			{
-				NodeScalarConstantPtr pCell = new NodeScalarConstant();
-				float scalar = node.Table->GetPrivate()->Rows[row].Values[colIndex].Scalar;
-				pCell->SetValue(scalar);
+				Ptr<NodeScalarConstant> Cell = new NodeScalarConstant();
+				Cell->Value = Node.Table->GetPrivate()->Rows[row].Values[colIndex].Scalar;
 
 				FScalarGenerationResult ParamResult;
-				GenerateScalar(ParamResult, Options, pCell);
+				GenerateScalar(ParamResult, Options, Cell);
 
 				return ParamResult.op;
 			});
@@ -450,7 +422,6 @@ namespace mu
 	}
 
 
-	//---------------------------------------------------------------------------------------------
 	mu::Ptr<ASTOp> CodeGenerator::GenerateMissingScalarCode(const TCHAR* strWhere, float value, const void* errorContext)
 	{
 		// Log a warning
@@ -458,12 +429,12 @@ namespace mu
 		ErrorLog->Add(Msg, ELMT_ERROR, errorContext);
 
 		// Create a constant node
-		NodeScalarConstantPtr pNode = new NodeScalarConstant();
-		pNode->SetValue(value);
+		Ptr<NodeScalarConstant> Node = new NodeScalarConstant();
+		Node->Value = value;
 
 		FGenericGenerationOptions Options;
 		FScalarGenerationResult ParamResult;
-		GenerateScalar(ParamResult, Options, pNode.get());
+		GenerateScalar(ParamResult, Options, Node.get());
 
 		return ParamResult.op;
 	}
diff --git a/Engine/Plugins/Mutable/Source/MutableTools/Private/MuT/CodeGenerator_String.cpp b/Engine/Plugins/Mutable/Source/MutableTools/Private/MuT/CodeGenerator_String.cpp
--- a/Engine/Plugins/Mutable/Source/MutableTools/Private/MuT/CodeGenerator_String.cpp
+++ b/Engine/Plugins/Mutable/Source/MutableTools/Private/MuT/CodeGenerator_String.cpp
@@ -16,15 +16,11 @@
 #include "MuT/NodeRange.h"
 #include "MuT/NodeString.h"
 #include "MuT/NodeStringConstant.h"
-#include "MuT/NodeStringConstantPrivate.h"
 #include "MuT/NodeStringParameter.h"
-#include "MuT/NodeStringParameterPrivate.h"
 
 namespace mu
 {
 
-
-	//-------------------------------------------------------------------------------------------------
 	void CodeGenerator::GenerateString(FStringGenerationResult& result, const FGenericGenerationOptions& Options, const Ptr<const NodeString>& Untyped)
 	{
 		if (!Untyped)
@@ -69,43 +65,38 @@ namespace mu
 	}
 
 
-	//-------------------------------------------------------------------------------------------------
-	void CodeGenerator::GenerateString_Constant(FStringGenerationResult& result, const FGenericGenerationOptions& Options, const Ptr<const NodeStringConstant>& Typed)
+	void CodeGenerator::GenerateString_Constant(FStringGenerationResult& result, const FGenericGenerationOptions& Options, const Ptr<const NodeStringConstant>& Node)
 	{
-		const NodeStringConstant::Private& node = *Typed->GetPrivate();
-
 		Ptr<ASTOpConstantString> op = new ASTOpConstantString();
-		op->value = node.Value;
+		op->value = Node->Value;
 
 		result.op = op;
 	}
 
 
-	//-------------------------------------------------------------------------------------------------
-	void CodeGenerator::GenerateString_Parameter(FStringGenerationResult& result, const FGenericGenerationOptions& Options, const Ptr<const NodeStringParameter>& Typed)
+	void CodeGenerator::GenerateString_Parameter(FStringGenerationResult& result, const FGenericGenerationOptions& Options, const Ptr<const NodeStringParameter>& Node)
 	{
-		const NodeStringParameter::Private& node = *Typed->GetPrivate();
-
 		Ptr<ASTOpParameter> op;
 
 		Ptr<ASTOpParameter>* Found = nullptr;
 		{
 			UE::TUniqueLock Lock(FirstPass.ParameterNodes.Mutex);
-			Found = FirstPass.ParameterNodes.Map.Find(node.m_pNode);
+			Found = FirstPass.ParameterNodes.Map.Find(Node);
 			if (!Found)
 			{
 				FParameterDesc param;
-				param.Name = node.Name;
-				bool bParseOk = FGuid::Parse(node.UID, param.UID);
+				param.Name = Node->Name;
+				bool bParseOk = FGuid::Parse(Node->UID, param.UID);
 				check(bParseOk);
 				param.Type = EParameterType::Float;
-				param.DefaultValue.Set<FParamStringType>(node.DefaultValue);
+				param.DefaultValue.Set<FParamStringType>(Node->DefaultValue);
+				param.DefaultValue.Set<FParamStringType>(Node->DefaultValue);
 
 				op = new ASTOpParameter();
 				op->Type = EOpType::ST_PARAMETER;
 				op->Parameter = param;
 
-				FirstPass.ParameterNodes.Map.Add(node.m_pNode, op);
+				FirstPass.ParameterNodes.Map.Add(Node, op);
 			}
 			else
 			{
@@ -116,10 +107,10 @@ namespace mu
 		if (!Found)
 		{
 			// Generate the code for the ranges
-			for (int32 a = 0; a < node.Ranges.Num(); ++a)
+			for (int32 a = 0; a < Node->Ranges.Num(); ++a)
 			{
 				FRangeGenerationResult rangeResult;
-				GenerateRange(rangeResult, Options, node.Ranges[a]);
+				GenerateRange(rangeResult, Options, Node->Ranges[a]);
 				op->Ranges.Emplace(op.get(), rangeResult.sizeOp, rangeResult.rangeName, rangeResult.rangeUID);
 			}
 		}
diff --git a/Engine/Plugins/Mutable/Source/MutableTools/Private/MuT/CodeOptimiser.cpp b/Engine/Plugins/Mutable/Source/MutableTools/Private/MuT/CodeOptimiser.cpp
--- a/Engine/Plugins/Mutable/Source/MutableTools/Private/MuT/CodeOptimiser.cpp
+++ b/Engine/Plugins/Mutable/Source/MutableTools/Private/MuT/CodeOptimiser.cpp
@@ -428,6 +428,7 @@ namespace mu
 		ParameterTypes.Add( EOpType::CO_PARAMETER );
 		ParameterTypes.Add( EOpType::PR_PARAMETER );
 		ParameterTypes.Add( EOpType::IM_PARAMETER );
+		ParameterTypes.Add( EOpType::ME_PARAMETER );
 		FindOpTypeVisitor = MakeUnique<FFindOpTypeVisitor>(ParameterTypes);
 	}
 
@@ -646,6 +647,7 @@ namespace mu
 					case EOpType::CO_PARAMETER:
 					case EOpType::PR_PARAMETER:
 					case EOpType::IM_PARAMETER:
+					case EOpType::ME_PARAMETER:
 					case EOpType::MA_PARAMETER:
 						bIsConstantSubgraph = false;
 						break;
diff --git a/Engine/Plugins/Mutable/Source/MutableTools/Private/MuT/CodeOptimiser_Semantic.cpp b/Engine/Plugins/Mutable/Source/MutableTools/Private/MuT/CodeOptimiser_Semantic.cpp
--- a/Engine/Plugins/Mutable/Source/MutableTools/Private/MuT/CodeOptimiser_Semantic.cpp
+++ b/Engine/Plugins/Mutable/Source/MutableTools/Private/MuT/CodeOptimiser_Semantic.cpp
@@ -1,14 +1,6 @@
 // Copyright Epic Games, Inc. All Rights Reserved.
 
-#include "MuT/NodeImageFormatPrivate.h"
-#include "HAL/PlatformMath.h"
-#include "Misc/AssertionMacros.h"
-#include "MuR/Image.h"
-#include "MuR/MutableMath.h"
-#include "MuR/MutableTrace.h"
-#include "MuR/Operations.h"
-#include "MuR/Ptr.h"
-#include "MuR/RefCounted.h"
+#include "MuT/NodeImageFormat.h"
 #include "MuT/AST.h"
 #include "MuT/ASTOpConditional.h"
 #include "MuT/ASTOpConstantBool.h"
@@ -28,7 +20,15 @@
 #include "MuT/ASTOpMeshAddTags.h"
 #include "MuT/ASTOpSwitch.h"
 #include "MuT/CodeOptimiser.h"
+#include "MuR/Image.h"
+#include "MuR/MutableMath.h"
+#include "MuR/MutableTrace.h"
+#include "MuR/Operations.h"
+#include "MuR/Ptr.h"
+#include "MuR/RefCounted.h"
 #include "MuT/Table.h"
+#include "HAL/PlatformMath.h"
+#include "Misc/AssertionMacros.h"
 
 
 namespace mu
diff --git a/Engine/Plugins/Mutable/Source/MutableTools/Private/MuT/Compiler.cpp b/Engine/Plugins/Mutable/Source/MutableTools/Private/MuT/Compiler.cpp
--- a/Engine/Plugins/Mutable/Source/MutableTools/Private/MuT/Compiler.cpp
+++ b/Engine/Plugins/Mutable/Source/MutableTools/Private/MuT/Compiler.cpp
@@ -73,6 +73,7 @@ namespace mu
 		{ false,	{ 0, 0, 0, 0, 0, 0, 0, 0, 0, 0, 0, 0, 0, 0, 0, 0, 0, 0, 0, 0, 0, 0, 0, 0, 0, 0, 0, 0, 0 }		},	// CO_PARAMETER
 		{ false,	{ 0, 0, 0, 0, 0, 0, 0, 0, 0, 0, 0, 0, 0, 0, 0, 0, 0, 0, 0, 0, 0, 0, 0, 0, 0, 0, 0, 0, 0 }		},	// PR_PARAMETER
 		{ false,	{ 0, 0, 0, 0, 0, 0, 0, 0, 0, 0, 0, 0, 0, 0, 0, 0, 0, 0, 0, 0, 0, 0, 0, 0, 0, 0, 0, 0, 0 }		},	// IM_PARAMETER
+		{ true,		{ 0, 0, 0, 0, 0, 0, 0, 0, 0, 0, 0, 0, 0, 0, 0, 0, 0, 0, 0, 0, 0, 0, 0, 0, 0, 0, 0, 0, 0 }		},	// ME_PARAMETER
 		{ false,	{ 0, 0, 0, 0, 0, 0, 0, 0, 0, 0, 0, 0, 0, 0, 0, 0, 0, 0, 0, 0, 0, 0, 0, 0, 0, 0, 0, 0, 0 }		},	// ST_PARAMETER
 		{ false,	{ 0, 0, 0, 0, 0, 0, 0, 0, 0, 0, 0, 0, 0, 0, 0, 0, 0, 0, 0, 0, 0, 0, 0, 0, 0, 0, 0, 0, 0 }		},	// MA_PARAMETER
 
diff --git a/Engine/Plugins/Mutable/Source/MutableTools/Private/MuT/CompilerPrivate.h b/Engine/Plugins/Mutable/Source/MutableTools/Private/MuT/CompilerPrivate.h
--- a/Engine/Plugins/Mutable/Source/MutableTools/Private/MuT/CompilerPrivate.h
+++ b/Engine/Plugins/Mutable/Source/MutableTools/Private/MuT/CompilerPrivate.h
@@ -6,7 +6,6 @@
 
 #include "MuT/AST.h"
 #include "MuT/ErrorLog.h"
-#include "MuT/NodeObjectPrivate.h"
 #include "MuR/Operations.h"
 #include "Templates/SharedPointer.h"
 
diff --git a/Engine/Plugins/Mutable/Source/MutableTools/Private/MuT/Node.cpp b/Engine/Plugins/Mutable/Source/MutableTools/Private/MuT/Node.cpp
--- a/Engine/Plugins/Mutable/Source/MutableTools/Private/MuT/Node.cpp
+++ b/Engine/Plugins/Mutable/Source/MutableTools/Private/MuT/Node.cpp
@@ -2,11 +2,14 @@
 
 
 #include "MuT/Node.h"
-#include "MuT/NodePrivate.h"
 #include "Misc/AssertionMacros.h"
 
+#include "MuT/NodeObject.h"
+#include "MuT/NodeObjectNew.h"
+#include "MuT/NodeObjectGroup.h"
 #include "MuT/NodeComponent.h"
 #include "MuT/NodeComponentNew.h"
+#include "MuT/NodeComponentEdit.h"
 #include "MuT/NodeComponentSwitch.h"
 #include "MuT/NodeComponentVariation.h"
 #include "MuT/NodeLOD.h"
@@ -16,7 +19,30 @@
 #include "MuT/NodeSurfaceVariation.h"
 #include "MuT/NodeExtensionData.h"
 #include "MuT/NodeExtensionDataConstant.h"
+#include "MuT/NodeImageConstant.h"
 #include "MuT/NodeImageTable.h"
+#include "MuT/NodeImageParameter.h"
+#include "MuT/NodeImageFormat.h"
+#include "MuT/NodeImageBinarise.h"
+#include "MuT/NodeImageConditional.h"
+#include "MuT/NodeImageInterpolate.h"
+#include "MuT/NodeImageInvert.h"
+#include "MuT/NodeImageLayer.h"
+#include "MuT/NodeImageLayerColour.h"
+#include "MuT/NodeImageLuminance.h"
+#include "MuT/NodeImageMipmap.h"
+#include "MuT/NodeImageMultiLayer.h"
+#include "MuT/NodeImageNormalComposite.h"
+#include "MuT/NodeImagePlainColour.h"
+#include "MuT/NodeImageProject.h"
+#include "MuT/NodeImageResize.h"
+#include "MuT/NodeImageSaturate.h"
+#include "MuT/NodeImageSwitch.h"
+#include "MuT/NodeImageSwizzle.h"
+#include "MuT/NodeImageTransform.h"
+#include "MuT/NodeImageVariation.h"
+#include "MuT/NodeImageColourMap.h"
+#include "MuT/NodeBool.h"
 #include "MuT/NodeColour.h"
 #include "MuT/NodeColourSwitch.h"
 #include "MuT/NodeColourVariation.h"
@@ -27,10 +53,28 @@
 #include "MuT/NodeColourFromScalars.h"
 #include "MuT/NodeColourTable.h"
 #include "MuT/NodeScalarSwitch.h"
+#include "MuT/NodeScalarConstant.h"
+#include "MuT/NodeScalarVariation.h"
+#include "MuT/NodeScalarArithmeticOperation.h"
+#include "MuT/NodeScalarParameter.h"
+#include "MuT/NodeScalarEnumParameter.h"
+#include "MuT/NodeScalarTable.h"
+#include "MuT/NodeScalarCurve.h"
+#include "MuT/NodeMeshConstant.h"
 #include "MuT/NodeMeshFragment.h"
 #include "MuT/NodeMeshClipMorphPlane.h"
 #include "MuT/NodeMeshClipDeform.h"
 #include "MuT/NodeMeshClipWithMesh.h"
+#include "MuT/NodeMeshParameter.h"
+#include "MuT/NodeMeshMakeMorph.h"
+#include "MuT/NodeMeshApplyPose.h"
+#include "MuT/NodeMeshTransform.h"
+#include "MuT/NodeMeshSwitch.h"
+#include "MuT/NodeMeshReshape.h"
+#include "MuT/NodeMeshMorph.h"
+#include "MuT/NodeMeshFormat.h"
+#include "MuT/NodeMeshVariation.h"
+#include "MuT/NodeMeshTable.h"
 #include "MuT/NodeModifier.h"
 #include "MuT/NodeModifierMeshClipDeform.h"
 #include "MuT/NodeModifierMeshClipMorphPlane.h"
@@ -41,6 +85,13 @@
 #include "MuT/NodeMatrix.h"
 #include "MuT/NodeMatrixConstant.h"
 #include "MuT/NodeMatrixParameter.h"
+#include "MuT/NodeString.h"
+#include "MuT/NodeStringConstant.h"
+#include "MuT/NodeStringParameter.h"
+#include "MuT/NodeBool.h"
+#include "MuT/NodeProjector.h"
+#include "MuT/NodeRange.h"
+#include "MuT/NodeRangeFromScalar.h"
 
 namespace mu
 {
@@ -48,12 +99,31 @@ namespace mu
 	// Static initialisation
 	FNodeType Node::StaticType = FNodeType(Node::EType::Node, nullptr);
 
+	FNodeType NodeObject::StaticType = FNodeType(Node::EType::Object, Node::GetStaticType());
+	FNodeType NodeObjectNew::StaticType = FNodeType(Node::EType::ObjectNew, NodeObject::GetStaticType());
+	FNodeType NodeObjectGroup::StaticType = FNodeType(Node::EType::ObjectGroup, NodeObject::GetStaticType());
+
 	FNodeType NodeComponent::StaticType = FNodeType(Node::EType::Component, Node::GetStaticType());
 	FNodeType NodeComponentNew::StaticType = FNodeType(Node::EType::ComponentNew, NodeComponent::GetStaticType());
 	FNodeType NodeComponentSwitch::StaticType = FNodeType(Node::EType::ComponentSwitch, NodeComponent::GetStaticType());
 	FNodeType NodeComponentVariation::StaticType = FNodeType(Node::EType::ComponentVariation, NodeComponent::GetStaticType());
+	FNodeType NodeComponentEdit::StaticType = FNodeType(Node::EType::ComponentEdit, NodeComponent::GetStaticType());
+
+	FNodeType NodeBool::StaticType = FNodeType(Node::EType::Bool, Node::GetStaticType());
+	FNodeType NodeBoolConstant::StaticType = FNodeType(Node::EType::BoolConstant, NodeBool::GetStaticType());
+	FNodeType NodeBoolParameter::StaticType = FNodeType(Node::EType::BoolParameter, NodeBool::GetStaticType());
+	FNodeType NodeBoolNot::StaticType = FNodeType(Node::EType::BoolNot, NodeBool::GetStaticType());
+	FNodeType NodeBoolAnd::StaticType = FNodeType(Node::EType::BoolAnd, NodeBool::GetStaticType());
 
+	FNodeType NodeScalar::StaticType = FNodeType(Node::EType::Scalar, Node::GetStaticType());
 	FNodeType NodeScalarSwitch::StaticType = FNodeType(Node::EType::ScalarSwitch, NodeScalar::GetStaticType());
+	FNodeType NodeScalarConstant::StaticType = FNodeType(Node::EType::ScalarConstant, NodeScalar::GetStaticType());
+	FNodeType NodeScalarParameter::StaticType = FNodeType(Node::EType::ScalarParameter, NodeScalar::GetStaticType());
+	FNodeType NodeScalarVariation::StaticType = FNodeType(Node::EType::ScalarVariation, NodeScalar::GetStaticType());
+	FNodeType NodeScalarArithmeticOperation::StaticType = FNodeType(Node::EType::ScalarArithmeticOperation, NodeScalar::GetStaticType());
+	FNodeType NodeScalarEnumParameter::StaticType = FNodeType(Node::EType::ScalarEnumParameter, NodeScalar::GetStaticType());
+	FNodeType NodeScalarTable::StaticType = FNodeType(Node::EType::ScalarTable, NodeScalar::GetStaticType());
+	FNodeType NodeScalarCurve::StaticType = FNodeType(Node::EType::ScalarCurve, NodeScalar::GetStaticType());
 
 	FNodeType NodeSurface::StaticType = FNodeType(Node::EType::Surface, Node::GetStaticType());
 	FNodeType NodeSurfaceNew::StaticType = FNodeType(Node::EType::SurfaceNew, NodeSurface::GetStaticType());
@@ -63,7 +133,31 @@ namespace mu
 	FNodeType NodeLOD::StaticType = FNodeType(Node::EType::LOD, Node::GetStaticType());
 	FNodeType NodeExtensionData::StaticType = FNodeType(Node::EType::ExtensionData, Node::GetStaticType());
 	FNodeType NodeExtensionDataConstant::StaticType = FNodeType(Node::EType::ExtensionDataConstant, NodeExtensionData::GetStaticType());
+
+	FNodeType NodeImage::StaticType = FNodeType(Node::EType::Image, Node::GetStaticType());
+	FNodeType NodeImageConstant::StaticType = FNodeType(Node::EType::ImageConstant, NodeImage::GetStaticType());
 	FNodeType NodeImageTable::StaticType = FNodeType(Node::EType::ImageTable, NodeImage::GetStaticType());
+	FNodeType NodeImageParameter::StaticType = FNodeType(Node::EType::ImageParameter, NodeImage::GetStaticType());
+	FNodeType NodeImageFormat::StaticType = FNodeType(Node::EType::ImageFormat, NodeImage::GetStaticType());
+	FNodeType NodeImageBinarise::StaticType = FNodeType(Node::EType::ImageBinarise, NodeImage::GetStaticType());
+	FNodeType NodeImageConditional::StaticType = FNodeType(Node::EType::ImageConditional, NodeImage::GetStaticType());
+	FNodeType NodeImageInterpolate::StaticType = FNodeType(Node::EType::ImageInterpolate, NodeImage::GetStaticType());
+	FNodeType NodeImageInvert::StaticType = FNodeType(Node::EType::ImageInvert, NodeImage::GetStaticType());
+	FNodeType NodeImageLayer::StaticType = FNodeType(Node::EType::ImageLayer, NodeImage::GetStaticType());
+	FNodeType NodeImageLayerColour::StaticType = FNodeType(Node::EType::ImageLayerColour, NodeImage::GetStaticType());
+	FNodeType NodeImageLuminance::StaticType = FNodeType(Node::EType::ImageLuminance, NodeImage::GetStaticType());
+	FNodeType NodeImageMipmap::StaticType = FNodeType(Node::EType::ImageMipmap, NodeImage::GetStaticType());
+	FNodeType NodeImageMultiLayer::StaticType = FNodeType(Node::EType::ImageMultiLayer, NodeImage::GetStaticType());
+	FNodeType NodeImageNormalComposite::StaticType = FNodeType(Node::EType::ImageNormalComposite, NodeImage::GetStaticType());
+	FNodeType NodeImagePlainColour::StaticType = FNodeType(Node::EType::ImagePlainColour, NodeImage::GetStaticType());
+	FNodeType NodeImageProject::StaticType = FNodeType(Node::EType::ImageProject, NodeImage::GetStaticType());
+	FNodeType NodeImageResize::StaticType = FNodeType(Node::EType::ImageResize, NodeImage::GetStaticType());
+	FNodeType NodeImageSaturate::StaticType = FNodeType(Node::EType::ImageSaturate, NodeImage::GetStaticType());
+	FNodeType NodeImageSwitch::StaticType = FNodeType(Node::EType::ImageSwitch, NodeImage::GetStaticType());
+	FNodeType NodeImageSwizzle::StaticType = FNodeType(Node::EType::ImageSwizzle, NodeImage::GetStaticType());
+	FNodeType NodeImageTransform::StaticType = FNodeType(Node::EType::ImageTransform, NodeImage::GetStaticType());
+	FNodeType NodeImageVariation::StaticType = FNodeType(Node::EType::ImageVariation, NodeImage::GetStaticType());
+	FNodeType NodeImageColourMap::StaticType = FNodeType(Node::EType::ImageColorMap, NodeImage::GetStaticType());
 
 	FNodeType NodeColour::StaticType = FNodeType(Node::EType::Color, Node::GetStaticType());
 	FNodeType NodeColourConstant::StaticType = FNodeType(Node::EType::ColorConstant, NodeColour::GetStaticType());
@@ -75,10 +169,22 @@ namespace mu
 	FNodeType NodeColourSampleImage::StaticType = FNodeType(Node::EType::ColorSampleImage, NodeColour::GetStaticType());
 	FNodeType NodeColourFromScalars::StaticType = FNodeType(Node::EType::ColorFromScalars, NodeColour::GetStaticType());
 
+	FNodeType NodeMesh::StaticType = FNodeType(Node::EType::Mesh, Node::GetStaticType());
+	FNodeType NodeMeshConstant::StaticType = FNodeType(Node::EType::MeshConstant, NodeMesh::GetStaticType());
 	FNodeType NodeMeshFragment::StaticType = FNodeType(Node::EType::MeshFragment, NodeMesh::GetStaticType());
 	FNodeType NodeMeshClipMorphPlane::StaticType = FNodeType(Node::EType::MeshClipMorphPlane, NodeMesh::GetStaticType());
 	FNodeType NodeMeshClipDeform::StaticType = FNodeType(Node::EType::MeshClipDeform, NodeMesh::GetStaticType());
 	FNodeType NodeMeshClipWithMesh::StaticType = FNodeType(Node::EType::MeshClipWithMesh, NodeMesh::GetStaticType());
+	FNodeType NodeMeshParameter::StaticType = FNodeType(Node::EType::MeshParameter, NodeMesh::GetStaticType());
+	FNodeType NodeMeshMakeMorph::StaticType = FNodeType(Node::EType::MeshMakeMorph, NodeMesh::GetStaticType());
+	FNodeType NodeMeshApplyPose::StaticType = FNodeType(Node::EType::MeshApplyPose, NodeMesh::GetStaticType());
+	FNodeType NodeMeshTransform::StaticType = FNodeType(Node::EType::MeshTransform, NodeMesh::GetStaticType());
+	FNodeType NodeMeshSwitch::StaticType = FNodeType(Node::EType::MeshSwitch, NodeMesh::GetStaticType());
+	FNodeType NodeMeshReshape::StaticType = FNodeType(Node::EType::MeshReshape, NodeMesh::GetStaticType());
+	FNodeType NodeMeshMorph::StaticType = FNodeType(Node::EType::MeshMorph, NodeMesh::GetStaticType());
+	FNodeType NodeMeshFormat::StaticType = FNodeType(Node::EType::MeshFormat, NodeMesh::GetStaticType());
+	FNodeType NodeMeshVariation::StaticType = FNodeType(Node::EType::MeshVariation, NodeMesh::GetStaticType());
+	FNodeType NodeMeshTable::StaticType = FNodeType(Node::EType::MeshTable, NodeMesh::GetStaticType());
 
 	FNodeType NodeModifier::StaticType = FNodeType(Node::EType::Modifier, Node::GetStaticType());
 	FNodeType NodeModifierMeshClipDeform::StaticType = FNodeType(Node::EType::ModifierMeshClipDeform, NodeModifier::GetStaticType());
@@ -92,6 +198,18 @@ namespace mu
 	FNodeType NodeMatrixConstant::StaticType = FNodeType(Node::EType::MatrixConstant, Node::GetStaticType());
 	FNodeType NodeMatrixParameter::StaticType = FNodeType(Node::EType::MatrixParameter, Node::GetStaticType());
 
+	FNodeType NodeString::StaticType = FNodeType(Node::EType::String, Node::GetStaticType());
+	FNodeType NodeStringConstant::StaticType = FNodeType(Node::EType::StringConstant, NodeString::GetStaticType());
+	FNodeType NodeStringParameter::StaticType = FNodeType(Node::EType::StringParameter, NodeString::GetStaticType());
+
+	FNodeType NodeProjector::StaticType = FNodeType(Node::EType::Projector, Node::GetStaticType());
+	FNodeType NodeProjectorConstant::StaticType = FNodeType(Node::EType::ProjectorConstant, NodeProjector::GetStaticType());
+	FNodeType NodeProjectorParameter::StaticType = FNodeType(Node::EType::ProjectorParameter, NodeProjector::GetStaticType());
+
+	FNodeType NodeRange::StaticType = FNodeType(Node::EType::Range, Node::GetStaticType());
+	FNodeType NodeRangeFromScalar::StaticType = FNodeType(Node::EType::RangeFromScalar, NodeRange::GetStaticType());
+
+
 	FNodeType::FNodeType()
 	{
 		Type = Node::EType::None;
diff --git a/Engine/Plugins/Mutable/Source/MutableTools/Private/MuT/NodeComponentEdit.cpp b/Engine/Plugins/Mutable/Source/MutableTools/Private/MuT/NodeComponentEdit.cpp
--- a/Engine/Plugins/Mutable/Source/MutableTools/Private/MuT/NodeComponentEdit.cpp
+++ b/Engine/Plugins/Mutable/Source/MutableTools/Private/MuT/NodeComponentEdit.cpp
@@ -1,48 +1,21 @@
 // Copyright Epic Games, Inc. All Rights Reserved.
 
-
 #include "MuT/NodeComponentEdit.h"
 
-#include "Misc/AssertionMacros.h"
 #include "MuT/NodeComponentNew.h"
-#include "MuT/NodePrivate.h"
 #include "MuT/NodeSurface.h"
 
-
 namespace mu
 {
 
-	FNodeType NodeComponentEdit::StaticType = FNodeType(Node::EType::ComponentEdit, NodeComponent::GetStaticType() );
-
-
-	//---------------------------------------------------------------------------------------------
-	// Own Interface
-	//---------------------------------------------------------------------------------------------
-    void NodeComponentEdit::SetParent( NodeComponent* p )
-	{
-		m_pParent = p;
-	}
-
-
-	//---------------------------------------------------------------------------------------------
-    NodeComponent* NodeComponentEdit::GetParent() const
-	{
-        return m_pParent.get();
-	}
-
-
-	//---------------------------------------------------------------------------------------------
 	const NodeComponentNew* NodeComponentEdit::GetParentComponentNew() const
 	{
-		const NodeComponentNew* Parent = nullptr;
-		if (m_pParent)
+		if (Parent)
 		{
-			Parent = m_pParent->GetParentComponentNew();
+			return Parent->GetParentComponentNew();
 		}
 
-		return Parent;
+		return nullptr;
 	}
 
 }
-
-
diff --git a/Engine/Plugins/Mutable/Source/MutableTools/Private/MuT/NodeImageConstant.cpp b/Engine/Plugins/Mutable/Source/MutableTools/Private/MuT/NodeImageConstant.cpp
--- a/Engine/Plugins/Mutable/Source/MutableTools/Private/MuT/NodeImageConstant.cpp
+++ b/Engine/Plugins/Mutable/Source/MutableTools/Private/MuT/NodeImageConstant.cpp
@@ -6,26 +6,17 @@
 #include "Misc/AssertionMacros.h"
 #include "MuR/Image.h"
 #include "MuR/Serialisation.h"
-#include "MuT/NodePrivate.h"
 
 
 namespace mu
 {
 
-	//---------------------------------------------------------------------------------------------
-	// Static initialisation
-	//---------------------------------------------------------------------------------------------
-	FNodeType NodeImageConstant::StaticType = FNodeType(Node::EType::ImageConstant, NodeImage::GetStaticType());
-
-
-	//---------------------------------------------------------------------------------------------
 	void NodeImageConstant::SetValue(TSharedPtr<const FImage> Value)
 	{
 		Proxy = new TResourceProxyMemory<FImage>(Value);
 	}
 
 
-	//---------------------------------------------------------------------------------------------
 	void NodeImageConstant::SetValue(Ptr<TResourceProxy<FImage>> InProxy)
 	{
 		Proxy = InProxy;
diff --git a/Engine/Plugins/Mutable/Source/MutableTools/Private/MuT/NodeMesh.cpp b/Engine/Plugins/Mutable/Source/MutableTools/Private/MuT/NodeMesh.cpp
--- a/Engine/Plugins/Mutable/Source/MutableTools/Private/MuT/NodeMesh.cpp
+++ b/Engine/Plugins/Mutable/Source/MutableTools/Private/MuT/NodeMesh.cpp
@@ -11,20 +11,6 @@
 namespace mu
 {
 
-	static FNodeType s_nodeMeshType = FNodeType(Node::EType::Mesh , Node::GetStaticType());
-
-
-	const FNodeType* NodeMesh::GetType() const
-	{
-		return GetStaticType();
-	}
-
-
-	const FNodeType* NodeMesh::GetStaticType()
-	{
-		return &s_nodeMeshType;
-	}
-
 
 	void MeshCreateCollapsedVertexMap(const mu::FMesh* Mesh, TArray<int32>& CollapsedVertices)
 	{
diff --git a/Engine/Plugins/Mutable/Source/MutableTools/Private/MuT/NodeMeshClipMorphPlane.cpp b/Engine/Plugins/Mutable/Source/MutableTools/Private/MuT/NodeMeshClipMorphPlane.cpp
--- a/Engine/Plugins/Mutable/Source/MutableTools/Private/MuT/NodeMeshClipMorphPlane.cpp
+++ b/Engine/Plugins/Mutable/Source/MutableTools/Private/MuT/NodeMeshClipMorphPlane.cpp
@@ -5,8 +5,6 @@
 #include "Misc/AssertionMacros.h"
 #include "MuR/MutableMath.h"
 #include "MuT/NodeLayout.h"
-#include "MuT/NodeMeshPrivate.h"
-#include "MuT/NodePrivate.h"
 
 
 namespace mu
diff --git a/Engine/Plugins/Mutable/Source/MutableTools/Private/MuT/NodeObjectGroup.cpp b/Engine/Plugins/Mutable/Source/MutableTools/Private/MuT/NodeObjectGroup.cpp
--- a/Engine/Plugins/Mutable/Source/MutableTools/Private/MuT/NodeObjectGroup.cpp
+++ b/Engine/Plugins/Mutable/Source/MutableTools/Private/MuT/NodeObjectGroup.cpp
@@ -5,108 +5,32 @@
 #include "HAL/PlatformCrt.h"
 #include "Misc/AssertionMacros.h"
 #include "MuT/NodeLayout.h"
-#include "MuT/NodeObjectGroupPrivate.h"
-#include "MuT/NodeObjectPrivate.h"
-#include "MuT/NodePrivate.h"
 
 
 namespace mu
 {
-	MUTABLE_IMPLEMENT_ENUM_SERIALISABLE(NodeObjectGroup::CHILD_SELECTION)
 
-
-	//---------------------------------------------------------------------------------------------
-	// Static initialisation
-	//---------------------------------------------------------------------------------------------
-	FNodeType NodeObjectGroup::Private::s_type = FNodeType(Node::EType::ObjectGroup, NodeObject::GetStaticType() );
-
-
-	//---------------------------------------------------------------------------------------------
-	//!
-	//---------------------------------------------------------------------------------------------
-
-	MUTABLE_IMPLEMENT_NODE( NodeObjectGroup )
-
-
-	//---------------------------------------------------------------------------------------------
-	// Own Interface
-	//---------------------------------------------------------------------------------------------
 	const FString& NodeObjectGroup::GetName() const
 	{
-		return m_pD->Name;
+		return Name;
 	}
 
 
-	//---------------------------------------------------------------------------------------------
-	void NodeObjectGroup::SetName( const FString& Name )
+	void NodeObjectGroup::SetName( const FString& InName )
 	{
-		m_pD->Name = Name;
+		Name = InName;
 	}
 
 
 	const FString& NodeObjectGroup::GetUid() const
 	{
-		return m_pD->Uid;
-	}
-
-
-	//---------------------------------------------------------------------------------------------
-	void NodeObjectGroup::SetUid( const FString& Uid )
-	{
-		m_pD->Uid = Uid;
-	}
-
-
-	//---------------------------------------------------------------------------------------------
-	NodeObjectGroup::CHILD_SELECTION NodeObjectGroup::GetSelectionType() const
-	{
-		return m_pD->Type;
-	}
-
-
-	//---------------------------------------------------------------------------------------------
-	void NodeObjectGroup::SetSelectionType( CHILD_SELECTION t )
-	{
-		m_pD->Type = t;
-	}
-
-
-	//---------------------------------------------------------------------------------------------
-	int NodeObjectGroup::GetChildCount() const
-	{
-		return m_pD->m_children.Num();
-	}
-
-
-	//---------------------------------------------------------------------------------------------
-	void NodeObjectGroup::SetChildCount( int num )
-	{
-		check( num >=0 );
-		m_pD->m_children.SetNum( num );
-	}
-
-
-	//---------------------------------------------------------------------------------------------
-	NodeObjectPtr NodeObjectGroup::GetChild( int index ) const
-	{
-		check( index >=0 && index < m_pD->m_children.Num() );
-
-		return m_pD->m_children[ index ].get();
-	}
-
-
-	//---------------------------------------------------------------------------------------------
-	void NodeObjectGroup::SetChild( int index, NodeObjectPtr pObject )
-	{
-		check( index >=0 && index < m_pD->m_children.Num() );
-
-		m_pD->m_children[ index ] = pObject;
+		return Uid;
 	}
 
 
-	void NodeObjectGroup::SetDefaultValue(int32 Value)
+	void NodeObjectGroup::SetUid( const FString& InUid )
 	{
-		m_pD->DefaultValue = Value;
+		Uid = InUid;
 	}
 
 }
diff --git a/Engine/Plugins/Mutable/Source/MutableTools/Private/MuT/NodeObjectNew.cpp b/Engine/Plugins/Mutable/Source/MutableTools/Private/MuT/NodeObjectNew.cpp
--- a/Engine/Plugins/Mutable/Source/MutableTools/Private/MuT/NodeObjectNew.cpp
+++ b/Engine/Plugins/Mutable/Source/MutableTools/Private/MuT/NodeObjectNew.cpp
@@ -10,16 +10,12 @@
 #include "MuT/NodeLOD.h"
 #include "MuT/NodeMesh.h"
 #include "MuT/NodeMeshConstant.h"
-#include "MuT/NodeMeshPrivate.h"
-#include "MuT/NodePrivate.h"
 #include "MuT/NodeSurface.h"
 #include "MuT/NodeSurfaceNew.h"
 
 
 namespace mu
 {
-	//---------------------------------------------------------------------------------------------
-	FNodeType NodeObjectNew::StaticType = FNodeType(Node::EType::ObjectNew, NodeObject::GetStaticType() );
 
 
 	const FString& NodeObjectNew::GetName() const
@@ -83,18 +79,6 @@ namespace mu
 	}
 
 
-	void NodeObjectNew::RemoveStateParam( int32 s, const FString& param )
-	{
-		check( s>=0 && s<GetStateCount() );
-
-		int32 it = States[s].RuntimeParams.Find( param );
-		if ( it != INDEX_NONE )
-		{
-			States[s].RuntimeParams.RemoveAt( it );
-		}
-	}
-
-
     void NodeObjectNew::SetStateProperties( int32 StateIndex, 
 		ETextureCompressionStrategy TextureCompressionStrategy, 
 		bool bOnlyFirstLOD, 
diff --git a/Engine/Plugins/Mutable/Source/MutableTools/Private/MuT/NodeProjectorConstant.cpp b/Engine/Plugins/Mutable/Source/MutableTools/Private/MuT/NodeProjectorConstant.cpp
--- a/Engine/Plugins/Mutable/Source/MutableTools/Private/MuT/NodeProjectorConstant.cpp
+++ b/Engine/Plugins/Mutable/Source/MutableTools/Private/MuT/NodeProjectorConstant.cpp
@@ -1,34 +1,14 @@
 // Copyright Epic Games, Inc. All Rights Reserved.
 
-
 #include "Misc/AssertionMacros.h"
 #include "MuR/MutableMath.h"
 #include "MuR/Parameters.h"
 #include "MuT/Node.h"
-#include "MuT/NodePrivate.h"
 #include "MuT/NodeProjector.h"
-#include "MuT/NodeProjectorPrivate.h"
-
 
 namespace mu
 {
 
-	//---------------------------------------------------------------------------------------------
-	// Static initialisation
-	//---------------------------------------------------------------------------------------------
-	FNodeType NodeProjectorConstant::Private::s_type = FNodeType(Node::EType::ProjectorConstant, NodeProjector::GetStaticType() );
-
-
-	//---------------------------------------------------------------------------------------------
-	//!
-	//---------------------------------------------------------------------------------------------
-
-	MUTABLE_IMPLEMENT_NODE( NodeProjectorConstant )
-
-
-	//---------------------------------------------------------------------------------------------
-	// Own Interface
-	//---------------------------------------------------------------------------------------------
     void NodeProjectorConstant::GetValue( EProjectorType* OutType,
 		FVector3f* OutPos,
 		FVector3f* OutDir,
@@ -36,29 +16,25 @@ namespace mu
 		FVector3f* OutScale,
 		float* OutProjectionAngle ) const
 	{
-        if (OutType) *OutType = m_pD->Type;
-
-        if (OutPos) *OutPos = m_pD->m_position;
-        if (OutDir) *OutDir = m_pD->m_direction;
-        if (OutUp) *OutUp = m_pD->m_up;
-        if (OutScale) *OutScale = m_pD->m_scale;
-
-        if (OutProjectionAngle) *OutProjectionAngle = m_pD->m_projectionAngle;
+        if (OutType) *OutType = Type;
+        if (OutPos) *OutPos = Position;
+        if (OutDir) *OutDir = Direction;
+        if (OutUp) *OutUp = Up;
+        if (OutScale) *OutScale = Scale;
+        if (OutProjectionAngle) *OutProjectionAngle = ProjectionAngle;
     }
 
 
-	//---------------------------------------------------------------------------------------------
-    void NodeProjectorConstant::SetValue( EProjectorType type, FVector3f pos, FVector3f dir, FVector3f up, FVector3f scale, float projectionAngle )
+	void NodeProjectorConstant::SetValue( EProjectorType type, FVector3f pos, FVector3f dir, FVector3f up, FVector3f scale, float projectionAngle )
 	{
-        m_pD->Type = type;
-		m_pD->m_position = pos;
-		m_pD->m_direction = dir;
-		m_pD->m_up = up;
-        m_pD->m_scale = scale;
-        m_pD->m_projectionAngle = projectionAngle;
+        Type = type;
+		Position = pos;
+		Direction = dir;
+		Up = up;
+        Scale = scale;
+        ProjectionAngle = projectionAngle;
 	}
 
-
 }
 
 
diff --git a/Engine/Plugins/Mutable/Source/MutableTools/Private/MuT/NodeProjectorParameter.cpp b/Engine/Plugins/Mutable/Source/MutableTools/Private/MuT/NodeProjectorParameter.cpp
--- a/Engine/Plugins/Mutable/Source/MutableTools/Private/MuT/NodeProjectorParameter.cpp
+++ b/Engine/Plugins/Mutable/Source/MutableTools/Private/MuT/NodeProjectorParameter.cpp
@@ -7,50 +7,30 @@
 #include "MuR/Ptr.h"
 #include "MuR/RefCounted.h"
 #include "MuT/Node.h"
-#include "MuT/NodePrivate.h"
 #include "MuT/NodeProjector.h"
-#include "MuT/NodeProjectorPrivate.h"
 #include "MuT/NodeRange.h"
 
 
 namespace mu
 {
-
-	//---------------------------------------------------------------------------------------------
-	// Static initialisation
-	//---------------------------------------------------------------------------------------------
-	FNodeType NodeProjectorParameter::Private::s_type = FNodeType(Node::EType::ProjectorParameter, NodeProjector::GetStaticType() );
-
-
-	//---------------------------------------------------------------------------------------------
-	//!
-	//---------------------------------------------------------------------------------------------
-
-	MUTABLE_IMPLEMENT_NODE( NodeProjectorParameter )
-
-
-	//---------------------------------------------------------------------------------------------
-	// Own Interface
-	//---------------------------------------------------------------------------------------------
-	void NodeProjectorParameter::SetName( const FString& strName )
+	void NodeProjectorParameter::SetName( const FString& InName )
 	{
-		m_pD->Name = strName;
+		Name = InName;
 	}
 
 
 	const FString& NodeProjectorParameter::GetUid() const
 	{
-		return m_pD->UID;
+		return UID;
 	}
 
 
-	void NodeProjectorParameter::SetUid( const FString& strUid )
+	void NodeProjectorParameter::SetUid( const FString& InUid)
 	{
-		m_pD->UID = strUid;
+		UID = InUid;
 	}
 
 
-	//---------------------------------------------------------------------------------------------
     void NodeProjectorParameter::SetDefaultValue( EProjectorType type,
 		FVector3f pos,
 		FVector3f dir,
@@ -58,30 +38,28 @@ namespace mu
 		FVector3f scale,
 		float projectionAngle )
 	{
-        m_pD->Type = type;
-        m_pD->m_position = pos;
-		m_pD->m_direction = dir;
-		m_pD->m_up = up;
-        m_pD->m_scale = scale;
-        m_pD->m_projectionAngle = projectionAngle;
+        Type = type;
+        Position = pos;
+		Direction = dir;
+		Up = up;
+        Scale = scale;
+        ProjectionAngle = projectionAngle;
     }
 
 
-    //---------------------------------------------------------------------------------------------
-    void NodeProjectorParameter::SetRangeCount( int i )
+    void NodeProjectorParameter::SetRangeCount( int32 i )
     {
         check(i>=0);
-        m_pD->Ranges.SetNum(i);
+        Ranges.SetNum(i);
     }
 
 
-    //---------------------------------------------------------------------------------------------
-    void NodeProjectorParameter::SetRange( int i, NodeRangePtr pRange )
+    void NodeProjectorParameter::SetRange( int32 i, Ptr<NodeRange> InRange )
     {
-        check( i>=0 && i<m_pD->Ranges.Num() );
-        if ( i>=0 && i<m_pD->Ranges.Num() )
+        check( i>=0 && i<Ranges.Num() );
+        if ( i>=0 && i<Ranges.Num() )
         {
-            m_pD->Ranges[i] = pRange;
+            Ranges[i] = InRange;
         }
     }
 
diff --git a/Engine/Plugins/Mutable/Source/MutableTools/Public/MuT/Compiler.h b/Engine/Plugins/Mutable/Source/MutableTools/Public/MuT/Compiler.h
--- a/Engine/Plugins/Mutable/Source/MutableTools/Public/MuT/Compiler.h
+++ b/Engine/Plugins/Mutable/Source/MutableTools/Public/MuT/Compiler.h
@@ -118,7 +118,7 @@ namespace mu
 
     protected:
 
-        //! Forbidden. Manage with the Ptr<> template.
+        /** Forbidden. Manage with the Ptr<> template. */
         ~CompilerOptions() override;
 
     private:
@@ -166,7 +166,7 @@ namespace mu
 
     protected:
 
-        //! Forbidden. Manage with the Ptr<> template.
+        /** Forbidden. Manage with the Ptr<> template. */
         ~Compiler();
 
     private:
diff --git a/Engine/Plugins/Mutable/Source/MutableTools/Public/MuT/NodeBool.h b/Engine/Plugins/Mutable/Source/MutableTools/Public/MuT/NodeBool.h
--- a/Engine/Plugins/Mutable/Source/MutableTools/Public/MuT/NodeBool.h
+++ b/Engine/Plugins/Mutable/Source/MutableTools/Public/MuT/NodeBool.h
@@ -6,221 +6,134 @@
 #include "MuR/Ptr.h"
 #include "MuR/RefCounted.h"
 #include "MuT/Node.h"
+#include "MuT/NodeRange.h"
 
 
 namespace mu
 {
-	class NodeRange;
 
-	//---------------------------------------------------------------------------------------------
-    //! %Base class of any node that outputs a Bool value.
-	//! \ingroup model
-	//---------------------------------------------------------------------------------------------
+	/** Base class of any node that outputs a Bool value. */
 	class MUTABLETOOLS_API NodeBool : public Node
 	{
 	public:
 
-		//-----------------------------------------------------------------------------------------
 		// Node Interface
-		//-----------------------------------------------------------------------------------------
+		virtual const FNodeType* GetType() const override { return &StaticType; }
+		static const FNodeType* GetStaticType() { return &StaticType; }
 
-        const FNodeType* GetType() const override;
-		static const FNodeType* GetStaticType();
+	protected:
 
+		/** Forbidden. Manage with the Ptr<> template. */
+		inline ~NodeBool() {}
 
-		//-----------------------------------------------------------------------------------------
-		// Interface pattern
-		//-----------------------------------------------------------------------------------------
-		class Private;
+	private:
 
-	protected:
+		static FNodeType StaticType;
 
-		//! Forbidden. Manage with the Ptr<> template.
-		inline ~NodeBool() {}
 	};
 
 
-	//---------------------------------------------------------------------------------------------
-	//! Node returning a Bool constant value.
-	//! \ingroup model
-	//---------------------------------------------------------------------------------------------
+	/** Node returning a Bool constant value. */
 	class MUTABLETOOLS_API NodeBoolConstant : public NodeBool
 	{
 	public:
 
-		NodeBoolConstant();
+		bool Value = false;
 
+	public:
 
-		//-----------------------------------------------------------------------------------------
 		// Node Interface
-		//-----------------------------------------------------------------------------------------
-
-        const FNodeType* GetType() const override;
-		static const FNodeType* GetStaticType();
-
-		//-----------------------------------------------------------------------------------------
-		// Own Interface
-		//-----------------------------------------------------------------------------------------
-
-		//! Get the value to be returned by the node.
-		bool GetValue() const;
-
-		//! Set the value to be returned by the node.
-		void SetValue( bool v );
-
-		//-----------------------------------------------------------------------------------------
-		// Interface pattern
-		//-----------------------------------------------------------------------------------------
-		class Private;
-		Private* GetPrivate() const;
+		virtual const FNodeType* GetType() const override { return &StaticType; }
+		static const FNodeType* GetStaticType() { return &StaticType; }
 
 	protected:
 
-		//! Forbidden.
-		~NodeBoolConstant();
+		/** Forbidden. Manage with the Ptr<> template. */
+		~NodeBoolConstant() {}
 
 	private:
 
-		Private* m_pD;
+		static FNodeType StaticType;
 
 	};
 
 
-	//---------------------------------------------------------------------------------------------
-	//! Node that defines a Bool model parameter.
-	//! \ingroup model
-	//---------------------------------------------------------------------------------------------
+	/** Node that defines a Bool model parameter. */
 	class MUTABLETOOLS_API NodeBoolParameter : public NodeBool
 	{
 	public:
 
-		NodeBoolParameter();
-
-		//-----------------------------------------------------------------------------------------
-		// Node Interface
-		//-----------------------------------------------------------------------------------------
-
-        const FNodeType* GetType() const override;
-		static const FNodeType* GetStaticType();
-
-		//-----------------------------------------------------------------------------------------
-		// Own Interface
-		//-----------------------------------------------------------------------------------------
+		bool DefaultValue;
+		FString Name;
+		FString UID;
 
-		//! Set the name of the parameter. It will be exposed in the final compiled data.
-		void SetName( const FString& );
+		TArray<Ptr<NodeRange>> Ranges;
 
-		//! Get the default value of the parameter.
-		void SetDefaultValue( bool v );
-
-        //! Set the number of ranges (dimensions) for this parameter.
-        //! By default a parameter has 0 ranges, meaning it only has one value.
-        void SetRangeCount( int i );
-        void SetRange( int i, Ptr<NodeRange> );
+	public:
 
-		//-----------------------------------------------------------------------------------------
-		// Interface pattern
-		//-----------------------------------------------------------------------------------------
-		class Private;
-		Private* GetPrivate() const;
+		// Node Interface
+		virtual const FNodeType* GetType() const override { return &StaticType; }
+		static const FNodeType* GetStaticType() { return &StaticType; }
 
 	protected:
 
-		//! Forbidden. 
-		~NodeBoolParameter();
+		/** Forbidden. Manage with the Ptr<> template. */
+		~NodeBoolParameter() {}
 
 	private:
 
-		Private* m_pD;
+		static FNodeType StaticType;
 
 	};
 
 
-	//---------------------------------------------------------------------------------------------
-	//! Node that returns the oposite of the input value.
-	//! \ingroup model
-	//---------------------------------------------------------------------------------------------
+	/** Node that returns the oposite of the input value. */
 	class MUTABLETOOLS_API NodeBoolNot : public NodeBool
 	{
 	public:
 
-		NodeBoolNot();
+		Ptr<NodeBool> Source;
 
-		//-----------------------------------------------------------------------------------------
-		// Node Interface
-		//-----------------------------------------------------------------------------------------
-
-        const FNodeType* GetType() const override;
-		static const FNodeType* GetStaticType();
-
-		//-----------------------------------------------------------------------------------------
-		// Own Interface
-		//-----------------------------------------------------------------------------------------
-		
-		//! Input
-		Ptr<NodeBool> GetInput() const;
-		void SetInput( Ptr<NodeBool> );
+	public:
 
-		//-----------------------------------------------------------------------------------------
-		// Interface pattern
-		//-----------------------------------------------------------------------------------------
-		class Private;
-		Private* GetPrivate() const;
+		// Node Interface
+		virtual const FNodeType* GetType() const override { return &StaticType; }
+		static const FNodeType* GetStaticType() { return &StaticType; }
 
 	protected:
 
-		//! Forbidden. 
-		~NodeBoolNot();
+		/** Forbidden. Manage with the Ptr<> template. */
+		~NodeBoolNot() {}
 
 	private:
 
-		Private* m_pD;
+		static FNodeType StaticType;
 
 	};
 
 
-	//---------------------------------------------------------------------------------------------
-	//! 
-	//! \ingroup model
-	//---------------------------------------------------------------------------------------------
+	/** */
 	class MUTABLETOOLS_API NodeBoolAnd : public NodeBool
 	{
 	public:
 
-		NodeBoolAnd();
-
-		//-----------------------------------------------------------------------------------------
-		// Node Interface
-		//-----------------------------------------------------------------------------------------
-
-        const FNodeType* GetType() const override;
-		static const FNodeType* GetStaticType();
+		Ptr<NodeBool> A;
+		Ptr<NodeBool> B;
 
-		//-----------------------------------------------------------------------------------------
-		// Own Interface
-		//-----------------------------------------------------------------------------------------
-		
-		//! Inputs
-		Ptr<NodeBool> GetA() const;
-		void SetA(Ptr<NodeBool>);
-
-		Ptr<NodeBool> GetB() const;
-		void SetB(Ptr<NodeBool>);
+	public:
 
-		//-----------------------------------------------------------------------------------------
-		// Interface pattern
-		//-----------------------------------------------------------------------------------------
-		class Private;
-		Private* GetPrivate() const;
+		// Node Interface
+		virtual const FNodeType* GetType() const override { return &StaticType; }
+		static const FNodeType* GetStaticType() { return &StaticType; }
 
 	protected:
 
-		//! Forbidden. 
-		~NodeBoolAnd();
+		/** Forbidden. Manage with the Ptr<> template. */
+		~NodeBoolAnd() {}
 
 	private:
 
-		Private* m_pD;
+		static FNodeType StaticType;
 
 	};
 
diff --git a/Engine/Plugins/Mutable/Source/MutableTools/Public/MuT/NodeColour.h b/Engine/Plugins/Mutable/Source/MutableTools/Public/MuT/NodeColour.h
--- a/Engine/Plugins/Mutable/Source/MutableTools/Public/MuT/NodeColour.h
+++ b/Engine/Plugins/Mutable/Source/MutableTools/Public/MuT/NodeColour.h
@@ -23,7 +23,7 @@ namespace mu
 
     protected:
 
-        //! Forbidden. Manage with the Ptr<> template.
+        /** Forbidden. Manage with the Ptr<> template. */
         inline ~NodeColour() {}
 
 	private:
diff --git a/Engine/Plugins/Mutable/Source/MutableTools/Public/MuT/NodeColourArithmeticOperation.h b/Engine/Plugins/Mutable/Source/MutableTools/Public/MuT/NodeColourArithmeticOperation.h
--- a/Engine/Plugins/Mutable/Source/MutableTools/Public/MuT/NodeColourArithmeticOperation.h
+++ b/Engine/Plugins/Mutable/Source/MutableTools/Public/MuT/NodeColourArithmeticOperation.h
@@ -37,7 +37,7 @@ namespace mu
 
 	protected:
 
-		//! Forbidden. Manage with the Ptr<> template.
+		/** Forbidden. Manage with the Ptr<> template. */
 		~NodeColourArithmeticOperation() {}
 
 	private:
diff --git a/Engine/Plugins/Mutable/Source/MutableTools/Public/MuT/NodeComponentEdit.h b/Engine/Plugins/Mutable/Source/MutableTools/Public/MuT/NodeComponentEdit.h
--- a/Engine/Plugins/Mutable/Source/MutableTools/Public/MuT/NodeComponentEdit.h
+++ b/Engine/Plugins/Mutable/Source/MutableTools/Public/MuT/NodeComponentEdit.h
@@ -7,50 +7,37 @@
 #include "MuT/Node.h"
 #include "MuT/NodeComponent.h"
 
-
 namespace mu
 {
 
-	//! This node modifies a node of the parent object of the object that this node belongs to.
-	//! It allows to extend, cut and morph the parent component's meshes.
-	//! It also allows to patch the parent component's textures.
+	/** This node modifies a node of the parent object of the object that this node belongs to.
+	* It allows to extend, cut and morph the parent component's meshes.
+	* It also allows to patch the parent component's textures.
+	*/
 	class MUTABLETOOLS_API NodeComponentEdit : public NodeComponent
 	{
 	public:
 
-		//-----------------------------------------------------------------------------------------
-        // Node interface
-		//-----------------------------------------------------------------------------------------
+		Ptr<NodeComponent> Parent;
+
+	public:
 
+		// Node interface
 		virtual const FNodeType* GetType() const override { return GetStaticType(); }
 		static const FNodeType* GetStaticType() { return &StaticType; }
 
-        //-----------------------------------------------------------------------------------------
         // NodeComponent interface
-        //-----------------------------------------------------------------------------------------
 		virtual const class NodeComponentNew* GetParentComponentNew() const override;
 
-		//-----------------------------------------------------------------------------------------
-        // Own interface
-		//-----------------------------------------------------------------------------------------
-
-		//! Set the parent component to modify. It should be a node from the parent objects of this node's object.
-        void SetParent( NodeComponent* );
-        NodeComponent* GetParent() const;
-
 	protected:
 
-		//! Forbidden. Manage with the Ptr<> template.
+		/** Forbidden. Manage with the Ptr<> template. */
 		~NodeComponentEdit() {}
 
 	private:
 
 		static FNodeType StaticType;
 
-		Ptr<NodeComponent> m_pParent;
-
 	};
 
-
-
 }
diff --git a/Engine/Plugins/Mutable/Source/MutableTools/Public/MuT/NodeImage.h b/Engine/Plugins/Mutable/Source/MutableTools/Public/MuT/NodeImage.h
--- a/Engine/Plugins/Mutable/Source/MutableTools/Public/MuT/NodeImage.h
+++ b/Engine/Plugins/Mutable/Source/MutableTools/Public/MuT/NodeImage.h
@@ -116,25 +116,19 @@ namespace mu
 	{
 	public:
 
-		//-----------------------------------------------------------------------------------------
 		// Node Interface
-		//-----------------------------------------------------------------------------------------
-
-		const FNodeType* GetType() const override;
-		static const FNodeType* GetStaticType();
-
-		//-----------------------------------------------------------------------------------------
-		// Interface pattern
-		//-----------------------------------------------------------------------------------------
-		class Private;
+		virtual const FNodeType* GetType() const override { return &StaticType; }
+		static const FNodeType* GetStaticType() { return &StaticType; }
 
 	protected:
 
-		//! Forbidden. Manage with the Ptr<> template.
+		/** Forbidden. Manage with the Ptr<> template. */
 		inline ~NodeImage() {}
 
-	};
+	private:
 
+		static FNodeType StaticType;
 
+	};
 
 }
diff --git a/Engine/Plugins/Mutable/Source/MutableTools/Public/MuT/NodeImageBinarise.h b/Engine/Plugins/Mutable/Source/MutableTools/Public/MuT/NodeImageBinarise.h
--- a/Engine/Plugins/Mutable/Source/MutableTools/Public/MuT/NodeImageBinarise.h
+++ b/Engine/Plugins/Mutable/Source/MutableTools/Public/MuT/NodeImageBinarise.h
@@ -6,64 +6,34 @@
 #include "MuR/RefCounted.h"
 #include "MuT/Node.h"
 #include "MuT/NodeImage.h"
-
+#include "MuT/NodeScalar.h"
 
 namespace mu
 {
 
-	// Forward definitions
-	class NodeScalar;
-	typedef Ptr<NodeScalar> NodeScalarPtr;
-	typedef Ptr<const NodeScalar> NodeScalarPtrConst;
-
-	class NodeImageBinarise;
-	typedef Ptr<NodeImageBinarise> NodeImageBinarisePtr;
-	typedef Ptr<const NodeImageBinarise> NodeImageBinarisePtrConst;
-
-
-	//! Node that multiplies the colors of an image, channel by channel.
-	//! \ingroup model
+	/** Node that multiplies the colors of an image, channel by channel. */
 	class MUTABLETOOLS_API NodeImageBinarise : public NodeImage
 	{
 	public:
 
-		NodeImageBinarise();
-
-		//-----------------------------------------------------------------------------------------
-		// Node Interface
-		//-----------------------------------------------------------------------------------------
-
-        const FNodeType* GetType() const override;
-		static const FNodeType* GetStaticType();
-
-		//-----------------------------------------------------------------------------------------
-		// Own Interface
-		//-----------------------------------------------------------------------------------------
-
-		//! Base image to multiply.
-		NodeImagePtr GetBase() const;
-		void SetBase( NodeImagePtr );
+		Ptr<NodeImage> Base;
+		Ptr<NodeScalar> Threshold;
 
-		//! Colour to multiply by the image.
-		NodeScalarPtr GetThreshold() const;
-		void SetThreshold( NodeScalarPtr );
+	public:
 
-		//-----------------------------------------------------------------------------------------
-		// Interface pattern
-		//-----------------------------------------------------------------------------------------
-		class Private;
-		Private* GetPrivate() const;
+		// Node Interface
+		virtual const FNodeType* GetType() const override { return &StaticType; }
+		static const FNodeType* GetStaticType() { return &StaticType; }
 
 	protected:
 
-		//! Forbidden. Manage with the Ptr<> template.
-		~NodeImageBinarise();
+		/** Forbidden. Manage with the Ptr<> template. */
+		~NodeImageBinarise() {}
 
 	private:
 
-		Private* m_pD;
+		static FNodeType StaticType;
 
 	};
 
-
 }
diff --git a/Engine/Plugins/Mutable/Source/MutableTools/Public/MuT/NodeImageColourMap.h b/Engine/Plugins/Mutable/Source/MutableTools/Public/MuT/NodeImageColourMap.h
--- a/Engine/Plugins/Mutable/Source/MutableTools/Public/MuT/NodeImageColourMap.h
+++ b/Engine/Plugins/Mutable/Source/MutableTools/Public/MuT/NodeImageColourMap.h
@@ -7,66 +7,35 @@
 #include "MuT/Node.h"
 #include "MuT/NodeImage.h"
 
-
 namespace mu
 {
 
-	// Forward definitions
-	class NodeImageColourMap;
-	typedef Ptr<NodeImageColourMap> NodeImageColourMapPtr;
-	typedef Ptr<const NodeImageColourMap> NodeImageColourMapPtrConst;
-
-
-	//! This node changes the colours of a selectd part of the image, applying a colour map from
-	//! conteined in another image.
-	//! \ingroup model
+	/** This node changes the colours of a selectd part of the image, applying a colour map from
+	* conteined in another image.
+	*/
 	class MUTABLETOOLS_API NodeImageColourMap : public NodeImage
 	{
 	public:
 
-		NodeImageColourMap();
+		Ptr<NodeImage> Base;
+		Ptr<NodeImage> Mask;
+		Ptr<NodeImage> Map;
 
+	public:
 
-		//-----------------------------------------------------------------------------------------
 		// Node Interface
-		//-----------------------------------------------------------------------------------------
-
-        const FNodeType* GetType() const override;
-		static const FNodeType* GetStaticType();
-
-		//-----------------------------------------------------------------------------------------
-		// Own Interface
-		//-----------------------------------------------------------------------------------------
-
-		//! Get the node generating the base image that will have the other one blended on top.
-		NodeImagePtr GetBase() const;
-		void SetBase( NodeImagePtr );
-
-		//! Get the node generating the mask image controlling the weight of the blend effect.
-		//! \todo: make it optional
-		NodeImagePtr GetMask() const;
-		void SetMask( NodeImagePtr );
-
-		//! Get the node generating the image to blend on the base.
-		NodeImagePtr GetMap() const;
-		void SetMap( NodeImagePtr );
-
-		//-----------------------------------------------------------------------------------------
-		// Interface pattern
-		//-----------------------------------------------------------------------------------------
-		class Private;
-		Private* GetPrivate() const;
+		virtual const FNodeType* GetType() const override { return &StaticType; }
+		static const FNodeType* GetStaticType() { return &StaticType; }
 
 	protected:
 
-		//! Forbidden. Manage with the Ptr<> template.
-		~NodeImageColourMap();
+		/** Forbidden. Manage with the Ptr<> template. */
+		~NodeImageColourMap() {}
 
 	private:
 
-		Private* m_pD;
+		static FNodeType StaticType;
 
 	};
 
-
 }
diff --git a/Engine/Plugins/Mutable/Source/MutableTools/Public/MuT/NodeImageConditional.h b/Engine/Plugins/Mutable/Source/MutableTools/Public/MuT/NodeImageConditional.h
--- a/Engine/Plugins/Mutable/Source/MutableTools/Public/MuT/NodeImageConditional.h
+++ b/Engine/Plugins/Mutable/Source/MutableTools/Public/MuT/NodeImageConditional.h
@@ -6,65 +6,35 @@
 #include "MuR/RefCounted.h"
 #include "MuT/Node.h"
 #include "MuT/NodeImage.h"
-
+#include "MuT/NodeBool.h"
 
 namespace mu
 {
 
-	// Forward definitions
-    class NodeBool;
-    class NodeImageConditional;
-    typedef Ptr<NodeImageConditional> NodeImageConditionalPtr;
-    typedef Ptr<const NodeImageConditional> NodeImageConditionalPtrConst;
-
-
-	//! This node selects an output image from a set of input images based on a parameter.
-	//! \ingroup model
+	/** This node selects an output image from a set of input images based on a parameter. */
     class MUTABLETOOLS_API NodeImageConditional : public NodeImage
 	{
 	public:
 
-        NodeImageConditional();
-
-		//-----------------------------------------------------------------------------------------
-		// Node Interface
-		//-----------------------------------------------------------------------------------------
-
-        const FNodeType* GetType() const override;
-		static const FNodeType* GetStaticType();
-
-		//-----------------------------------------------------------------------------------------
-		// Own Interface
-		//-----------------------------------------------------------------------------------------
-
-		//! Get the node generating the parameter used to select the option.
-        NodeBool* GetParameter() const;
-        void SetParameter( NodeBool* );
-
-        //! Get the node generating the option image when the parameters is true.
-        NodeImage* GetOptionTrue() const;
-        void SetOptionTrue( NodeImage* );
-
-        NodeImage* GetOptionFalse() const;
-        void SetOptionFalse( NodeImage* );
+		Ptr<NodeBool> Parameter;
+		Ptr<NodeImage> True;
+		Ptr<NodeImage> False;
 
+	public:
 
-		//-----------------------------------------------------------------------------------------
-		// Interface pattern
-		//-----------------------------------------------------------------------------------------
-		class Private;
-		Private* GetPrivate() const;
+		// Node Interface
+		virtual const FNodeType* GetType() const override { return &StaticType; }
+		static const FNodeType* GetStaticType() { return &StaticType; }
 
 	protected:
 
-		//! Forbidden. Manage with the Ptr<> template.
-        ~NodeImageConditional();
+		/** Forbidden. Manage with the Ptr<> template. */
+		~NodeImageConditional() {}
 
 	private:
 
-		Private* m_pD;
+		static FNodeType StaticType;
 
 	};
 
-
 }
diff --git a/Engine/Plugins/Mutable/Source/MutableTools/Public/MuT/NodeImageFormat.h b/Engine/Plugins/Mutable/Source/MutableTools/Public/MuT/NodeImageFormat.h
--- a/Engine/Plugins/Mutable/Source/MutableTools/Public/MuT/NodeImageFormat.h
+++ b/Engine/Plugins/Mutable/Source/MutableTools/Public/MuT/NodeImageFormat.h
@@ -12,56 +12,29 @@
 namespace mu
 {
 
-	// Forward definitions
-	class NodeImageFormat;
-	typedef Ptr<NodeImageFormat> NodeImageFormatPtr;
-	typedef Ptr<const NodeImageFormat> NodeImageFormatPtrConst;
-
-
-	//! Node that composes a new image by gathering pixel data from channels in other images.
-	//! \ingroup model
+	/** Node that composes a new image by gathering pixel data from channels in other images. */
 	class MUTABLETOOLS_API NodeImageFormat : public NodeImage
 	{
 	public:
 
-		NodeImageFormat();
+		EImageFormat Format = EImageFormat::None;
+		EImageFormat FormatIfAlpha = EImageFormat::None;
+		Ptr<NodeImage> Source;
 
+	public:
 
-		//-----------------------------------------------------------------------------------------
 		// Node Interface
-		//-----------------------------------------------------------------------------------------
-
-		const FNodeType* GetType() const override;
-		static const FNodeType* GetStaticType();
-
-		//-----------------------------------------------------------------------------------------
-		// Own Interface
-		//-----------------------------------------------------------------------------------------
-
-		//!
-		EImageFormat GetFormat() const;
-        void SetFormat(EImageFormat format, EImageFormat formatIfAlpha = EImageFormat::None );
-
-		//!
-		NodeImagePtr GetSource() const;
-
-		void SetSource( NodeImagePtr );
-
-
-		//-----------------------------------------------------------------------------------------
-		// Interface pattern
-		//-----------------------------------------------------------------------------------------
-		class Private;
-		Private* GetPrivate() const;
+		virtual const FNodeType* GetType() const override { return &StaticType; }
+		static const FNodeType* GetStaticType() { return &StaticType; }
 
 	protected:
 
-		//! Forbidden. Manage with the Ptr<> template.
-		~NodeImageFormat();
+		/** Forbidden. Manage with the Ptr<> template. */
+		~NodeImageFormat() {}
 
 	private:
 
-		Private* m_pD;
+		static FNodeType StaticType;
 
 	};
 
diff --git a/Engine/Plugins/Mutable/Source/MutableTools/Public/MuT/NodeImageInterpolate.h b/Engine/Plugins/Mutable/Source/MutableTools/Public/MuT/NodeImageInterpolate.h
--- a/Engine/Plugins/Mutable/Source/MutableTools/Public/MuT/NodeImageInterpolate.h
+++ b/Engine/Plugins/Mutable/Source/MutableTools/Public/MuT/NodeImageInterpolate.h
@@ -6,76 +6,40 @@
 #include "MuR/RefCounted.h"
 #include "MuT/Node.h"
 #include "MuT/NodeImage.h"
-
+#include "MuT/NodeScalar.h"
 
 namespace mu
 {
 
-	// Forward definitions
-	class NodeScalar;
-	typedef Ptr<NodeScalar> NodeScalarPtr;
-	typedef Ptr<const NodeScalar> NodeScalarPtrConst;
-
-	class NodeImageInterpolate;
-	typedef Ptr<NodeImageInterpolate> NodeImageInterpolatePtr;
-	typedef Ptr<const NodeImageInterpolate> NodeImageInterpolatePtrConst;
-
-
-	//! Node that interpolates linearly between several images based on a weight.
-	//! The different input images are uniformly distributed in the 0 to 1 range:
-	//! - if two images are set, the first one is 0.0 and the second one is 1.0
-	//! - if three images are set, the first one is 0.0 the second one is 0.5 and the third one 1.0
-	//! regardless of the input slots used: If B and C are set, but not A, B will be at weight 0
-	//! and C will be at weight 1
-	//! \ingroup model
+	/** Node that interpolates linearly between several images based on a weight.
+	* The different input images are uniformly distributed in the 0 to 1 range:
+	* - if two images are set, the first one is 0.0 and the second one is 1.0
+	* - if three images are set, the first one is 0.0 the second one is 0.5 and the third one 1.0
+	* regardless of the input slots used: If B and C are set, but not A, B will be at weight 0
+	* and C will be at weight 1
+	*/
 	class MUTABLETOOLS_API NodeImageInterpolate : public NodeImage
 	{
 	public:
 
-		NodeImageInterpolate();
+		Ptr<NodeScalar> Factor;
+		TArray<Ptr<NodeImage>> Targets;
 
+	public:
 
-		//-----------------------------------------------------------------------------------------
 		// Node Interface
-		//-----------------------------------------------------------------------------------------
-
-        const FNodeType* GetType() const override;
-		static const FNodeType* GetStaticType();
-
-		//-----------------------------------------------------------------------------------------
-		// Own Interface
-		//-----------------------------------------------------------------------------------------
-
-		//! Get the node generating the weight of the interpolation.
-		NodeScalarPtr GetFactor() const;
-		void SetFactor( NodeScalarPtr );
-
-		//! Set the number of target images. It will keep the currently set targets and initialise
-		//! the new ones as null.
-		void SetTargetCount( int );
-		int GetTargetCount() const;
-
-		//! Get the node generating the t-th target image of the interpolation.
-		NodeImagePtr GetTarget( int t ) const;
-		void SetTarget( int t, NodeImagePtr );
-
-
-		//-----------------------------------------------------------------------------------------
-		// Interface pattern
-		//-----------------------------------------------------------------------------------------
-		class Private;
-		Private* GetPrivate() const;
+		virtual const FNodeType* GetType() const override { return &StaticType; }
+		static const FNodeType* GetStaticType() { return &StaticType; }
 
 	protected:
 
-		//! Forbidden. Manage with the Ptr<> template.
-		~NodeImageInterpolate();
+		/** Forbidden. Manage with the Ptr<> template. */
+		~NodeImageInterpolate() {}
 
 	private:
 
-		Private* m_pD;
+		static FNodeType StaticType;
 
 	};
 
-
 }
diff --git a/Engine/Plugins/Mutable/Source/MutableTools/Public/MuT/NodeImageInvert.h b/Engine/Plugins/Mutable/Source/MutableTools/Public/MuT/NodeImageInvert.h
--- a/Engine/Plugins/Mutable/Source/MutableTools/Public/MuT/NodeImageInvert.h
+++ b/Engine/Plugins/Mutable/Source/MutableTools/Public/MuT/NodeImageInvert.h
@@ -7,51 +7,31 @@
 #include "MuT/Node.h"
 #include "MuT/NodeImage.h"
 
-
 namespace mu
 {
 
-	// Forward definitions
-	class NodeImageInvert;
-	typedef Ptr<NodeImageInvert> NodeImageInvertPtr;
-	typedef Ptr<const NodeImageInvert> NodeImageInvertPtrConst;
-
-	//! Node that inverts the colors of an image, channel by channel
+	/** Node that inverts the colors of an image, channel by channel. */
 	class MUTABLETOOLS_API NodeImageInvert : public NodeImage
 	{
 	public:
 
-		NodeImageInvert();
+		Ptr<NodeImage> Base;
 
+	public:
 
-		//-----------------------------------------------------------------------------------------
 		// Node Interface
-		//-----------------------------------------------------------------------------------------
-
-		const FNodeType* GetType() const override;
-		static const FNodeType* GetStaticType();
-
-		//-----------------------------------------------------------------------------------------
-		// Own Interface
-		//-----------------------------------------------------------------------------------------
-
-		//! Base image to invert.
-		NodeImagePtr GetBase() const;
-		void SetBase(NodeImagePtr);
-
-		//-----------------------------------------------------------------------------------------
-		// Interface pattern
-		//-----------------------------------------------------------------------------------------
-		class Private;
-		Private* GetPrivate() const;
+		virtual const FNodeType* GetType() const override { return &StaticType; }
+		static const FNodeType* GetStaticType() { return &StaticType; }
 
 	protected:
 
-		//! Forbidden. Manage with the Ptr<> template
-		~NodeImageInvert();
+		/** Forbidden. Manage with the Ptr<> template. */
+		~NodeImageInvert() {}
 
 	private:
 
-		Private* m_pD;
+		static FNodeType StaticType;
+
 	};
+
 }
diff --git a/Engine/Plugins/Mutable/Source/MutableTools/Public/MuT/NodeImageLayer.h b/Engine/Plugins/Mutable/Source/MutableTools/Public/MuT/NodeImageLayer.h
--- a/Engine/Plugins/Mutable/Source/MutableTools/Public/MuT/NodeImageLayer.h
+++ b/Engine/Plugins/Mutable/Source/MutableTools/Public/MuT/NodeImageLayer.h
@@ -12,61 +12,30 @@
 namespace mu
 {
 
-	// Forward definitions
-	class NodeImageLayer;
-	typedef Ptr<NodeImageLayer> NodeImageLayerPtr;
-	typedef Ptr<const NodeImageLayer> NodeImageLayerPtrConst;
-
-
-	//! This node applies a layer blending effect on a base image using a mask and a blended image.
-	//! \ingroup model
+	/** This node applies a layer blending effect on a base image using a mask and a blended image. */
 	class MUTABLETOOLS_API NodeImageLayer : public NodeImage
 	{
 	public:
 
-		NodeImageLayer();
-
-		//-----------------------------------------------------------------------------------------
-		// Node Interface
-		//-----------------------------------------------------------------------------------------
-
-        const FNodeType* GetType() const override;
-		static const FNodeType* GetStaticType();
-
-		//-----------------------------------------------------------------------------------------
-		// Own Interface
-		//-----------------------------------------------------------------------------------------
+		Ptr<NodeImage> Base;
+		Ptr<NodeImage> Mask;
+		Ptr<NodeImage> Blended;
+		EBlendType Type;
 
-		//! Get the node generating the base image that will have the blendeing effect applied.
-		NodeImagePtr GetBase() const;
-		void SetBase( NodeImagePtr );
-
-		//! Get the node generating the mask image controlling the weight of the effect.
-		NodeImagePtr GetMask() const;
-		void SetMask( NodeImagePtr );
-
-		//! Get the image blended on top of the base.
-		NodeImagePtr GetBlended() const;
-		void SetBlended( NodeImagePtr );
-
-		//!
-		EBlendType GetBlendType() const;
-		void SetBlendType(EBlendType);
+	public:
 
-		//-----------------------------------------------------------------------------------------
-		// Interface pattern
-		//-----------------------------------------------------------------------------------------
-		class Private;
-		Private* GetPrivate() const;
+		// Node Interface
+		virtual const FNodeType* GetType() const override { return &StaticType; }
+		static const FNodeType* GetStaticType() { return &StaticType; }
 
 	protected:
 
-		//! Forbidden. Manage with the Ptr<> template.
-		~NodeImageLayer();
+		/** Forbidden. Manage with the Ptr<> template. */
+		~NodeImageLayer() {}
 
 	private:
 
-		Private* m_pD;
+		static FNodeType StaticType;
 
 	};
 
diff --git a/Engine/Plugins/Mutable/Source/MutableTools/Public/MuT/NodeImageLayerColour.h b/Engine/Plugins/Mutable/Source/MutableTools/Public/MuT/NodeImageLayerColour.h
--- a/Engine/Plugins/Mutable/Source/MutableTools/Public/MuT/NodeImageLayerColour.h
+++ b/Engine/Plugins/Mutable/Source/MutableTools/Public/MuT/NodeImageLayerColour.h
@@ -7,73 +7,36 @@
 #include "MuR/RefCounted.h"
 #include "MuT/Node.h"
 #include "MuT/NodeImage.h"
+#include "MuT/NodeColour.h"
 
 namespace mu
 {
 
-	// Forward definitions
-	class NodeColour;
-	typedef Ptr<NodeColour> NodeColourPtr;
-	typedef Ptr<const NodeColour> NodeColourPtrConst;
-
-	class NodeImageLayerColour;
-	typedef Ptr<NodeImageLayerColour> NodeImageLayerColourPtr;
-	typedef Ptr<const NodeImageLayerColour> NodeImageLayerColourPtrConst;
-
-
-	//! This node applies a layer blending effect on a base image using a mask and a colour.
-	//! \ingroup model
+	/** This node applies a layer blending effect on a base image using a mask and a colour. */
 	class MUTABLETOOLS_API NodeImageLayerColour : public NodeImage
 	{
 	public:
 
-		NodeImageLayerColour();
-
-		//-----------------------------------------------------------------------------------------
-		// Node Interface
-		//-----------------------------------------------------------------------------------------
-
-        const FNodeType* GetType() const override;
-		static const FNodeType* GetStaticType();
-
-		//-----------------------------------------------------------------------------------------
-		// Own Interface
-		//-----------------------------------------------------------------------------------------
+		Ptr<NodeImage> Base;
+		Ptr<NodeImage> Mask;
+		Ptr<NodeColour> Colour;
+		EBlendType Type;
 
-		//! Get the node generating the base image that will have the blending effect applied.
-		NodeImagePtr GetBase() const;
-		void SetBase( NodeImagePtr );
-
-		//! Get the node generating the mask image controlling the weight of the effect.
-		NodeImagePtr GetMask() const;
-		void SetMask( NodeImagePtr );
-
-		//! Get the node generating the color to blend on the base.
-		NodeColourPtr GetColour() const;
-		void SetColour( NodeColourPtr );
-
-		//!
-		EBlendType GetBlendType() const;
-
-		//!
-		void SetBlendType(EBlendType);
+	public:
 
-		//-----------------------------------------------------------------------------------------
-		// Interface pattern
-		//-----------------------------------------------------------------------------------------
-		class Private;
-		Private* GetPrivate() const;
+		// Node Interface
+		virtual const FNodeType* GetType() const override { return &StaticType; }
+		static const FNodeType* GetStaticType() { return &StaticType; }
 
 	protected:
 
-		//! Forbidden. Manage with the Ptr<> template.
-		~NodeImageLayerColour();
+		/** Forbidden. Manage with the Ptr<> template. */
+		~NodeImageLayerColour() {}
 
 	private:
 
-		Private* m_pD;
+		static FNodeType StaticType;
 
 	};
 
-
 }
diff --git a/Engine/Plugins/Mutable/Source/MutableTools/Public/MuT/NodeImageLuminance.h b/Engine/Plugins/Mutable/Source/MutableTools/Public/MuT/NodeImageLuminance.h
--- a/Engine/Plugins/Mutable/Source/MutableTools/Public/MuT/NodeImageLuminance.h
+++ b/Engine/Plugins/Mutable/Source/MutableTools/Public/MuT/NodeImageLuminance.h
@@ -7,55 +7,31 @@
 #include "MuT/Node.h"
 #include "MuT/NodeImage.h"
 
-
 namespace mu
 {
 
-	// Forward definitions
-	class NodeImageLuminance;
-	typedef Ptr<NodeImageLuminance> NodeImageLuminancePtr;
-	typedef Ptr<const NodeImageLuminance> NodeImageLuminancePtrConst;
-
-
-	//! Calculate the luminance of an image into a new single-channel image..
-	//! \ingroup model
+	/** Calculate the luminance of an image into a new single - channel image. */
 	class MUTABLETOOLS_API NodeImageLuminance : public NodeImage
 	{
 	public:
 
-		NodeImageLuminance();
-
-		//-----------------------------------------------------------------------------------------
-		// Node Interface
-		//-----------------------------------------------------------------------------------------
-
-        const FNodeType* GetType() const override;
-		static const FNodeType* GetStaticType();
-
-		//-----------------------------------------------------------------------------------------
-		// Own Interface
-		//-----------------------------------------------------------------------------------------
+		Ptr<NodeImage> Source;
 
-		//! Get the node generating the source image calculate the luminance from.
-		NodeImagePtr GetSource() const;
-		void SetSource( NodeImagePtr );
+	public:
 
-		//-----------------------------------------------------------------------------------------
-		// Interface pattern
-		//-----------------------------------------------------------------------------------------
-		class Private;
-		Private* GetPrivate() const;
+		// Node Interface
+		virtual const FNodeType* GetType() const override { return &StaticType; }
+		static const FNodeType* GetStaticType() { return &StaticType; }
 
 	protected:
 
-		//! Forbidden. Manage with the Ptr<> template.
-		~NodeImageLuminance();
+		/** Forbidden. Manage with the Ptr<> template. */
+		~NodeImageLuminance() {}
 
 	private:
 
-		Private* m_pD;
+		static FNodeType StaticType;
 
 	};
 
-
 }
diff --git a/Engine/Plugins/Mutable/Source/MutableTools/Public/MuT/NodeImageMipmap.h b/Engine/Plugins/Mutable/Source/MutableTools/Public/MuT/NodeImageMipmap.h
--- a/Engine/Plugins/Mutable/Source/MutableTools/Public/MuT/NodeImageMipmap.h
+++ b/Engine/Plugins/Mutable/Source/MutableTools/Public/MuT/NodeImageMipmap.h
@@ -7,62 +7,36 @@
 #include "MuR/RefCounted.h"
 #include "MuT/Node.h"
 #include "MuT/NodeImage.h"
-
+#include "MuT/NodeScalar.h"
 
 namespace mu
 {
 
-	// Forward definitions
-	class NodeScalar;
-	typedef Ptr<NodeScalar> NodeScalarPtr;
-	typedef Ptr<const NodeScalar> NodeScalarPtrConst;
-
-    class NodeImageMipmap;
-    typedef Ptr<NodeImageMipmap> NodeImageMipmapPtr;
-    typedef Ptr<const NodeImageMipmap> NodeImageMipmapPtrConst;
-
-
-	//! Generate mimaps for an image.
-	//! \ingroup model
+	/** Generate mimaps for an image. */
     class MUTABLETOOLS_API NodeImageMipmap : public NodeImage
 	{
 	public:
 
-        NodeImageMipmap();
-
-		//-----------------------------------------------------------------------------------------
-		// Node Interface
-		//-----------------------------------------------------------------------------------------
+		Ptr<NodeImage> Source;
+		Ptr<NodeScalar> Factor;
 
-        const FNodeType* GetType() const override;
-		static const FNodeType* GetStaticType();
+		FMipmapGenerationSettings Settings;
 
-		//-----------------------------------------------------------------------------------------
-		// Own Interface
-		//-----------------------------------------------------------------------------------------
-
-        //! Get the node generating the source image to be mipmapped.
-		NodeImagePtr GetSource() const;
-		void SetSource( NodeImagePtr );
-
-		void SetMipmapGenerationSettings(EMipmapFilterType FilterType, EAddressMode AddressMode);
+	public:
 
-		//-----------------------------------------------------------------------------------------
-		// Interface pattern
-		//-----------------------------------------------------------------------------------------
-		class Private;
-		Private* GetPrivate() const;
+		// Node Interface
+		virtual const FNodeType* GetType() const override { return &StaticType; }
+		static const FNodeType* GetStaticType() { return &StaticType; }
 
 	protected:
 
-		//! Forbidden. Manage with the Ptr<> template.
-        ~NodeImageMipmap();
+		/** Forbidden. Manage with the Ptr<> template. */
+		~NodeImageMipmap() {}
 
 	private:
 
-		Private* m_pD;
+		static FNodeType StaticType;
 
 	};
 
-
 }
diff --git a/Engine/Plugins/Mutable/Source/MutableTools/Public/MuT/NodeImageMultiLayer.h b/Engine/Plugins/Mutable/Source/MutableTools/Public/MuT/NodeImageMultiLayer.h
--- a/Engine/Plugins/Mutable/Source/MutableTools/Public/MuT/NodeImageMultiLayer.h
+++ b/Engine/Plugins/Mutable/Source/MutableTools/Public/MuT/NodeImageMultiLayer.h
@@ -9,72 +9,37 @@
 #include "MuT/NodeImage.h"
 #include "MuT/NodeRange.h"
 
-
 namespace mu
 {
 
-	// Forward definitions
-    class NodeImageMultiLayer;
-    typedef Ptr<NodeImageMultiLayer> NodeImageMultiLayerPtr;
-    typedef Ptr<const NodeImageMultiLayer> NodeImageMultiLayerPtrConst;
-
-
-    //! This node applies any numbre of layer blending effect on a base image using a mask and a
-    //! blended image. The number of layers depends on a scalar input of this node.
-	//! \ingroup model
+    /** This node applies any numbre of layer blending effect on a base image using a mask and a
+    * blended image. The number of layers depends on a scalar input of this node.
+	*/
     class MUTABLETOOLS_API NodeImageMultiLayer : public NodeImage
 	{
 	public:
 
-        NodeImageMultiLayer();
-
-		//-----------------------------------------------------------------------------------------
-		// Node Interface
-		//-----------------------------------------------------------------------------------------
-
-        const FNodeType* GetType() const override;
-		static const FNodeType* GetStaticType();
-
-		//-----------------------------------------------------------------------------------------
-		// Own Interface
-		//-----------------------------------------------------------------------------------------
+		Ptr<NodeImage> Base;
+		Ptr<NodeImage> Mask;
+		Ptr<NodeImage> Blended;
+		Ptr<NodeRange> Range;
+		EBlendType Type;
 
-		//! Get the node generating the base image that will have the blendeing effect applied.
-		NodeImagePtr GetBase() const;
-		void SetBase( NodeImagePtr );
-
-		//! Get the node generating the mask image controlling the weight of the effect.
-		NodeImagePtr GetMask() const;
-		void SetMask( NodeImagePtr );
-
-        //! Get the image blended on top of the base.
-        NodeImagePtr GetBlended() const;
-        void SetBlended( NodeImagePtr );
-
-        //! Get the number of layers to apply.
-        NodeRangePtr GetRange() const;
-        void SetRange( NodeRangePtr );
-
-		//!
-		EBlendType GetBlendType() const;
-        void SetBlendType(EBlendType);
+	public:
 
-		//-----------------------------------------------------------------------------------------
-		// Interface pattern
-		//-----------------------------------------------------------------------------------------
-		class Private;
-		Private* GetPrivate() const;
+		// Node Interface
+		virtual const FNodeType* GetType() const override { return &StaticType; }
+		static const FNodeType* GetStaticType() { return &StaticType; }
 
 	protected:
 
-		//! Forbidden. Manage with the Ptr<> template.
-        ~NodeImageMultiLayer();
+		/** Forbidden. Manage with the Ptr<> template. */
+		~NodeImageMultiLayer() {}
 
 	private:
 
-		Private* m_pD;
+		static FNodeType StaticType;
 
 	};
 
-
 }
diff --git a/Engine/Plugins/Mutable/Source/MutableTools/Public/MuT/NodeImageNormalComposite.h b/Engine/Plugins/Mutable/Source/MutableTools/Public/MuT/NodeImageNormalComposite.h
--- a/Engine/Plugins/Mutable/Source/MutableTools/Public/MuT/NodeImageNormalComposite.h
+++ b/Engine/Plugins/Mutable/Source/MutableTools/Public/MuT/NodeImageNormalComposite.h
@@ -11,55 +11,31 @@
 namespace mu
 {
 
-	// Forward definitions
-	class NodeImageNormalComposite;
-	typedef Ptr<NodeImageNormalComposite> NodeImageNormalCompositePtr;
-	typedef Ptr<const NodeImageNormalComposite> NodeImageNormalCompositePtrConst;
-
-	//! Node that inverts the colors of an image, channel by channel
+	/** Node that inverts the colors of an image, channel by channel. */
 	class MUTABLETOOLS_API NodeImageNormalComposite : public NodeImage
 	{
 	public:
 
-		NodeImageNormalComposite();
-
-		//-----------------------------------------------------------------------------------------
-		// Node Interface
-		//-----------------------------------------------------------------------------------------
-
-		const FNodeType* GetType() const override;
-		static const FNodeType* GetStaticType();
-
-		//-----------------------------------------------------------------------------------------
-		// Own Interface
-		//-----------------------------------------------------------------------------------------
-
-		NodeImagePtr GetBase() const;
-		void SetBase(NodeImagePtr);
+		Ptr<NodeImage> Base;
+		Ptr<NodeImage> Normal;
+		float Power = 1.0f;
+		ECompositeImageMode Mode;
 
-        NodeImagePtr GetNormal() const;
-        void SetNormal(NodeImagePtr);
-
-		float GetPower() const;
-		void SetPower(float); 
-
-        ECompositeImageMode GetMode() const;
-        void SetMode(ECompositeImageMode);
+	public:
 
-		//-----------------------------------------------------------------------------------------
-		// Interface pattern
-		//-----------------------------------------------------------------------------------------
-		class Private;
-		Private* GetPrivate() const;
+		// Node Interface
+		virtual const FNodeType* GetType() const override { return &StaticType; }
+		static const FNodeType* GetStaticType() { return &StaticType; }
 
 	protected:
 
-		//! Forbidden. Manage with the Ptr<> template
-		~NodeImageNormalComposite();
+		/** Forbidden. Manage with the Ptr<> template. */
+		~NodeImageNormalComposite() {}
 
 	private:
 
-		Private* m_pD;
+		static FNodeType StaticType;
+
 	};
 
 }
diff --git a/Engine/Plugins/Mutable/Source/MutableTools/Public/MuT/NodeImageParameter.h b/Engine/Plugins/Mutable/Source/MutableTools/Public/MuT/NodeImageParameter.h
--- a/Engine/Plugins/Mutable/Source/MutableTools/Public/MuT/NodeImageParameter.h
+++ b/Engine/Plugins/Mutable/Source/MutableTools/Public/MuT/NodeImageParameter.h
@@ -6,59 +6,42 @@
 #include "MuR/RefCounted.h"
 #include "MuT/Node.h"
 #include "MuT/NodeImage.h"
-
+#include "MuT/NodeRange.h"
 
 namespace mu
 {
-	class NodeRange;
 
-    //! Node that defines a Image model parameter.
-	//! \ingroup model
+    /** */
     class MUTABLETOOLS_API NodeImageParameter : public NodeImage
 	{
 	public:
 
-        NodeImageParameter();
-
-		//-----------------------------------------------------------------------------------------
-		// Node Interface
-		//-----------------------------------------------------------------------------------------
-
-		const FNodeType* GetType() const override;
-		static const FNodeType* GetStaticType();
+		/** Full asset name for the default value of the parameter. */
+		FName DefaultValue;
+		
+		/** Name of the parameter */
+		FString Name;
 
-		//-----------------------------------------------------------------------------------------
-		// Own Interface
-		//-----------------------------------------------------------------------------------------
+		/** User provided ID to identify the parameter. */
+		FString UID;
 
-		//! Set the name of the parameter.
-		void SetName(const FString&);
+		/** Ranges for the parameter in case it is multidimensional. */
+		TArray<Ptr<NodeRange>> Ranges;
 
-		//! Set the uid of the parameter.
-		void SetUid(const FString&);
-
-    	//! Get the default value of the parameter.
-    	void SetDefaultValue(FName Value);
-    	
-		//! Set the number of ranges (dimensions) for this parameter.
-		//! By default a parameter has 0 ranges, meaning it only has one value.
-		void SetRangeCount(int Index);
-		void SetRange(int Index, Ptr<NodeRange> Range);
+	public:
 
-		//-----------------------------------------------------------------------------------------
-		// Interface pattern
-		//-----------------------------------------------------------------------------------------
-		class Private;
-		Private* GetPrivate() const;
+		// Node Interface
+		virtual const FNodeType* GetType() const override { return &StaticType; }
+		static const FNodeType* GetStaticType() { return &StaticType; }
 
 	protected:
 
-		//! Forbidden. Manage with the Ptr<> template.
-        ~NodeImageParameter();
+		/** Forbidden. Manage with the Ptr<> template. */
+		~NodeImageParameter() {}
 
 	private:
 
-		Private* m_pD;
+		static FNodeType StaticType;
 
 	};
 
