[mutable] Initial boilerplate for the implementation of Mesh Parameters.

diff --git a/Engine/Plugins/Mutable/Source/MutableTools/Public/MuT/NodeImagePlainColour.h b/Engine/Plugins/Mutable/Source/MutableTools/Public/MuT/NodeImagePlainColour.h
--- a/Engine/Plugins/Mutable/Source/MutableTools/Public/MuT/NodeImagePlainColour.h
+++ b/Engine/Plugins/Mutable/Source/MutableTools/Public/MuT/NodeImagePlainColour.h
@@ -6,64 +6,37 @@
 #include "MuR/RefCounted.h"
 #include "MuT/Node.h"
 #include "MuT/NodeImage.h"
-
+#include "MuT/NodeColour.h"
+#include "MuR/Image.h"
 
 namespace mu
 {
 
-	// Forward definitions
-	class NodeColour;
-	typedef Ptr<NodeColour> NodeColourPtr;
-	typedef Ptr<const NodeColour> NodeColourPtrConst;
-
-	class NodeImagePlainColour;
-	typedef Ptr<NodeImagePlainColour> NodeImagePlainColourPtr;
-	typedef Ptr<const NodeImagePlainColour> NodeImagePlainColourPtrConst;
-
-
-	//! Node that multiplies the colors of an image, channel by channel.
-	//! \ingroup model
+	/** Node that multiplies the colors of an image, channel by channel. */
 	class MUTABLETOOLS_API NodeImagePlainColour : public NodeImage
 	{
 	public:
 
-		NodeImagePlainColour();
+		Ptr<NodeColour> Colour;
+		int32 SizeX = 4;
+		int32 SizeY = 4;
+		EImageFormat Format = EImageFormat::RGBA_UByte;
 
+	public:
 
-		//-----------------------------------------------------------------------------------------
 		// Node Interface
-		//-----------------------------------------------------------------------------------------
-
-        const FNodeType* GetType() const override;
-		static const FNodeType* GetStaticType();
-
-		//-----------------------------------------------------------------------------------------
-		// Own Interface
-		//-----------------------------------------------------------------------------------------
-
-		//! Colour of the image.
-		NodeColourPtr GetColour() const;
-		void SetColour( NodeColourPtr );
-
-		//! New size or relative factor
-		void SetSize( int x, int y );
-
-		//-----------------------------------------------------------------------------------------
-		// Interface pattern
-		//-----------------------------------------------------------------------------------------
-		class Private;
-		Private* GetPrivate() const;
+		virtual const FNodeType* GetType() const override { return &StaticType; }
+		static const FNodeType* GetStaticType() { return &StaticType; }
 
 	protected:
 
-		//! Forbidden. Manage with the Ptr<> template.
-		~NodeImagePlainColour();
+		/** Forbidden. Manage with the Ptr<> template. */
+		~NodeImagePlainColour() {}
 
 	private:
 
-		Private* m_pD;
+		static FNodeType StaticType;
 
 	};
 
-
 }
diff --git a/Engine/Plugins/Mutable/Source/MutableTools/Public/MuT/NodeImageProject.h b/Engine/Plugins/Mutable/Source/MutableTools/Public/MuT/NodeImageProject.h
--- a/Engine/Plugins/Mutable/Source/MutableTools/Public/MuT/NodeImageProject.h
+++ b/Engine/Plugins/Mutable/Source/MutableTools/Public/MuT/NodeImageProject.h
@@ -1,98 +1,53 @@
 // Copyright Epic Games, Inc. All Rights Reserved.
 
-
 #pragma once
 
 #include "MuR/Ptr.h"
 #include "MuR/RefCounted.h"
 #include "MuT/Node.h"
 #include "MuT/NodeImage.h"
+#include "MuT/NodeMesh.h"
+#include "MuT/NodeProjector.h"
+#include "MuT/NodeScalar.h"
 #include "MuR/Image.h"
 
 namespace mu
 {
 
-	// Forward definitions
-	class NodeScalar;
-	class NodeMesh;
-	class NodeProjector;
-	class NodeImageProject;
-
-	//!
-	//! \ingroup model
+	/** */
 	class MUTABLETOOLS_API NodeImageProject : public NodeImage
 	{
 	public:
 
-		NodeImageProject();
-
-		//-----------------------------------------------------------------------------------------
-		// Node Interface
-		//-----------------------------------------------------------------------------------------
-
-        const FNodeType* GetType() const override;
-		static const FNodeType* GetStaticType();
-
-		//-----------------------------------------------------------------------------------------
-		// Own Interface
-		//-----------------------------------------------------------------------------------------
-
-		//! Get the node generating the mesh to be used for the projection.
-		void SetMesh( Ptr<NodeMesh> );
-
-		//! Get the node generating the projector to be used for the projection.
-		void SetProjector( Ptr<NodeProjector> );
-
-		//! Get the node generating a mask to filter out the projected images. It is optional.
-		void SetTargetMask( Ptr<NodeImage> );
-
-		/** Eanble or disable the additional operations to correct texture UV seam artifact correction. Default is enabled.*/
-		void SetEnableSeamCorrection(bool bEnabled);
+		Ptr<NodeProjector> Projector;
+		Ptr<NodeMesh> Mesh;
+		Ptr<NodeScalar> AngleFadeStart;
+		Ptr<NodeScalar> AngleFadeEnd;
+		Ptr<NodeImage> Image;
+		Ptr<NodeImage> Mask;
+		FUintVector2 ImageSize;
+		uint8 Layout = 0;
+		bool bIsRGBFadingEnabled = true;
+		bool bIsAlphaFadingEnabled = true;
+		bool bEnableTextureSeamCorrection = true;
+		ESamplingMethod SamplingMethod = ESamplingMethod::Point;
+		EMinFilterMethod MinFilterMethod = EMinFilterMethod::None;
 
-		//! Set the angle-based fading behaviour for this projector. It must be set separately for 
-		//! RGB and Alpha channels of the images to project. By default it is enabled for all channels.
-		//! Disabling the angle fade will still use the AngleFadeEnd below to clamp the projection,
-		//! but it will ignore the start angle, so it will not apply a fading gradient.
-		void SetAngleFadeChannels( bool bFadeRGB, bool bFadeA );
-
-        //! Set the node generating the fading start angle. Only relevant if fading is enabled with 
-		void SetAngleFadeStart( Ptr<NodeScalar> );
-
-		//! Set the node generating the fading end angle
-		void SetAngleFadeEnd( Ptr<NodeScalar> );
-
-		//! Set sampling method.
-		void SetSamplingMethod(ESamplingMethod SamplingMethod);
-		
-		//! Set min filter method
-		void SetMinFilterMethod(EMinFilterMethod MinFilterMethod);
-
-        //! Get the node generating the image to project.
-        void SetImage( Ptr<NodeImage> );
-
-        //! UV layout of the mesh to use for the generated image. Defaults to 0.
-        void SetLayout( uint8  );
-
-		//! Set the size of the image to generate with the projection. If set to 0 (default) a size
-		//! that matches how this node is used will try to be guessed. 
-		void SetImageSize( const FUintVector2& size );
+	public:
 
-		//-----------------------------------------------------------------------------------------
-		// Interface pattern
-		//-----------------------------------------------------------------------------------------
-		class Private;
-		Private* GetPrivate() const;
+		// Node Interface
+		virtual const FNodeType* GetType() const override { return &StaticType; }
+		static const FNodeType* GetStaticType() { return &StaticType; }
 
 	protected:
 
-		//! Forbidden. Manage with the Ptr<> template.
-		~NodeImageProject();
+		/** Forbidden. Manage with the Ptr<> template. */
+		~NodeImageProject() {}
 
 	private:
 
-		Private* m_pD;
+		static FNodeType StaticType;
 
 	};
 
-
 }
diff --git a/Engine/Plugins/Mutable/Source/MutableTools/Public/MuT/NodeImageResize.h b/Engine/Plugins/Mutable/Source/MutableTools/Public/MuT/NodeImageResize.h
--- a/Engine/Plugins/Mutable/Source/MutableTools/Public/MuT/NodeImageResize.h
+++ b/Engine/Plugins/Mutable/Source/MutableTools/Public/MuT/NodeImageResize.h
@@ -7,66 +7,36 @@
 #include "MuT/Node.h"
 #include "MuT/NodeImage.h"
 
-
 namespace mu
 {
 
-	// Forward definitions
-	class NodeScalar;
-	typedef Ptr<NodeScalar> NodeScalarPtr;
-	typedef Ptr<const NodeScalar> NodeScalarPtrConst;
-
-	class NodeImageResize;
-	typedef Ptr<NodeImageResize> NodeImageResizePtr;
-	typedef Ptr<const NodeImageResize> NodeImageResizePtrConst;
-
-
-	//! Node that multiplies the colors of an image, channel by channel.
-	//! \ingroup model
+	/** Node that multiplies the colors of an image, channel by channel. */
 	class MUTABLETOOLS_API NodeImageResize : public NodeImage
 	{
 	public:
 
-		NodeImageResize();
-
-
-		//-----------------------------------------------------------------------------------------
-		// Node Interface
-		//-----------------------------------------------------------------------------------------
-
-        const FNodeType* GetType() const override;
-		static const FNodeType* GetStaticType();
-
-		//-----------------------------------------------------------------------------------------
-		// Own Interface
-		//-----------------------------------------------------------------------------------------
+		Ptr<NodeImage> Base;
 
-		//! Base image to resize.
-		NodeImagePtr GetBase() const;
-		void SetBase( NodeImagePtr );
+		bool bRelative = true;
 
-		//! Is the size a relative factor or an absolute size?
-		void SetRelative( bool );
+		float SizeX = 0.5f;
+		float SizeY = 0.5f;
 
-		//! New size or relative factor
-		void SetSize( float x, float y );
+	public:
 
-		//-----------------------------------------------------------------------------------------
-		// Interface pattern
-		//-----------------------------------------------------------------------------------------
-		class Private;
-		Private* GetPrivate() const;
+		// Node Interface
+		virtual const FNodeType* GetType() const override { return &StaticType; }
+		static const FNodeType* GetStaticType() { return &StaticType; }
 
 	protected:
 
-		//! Forbidden. Manage with the Ptr<> template.
-		~NodeImageResize();
+		/** Forbidden. Manage with the Ptr<> template. */
+		~NodeImageResize() {}
 
 	private:
 
-		Private* m_pD;
+		static FNodeType StaticType;
 
 	};
 
-
 }
diff --git a/Engine/Plugins/Mutable/Source/MutableTools/Public/MuT/NodeImageSaturate.h b/Engine/Plugins/Mutable/Source/MutableTools/Public/MuT/NodeImageSaturate.h
--- a/Engine/Plugins/Mutable/Source/MutableTools/Public/MuT/NodeImageSaturate.h
+++ b/Engine/Plugins/Mutable/Source/MutableTools/Public/MuT/NodeImageSaturate.h
@@ -6,68 +6,38 @@
 #include "MuR/RefCounted.h"
 #include "MuT/Node.h"
 #include "MuT/NodeImage.h"
-
+#include "MuT/NodeScalar.h"
 
 namespace mu
 {
 
-	// Forward definitions
-	class NodeScalar;
-	typedef Ptr<NodeScalar> NodeScalarPtr;
-	typedef Ptr<const NodeScalar> NodeScalarPtrConst;
-
-	class NodeImageSaturate;
-	typedef Ptr<NodeImageSaturate> NodeImageSaturatePtr;
-	typedef Ptr<const NodeImageSaturate> NodeImageSaturatePtrConst;
-
-
-	//! Change the saturation of an image. This node can be used to increase the saturation with a
-	//! factor bigger than 1, or to decrease it or desaturate it completely with a factor smaller
-	//! than 1 or 0.
+	/** Change the saturation of an image.This node can be used to increase the saturation with a
+	* factor bigger than 1, or to decrease it or desaturate it completely with a factor smaller
+	* than 1 or 0.
+	*/
 	class MUTABLETOOLS_API NodeImageSaturate : public NodeImage
 	{
 	public:
 
-		NodeImageSaturate();
-
-		//-----------------------------------------------------------------------------------------
-		// Node Interface
-		//-----------------------------------------------------------------------------------------
-
-        const FNodeType* GetType() const override;
-		static const FNodeType* GetStaticType();
-
-		//-----------------------------------------------------------------------------------------
-		// Own Interface
-		//-----------------------------------------------------------------------------------------
-
-		//! Get the node generating the saturation factor.
-		//! A value of 0 completely desaturates. 1 will leave the same saturation and bigger than 1
-		//! will increase it.
-		NodeScalarPtr GetFactor() const;
-		void SetFactor( NodeScalarPtr );
+		Ptr<NodeImage> Source;
 
-		//! Get the node generating the source image to be [de]saturated.
-		NodeImagePtr GetSource() const;
-		void SetSource( NodeImagePtr );
+		Ptr<NodeScalar> Factor;
 
-		//-----------------------------------------------------------------------------------------
-		// Interface pattern
-		//-----------------------------------------------------------------------------------------
-		class Private;
+	public:
 
-		Private* GetPrivate() const;
+		// Node Interface
+		virtual const FNodeType* GetType() const override { return &StaticType; }
+		static const FNodeType* GetStaticType() { return &StaticType; }
 
 	protected:
 
-		//! Forbidden. Manage with the Ptr<> template.
-		~NodeImageSaturate();
+		/** Forbidden. Manage with the Ptr<> template. */
+		~NodeImageSaturate() {}
 
 	private:
 
-		Private* m_pD;
+		static FNodeType StaticType;
 
 	};
 
-
 }
diff --git a/Engine/Plugins/Mutable/Source/MutableTools/Public/MuT/NodeImageSwitch.h b/Engine/Plugins/Mutable/Source/MutableTools/Public/MuT/NodeImageSwitch.h
--- a/Engine/Plugins/Mutable/Source/MutableTools/Public/MuT/NodeImageSwitch.h
+++ b/Engine/Plugins/Mutable/Source/MutableTools/Public/MuT/NodeImageSwitch.h
@@ -6,69 +6,35 @@
 #include "MuR/RefCounted.h"
 #include "MuT/Node.h"
 #include "MuT/NodeImage.h"
-
+#include "MuT/NodeScalar.h"
 
 namespace mu
 {
 
-	// Forward definitions
-	class NodeScalar;
-	typedef Ptr<NodeScalar> NodeScalarPtr;
-	typedef Ptr<const NodeScalar> NodeScalarPtrConst;
-
-	class NodeImageSwitch;
-	typedef Ptr<NodeImageSwitch> NodeImageSwitchPtr;
-	typedef Ptr<const NodeImageSwitch> NodeImageSwitchPtrConst;
-
-
 	//! This node selects an output image from a set of input images based on a parameter.
 	class MUTABLETOOLS_API NodeImageSwitch : public NodeImage
 	{
 	public:
 
-		NodeImageSwitch();
-
-		//-----------------------------------------------------------------------------------------
-		// Node Interface
-		//-----------------------------------------------------------------------------------------
-
-        const FNodeType* GetType() const override;
-		static const FNodeType* GetStaticType();
+		Ptr<NodeScalar> Parameter;
 
-		//-----------------------------------------------------------------------------------------
-		// Own Interface
-		//-----------------------------------------------------------------------------------------
-
-		//! Get the node generating the parameter used to select the option.
-		NodeScalarPtr GetParameter() const;
-		void SetParameter( NodeScalarPtr );
-
-		//! Set the number of option images. It will keep the currently set targets and initialise
-		//! the new ones as null.
-		void SetOptionCount( int );
-		int GetOptionCount() const;
-
-		//! Get the node generating the t-th option image.
-		NodeImagePtr GetOption( int t ) const;
-		void SetOption( int t, NodeImagePtr );
+		TArray<Ptr<NodeImage>> Options;
 
+	public:
 
-		//-----------------------------------------------------------------------------------------
-		// Interface pattern
-		//-----------------------------------------------------------------------------------------
-		class Private;
-		Private* GetPrivate() const;
+		// Node Interface
+		virtual const FNodeType* GetType() const override { return &StaticType; }
+		static const FNodeType* GetStaticType() { return &StaticType; }
 
 	protected:
 
-		//! Forbidden. Manage with the Ptr<> template.
-		~NodeImageSwitch();
+		/** Forbidden. Manage with the Ptr<> template. */
+		~NodeImageSwitch() {}
 
 	private:
 
-		Private* m_pD;
+		static FNodeType StaticType;
 
 	};
 
-
 }
diff --git a/Engine/Plugins/Mutable/Source/MutableTools/Public/MuT/NodeImageSwizzle.h b/Engine/Plugins/Mutable/Source/MutableTools/Public/MuT/NodeImageSwizzle.h
--- a/Engine/Plugins/Mutable/Source/MutableTools/Public/MuT/NodeImageSwizzle.h
+++ b/Engine/Plugins/Mutable/Source/MutableTools/Public/MuT/NodeImageSwizzle.h
@@ -8,62 +8,43 @@
 #include "MuT/Node.h"
 #include "MuT/NodeImage.h"
 
-
 namespace mu
 {
 
-	// Forward definitions
-	class NodeImageSwizzle;
-	typedef Ptr<NodeImageSwizzle> NodeImageSwizzlePtr;
-	typedef Ptr<const NodeImageSwizzle> NodeImageSwizzlePtrConst;
-
-
-	//! Node that composes a new image by gathering pixel data from channels in other images.
-	//! \ingroup model
+	/** Node that composes a new image by gathering pixel data from channels in other images. */
 	class MUTABLETOOLS_API NodeImageSwizzle : public NodeImage
 	{
 	public:
 
-		NodeImageSwizzle();
-
-		//-----------------------------------------------------------------------------------------
-		// Node Interface
-		//-----------------------------------------------------------------------------------------
-
-        const FNodeType* GetType() const override;
-		static const FNodeType* GetStaticType();
-
-		//-----------------------------------------------------------------------------------------
-		// Own Interface
-		//-----------------------------------------------------------------------------------------
+		EImageFormat NewFormat;
+		TArray<Ptr<NodeImage>> Sources;
+		TArray<int32> SourceChannels;
 
-		//!
-		EImageFormat GetFormat() const;
-		void SetFormat(EImageFormat);
+	public:
 
-		//!
-		NodeImagePtr GetSource( int32 ) const;
-		int GetSourceChannel( int32 ) const;
+		// Node Interface
+		virtual const FNodeType* GetType() const override { return &StaticType; }
+		static const FNodeType* GetStaticType() { return &StaticType; }
 
-		void SetSource( int32, NodeImagePtr );
-		void SetSourceChannel( int32 OutputChannel, int32 SourceChannel );
+		// Own interface 
+		inline void SetFormat(EImageFormat InFormat)
+		{
+			NewFormat = InFormat;
 
-		//-----------------------------------------------------------------------------------------
-		// Interface pattern
-		//-----------------------------------------------------------------------------------------
-		class Private;
-		Private* GetPrivate() const;
+			int32 ChannelCount = GetImageFormatData(InFormat).Channels;
+			Sources.SetNum(ChannelCount);
+			SourceChannels.SetNum(ChannelCount);
+		}
 
 	protected:
 
-		//! Forbidden. Manage with the Ptr<> template.
-		~NodeImageSwizzle();
+		/** Forbidden. Manage with the Ptr<> template. */
+		~NodeImageSwizzle() {}
 
 	private:
 
-		Private* m_pD;
+		static FNodeType StaticType;
 
 	};
 
-
 }
diff --git a/Engine/Plugins/Mutable/Source/MutableTools/Public/MuT/NodeImageTransform.h b/Engine/Plugins/Mutable/Source/MutableTools/Public/MuT/NodeImageTransform.h
--- a/Engine/Plugins/Mutable/Source/MutableTools/Public/MuT/NodeImageTransform.h
+++ b/Engine/Plugins/Mutable/Source/MutableTools/Public/MuT/NodeImageTransform.h
@@ -2,87 +2,48 @@
 
 #pragma once
 
-#include "MuR/Ptr.h"
-#include "MuR/RefCounted.h"
 #include "MuT/Node.h"
 #include "MuT/NodeImage.h"
+#include "MuT/NodeScalar.h"
+#include "MuR/ImageTypes.h"
+#include "MuR/Ptr.h"
 
 namespace mu
 {
 
-	// Forward definitions
-	class NodeScalar;
-	typedef Ptr<NodeScalar> NodeScalarPtr;
-	typedef Ptr<const NodeScalar> NodeScalarPtrConst;
-
-	class NodeImageTransform;
-	typedef Ptr<NodeImageTransform> NodeImageTransformPtr;
-	typedef Ptr<const NodeImageTransform> NodeImageTransformPtrConst;
-
-	enum class EAddressMode;
-
-	//! Node that multiplies the colors of an image, channel by channel.
-	//! \ingroup model
+	/**  */
 	class MUTABLETOOLS_API NodeImageTransform : public NodeImage
 	{
 	public:
 
-		NodeImageTransform();
-
-		//-----------------------------------------------------------------------------------------
-		// Node Interface
-		//-----------------------------------------------------------------------------------------
-
-        const FNodeType* GetType() const override;
-		static const FNodeType* GetStaticType();
-
-		//-----------------------------------------------------------------------------------------
-		// Own Interface
-		//-----------------------------------------------------------------------------------------
-
-		//! Base image to multiply.
-		NodeImagePtr GetBase() const;
-		void SetBase( NodeImagePtr );
-	
-		NodeScalarPtr GetOffsetX() const;
-		NodeScalarPtr GetOffsetY() const;
-		NodeScalarPtr GetScaleX() const;
-		NodeScalarPtr GetScaleY() const;
-		NodeScalarPtr GetRotation() const;
-		void SetOffsetX( NodeScalarPtr pNode );
-		void SetOffsetY( NodeScalarPtr pNode );
-		void SetScaleX( NodeScalarPtr pNode );
-		void SetScaleY( NodeScalarPtr pNode );
-		void SetRotation( NodeScalarPtr pNode );
-
-		EAddressMode GetAddressMode() const;
-		void SetAddressMode(EAddressMode AddressMode);
+		Ptr<NodeImage> Base;
+		Ptr<NodeScalar> OffsetX;
+		Ptr<NodeScalar> OffsetY;
+		Ptr<NodeScalar> ScaleX;
+		Ptr<NodeScalar> ScaleY;
+		Ptr<NodeScalar> Rotation;
 
-		uint16 GetSizeX() const;
-		void SetSizeX(uint16 SizeX); 
+		EAddressMode AddressMode = EAddressMode::Wrap;
+		uint32 SizeX = 0;
+		uint32 SizeY = 0;
 
-		uint16 GetSizeY() const;
-		void SetSizeY(uint16 SizeY); 
+		bool bKeepAspectRatio = false;
 
-		bool GetKeepAspectRatio() const;
-		void SetKeepAspectRatio(bool bKeepAspectRatio);
+	public:
 
-		//-----------------------------------------------------------------------------------------
-		// Interface pattern
-		//-----------------------------------------------------------------------------------------
-		class Private;
-		Private* GetPrivate() const;
+		// Node Interface
+		virtual const FNodeType* GetType() const override { return &StaticType; }
+		static const FNodeType* GetStaticType() { return &StaticType; }
 
 	protected:
 
-		//! Forbidden. Manage with the Ptr<> template.
-		~NodeImageTransform();
+		/** Forbidden. Manage with the Ptr<> template. */
+		~NodeImageTransform() {}
 
 	private:
 
-		Private* m_pD;
+		static FNodeType StaticType;
 
 	};
 
-
 }
diff --git a/Engine/Plugins/Mutable/Source/MutableTools/Public/MuT/NodeImageVariation.h b/Engine/Plugins/Mutable/Source/MutableTools/Public/MuT/NodeImageVariation.h
--- a/Engine/Plugins/Mutable/Source/MutableTools/Public/MuT/NodeImageVariation.h
+++ b/Engine/Plugins/Mutable/Source/MutableTools/Public/MuT/NodeImageVariation.h
@@ -7,69 +7,40 @@
 #include "MuT/Node.h"
 #include "MuT/NodeImage.h"
 
-
 namespace mu
 {
 
-    // Forward definitions
-    class NodeImageVariation;
-    typedef Ptr<NodeImageVariation> NodeImageVariationPtr;
-    typedef Ptr<const NodeImageVariation> NodeImageVariationPtrConst;
-
-
-    //!
-    //! \ingroup model
+    /** */
     class MUTABLETOOLS_API NodeImageVariation : public NodeImage
     {
-    public:
-        //-----------------------------------------------------------------------------------------
-        // Life cycle
-        //-----------------------------------------------------------------------------------------
-
-        NodeImageVariation();
-
-        //-----------------------------------------------------------------------------------------
-        // Node Interface
-        //-----------------------------------------------------------------------------------------
-
-        const FNodeType* GetType() const override;
-        static const FNodeType* GetStaticType();
+	public:
 
-        //-----------------------------------------------------------------------------------------
-        // Own Interface
-        //-----------------------------------------------------------------------------------------
+		Ptr<NodeImage> DefaultImage;
 
-        //!
-        void SetDefaultImage( NodeImage* Image );
+		struct FVariation
+		{
+			Ptr<NodeImage> Image;
+			FString Tag;
+		};
 
-        //! Set the number of tags to consider in this variation
-        void SetVariationCount( int count );
-
-        //!
-        int GetVariationCount() const;
-
-        //! Set the tag or state name that will enable a specific vartiation
-        void SetVariationTag( int index, const FString& Tag );
-
-        //!
-        void SetVariationImage( int index, NodeImage* Image );
-
-        //!}
+		TArray<FVariation> Variations;
 
+    public:
 
-        //-----------------------------------------------------------------------------------------
-        // Interface pattern
-        //-----------------------------------------------------------------------------------------
-        class Private;
-        Private* GetPrivate() const;
+		// Node Interface
+		virtual const FNodeType* GetType() const override { return &StaticType; }
+		static const FNodeType* GetStaticType() { return &StaticType; }
 
     protected:
-        //! Forbidden. Manage with the Ptr<> template.
-        ~NodeImageVariation();
+
+        /** Forbidden. Manage with the Ptr<> template. */
+		~NodeImageVariation() {}
 
     private:
-        Private* m_pD;
-    };
 
+		static FNodeType StaticType;
+	
+	};
+
+}
 
-} // namespace mu
diff --git a/Engine/Plugins/Mutable/Source/MutableTools/Public/MuT/NodeMesh.h b/Engine/Plugins/Mutable/Source/MutableTools/Public/MuT/NodeMesh.h
--- a/Engine/Plugins/Mutable/Source/MutableTools/Public/MuT/NodeMesh.h
+++ b/Engine/Plugins/Mutable/Source/MutableTools/Public/MuT/NodeMesh.h
@@ -25,16 +25,18 @@ namespace mu
 	public:
 
 		// Node Interface
-        const FNodeType* GetType() const override;
-		static const FNodeType* GetStaticType();
-
-		// Interface pattern
-		class Private;
+		virtual const FNodeType* GetType() const override { return &StaticType; }
+		static const FNodeType* GetStaticType() { return &StaticType; }
 
 	protected:
 
 		// Forbidden. Manage with the Ptr<> template.
 		inline ~NodeMesh() {}
+
+	private:
+
+		static FNodeType StaticType;
+
 	};
 
 
diff --git a/Engine/Plugins/Mutable/Source/MutableTools/Public/MuT/NodeMeshApplyPose.h b/Engine/Plugins/Mutable/Source/MutableTools/Public/MuT/NodeMeshApplyPose.h
--- a/Engine/Plugins/Mutable/Source/MutableTools/Public/MuT/NodeMeshApplyPose.h
+++ b/Engine/Plugins/Mutable/Source/MutableTools/Public/MuT/NodeMeshApplyPose.h
@@ -7,60 +7,32 @@
 #include "MuT/Node.h"
 #include "MuT/NodeMesh.h"
 
-
 namespace mu
 {
 
-	// Forward definitions
-    class NodeMeshApplyPose;
-    using NodeMeshApplyPosePtr = Ptr<NodeMeshApplyPose>;
-    using NodeMeshApplyPosePtrConst = Ptr<const NodeMeshApplyPose>;
-
-
-    //! Node that applies a pose to a mesh, baking it into the vertex data
+    /** Node that applies a pose to a mesh, baking it into the vertex data. */
     class MUTABLETOOLS_API NodeMeshApplyPose : public NodeMesh
 	{
 	public:
 
-        NodeMeshApplyPose();
+		Ptr<NodeMesh> Base;
+		Ptr<NodeMesh> Pose;
 
+	public:
 
-		//-----------------------------------------------------------------------------------------
 		// Node Interface
-		//-----------------------------------------------------------------------------------------
-
-		const FNodeType* GetType() const override;
-		static const FNodeType* GetStaticType();
-
-		//-----------------------------------------------------------------------------------------
-		// Own Interface
-		//-----------------------------------------------------------------------------------------
-
-		//! Get the nodes generating the base mesh to be morphed
-		NodeMeshPtr GetBase() const;
-		void SetBase( NodeMeshPtr );
-
-        //! Get the nodes generating the pose to apply. A pose is represented with a mesh object
-        //! with no geometry: only bone matrices and names are relevant.
-        NodeMeshPtr GetPose() const;
-        void SetPose( NodeMeshPtr );
-
-		//-----------------------------------------------------------------------------------------
-		// Interface pattern
-		//-----------------------------------------------------------------------------------------
-		class Private;
-		Private* GetPrivate() const;
+		virtual const FNodeType* GetType() const override { return &StaticType; }
+		static const FNodeType* GetStaticType() { return &StaticType; }
 
 	protected:
 
-		//! Forbidden. Manage with the Ptr<> template.
-        ~NodeMeshApplyPose();
+		/** Forbidden.Manage with the Ptr<> template. */
+		~NodeMeshApplyPose() {}
 
 	private:
 
-		Private* m_pD;
+		static FNodeType StaticType;
 
 	};
 
-
 }
diff --git a/Engine/Plugins/Mutable/Source/MutableTools/Public/MuT/NodeMeshConstant.h b/Engine/Plugins/Mutable/Source/MutableTools/Public/MuT/NodeMeshConstant.h
--- a/Engine/Plugins/Mutable/Source/MutableTools/Public/MuT/NodeMeshConstant.h
+++ b/Engine/Plugins/Mutable/Source/MutableTools/Public/MuT/NodeMeshConstant.h
@@ -2,77 +2,44 @@
 
 #pragma once
 
-#include "MuR/Ptr.h"
-#include "MuR/RefCounted.h"
 #include "MuT/Node.h"
 #include "MuT/NodeMesh.h"
 #include "MuT/NodeImage.h"
+#include "MuT/NodeLayout.h"
+#include "MuR/Mesh.h"
 #include "Templates/SharedPointer.h"
 
 namespace mu
 {
 
-	// Forward definitions
-	class NodeMeshConstant;
-	typedef Ptr<NodeMeshConstant> NodeMeshConstantPtr;
-	typedef Ptr<const NodeMeshConstant> NodeMeshConstantPtrConst;
-
-	class FMesh;
-	class NodeLayout;
-
-	//! Node that outputs a constant mesh.
-	//! It allows to define the layouts for the texture channels of the constant mesh
+	/** Node that outputs a constant mesh.
+	* It allows to define the layouts for the texture channels of the constant mesh
+	*/
 	class MUTABLETOOLS_API NodeMeshConstant : public NodeMesh
 	{
 	public:
 
 		FSourceDataDescriptor SourceDataDescriptor;
 
-	public:
-
-		NodeMeshConstant();
-
-		//-----------------------------------------------------------------------------------------
-		// Node Interface
-		//-----------------------------------------------------------------------------------------
-
-		const FNodeType* GetType() const override;
-		static const FNodeType* GetStaticType();
+		TSharedPtr<FMesh> Value;
 
-		//-----------------------------------------------------------------------------------------
-		// Own Interface
-		//-----------------------------------------------------------------------------------------
+		TArray<Ptr<NodeLayout>> Layouts;
 
-		//! Get the constant mesh that will be returned.
-		TSharedPtr<FMesh> GetValue() const;
-
-		//! Set the constant mesh that will be returned.
-		void SetValue(TSharedPtr<FMesh>);
-
-		//! Get the number of layouts defined in this mesh.
-		int32 GetLayoutCount() const;
-		void SetLayoutCount( int32 );
-
-		//! Get the node defining a layout of the returned mesh.
-		Ptr<NodeLayout> GetLayout( int32 index ) const;
-		void SetLayout( int32 index, Ptr<NodeLayout>);
+	public:
 
-		//-----------------------------------------------------------------------------------------
-		// Interface pattern
-		//-----------------------------------------------------------------------------------------
-		class Private;
-		Private* GetPrivate() const;
+		// Node Interface
+		virtual const FNodeType* GetType() const override { return &StaticType; }
+		static const FNodeType* GetStaticType() { return &StaticType; }
 
 	protected:
 
-		//! Forbidden. Manage with the Ptr<> template.
-		~NodeMeshConstant();
+		/** Forbidden. Manage with the Ptr<> template. */
+		~NodeMeshConstant() {}
 
 	private:
 
-		Private* m_pD;
+		static FNodeType StaticType;
 
 	};
 
-
 }
diff --git a/Engine/Plugins/Mutable/Source/MutableTools/Public/MuT/NodeMeshFormat.h b/Engine/Plugins/Mutable/Source/MutableTools/Public/MuT/NodeMeshFormat.h
--- a/Engine/Plugins/Mutable/Source/MutableTools/Public/MuT/NodeMeshFormat.h
+++ b/Engine/Plugins/Mutable/Source/MutableTools/Public/MuT/NodeMeshFormat.h
@@ -6,6 +6,7 @@
 #include "MuR/RefCounted.h"
 #include "MuT/Node.h"
 #include "MuT/NodeMesh.h"
+#include "MuR/Mesh.h"
 
 namespace mu
 {
@@ -18,55 +19,30 @@ namespace mu
 	{
 	public:
 
-		NodeMeshFormat();
+		/** Source mesh to transform. */
+		Ptr<NodeMesh> Source;
+
+		/** New mesh format.The buffers in the sets have no elements, but they define the formats. */
+		FMeshBufferSet VertexBuffers;
+		FMeshBufferSet IndexBuffers;
+
+		/** */
+		bool bOptimizeBuffers = false;
+
+	public:
 
-		//-----------------------------------------------------------------------------------------
 		// Node Interface
-		//-----------------------------------------------------------------------------------------
-
-		const FNodeType* GetType() const override;
-		static const FNodeType* GetStaticType();
-
-		//-----------------------------------------------------------------------------------------
-		// Own Interface
-		//-----------------------------------------------------------------------------------------
-
-		/** Mesh to be reformatted. */
-		Ptr<NodeMesh> GetSource() const;
-		void SetSource( NodeMesh* );
-
-		/** Access the MeshBufferSet that defines the new format for the mesh vertices.These
-		* buffers don't really contain any data (they have 0 elements) but they define the
-		* structure. If this is null, the vertex buffers will not be changed.
-		*/
-		FMeshBufferSet& GetVertexBuffers();
-
-		/** Access the MeshBufferSet that defines the new format for the mesh indices. These
-		* buffers don't really contain any data (they have 0 elements) but they define the
-		* structure. If this is null, the vertex buffers will not be changed.
-		*/
-		FMeshBufferSet& GetIndexBuffers();
-
-		/** Optimize the buffers is possible. This may change the target format to reduce the number
-		* of channels or even the type if possible to minimize size.
-		* By default it is disabled.
-		*/
-		void SetOptimizeBuffers(bool bEnable);
-
-		//-----------------------------------------------------------------------------------------
-		// Interface pattern
-		//-----------------------------------------------------------------------------------------
-		class Private;
-		Private* GetPrivate() const;
+		virtual const FNodeType* GetType() const override { return &StaticType; }
+		static const FNodeType* GetStaticType() { return &StaticType; }
 
 	protected:
 
-		//! Forbidden. Manage with the Ptr<> template.
-		~NodeMeshFormat();
+		/** Forbidden. Manage with the Ptr<> template. */
+		~NodeMeshFormat() {}
 
 	private:
 
-		Private* m_pD;
+		static FNodeType StaticType;
 
 	};
 
diff --git a/Engine/Plugins/Mutable/Source/MutableTools/Public/MuT/NodeMeshMakeMorph.h b/Engine/Plugins/Mutable/Source/MutableTools/Public/MuT/NodeMeshMakeMorph.h
--- a/Engine/Plugins/Mutable/Source/MutableTools/Public/MuT/NodeMeshMakeMorph.h
+++ b/Engine/Plugins/Mutable/Source/MutableTools/Public/MuT/NodeMeshMakeMorph.h
@@ -11,57 +11,31 @@
 namespace mu
 {
 
-	// Forward definitions
-    class NodeMeshMakeMorph;
-    typedef Ptr<NodeMeshMakeMorph> NodeMeshMakeMorphPtr;
-    typedef Ptr<const NodeMeshMakeMorph> NodeMeshMakeMorphPtrConst;
-
-
-	//! This node removes from a mesh A all the faces that are also part of a mesh B.
+	/** This node removes from a mesh A all the faces that are also part of a mesh B. */
     class MUTABLETOOLS_API NodeMeshMakeMorph : public NodeMesh
 	{
 	public:
 
-        NodeMeshMakeMorph();
-
-		//-----------------------------------------------------------------------------------------
-		// Node Interface
-		//-----------------------------------------------------------------------------------------
-
-		const FNodeType* GetType() const override;
-		static const FNodeType* GetStaticType();
+		Ptr<NodeMesh> Base;
+		Ptr<NodeMesh> Target;
 
-		//-----------------------------------------------------------------------------------------
-		// Own Interface
-		//-----------------------------------------------------------------------------------------
+		bool bOnlyPositionAndNormal = false;
 
-		//!
-		NodeMeshPtr GetBase() const;
-		void SetBase( NodeMesh* p );
-
-		//!
-		NodeMeshPtr GetTarget() const;
-		void SetTarget( NodeMesh* p );
-
-		bool GetOnlyPositionAndNormal() const;
-		void SetOnlyPositionAndNormal(bool bOnlyPositionAndNormal);
+	public:
 
-		//-----------------------------------------------------------------------------------------
-		// Interface pattern
-		//-----------------------------------------------------------------------------------------
-		class Private;
-		Private* GetPrivate() const;
+		// Node Interface
+		virtual const FNodeType* GetType() const override { return &StaticType; }
+		static const FNodeType* GetStaticType() { return &StaticType; }
 
 	protected:
 
-		//! Forbidden. Manage with the Ptr<> template.
-        ~NodeMeshMakeMorph();
+		/** Forbidden. Manage with the Ptr<> template. */
+		~NodeMeshMakeMorph() {}
 
 	private:
 
-		Private* m_pD;
+		static FNodeType StaticType;
 
 	};
 
-
 }
diff --git a/Engine/Plugins/Mutable/Source/MutableTools/Public/MuT/NodeMeshMorph.h b/Engine/Plugins/Mutable/Source/MutableTools/Public/MuT/NodeMeshMorph.h
--- a/Engine/Plugins/Mutable/Source/MutableTools/Public/MuT/NodeMeshMorph.h
+++ b/Engine/Plugins/Mutable/Source/MutableTools/Public/MuT/NodeMeshMorph.h
@@ -6,73 +6,42 @@
 #include "MuR/RefCounted.h"
 #include "MuT/Node.h"
 #include "MuT/NodeMesh.h"
-
+#include "MuT/NodeScalar.h"
+#include "MuR/Skeleton.h"
 
 namespace mu
 {
 
-	// Forward definitions
-	class NodeScalar;
-	struct FBoneName;
-
 	/** Node that morphs a base mesh with one weighted target. */
 	class MUTABLETOOLS_API NodeMeshMorph : public NodeMesh
 	{
 	public:
 
-		NodeMeshMorph();
-
-		//-----------------------------------------------------------------------------------------
-		// Node Interface
-		//-----------------------------------------------------------------------------------------		
-
-		const FNodeType* GetType() const override;
-		static const FNodeType* GetStaticType();
-
-		//-----------------------------------------------------------------------------------------
-		// Own Interface
-		//-----------------------------------------------------------------------------------------
-
-		/** Get the node generating the weight used to select(and combine) the morphs to apply. */
-		Ptr<NodeScalar> GetFactor() const;
-		void SetFactor( Ptr<NodeScalar> );
+		Ptr<NodeScalar> Factor;
+		Ptr<NodeMesh> Base;
+		Ptr<NodeMesh> Morph;
 
-		/** Get the nodes generating the base mesh to be morphed. */
-		Ptr<NodeMesh> GetBase() const;
-		void SetBase( Ptr<NodeMesh> );
+		bool bReshapeSkeleton = false;
+		bool bReshapePhysicsVolumes = false;
 
-		/** Get the nodes generating the several morph targets. */
-		Ptr<NodeMesh> GetMorph() const;
-		void SetMorph( Ptr<NodeMesh> );
+		TArray<FBoneName> BonesToDeform;
+		TArray<FBoneName> PhysicsToDeform;
 
-		/** Also deform the mesh skeleton. Disabled by default. */
-		void SetReshapeSkeleton(bool);	
-
-		/** Deform Mesh Physics Volumes */
-		void SetReshapePhysicsVolumes(bool);
-
-		/** Sets the number of bones that will be deform */
-		void AddBoneToDeform(const FBoneName& BoneName);	
-
-		/** Add a Physics Body to deform */
-		void AddPhysicsBodyToDeform(const FBoneName& BoneName);
+	public:
 
-        //-----------------------------------------------------------------------------------------
-		// Interface pattern
-		//-----------------------------------------------------------------------------------------
-		class Private;
-		Private* GetPrivate() const;
+		// Node Interface
+		virtual const FNodeType* GetType() const override { return &StaticType; }
+		static const FNodeType* GetStaticType() { return &StaticType; }
 
 	protected:
 
-		//! Forbidden. Manage with the Ptr<> template.
-		~NodeMeshMorph();
+		/** Forbidden. Manage with the Ptr<> template. */
+		~NodeMeshMorph() {}
 
 	private:
 
-		Private* m_pD;
+		static FNodeType StaticType;
 
 	};
 
-
 }
diff --git a/Engine/Plugins/Mutable/Source/MutableTools/Public/MuT/NodeMeshParameter.h b/Engine/Plugins/Mutable/Source/MutableTools/Public/MuT/NodeMeshParameter.h
new file mode 100644
--- /dev/null
+++ b/Engine/Plugins/Mutable/Source/MutableTools/Public/MuT/NodeMeshParameter.h
@@ -0,0 +1,48 @@
+// Copyright Epic Games, Inc. All Rights Reserved.
+
+#pragma once
+
+#include "MuR/Ptr.h"
+#include "MuR/RefCounted.h"
+#include "MuT/Node.h"
+#include "MuT/NodeMesh.h"
+#include "MuT/NodeRange.h"
+
+namespace mu
+{
+
+	/** */
+	class MUTABLETOOLS_API NodeMeshParameter : public NodeMesh
+	{
+	public:
+
+		/** Full asset name for the default value of the parameter. */
+		FName DefaultValue;
+
+		/** Name of the parameter */
+		FString Name;
+
+		/** User provided ID to identify the parameter. */
+		FString UID;
+
+		/** Ranges for the parameter in case it is multidimensional. */
+		TArray<Ptr<NodeRange>> Ranges;
+
+	public:
+
+		// Node Interface
+		virtual const FNodeType* GetType() const override { return &StaticType; }
+		static const FNodeType* GetStaticType() { return &StaticType; }
+
+	protected:
+
+		/** Forbidden. Manage with the Ptr<> template. */
+		~NodeMeshParameter() {}
+
+	private:
+
+		static FNodeType StaticType;
+
+	};
+
+}
diff --git a/Engine/Plugins/Mutable/Source/MutableTools/Public/MuT/NodeMeshReshape.h b/Engine/Plugins/Mutable/Source/MutableTools/Public/MuT/NodeMeshReshape.h
--- a/Engine/Plugins/Mutable/Source/MutableTools/Public/MuT/NodeMeshReshape.h
+++ b/Engine/Plugins/Mutable/Source/MutableTools/Public/MuT/NodeMeshReshape.h
@@ -6,77 +6,47 @@
 #include "MuT/Node.h"
 #include "MuT/NodeMesh.h"
 #include "MuR/Mesh.h"
+#include "MuR/Skeleton.h"
 
 namespace mu
 {
-	// Forward declarations
-	struct FBoneName;
 
-	//! Node that morphs a base mesh with one or two weighted targets from a sequence.
+	/** Node that morphs a base mesh with one or two weighted targets from a sequence. */
 	class MUTABLETOOLS_API NodeMeshReshape : public NodeMesh
 	{
 	public:
 
-		NodeMeshReshape();
+		Ptr<NodeMesh> BaseMesh;
+		Ptr<NodeMesh> BaseShape;
+		Ptr<NodeMesh> TargetShape;
+		bool bReshapeVertices = true;
+		bool bRecomputeNormals = false;
+		bool bApplyLaplacian = false;
+		bool bReshapeSkeleton = false;
+		bool bReshapePhysicsVolumes = false;
 
-		//-----------------------------------------------------------------------------------------
-		// Node Interface
-		//-----------------------------------------------------------------------------------------
-
-		const FNodeType* GetType() const override;
-		static const FNodeType* GetStaticType();
-
-		//-----------------------------------------------------------------------------------------
-		// Own Interface
-		//-----------------------------------------------------------------------------------------
-
-		const NodeMeshPtr& GetBaseMesh() const;
-		void SetBaseMesh( const NodeMeshPtr& );
-
-		const NodeMeshPtr& GetBaseShape() const;
-		void SetBaseShape(const NodeMeshPtr&);
-
-		const NodeMeshPtr& GetTargetShape() const;
-		void SetTargetShape(const NodeMeshPtr&);
+		EVertexColorUsage ColorRChannelUsage = EVertexColorUsage::None;
+		EVertexColorUsage ColorGChannelUsage = EVertexColorUsage::None;
+		EVertexColorUsage ColorBChannelUsage = EVertexColorUsage::None;
+		EVertexColorUsage ColorAChannelUsage = EVertexColorUsage::None;
 
-		void SetReshapeVertices(bool);
+		TArray<FBoneName> BonesToDeform;
+		TArray<FBoneName> PhysicsToDeform;
 
-		/** Recompute normals after reshaping. This will ignore the reshaped normals. Disabled by default. */
-		void SetRecomputeNormals(bool);
-
-
-		/** Also deform the mesh skeleton. Disabled by default. */
-		void SetReshapeSkeleton(bool);
-	
-		/** Apply Laplacian smoothing to the reshaped mesh.  */
-		void SetApplyLaplacian(bool);
-		
-		/** Set vertex color channel usages for Reshape operations. */
-		void SetColorUsages(EVertexColorUsage R, EVertexColorUsage G, EVertexColorUsage B, EVertexColorUsage A);
-
-		/** Deform Mesh Physics Volumes */
-		void SetReshapePhysicsVolumes(bool);
+	public:
 
-		/** Sets the number of bones that will be deform */
-		void AddBoneToDeform(const FBoneName& BoneId);
-	
-		/** Add a Physics Body to deform */
-		void AddPhysicsBodyToDeform(const FBoneName& BoneId);
-        
-		//-----------------------------------------------------------------------------------------
-		// Interface pattern
-		//-----------------------------------------------------------------------------------------
-		class Private;
-		Private* GetPrivate() const;
+		// Node Interface
+		virtual const FNodeType* GetType() const override { return &StaticType; }
+		static const FNodeType* GetStaticType() { return &StaticType; }
 
 	protected:
 
-		//! Forbidden. Manage with the Ptr<> template.
-		~NodeMeshReshape();
+		/** Forbidden. Manage with the Ptr<> template. */
+		~NodeMeshReshape() {}
 
 	private:
 
-		Private* m_pD;
+		static FNodeType StaticType;
 
 	};
 
diff --git a/Engine/Plugins/Mutable/Source/MutableTools/Public/MuT/NodeMeshSwitch.h b/Engine/Plugins/Mutable/Source/MutableTools/Public/MuT/NodeMeshSwitch.h
--- a/Engine/Plugins/Mutable/Source/MutableTools/Public/MuT/NodeMeshSwitch.h
+++ b/Engine/Plugins/Mutable/Source/MutableTools/Public/MuT/NodeMeshSwitch.h
@@ -3,71 +3,36 @@
 #pragma once
 
 #include "MuR/Ptr.h"
-#include "MuR/RefCounted.h"
-#include "MuT/Node.h"
 #include "MuT/NodeMesh.h"
-
+#include "MuT/NodeScalar.h"
 
 namespace mu
 {
 
-	// Forward definitions
-	class NodeScalar;
-	typedef Ptr<NodeScalar> NodeScalarPtr;
-	typedef Ptr<const NodeScalar> NodeScalarPtrConst;
-
-	class NodeMeshSwitch;
-	typedef Ptr<NodeMeshSwitch> NodeMeshSwitchPtr;
-	typedef Ptr<const NodeMeshSwitch> NodeMeshSwitchPtrConst;
-
-
-	//! This node selects an output Mesh from a set of input Meshs based on a parameter.
+	/** This node selects an output Mesh from a set of input Meshs based on a parameter. */
 	class MUTABLETOOLS_API NodeMeshSwitch : public NodeMesh
 	{
 	public:
 
-		NodeMeshSwitch();
-
-		//-----------------------------------------------------------------------------------------
-		// Node Interface
-		//-----------------------------------------------------------------------------------------
-
-		const FNodeType* GetType() const override;
-		static const FNodeType* GetStaticType();
+		Ptr<NodeScalar> Parameter;
 
-		//-----------------------------------------------------------------------------------------
-		// Own Interface
-		//-----------------------------------------------------------------------------------------
-
-		//! Get the node generating the parameter used to select the option.
-		NodeScalarPtr GetParameter() const;
-		void SetParameter( NodeScalarPtr );
-
-		//! Set the number of option Meshs. It will keep the currently set targets and initialise
-		//! the new ones as null.
-		void SetOptionCount( int );
-
-		//! Get the node generating the t-th option Mesh.
-		NodeMeshPtr GetOption( int t ) const;
-		void SetOption( int t, NodeMeshPtr );
+		TArray<Ptr<NodeMesh>> Options;
 
+	public:
 
-		//-----------------------------------------------------------------------------------------
-		// Interface pattern
-		//-----------------------------------------------------------------------------------------
-		class Private;
-		Private* GetPrivate() const;
+		// Node Interface
+		virtual const FNodeType* GetType() const override { return &StaticType; }
+		static const FNodeType* GetStaticType() { return &StaticType; }
 
 	protected:
 
-		//! Forbidden. Manage with the Ptr<> template.
-		~NodeMeshSwitch();
+		/** Forbidden. Manage with the Ptr<> template. */
+		~NodeMeshSwitch() {}
 
 	private:
 
-		Private* m_pD;
+		static FNodeType StaticType;
 
 	};
 
-
 }
diff --git a/Engine/Plugins/Mutable/Source/MutableTools/Public/MuT/NodeMeshTable.h b/Engine/Plugins/Mutable/Source/MutableTools/Public/MuT/NodeMeshTable.h
--- a/Engine/Plugins/Mutable/Source/MutableTools/Public/MuT/NodeMeshTable.h
+++ b/Engine/Plugins/Mutable/Source/MutableTools/Public/MuT/NodeMeshTable.h
@@ -10,16 +10,9 @@
 #include "MuT/NodeImage.h"
 #include "MuT/Table.h"
 
-
 namespace mu
 {
 
-	// Forward definitions
-	class NodeMeshTable;
-	typedef Ptr<NodeMeshTable> NodeMeshTablePtr;
-	typedef Ptr<const NodeMeshTable> NodeMeshTablePtrConst;
-
-
 	//! This node provides the meshes stored in the column of a table.
 	class MUTABLETOOLS_API NodeMeshTable : public NodeMesh
 	{
@@ -42,34 +35,9 @@ namespace mu
 		virtual const FNodeType* GetType() const override { return GetStaticType(); }
 		static const FNodeType* GetStaticType() { return &StaticType; }
 
-		//-----------------------------------------------------------------------------------------
-		// Own Interface
-		//-----------------------------------------------------------------------------------------
-
-		//! Set the name of the implicit table parameter.
-		void SetParameterName( const FString& strName );
-
-		//!
-		void SetColumn( const FString& strName );
-
-		//! Get the number of layouts defined in the meshes on this column.
-		int GetLayoutCount() const;
-		void SetLayoutCount( int );
-
-		//! Get the node defining a layout of the meshes on this column.
-		Ptr<NodeLayout> GetLayout( int index ) const;
-		void SetLayout( int index, Ptr<NodeLayout>);
-
-		//! Adds the "None" option to the parameter that represents this table column
-		void SetNoneOption(bool bAddOption);
-
-		//! Set the row name to be used as default value
-		void SetDefaultRowName(const FString& RowName);
-
-
 	protected:
 
-		//! Forbidden. Manage with the Ptr<> template.
+		/** Forbidden. Manage with the Ptr<> template. */
 		~NodeMeshTable() {}
 
 	private:
@@ -78,5 +46,4 @@ namespace mu
 
 	};
 
-
 }
diff --git a/Engine/Plugins/Mutable/Source/MutableTools/Public/MuT/NodeMeshTransform.h b/Engine/Plugins/Mutable/Source/MutableTools/Public/MuT/NodeMeshTransform.h
--- a/Engine/Plugins/Mutable/Source/MutableTools/Public/MuT/NodeMeshTransform.h
+++ b/Engine/Plugins/Mutable/Source/MutableTools/Public/MuT/NodeMeshTransform.h
@@ -6,58 +6,34 @@
 #include "MuR/RefCounted.h"
 #include "MuT/Node.h"
 #include "MuT/NodeMesh.h"
-#include "Math/MathFwd.h"
+#include "Math/Matrix.h"
 
 namespace mu
 {
 
-	// Forward definitions
-    class NodeMeshTransform;
-    typedef Ptr<NodeMeshTransform> NodeMeshTransformPtr;
-    typedef Ptr<const NodeMeshTransform> NodeMeshTransformPtrConst;
-
-
-    //! This node applies a geometric transform represented by a 4x4 matrix to a mesh
+    /** This node applies a geometric transform represented by a 4x4 matrix to a mesh. */
     class MUTABLETOOLS_API NodeMeshTransform : public NodeMesh
 	{
 	public:
 
-        NodeMeshTransform();
-
-		//-----------------------------------------------------------------------------------------
-		// Node Interface
-		//-----------------------------------------------------------------------------------------
-
-        const FNodeType* GetType() const override;
-		static const FNodeType* GetStaticType();
+		Ptr<NodeMesh> Source;
+		FMatrix44f Transform;
 
-		//-----------------------------------------------------------------------------------------
-		// Own Interface
-		//-----------------------------------------------------------------------------------------
-
-		//! Source mesh to be re-formatted
-        NodeMeshPtr GetSource() const;
-		void SetSource( NodeMesh* );
-
-        void SetTransform(const FMatrix44f&);
-		const FMatrix44f& GetTransform() const;
+	public:
 
-		//-----------------------------------------------------------------------------------------
-		// Interface pattern
-		//-----------------------------------------------------------------------------------------
-		class Private;
-		Private* GetPrivate() const;
+		// Node Interface
+		virtual const FNodeType* GetType() const override { return &StaticType; }
+		static const FNodeType* GetStaticType() { return &StaticType; }
 
 	protected:
 
-		//! Forbidden. Manage with the Ptr<> template.
-        ~NodeMeshTransform();
+		/** Forbidden. Manage with the Ptr<> template. */
+		~NodeMeshTransform() {}
 
 	private:
 
-		Private* m_pD;
+		static FNodeType StaticType;
 
 	};
 
-
 }
diff --git a/Engine/Plugins/Mutable/Source/MutableTools/Public/MuT/NodeMeshVariation.h b/Engine/Plugins/Mutable/Source/MutableTools/Public/MuT/NodeMeshVariation.h
--- a/Engine/Plugins/Mutable/Source/MutableTools/Public/MuT/NodeMeshVariation.h
+++ b/Engine/Plugins/Mutable/Source/MutableTools/Public/MuT/NodeMeshVariation.h
@@ -7,65 +7,40 @@
 #include "MuT/Node.h"
 #include "MuT/NodeMesh.h"
 
-
 namespace mu
 {
 
-    // Forward definitions
-    class NodeMeshVariation;
-    typedef Ptr<NodeMeshVariation> NodeMeshVariationPtr;
-    typedef Ptr<const NodeMeshVariation> NodeMeshVariationPtrConst;
-
-    //!
-    //! \ingroup model
-    class MUTABLETOOLS_API NodeMeshVariation : public NodeMesh
+	/** */
+	class MUTABLETOOLS_API NodeMeshVariation : public NodeMesh
     {
-    public:
-
-        NodeMeshVariation();
-
-        //-----------------------------------------------------------------------------------------
-        // Node Interface
-        //-----------------------------------------------------------------------------------------
-
-        const FNodeType* GetType() const override;
-        static const FNodeType* GetStaticType();
+	public:
 
-        //-----------------------------------------------------------------------------------------
-        // Own Interface
-        //-----------------------------------------------------------------------------------------
+		Ptr<NodeMesh> DefaultMesh;
 
-        //!
-        void SetDefaultMesh( NodeMesh* Mesh );
+		struct FVariation
+		{
+			Ptr<NodeMesh> Mesh;
+			FString Tag;
+		};
 
-        //! Set the number of tags to consider in this variation
-        void SetVariationCount( int count );
+		TArray<FVariation> Variations;
 
-        //!
-        int GetVariationCount() const;
-
-        //! Set the tag or state name that will enable a specific vartiation
-        void SetVariationTag( int index, const char* strTag );
+    public:
 
-        //!
-        void SetVariationMesh( int index, NodeMesh* Mesh );
+		// Node Interface
+		virtual const FNodeType* GetType() const override { return &StaticType; }
+		static const FNodeType* GetStaticType() { return &StaticType; }
 
-        //!}
+    protected:
 
+		/** Forbidden. Manage with the Ptr<> template. */
+		~NodeMeshVariation() {}
 
-        //-----------------------------------------------------------------------------------------
-        // Interface pattern
-        //-----------------------------------------------------------------------------------------
-        class Private;
-        Private* GetPrivate() const;
+    private:
 
-    protected:
-        //! Forbidden. Manage with the Ptr<> template.
-        ~NodeMeshVariation();
+		static FNodeType StaticType;
 
-    private:
-        Private* m_pD;
     };
 
+}
 
-} // namespace mu
diff --git a/Engine/Plugins/Mutable/Source/MutableTools/Public/MuT/NodeModifierMeshClipMorphPlane.h b/Engine/Plugins/Mutable/Source/MutableTools/Public/MuT/NodeModifierMeshClipMorphPlane.h
--- a/Engine/Plugins/Mutable/Source/MutableTools/Public/MuT/NodeModifierMeshClipMorphPlane.h
+++ b/Engine/Plugins/Mutable/Source/MutableTools/Public/MuT/NodeModifierMeshClipMorphPlane.h
@@ -9,12 +9,9 @@
 #include "MuT/Node.h"
 #include "MuT/NodeModifier.h"
 
-
 namespace mu
 {
 
-	/** */
-
 	struct FBoneName;
 
 	struct FClipMorphPlaneParameters
@@ -77,7 +74,7 @@ namespace mu
 
 	protected:
 
-		//! Forbidden. Manage with the Ptr<> template.
+		/** Forbidden. Manage with the Ptr<> template. */
 		~NodeModifierMeshClipMorphPlane() {}
 
 	private:
diff --git a/Engine/Plugins/Mutable/Source/MutableTools/Public/MuT/NodeObject.h b/Engine/Plugins/Mutable/Source/MutableTools/Public/MuT/NodeObject.h
--- a/Engine/Plugins/Mutable/Source/MutableTools/Public/MuT/NodeObject.h
+++ b/Engine/Plugins/Mutable/Source/MutableTools/Public/MuT/NodeObject.h
@@ -7,32 +7,19 @@
 #include "MuR/RefCounted.h"
 #include "MuT/Node.h"
 
-
 namespace mu
 {
 
-	// Forward definitions
-	class NodeObject;
-	typedef Ptr<NodeObject> NodeObjectPtr;
-	typedef Ptr<const NodeObject> NodeObjectPtrConst;
-
-
-    //! %Base class for any node that outputs an object.
-	//! \ingroup model
+	/** Base class for any node that outputs an object. */
 	class MUTABLETOOLS_API NodeObject : public Node
 	{
 	public:
 
-		//-----------------------------------------------------------------------------------------
 		// Node Interface
-		//-----------------------------------------------------------------------------------------
-
-		const FNodeType* GetType() const override;
-		static const FNodeType* GetStaticType();
+		virtual const FNodeType* GetType() const override { return &StaticType; }
+		static const FNodeType* GetStaticType() { return &StaticType; }
 
-		//-----------------------------------------------------------------------------------------
 		// Own interface
-		//-----------------------------------------------------------------------------------------
 
 		//! Get the name of the object.
 		virtual const FString& GetName() const = 0;
@@ -46,17 +33,15 @@ namespace mu
 		//! Set the uid of the object.
 		virtual void SetUid( const FString& ) = 0;
 
-		//-----------------------------------------------------------------------------------------
-		// Interface pattern
-		//-----------------------------------------------------------------------------------------
-		class Private;
-
 	protected:
 
-		//! Forbidden. Manage with the Ptr<> template.
+		/** Forbidden. Manage with the Ptr<> template. */
 		inline ~NodeObject() {}
 
-	};
+	private:
 
+		static FNodeType StaticType;
+
+	};
 
 }
diff --git a/Engine/Plugins/Mutable/Source/MutableTools/Public/MuT/NodeObjectGroup.h b/Engine/Plugins/Mutable/Source/MutableTools/Public/MuT/NodeObjectGroup.h
--- a/Engine/Plugins/Mutable/Source/MutableTools/Public/MuT/NodeObjectGroup.h
+++ b/Engine/Plugins/Mutable/Source/MutableTools/Public/MuT/NodeObjectGroup.h
@@ -7,42 +7,16 @@
 #include "MuT/Node.h"
 #include "MuT/NodeObject.h"
 
-
 namespace mu
 {
 
-	// Forward definitions
-	class NodeObjectGroup;
-	typedef Ptr<NodeObjectGroup> NodeObjectGroupPtr;
-	typedef Ptr<const NodeObjectGroup>NodeObjectGroupPtrConst;
-
-
-	//! Node that creates a group of objects and describes how they are selected.
-	//! \ingroup model
+	/** Node that creates a group of objects and describes how they are selected. */
 	class MUTABLETOOLS_API NodeObjectGroup : public NodeObject
 	{
 	public:
 
-		NodeObjectGroup();
-
-
-		//-----------------------------------------------------------------------------------------
-		// Node Interface
-		//-----------------------------------------------------------------------------------------
-        const FNodeType* GetType() const override;
-		static const FNodeType* GetStaticType();
-
-		//-----------------------------------------------------------------------------------------
-		// NodeObject Interface
-		//-----------------------------------------------------------------------------------------
-        virtual const FString& GetName() const override;
-		virtual void SetName( const FString&) override;
-		virtual const FString& GetUid() const override;
-		virtual void SetUid( const FString&) override;
-
-		//-----------------------------------------------------------------------------------------
-		// Own Interface
-		//-----------------------------------------------------------------------------------------
+		FString Name;
+		FString Uid;
 
 		//! Typed of child selection
 		typedef enum
@@ -62,46 +36,38 @@ namespace mu
 			//! A boolean parameter will be generated for every child.
 			CS_TOGGLE_EACH
 
-		} CHILD_SELECTION;
+		} EChildSelection;
 
-		//! Get child selection type
-		CHILD_SELECTION GetSelectionType() const;
+		EChildSelection Type;
 
-		//! Set the child selection type
-		void SetSelectionType( CHILD_SELECTION );
+		TArray<Ptr<NodeObject>> Children;
 
-		//! Get the number of child objects
-		int GetChildCount() const;
-
-		//! Set the number of child objects
-		void SetChildCount( int );
-
-		//! Get a child object node
-		NodeObjectPtr GetChild( int index ) const;
-
-		//! Set a child object node
-		void SetChild( int index, NodeObjectPtr );
-
-		//! Set default value for CS_ONE_OR_NONE or CS_ALWAYS_ONE groups
+		//! Default value for CS_ONE_OR_NONE or CS_ALWAYS_ONE groups
 		//! the value is the index of the child option in the group
 		//! -1 is the value for the None option.
 		//! 0 is the first child whether or not the NONE option is present.
-		void SetDefaultValue(int32 Value);
+		int32 DefaultValue;
 
-		//-----------------------------------------------------------------------------------------
-		// Interface pattern
-		//-----------------------------------------------------------------------------------------
-		class Private;
-		Private* GetPrivate() const;
+	public:
+
+		// Node Interface
+		virtual const FNodeType* GetType() const override { return &StaticType; }
+		static const FNodeType* GetStaticType() { return &StaticType; }
+
+		// NodeObject Interface
+        virtual const FString& GetName() const override;
+		virtual void SetName( const FString&) override;
+		virtual const FString& GetUid() const override;
+		virtual void SetUid( const FString&) override;
 
 	protected:
 
-		//! Forbidden. Manage with the Ptr<> template.
-		~NodeObjectGroup();
+		/** Forbidden. Manage with the Ptr<> template. */
+		~NodeObjectGroup() {}
 
 	private:
 
-		Private* m_pD;
+		static FNodeType StaticType;
 
 	};
 
diff --git a/Engine/Plugins/Mutable/Source/MutableTools/Public/MuT/NodeObjectNew.h b/Engine/Plugins/Mutable/Source/MutableTools/Public/MuT/NodeObjectNew.h
--- a/Engine/Plugins/Mutable/Source/MutableTools/Public/MuT/NodeObjectNew.h
+++ b/Engine/Plugins/Mutable/Source/MutableTools/Public/MuT/NodeObjectNew.h
@@ -49,23 +49,17 @@ namespace mu
 
 	public:
 
-		//-----------------------------------------------------------------------------------------
 		// Node Interface
-		//-----------------------------------------------------------------------------------------
-		virtual const FNodeType* GetType() const override { return GetStaticType(); }
+		virtual const FNodeType* GetType() const override { return &StaticType; }
 		static const FNodeType* GetStaticType() { return &StaticType; }
 
-		//-----------------------------------------------------------------------------------------
 		// NodeObject Interface
-		//-----------------------------------------------------------------------------------------
         virtual const FString& GetName() const override;
 		virtual void SetName( const FString& ) override;
 		virtual const FString& GetUid() const override;
 		virtual void SetUid( const FString& ) override;
 
-		//-----------------------------------------------------------------------------------------
 		// Own Interface
-		//-----------------------------------------------------------------------------------------
 
 		//! Set the number of states that the model can be in.
 		int32 GetStateCount() const;
@@ -80,9 +74,6 @@ namespace mu
 		//! Add a runtime parameter to the state.
 		void AddStateParam( int32 s, const FString& param );
 
-		//! Remove a runtime parameter from a state.
-		void RemoveStateParam( int32 s, const FString& param );
-
         //! Set the optimisation properties of a state
 		void SetStateProperties(int32 StateIndex,
 			ETextureCompressionStrategy TextureCompressionStrategy,
@@ -94,14 +85,13 @@ namespace mu
 
 	protected:
 
-		//! Forbidden. Manage with the Ptr<> template.
+		/** Forbidden. Manage with the Ptr<> template. */
 		~NodeObjectNew() {}
 
 	private:
 
 		static FNodeType StaticType;
 
-
 	};
 
 
diff --git a/Engine/Plugins/Mutable/Source/MutableTools/Public/MuT/NodeProjector.h b/Engine/Plugins/Mutable/Source/MutableTools/Public/MuT/NodeProjector.h
--- a/Engine/Plugins/Mutable/Source/MutableTools/Public/MuT/NodeProjector.h
+++ b/Engine/Plugins/Mutable/Source/MutableTools/Public/MuT/NodeProjector.h
@@ -7,51 +7,30 @@
 #include "MuR/Ptr.h"
 #include "MuR/RefCounted.h"
 #include "MuT/Node.h"
+#include "MuT/NodeRange.h"
 
 
 namespace mu
 {
 
-	class NodeProjector;
-    using NodeProjectorPtr = Ptr<NodeProjector>;
-    using NodeProjectorPtrConst = Ptr<const NodeProjector>;
-
-	class NodeProjectorConstant;
-    using NodeProjectorConstantPtr = Ptr<NodeProjectorConstant>;
-    using NodeProjectorConstantPtrConst = Ptr<const NodeProjectorConstant>;
-
-	class NodeProjectorParameter;
-    using NodeProjectorParameterPtr = Ptr<NodeProjectorParameter>;
-    using NodeProjectorParameterPtrConst = Ptr<const NodeProjectorParameter>;
-
-    class NodeRange;
-    using NodeRangePtr = Ptr<NodeRange>;
-    using NodeRangePtrConst = Ptr<const NodeRange>;
-
-
     //! Base class of any node that outputs a Projector.
 	class MUTABLETOOLS_API NodeProjector : public Node
 	{
 	public:
 
-		//-----------------------------------------------------------------------------------------
 		// Node Interface
-		//-----------------------------------------------------------------------------------------
-
-        const FNodeType* GetType() const override;
-		static const FNodeType* GetStaticType();
-
-
-		//-----------------------------------------------------------------------------------------
-		// Interface pattern
-		//-----------------------------------------------------------------------------------------
-		class Private;
+		virtual const FNodeType* GetType() const override { return &StaticType; }
+		static const FNodeType* GetStaticType() { return &StaticType; }
 
 	protected:
 
-		//! Forbidden. Manage with the Ptr<> template.
+		/** Forbidden. Manage with the Ptr<> template. */
 		inline ~NodeProjector() {}
 
+	private:
+
+		static FNodeType StaticType;
+
 	};
 
 
@@ -60,18 +39,20 @@ namespace mu
 	{
 	public:
 
-		NodeProjectorConstant();
+		EProjectorType Type = EProjectorType::Planar;
+		FVector3f Position;
+		FVector3f Direction;
+		FVector3f Up;
+		FVector3f Scale;
+		float ProjectionAngle = 0.0f;
 
-		//-----------------------------------------------------------------------------------------
-		// Node Interface
-		//-----------------------------------------------------------------------------------------
+	public:
 
-        const FNodeType* GetType() const override;
-		static const FNodeType* GetStaticType();
+		// Node Interface
+		virtual const FNodeType* GetType() const override { return &StaticType; }
+		static const FNodeType* GetStaticType() { return &StaticType; }
 
-		//-----------------------------------------------------------------------------------------
 		// Own Interface
-		//-----------------------------------------------------------------------------------------
 
 		//! Get the value that this node returns
         void GetValue( EProjectorType* OutType,
@@ -84,20 +65,14 @@ namespace mu
 		//! Set the value to be returned by this node
         void SetValue( EProjectorType type, FVector3f pos, FVector3f dir, FVector3f up, FVector3f scale, float projectionAngle);
 
-		//-----------------------------------------------------------------------------------------
-		// Interface pattern
-		//-----------------------------------------------------------------------------------------
-		class Private;
-		Private* GetPrivate() const;
-
 	protected:
 
-		//! Forbidden. Manage with the Ptr<> template.
-		~NodeProjectorConstant();
+		/** Forbidden. Manage with the Ptr<> template. */
+		~NodeProjectorConstant() {}
 
 	private:
 
-		Private* m_pD;
+		static FNodeType StaticType;
 
 	};
 
@@ -107,23 +82,29 @@ namespace mu
 	{
 	public:
 
-		NodeProjectorParameter();
+		EProjectorType Type = EProjectorType::Planar;
+		FVector3f Position;
+		FVector3f Direction;
+		FVector3f Up;
+		FVector3f Scale;
+		float ProjectionAngle = 0.0f;
 
-		//-----------------------------------------------------------------------------------------
-		// Node Interface
-		//-----------------------------------------------------------------------------------------
+		FString Name;
+		FString UID;
 
-        const FNodeType* GetType() const override;
-		static const FNodeType* GetStaticType();
+		TArray<Ptr<NodeRange>> Ranges;
+
+	public:
+
+		// Node Interface
+		virtual const FNodeType* GetType() const override { return &StaticType; }
+		static const FNodeType* GetStaticType() { return &StaticType; }
 
-		//-----------------------------------------------------------------------------------------
 		// Own Interface
-		//-----------------------------------------------------------------------------------------
 
 		//! Set the name of the parameter.
 		void SetName( const FString& );
 
-
 		//! Get the uid of the parameter. It will be exposed in the final compiled data.
 		const FString& GetUid() const;
 		void SetUid( const FString& );
@@ -138,23 +119,17 @@ namespace mu
 
         //! Set the number of ranges (dimensions) for this parameter.
         //! By default a parameter has 0 ranges, meaning it only has one value.
-        void SetRangeCount( int i );
-        void SetRange( int i, NodeRangePtr pRange );
-
-		//-----------------------------------------------------------------------------------------
-		// Interface pattern
-		//-----------------------------------------------------------------------------------------
-		class Private;
-		Private* GetPrivate() const;
+        void SetRangeCount( int32 i );
+        void SetRange( int32 i, Ptr<NodeRange> Range );
 
 	protected:
 
-		//! Forbidden. Manage with the Ptr<> template.
-		~NodeProjectorParameter();
+		/** Forbidden. Manage with the Ptr<> template. */
+		~NodeProjectorParameter() {}
 
 	private:
 
-		Private* m_pD;
+		static FNodeType StaticType;
 
 	};
 
diff --git a/Engine/Plugins/Mutable/Source/MutableTools/Public/MuT/NodeRange.h b/Engine/Plugins/Mutable/Source/MutableTools/Public/MuT/NodeRange.h
--- a/Engine/Plugins/Mutable/Source/MutableTools/Public/MuT/NodeRange.h
+++ b/Engine/Plugins/Mutable/Source/MutableTools/Public/MuT/NodeRange.h
@@ -7,41 +7,28 @@
 #include "MuR/RefCounted.h"
 #include "MuT/Node.h"
 
-
 namespace mu
 {
 
-    // Forward definitions
-    class NodeRange;
-    typedef Ptr<NodeRange> NodeRangePtr;
-    typedef Ptr<const NodeRange> NodeRangePtrConst;
-
-
-    //! %Base class of any node that outputs a range.
+    /** Base class of any node that outputs a range. */
     class MUTABLETOOLS_API NodeRange : public Node
 	{
 	public:
 
-		//-----------------------------------------------------------------------------------------
 		// Node Interface
-		//-----------------------------------------------------------------------------------------
-
-		const FNodeType* GetType() const override;
-		static const FNodeType* GetStaticType();
-
-
-		//-----------------------------------------------------------------------------------------
-		// Interface pattern
-		//-----------------------------------------------------------------------------------------
-		class Private;
+		virtual const FNodeType* GetType() const override { return &StaticType; }
+		static const FNodeType* GetStaticType() { return &StaticType; }
 
 	protected:
 
-		//! Forbidden. Manage with the Ptr<> template.
+		/** Forbidden. Manage with the Ptr<> template. */
         inline ~NodeRange() {}
 
-	};
+	private:
 
+		static FNodeType StaticType;
+
+	};
 
 }
 
diff --git a/Engine/Plugins/Mutable/Source/MutableTools/Public/MuT/NodeRangeFromScalar.h b/Engine/Plugins/Mutable/Source/MutableTools/Public/MuT/NodeRangeFromScalar.h
--- a/Engine/Plugins/Mutable/Source/MutableTools/Public/MuT/NodeRangeFromScalar.h
+++ b/Engine/Plugins/Mutable/Source/MutableTools/Public/MuT/NodeRangeFromScalar.h
@@ -3,8 +3,8 @@
 #pragma once
 
 #include "MuR/Ptr.h"
-#include "MuT/Node.h"
 #include "MuT/NodeRange.h"
+#include "MuT/NodeScalar.h"
 #include "Containers/UnrealString.h"
 
 namespace mu
@@ -17,48 +17,25 @@ namespace mu
 	{
 	public:
 
-		//-----------------------------------------------------------------------------------------
-		// Life cycle
-		//-----------------------------------------------------------------------------------------
+		Ptr<NodeScalar> Size;
 
-        NodeRangeFromScalar();
+		FString Name;
 
+	public:
 
-		//-----------------------------------------------------------------------------------------
 		// Node Interface
-		//-----------------------------------------------------------------------------------------
-
-		const FNodeType* GetType() const override;
-		static const FNodeType* GetStaticType();
-
-		//-----------------------------------------------------------------------------------------
-		// Own Interface
-		//-----------------------------------------------------------------------------------------
-
-        //!
-        Ptr<NodeScalar> GetSize() const;
-        void SetSize( const Ptr<NodeScalar>& );
-
-        //!
-		const FString& GetName() const;
-		void SetName(const FString& Name);
-
-		//-----------------------------------------------------------------------------------------
-		// Interface pattern
-		//-----------------------------------------------------------------------------------------
-		class Private;
-		Private* GetPrivate() const;
+		virtual const FNodeType* GetType() const override { return &StaticType; }
+		static const FNodeType* GetStaticType() { return &StaticType; }
 
 	protected:
 
-		//! Forbidden. Manage with the Ptr<> template.
-        ~NodeRangeFromScalar();
+		/** Forbidden. Manage with the Ptr<> template. */
+		~NodeRangeFromScalar() {}
 
 	private:
 
-		Private* m_pD;
+		static FNodeType StaticType;
 
 	};
 
-
 }
diff --git a/Engine/Plugins/Mutable/Source/MutableTools/Public/MuT/NodeScalar.h b/Engine/Plugins/Mutable/Source/MutableTools/Public/MuT/NodeScalar.h
--- a/Engine/Plugins/Mutable/Source/MutableTools/Public/MuT/NodeScalar.h
+++ b/Engine/Plugins/Mutable/Source/MutableTools/Public/MuT/NodeScalar.h
@@ -2,45 +2,29 @@
 
 #pragma once
 
-#include "HAL/Platform.h"
-#include "MuR/Ptr.h"
-#include "MuR/RefCounted.h"
 #include "MuT/Node.h"
 
-
 namespace mu
 {
 
-	// Forward definitions
-	class NodeScalar;
-	typedef Ptr<NodeScalar> NodeScalarPtr;
-	typedef Ptr<const NodeScalar> NodeScalarPtrConst;
-
-
-    //! %Base class of any node that outputs a scalar value.
+    /** Base class of any node that outputs a scalar value. */
 	class MUTABLETOOLS_API NodeScalar : public Node
 	{
 	public:
 
-		//-----------------------------------------------------------------------------------------
 		// Node Interface
-		//-----------------------------------------------------------------------------------------
-
-		const FNodeType* GetType() const override;
-		static const FNodeType* GetStaticType();
-
-
-		//-----------------------------------------------------------------------------------------
-		// Interface pattern
-		//-----------------------------------------------------------------------------------------
-		class Private;
+		virtual const FNodeType* GetType() const override { return &StaticType; }
+		static const FNodeType* GetStaticType() { return &StaticType; }
 
 	protected:
 
-		//! Forbidden. Manage with the Ptr<> template.
+		/** Forbidden. Manage with the Ptr<> template. */
 		inline ~NodeScalar() {}
 
-	};
+	private:
 
+		static FNodeType StaticType;
+
+	};
 
 }
diff --git a/Engine/Plugins/Mutable/Source/MutableTools/Public/MuT/NodeScalarArithmeticOperation.h b/Engine/Plugins/Mutable/Source/MutableTools/Public/MuT/NodeScalarArithmeticOperation.h
--- a/Engine/Plugins/Mutable/Source/MutableTools/Public/MuT/NodeScalarArithmeticOperation.h
+++ b/Engine/Plugins/Mutable/Source/MutableTools/Public/MuT/NodeScalarArithmeticOperation.h
@@ -7,80 +7,43 @@
 #include "MuT/Node.h"
 #include "MuT/NodeScalar.h"
 
-
 namespace mu
 {
 
-	// Forward definitions
-    class NodeScalarArithmeticOperation;
-    typedef Ptr<NodeScalarArithmeticOperation> NodeScalarArithmeticOperationPtr;
-    typedef Ptr<const NodeScalarArithmeticOperation> NodeScalarArithmeticOperationPtrConst;
-
-
-    //! Perform an arithmetic operation between two scalars.
-	//! \ingroup model
+    /** Perform an arithmetic operation between two scalars. */
     class MUTABLETOOLS_API NodeScalarArithmeticOperation : public NodeScalar
 	{
 	public:
 
-        NodeScalarArithmeticOperation();
-
-
-		//-----------------------------------------------------------------------------------------
-		// Node Interface
-		//-----------------------------------------------------------------------------------------
-
-        const FNodeType* GetType() const override;
-		static const FNodeType* GetStaticType();
-
-
-		//-----------------------------------------------------------------------------------------
-		// Own Interface
-		//-----------------------------------------------------------------------------------------
-
-		//! Possible arithmetic operations
+		/** Possible arithmetic operations. */
 		typedef enum
 		{
-            AO_ADD = 0,
+			AO_ADD = 0,
 			AO_SUBTRACT,
 			AO_MULTIPLY,
-            AO_DIVIDE,
-            _AO_COUNT
-		} OPERATION;
-
-        static const char* s_opTypeName[ _AO_COUNT ];
-
-		//! Get child selection type
-		OPERATION GetOperation() const;
-
-		//! Set the child selection type
-		void SetOperation(OPERATION);
+			AO_DIVIDE,
+			_AO_COUNT
+		} EOperation;
 
+		EOperation Operation;
+		Ptr<NodeScalar> A;
+		Ptr<NodeScalar> B;
 
-		//! Get the first operand
-        NodeScalarPtr GetA() const;
-        void SetA(NodeScalarPtr);
-
-		// Get the second operand
-        NodeScalarPtr GetB() const;
-        void SetB(NodeScalarPtr);
+	public:
 
-		//-----------------------------------------------------------------------------------------
-		// Interface pattern
-		//-----------------------------------------------------------------------------------------
-		class Private;
-		Private* GetPrivate() const;
+		// Node Interface
+		virtual const FNodeType* GetType() const override { return &StaticType; }
+		static const FNodeType* GetStaticType() { return &StaticType; }
 
 	protected:
 
-		//! Forbidden. Manage with the Ptr<> template.
-        ~NodeScalarArithmeticOperation();
+		/** Forbidden. Manage with the Ptr<> template. */
+		~NodeScalarArithmeticOperation() {}
 
 	private:
 
-		Private* m_pD;
+		static FNodeType StaticType;
 
 	};
 
-
 }
diff --git a/Engine/Plugins/Mutable/Source/MutableTools/Public/MuT/NodeScalarConstant.h b/Engine/Plugins/Mutable/Source/MutableTools/Public/MuT/NodeScalarConstant.h
--- a/Engine/Plugins/Mutable/Source/MutableTools/Public/MuT/NodeScalarConstant.h
+++ b/Engine/Plugins/Mutable/Source/MutableTools/Public/MuT/NodeScalarConstant.h
@@ -7,59 +7,31 @@
 #include "MuT/Node.h"
 #include "MuT/NodeScalar.h"
 
-
 namespace mu
 {
 
-	// Forward definitions
-	class NodeScalarConstant;
-	typedef Ptr<NodeScalarConstant> NodeScalarConstantPtr;
-	typedef Ptr<const NodeScalarConstant> NodeScalarConstantPtrConst;
-
-
-	//! Node returning a scalar constant value.
-	//! \ingroup model
+	/** Node returning a scalar constant value. */
 	class MUTABLETOOLS_API NodeScalarConstant : public NodeScalar
 	{
 	public:
 
-		NodeScalarConstant();
+		float Value;
 
+	public:
 
-		//-----------------------------------------------------------------------------------------
 		// Node Interface
-		//-----------------------------------------------------------------------------------------
-
-        const FNodeType* GetType() const override;
-		static const FNodeType* GetStaticType();
-
-		//-----------------------------------------------------------------------------------------
-		// Own Interface
-		//-----------------------------------------------------------------------------------------
-
-		//! Get the value to be returned by the node.
-		float GetValue() const;
-
-		//! Set the value to be returned by the node.
-		void SetValue( float v );
-
-		//-----------------------------------------------------------------------------------------
-		// Interface pattern
-		//-----------------------------------------------------------------------------------------
-		class Private;
-		Private* GetPrivate() const;
+		virtual const FNodeType* GetType() const override { return &StaticType; }
+		static const FNodeType* GetStaticType() { return &StaticType; }
 
 	protected:
 
-		//! Forbidden. Manage with the Ptr<> template.
-		~NodeScalarConstant();
+		/** Forbidden. Manage with the Ptr<> template. */
+		~NodeScalarConstant() {}
 
 	private:
 
-		Private* m_pD;
+		static FNodeType StaticType;
 
 	};
 
-
-
 }
diff --git a/Engine/Plugins/Mutable/Source/MutableTools/Public/MuT/NodeScalarCurve.h b/Engine/Plugins/Mutable/Source/MutableTools/Public/MuT/NodeScalarCurve.h
--- a/Engine/Plugins/Mutable/Source/MutableTools/Public/MuT/NodeScalarCurve.h
+++ b/Engine/Plugins/Mutable/Source/MutableTools/Public/MuT/NodeScalarCurve.h
@@ -6,7 +6,6 @@
 #include "MuR/RefCounted.h"
 #include "MuT/Node.h"
 #include "MuT/NodeScalar.h"
-
 #include "Curves/RichCurve.h"
 
 namespace mu
@@ -22,35 +21,19 @@ namespace mu
 
 	public:
 
-        NodeScalarCurve();
-
-		//-----------------------------------------------------------------------------------------
 		// Node Interface
-		//-----------------------------------------------------------------------------------------
-
-        const FNodeType* GetType() const override;
-		static const FNodeType* GetStaticType();
-
-		//-----------------------------------------------------------------------------------------
-		// Own Interface
-		//-----------------------------------------------------------------------------------------
-
-		//-----------------------------------------------------------------------------------------
-		// Interface pattern
-		//-----------------------------------------------------------------------------------------
-		class Private;
-		Private* GetPrivate() const;
+		virtual const FNodeType* GetType() const override { return &StaticType; }
+		static const FNodeType* GetStaticType() { return &StaticType; }
 
 	protected:
 
-		//! Forbidden. Manage with the Ptr<> template.
-        ~NodeScalarCurve();
+		/** Forbidden. Manage with the Ptr<> template. */
+		~NodeScalarCurve() {}
 
 	private:
 
-		Private* m_pD;
+		static FNodeType StaticType;
 
 	};
 
-
 }
diff --git a/Engine/Plugins/Mutable/Source/MutableTools/Public/MuT/NodeScalarEnumParameter.h b/Engine/Plugins/Mutable/Source/MutableTools/Public/MuT/NodeScalarEnumParameter.h
--- a/Engine/Plugins/Mutable/Source/MutableTools/Public/MuT/NodeScalarEnumParameter.h
+++ b/Engine/Plugins/Mutable/Source/MutableTools/Public/MuT/NodeScalarEnumParameter.h
@@ -8,83 +8,43 @@
 #include "MuT/NodeRange.h"
 #include "MuT/NodeScalar.h"
 
-
 namespace mu
 {
 
-	// Forward definitions
-	class NodeImage;
-    using NodeImagePtr = Ptr<NodeImage>;
-    using NodeImagePtrConst = Ptr<const NodeImage>;
-
-	class NodeScalarEnumParameter;
-    using NodeScalarEnumParameterPtr = Ptr<NodeScalarEnumParameter>;
-    using NodeScalarEnumParameterPtrConst = Ptr<const NodeScalarEnumParameter>;
-
-    class NodeRange;
-    using NodeRangePtr = Ptr<NodeRange>;
-    using NodeRangePtrConst = Ptr<const NodeRange>;
-
-
-	//! Node that defines a scalar model parameter to be selected from a set of named values.
+	/** Node that defines a scalar model parameter to be selected from a set of named values. */
 	class MUTABLETOOLS_API NodeScalarEnumParameter : public NodeScalar
 	{
 	public:
 
-		NodeScalarEnumParameter();
-
-		//-----------------------------------------------------------------------------------------
-		// Node Interface
-		//-----------------------------------------------------------------------------------------
-        const FNodeType* GetType() const override;
-		static const FNodeType* GetStaticType();
-
-		//-----------------------------------------------------------------------------------------
-		// Own Interface
-		//-----------------------------------------------------------------------------------------
-
-		//! Get the name of the parameter. It will be exposed in the final compiled data.
-		const FString& GetName() const;
-		void SetName( const FString&);
-
-		//! Get the uid of the parameter. It will be exposed in the final compiled data.
-		const FString& GetUid() const;
-		void SetUid( const FString&);
-
-		//! Get the index of the default value of the parameter.
-		int GetDefaultValueIndex() const;
-		void SetDefaultValueIndex( int v );
+		FString Name;
+		FString UID;
+		int32 DefaultValue = 0;
 
-		//! Set the number of possible values.
-		void SetValueCount( int i );
+		struct FOption
+		{
+			FString Name;
+			float Value = 0.0f;
+		};
 
-		//! Get the number of possible values.
-		int GetValueCount() const;
+		TArray<FOption> Options;
 
-		//! Set the data of one of the possible values of the parameter.
-		void SetValue( int i, float value, const FString& strName );
+		TArray<Ptr<NodeRange>> Ranges;
 
-        //! Set the number of ranges (dimensions) for this parameter.
-        //! By default a parameter has 0 ranges, meaning it only has one value.
-        void SetRangeCount( int i );
-        void SetRange( int i, NodeRangePtr pRange );
+	public:
 
-		//-----------------------------------------------------------------------------------------
-		// Interface pattern
-		//-----------------------------------------------------------------------------------------
-		class Private;
-		Private* GetPrivate() const;
+		// Node Interface
+		virtual const FNodeType* GetType() const override { return &StaticType; }
+		static const FNodeType* GetStaticType() { return &StaticType; }
 
 	protected:
 
-		//! Forbidden. Manage with the Ptr<> template.
-		~NodeScalarEnumParameter();
+		/** Forbidden. Manage with the Ptr<> template. */
+		~NodeScalarEnumParameter() {}
 
 	private:
 
-		Private* m_pD;
+		static FNodeType StaticType;
 
 	};
 
-
 }
diff --git a/Engine/Plugins/Mutable/Source/MutableTools/Public/MuT/NodeScalarParameter.h b/Engine/Plugins/Mutable/Source/MutableTools/Public/MuT/NodeScalarParameter.h
--- a/Engine/Plugins/Mutable/Source/MutableTools/Public/MuT/NodeScalarParameter.h
+++ b/Engine/Plugins/Mutable/Source/MutableTools/Public/MuT/NodeScalarParameter.h
@@ -7,75 +7,36 @@
 #include "MuR/RefCounted.h"
 #include "MuT/Node.h"
 #include "MuT/NodeScalar.h"
-
+#include "MuT/NodeRange.h"
 
 namespace mu
 {
 
-	// Forward definitions
-    class NodeImage;
-    using NodeImagePtr = Ptr<NodeImage>;
-    using NodeImagePtrConst = Ptr<const NodeImage>;
-
-    class NodeRange;
-    using NodeRangePtr = Ptr<NodeRange>;
-    using NodeRangePtrConst = Ptr<const NodeRange>;
-
-    class NodeScalarParameter;
-    using NodeScalarParameterPtr = Ptr<NodeScalarParameter>;
-    using NodeScalarParameterPtrConst = Ptr<const NodeScalarParameter>;
-
-
 	//! Node that defines a scalar model parameter.
 	class MUTABLETOOLS_API NodeScalarParameter : public NodeScalar
 	{
 	public:
 
-		//-----------------------------------------------------------------------------------------
-		// Life cycle
-		//-----------------------------------------------------------------------------------------
-
-		NodeScalarParameter();
-
-		//-----------------------------------------------------------------------------------------
-		// Node Interface
-		//-----------------------------------------------------------------------------------------
-
-        const FNodeType* GetType() const override;
-		static const FNodeType* GetStaticType();
-
-		//-----------------------------------------------------------------------------------------
-		// Own Interface
-		//-----------------------------------------------------------------------------------------
-
-		//! Get the name of the parameter. It will be exposed in the final compiled data.
-		void SetName( const FString& );
+		float DefaultValue = 0.0f;
+		FString Name;
+		FString UID;
 
-		//! Get the uid of the parameter. It will be exposed in the final compiled data.
-		void SetUid( const FString&);
+		TArray<Ptr<NodeRange>> Ranges;
 
-		//! Get the default value of the parameter.
-		void SetDefaultValue( float v );
-
-        //! Set the number of ranges (dimensions) for this parameter.
-        //! By default a parameter has 0 ranges, meaning it only has one value.
-        void SetRangeCount( int i );
-        void SetRange( int i, NodeRangePtr pRange );
+	public:
 
-		//-----------------------------------------------------------------------------------------
-		// Interface pattern
-		//-----------------------------------------------------------------------------------------
-		class Private;
-		Private* GetPrivate() const;
+		// Node Interface
+		virtual const FNodeType* GetType() const override { return &StaticType; }
+		static const FNodeType* GetStaticType() { return &StaticType; }
 
 	protected:
 
-		//! Forbidden. Manage with the Ptr<> template.
-		~NodeScalarParameter();
+		/** Forbidden. Manage with the Ptr<> template. */
+		~NodeScalarParameter() {}
 
 	private:
 
-		Private* m_pD;
+		static FNodeType StaticType;
 
 	};
 
diff --git a/Engine/Plugins/Mutable/Source/MutableTools/Public/MuT/NodeScalarTable.h b/Engine/Plugins/Mutable/Source/MutableTools/Public/MuT/NodeScalarTable.h
--- a/Engine/Plugins/Mutable/Source/MutableTools/Public/MuT/NodeScalarTable.h
+++ b/Engine/Plugins/Mutable/Source/MutableTools/Public/MuT/NodeScalarTable.h
@@ -8,17 +8,10 @@
 #include "MuT/NodeScalar.h"
 #include "MuT/Table.h"
 
-
 namespace mu
 {
 
-	// Forward definitions
-	class NodeScalarTable;
-	typedef Ptr<NodeScalarTable> NodeScalarTablePtr;
-	typedef Ptr<const NodeScalarTable> NodeScalarTablePtrConst;
-
-
-	//! This node provides the meshes stored in the column of a table.
+	/** This node provides the meshes stored in the column of a table. */
 	class MUTABLETOOLS_API NodeScalarTable : public NodeScalar
 	{
 	public:
@@ -35,26 +28,9 @@ namespace mu
 		virtual const FNodeType* GetType() const override { return GetStaticType(); }
 		static const FNodeType* GetStaticType() { return &StaticType; }
 
-		//-----------------------------------------------------------------------------------------
-		// Own Interface
-		//-----------------------------------------------------------------------------------------
-
-		//! Set the name of the implicit table parameter.
-		void SetParameterName( const FString& strName );
-
-		//!
-		void SetColumn( const FString& strName );
-
-		//! Adds the "None" option to the parameter that represents this table column
-		void SetNoneOption(bool bAddOption);
-
-		//! Set the row name to be used as default value
-		void SetDefaultRowName(const FString& RowName);
-
-
 	protected:
 
-		//! Forbidden. Manage with the Ptr<> template.
+		/** Forbidden. Manage with the Ptr<> template. */
 		~NodeScalarTable() {}
 
 	private:
@@ -63,5 +39,4 @@ namespace mu
 
 	};
 
-
 }
diff --git a/Engine/Plugins/Mutable/Source/MutableTools/Public/MuT/NodeScalarVariation.h b/Engine/Plugins/Mutable/Source/MutableTools/Public/MuT/NodeScalarVariation.h
--- a/Engine/Plugins/Mutable/Source/MutableTools/Public/MuT/NodeScalarVariation.h
+++ b/Engine/Plugins/Mutable/Source/MutableTools/Public/MuT/NodeScalarVariation.h
@@ -7,61 +7,39 @@
 #include "MuT/Node.h"
 #include "MuT/NodeScalar.h"
 
-
 namespace mu
 {
 
-    // Forward definitions
-    class NodeScalarVariation;
-    typedef Ptr<NodeScalarVariation> NodeScalarVariationPtr;
-    typedef Ptr<const NodeScalarVariation> NodeScalarVariationPtrConst;
-
-
-     class MUTABLETOOLS_API NodeScalarVariation : public NodeScalar
+	/** */
+    class MUTABLETOOLS_API NodeScalarVariation : public NodeScalar
     {
-    public:
- 
-        NodeScalarVariation();
-
-        //-----------------------------------------------------------------------------------------
-        // Node Interface
-        //-----------------------------------------------------------------------------------------
-
-        const FNodeType* GetType() const override;
-        static const FNodeType* GetStaticType();
-
-        //-----------------------------------------------------------------------------------------
-        // Own Interface
-        //-----------------------------------------------------------------------------------------
+	public:
 
-        //!
-        void SetDefaultScalar( NodeScalar* Scalar );
+		Ptr<NodeScalar> DefaultScalar;
 
-        //! Set the number of tags to consider in this variation
-        void SetVariationCount( int32 count );
+		struct FVariation
+		{
+			Ptr<NodeScalar> Scalar;
+			FString Tag;
+		};
 
-        //! Set the tag or state name that will enable a specific vartiation
-        void SetVariationTag( int32 index, const FString& strTag );
+		TArray<FVariation> Variations;
 
-        //!
-        void SetVariationScalar( int32 index, NodeScalar* Scalar );
-
-        //!}
-
-
-        //-----------------------------------------------------------------------------------------
-        // Interface pattern
-        //-----------------------------------------------------------------------------------------
-        class Private;
-        Private* GetPrivate() const;
+    public:
+ 
+		// Node Interface
+		virtual const FNodeType* GetType() const override { return &StaticType; }
+		static const FNodeType* GetStaticType() { return &StaticType; }
 
     protected:
-        //! Forbidden. Manage with the Ptr<> template.
-        ~NodeScalarVariation();
+
+		/** Forbidden. Manage with the Ptr<> template. */
+		~NodeScalarVariation() {}
 
     private:
-        Private* m_pD;
-    };
 
+		static FNodeType StaticType;
+
+	};
 
-} // namespace mu
+}
diff --git a/Engine/Plugins/Mutable/Source/MutableTools/Public/MuT/NodeString.h b/Engine/Plugins/Mutable/Source/MutableTools/Public/MuT/NodeString.h
--- a/Engine/Plugins/Mutable/Source/MutableTools/Public/MuT/NodeString.h
+++ b/Engine/Plugins/Mutable/Source/MutableTools/Public/MuT/NodeString.h
@@ -11,34 +11,23 @@
 namespace mu
 {
 
-	// Forward definitions
-	class NodeString;
-	typedef Ptr<NodeString> NodeStringPtr;
-	typedef Ptr<const NodeString> NodeStringPtrConst;
-
-
-    //! %Base class of any node that outputs a string value.
+    /** Base class of any node that outputs a string value. */
 	class MUTABLETOOLS_API NodeString : public Node
 	{
 	public:
 
-		//-----------------------------------------------------------------------------------------
 		// Node Interface
-		//-----------------------------------------------------------------------------------------
-
-		const FNodeType* GetType() const override;
-		static const FNodeType* GetStaticType();
+		virtual const FNodeType* GetType() const override { return &StaticType; }
+		static const FNodeType* GetStaticType() { return &StaticType; }
 
+	protected:
 
-		//-----------------------------------------------------------------------------------------
-		// Interface pattern
-		//-----------------------------------------------------------------------------------------
-		class Private;
+		/** Forbidden. Manage with the Ptr<> template. */
+		~NodeString() {}
 
-	protected:
+	private:
 
-		//! Forbidden. Manage with the Ptr<> template.
-		inline ~NodeString() {}
+		static FNodeType StaticType;
 
 	};
 
diff --git a/Engine/Plugins/Mutable/Source/MutableTools/Public/MuT/NodeStringConstant.h b/Engine/Plugins/Mutable/Source/MutableTools/Public/MuT/NodeStringConstant.h
--- a/Engine/Plugins/Mutable/Source/MutableTools/Public/MuT/NodeStringConstant.h
+++ b/Engine/Plugins/Mutable/Source/MutableTools/Public/MuT/NodeStringConstant.h
@@ -11,48 +11,27 @@
 namespace mu
 {
 
-	// Forward definitions
-	class NodeStringConstant;
-	typedef Ptr<NodeStringConstant> NodeStringConstantPtr;
-	typedef Ptr<const NodeStringConstant> NodeStringConstantPtrConst;
-
-
-	//! Node returning a string constant value.
-	//! \ingroup model
+	/** Node returning a string constant value. */
 	class MUTABLETOOLS_API NodeStringConstant : public NodeString
 	{
 	public:
 
-		NodeStringConstant();
+		FString Value;
 
+	public:
 
-		//-----------------------------------------------------------------------------------------
 		// Node Interface
-		//-----------------------------------------------------------------------------------------
-        const FNodeType* GetType() const override;
-		static const FNodeType* GetStaticType();
-
-		//-----------------------------------------------------------------------------------------
-		// Own Interface
-		//-----------------------------------------------------------------------------------------
-
-		//! Set the value to be returned by the node.
-		void SetValue( const FString& v );
-
-		//-----------------------------------------------------------------------------------------
-		// Interface pattern
-		//-----------------------------------------------------------------------------------------
-		class Private;
-		Private* GetPrivate() const;
+		virtual const FNodeType* GetType() const override { return &StaticType; }
+		static const FNodeType* GetStaticType() { return &StaticType; }
 
 	protected:
 
-		//! Forbidden. Manage with the Ptr<> template.
-		~NodeStringConstant();
+		/** Forbidden. Manage with the Ptr<> template. */
+		~NodeStringConstant() {}
 
 	private:
 
-		Private* m_pD;
+		static FNodeType StaticType;
 
 	};
 
diff --git a/Engine/Plugins/Mutable/Source/MutableTools/Public/MuT/NodeStringParameter.h b/Engine/Plugins/Mutable/Source/MutableTools/Public/MuT/NodeStringParameter.h
--- a/Engine/Plugins/Mutable/Source/MutableTools/Public/MuT/NodeStringParameter.h
+++ b/Engine/Plugins/Mutable/Source/MutableTools/Public/MuT/NodeStringParameter.h
@@ -7,71 +7,36 @@
 #include "MuR/RefCounted.h"
 #include "MuT/Node.h"
 #include "MuT/NodeString.h"
-
+#include "MuT/NodeRange.h"
 
 namespace mu
 {
 
-	// Forward definitions
-    class NodeImage;
-    using NodeImagePtr = Ptr<NodeImage>;
-    using NodeImagePtrConst = Ptr<const NodeImage>;
-
-    class NodeRange;
-    using NodeRangePtr = Ptr<NodeRange>;
-    using NodeRangePtrConst = Ptr<const NodeRange>;
-
-    class NodeStringParameter;
-    using NodeStringParameterPtr = Ptr<NodeStringParameter>;
-    using NodeStringParameterPtrConst = Ptr<const NodeStringParameter>;
-
-
-	//! Node that defines a string model parameter.
+	/** Node that defines a string model parameter. */
 	class MUTABLETOOLS_API NodeStringParameter : public NodeString
 	{
 	public:
 
-		NodeStringParameter();
-
-		//-----------------------------------------------------------------------------------------
-		// Node Interface
-		//-----------------------------------------------------------------------------------------
+		FString DefaultValue;
+		FString Name;
+		FString UID;
 
-        const FNodeType* GetType() const override;
-		static const FNodeType* GetStaticType();
+		TArray<Ptr<NodeRange>> Ranges;
 
-		//-----------------------------------------------------------------------------------------
-		// Own Interface
-		//-----------------------------------------------------------------------------------------
-
-		//! Set the name of the parameter. It will be exposed in the final compiled data.
-		void SetName( const FString& );
-
-		//! Set the uid of the parameter. It will be exposed in the final compiled data.
-		void SetUid( const FString& );
-
-		//! Set the default value of the parameter.
-		void SetDefaultValue( const FString& );
-
-        //! Set the number of ranges (dimensions) for this parameter.
-        //! By default a parameter has 0 ranges, meaning it only has one value.
-        void SetRangeCount( int i );
-        void SetRange( int i, NodeRangePtr pRange );
+	public:
 
-		//-----------------------------------------------------------------------------------------
-		// Interface pattern
-		//-----------------------------------------------------------------------------------------
-		class Private;
-		Private* GetPrivate() const;
+		// Node Interface
+		virtual const FNodeType* GetType() const override { return &StaticType; }
+		static const FNodeType* GetStaticType() { return &StaticType; }
 
 	protected:
 
-		//! Forbidden. Manage with the Ptr<> template.
-		~NodeStringParameter();
+		/** Forbidden. Manage with the Ptr<> template. */
+		~NodeStringParameter() {}
 
 	private:
 
-		Private* m_pD;
+		static FNodeType StaticType;
 
 	};
 
