[AutoRTFM] Enable AutoRTFM for MutableRuntime.

This is being called from the closed.

diff --git a/Engine/Plugins/Mutable/Source/MutableRuntime/MutableRuntime.Build.cs b/Engine/Plugins/Mutable/Source/MutableRuntime/MutableRuntime.Build.cs
--- a/Engine/Plugins/Mutable/Source/MutableRuntime/MutableRuntime.Build.cs
+++ b/Engine/Plugins/Mutable/Source/MutableRuntime/MutableRuntime.Build.cs
@@ -28,8 +28,6 @@ namespace UnrealBuildTool.Rules
 				"TraceLog"
 				}
 			);
-
-			bDisableAutoRTFMInstrumentation = true;
 		}
 	}
 }
