Mutable: Bump DDC version as DDC pollution is suspected as of 40275509

diff --git a/Engine/Plugins/Mutable/Source/CustomizableObject/Internal/MuCO/CustomizableObjectPrivate.h b/Engine/Plugins/Mutable/Source/CustomizableObject/Internal/MuCO/CustomizableObjectPrivate.h
--- a/Engine/Plugins/Mutable/Source/CustomizableObject/Internal/MuCO/CustomizableObjectPrivate.h
+++ b/Engine/Plugins/Mutable/Source/CustomizableObject/Internal/MuCO/CustomizableObjectPrivate.h
@@ -1848,6 +1848,6 @@ public:
 	
 	static constexpr int32 CurrentSupportedVersion = ECustomizableObjectVersions::LastCustomizableObjectVersion;
 	
-	static constexpr int32 DerivedDataVersion = 4;
+	static constexpr int32 DerivedDataVersion = 5;
 };
 
