[Backout] - CL40275509

diff --git a/Engine/Plugins/Mutable/Source/CustomizableObjectEditor/Private/MuCOE/GenerateMutableSource/GenerateMutableSourceMesh.cpp b/Engine/Plugins/Mutable/Source/CustomizableObjectEditor/Private/MuCOE/GenerateMutableSource/GenerateMutableSourceMesh.cpp
--- a/Engine/Plugins/Mutable/Source/CustomizableObjectEditor/Private/MuCOE/GenerateMutableSource/GenerateMutableSourceMesh.cpp
+++ b/Engine/Plugins/Mutable/Source/CustomizableObjectEditor/Private/MuCOE/GenerateMutableSource/GenerateMutableSourceMesh.cpp
@@ -34,7 +34,6 @@
 #include "MuCOE/Nodes/CustomizableObjectNodeMeshReshapeCommon.h"
 #include "MuCOE/Nodes/CustomizableObjectNodeMeshSwitch.h"
 #include "MuCOE/Nodes/CustomizableObjectNodeMeshVariation.h"
-#include "MuCOE/Nodes/CustomizableObjectNodeMeshParameter.h"
 #include "MuCOE/Nodes/CustomizableObjectNodeSkeletalMesh.h"
 #include "MuCOE/Nodes/CustomizableObjectNodeStaticMesh.h"
 #include "MuCOE/Nodes/CustomizableObjectNodeTable.h"
@@ -46,7 +45,6 @@
 #include "MuT/NodeMeshSwitch.h"
 #include "MuT/NodeMeshTable.h"
 #include "MuT/NodeMeshVariation.h"
-#include "MuT/NodeMeshParameter.h"
 #include "MuR/Types.h"
 #include "MuR/OpMeshFormat.h"
 #include "PhysicsEngine/PhysicsAsset.h"
@@ -2999,8 +2997,8 @@ void GenerateMorphFactor(const UCustomizableObjectNode* Node, const UEdGraphPin&
 
 		if (validStaticFactor)
 		{
-			mu::Ptr<mu::NodeScalar> FactorNode = GenerateMutableSourceFloat(ConnectedPin, GenerationContext);
-			MeshNode->Factor = FactorNode;
+			mu::NodeScalarPtr FactorNode = GenerateMutableSourceFloat(ConnectedPin, GenerationContext);
+			MeshNode->SetFactor(FactorNode);
 		}
 	}
 }
@@ -3471,12 +3469,12 @@ mu::NodeMeshPtr GenerateMorphMesh(const UEdGraphPin* Pin,
 	// Base
 	if (MorphIndex == TypedNodeMorphs.Num() - 1)
 	{
-		Result->Base = SourceNode;
+		Result->SetBase(SourceNode);
 	}
 	else
 	{
 		mu::NodeMeshPtr NextMorph = GenerateMorphMesh(Pin, TypedNodeMorphs, MorphIndex + 1, SourceNode, GenerationContext, bOnlyConnectedLOD, TableColumnName); // TODO FutureGMT change to a for. This recursion can be problematic with the production cache
-		Result->Base = NextMorph;
+		Result->SetBase(NextMorph);
 	}
 	
 	// Target
@@ -3526,21 +3524,21 @@ mu::NodeMeshPtr GenerateMorphMesh(const UEdGraphPin* Pin,
 		{
 			bSuccess = true;
 
-			mu::Ptr<mu::NodeMeshTable> MorphedSourceMeshNodeTable = new mu::NodeMeshTable;
+			mu::NodeMeshTablePtr MorphedSourceMeshNodeTable = new mu::NodeMeshTable;
 			MorphedSourceMeshNodeTable->Table = Table;
-			MorphedSourceMeshNodeTable->ColumnName = ColumnName;
-			MorphedSourceMeshNodeTable->ParameterName = TypedNodeTable->ParameterName;
-			MorphedSourceMeshNodeTable->SetMessageContext( MorphNode );
+			MorphedSourceMeshNodeTable->SetColumn(ColumnName);
+			MorphedSourceMeshNodeTable->SetParameterName(TypedNodeTable->ParameterName);
+			MorphedSourceMeshNodeTable->SetMessageContext(MorphNode);
 			MorphedSourceMeshNodeTable->SourceDataDescriptor.SourceId = TableId;
 			MorphedSourceMeshNodeTable->SourceDataDescriptor.OptionalMaxLODSize = 0;
 
-			mu::Ptr<mu::NodeMeshMakeMorph> Morph = new mu::NodeMeshMakeMorph;
-			Morph->Base = BaseSourceMesh;
-			Morph->Target = MorphedSourceMeshNodeTable;
-			Morph->bOnlyPositionAndNormal = true;
+			mu::NodeMeshMakeMorphPtr Morph = new mu::NodeMeshMakeMorph;
+			Morph->SetBase(BaseSourceMesh.get());
+			Morph->SetTarget(MorphedSourceMeshNodeTable.get());
+			Morph->SetOnlyPositionAndNormal(true);
 			Morph->SetMessageContext(MorphNode);
 
-			Result->Morph = Morph;
+			Result->SetMorph(Morph);
 		}
 	}
 	else
@@ -3551,22 +3549,22 @@ mu::NodeMeshPtr GenerateMorphMesh(const UEdGraphPin* Pin,
 		{
 			bSuccess = true;
 
-			mu::Ptr<mu::NodeMeshConstant> MorphedSourceMeshNode = new mu::NodeMeshConstant;
-			MorphedSourceMeshNode->Value = MorphedSourceMesh;
+			mu::NodeMeshConstantPtr MorphedSourceMeshNode = new mu::NodeMeshConstant;
+			MorphedSourceMeshNode->SetValue(MorphedSourceMesh);
 			MorphedSourceMeshNode->SetMessageContext(MorphNode);
 
-			mu::Ptr<mu::NodeMeshMakeMorph> Morph = new mu::NodeMeshMakeMorph;
-			Morph->Base = BaseSourceMesh;
-			Morph->Target = MorphedSourceMeshNode;
-			Morph->bOnlyPositionAndNormal = true;
+			mu::NodeMeshMakeMorphPtr Morph = new mu::NodeMeshMakeMorph;
+			Morph->SetBase(BaseSourceMesh.get());
+			Morph->SetTarget(MorphedSourceMeshNode.get());
+			Morph->SetOnlyPositionAndNormal(true);
 			Morph->SetMessageContext(MorphNode);
 
-			Result->Morph = Morph;
+			Result->SetMorph(Morph);
 
 			if (UCustomizableObjectNodeMeshMorph* TypedMorphNode = Cast<UCustomizableObjectNodeMeshMorph>(TypedNodeMorphs[MorphIndex].OwningNode))
 			{
-				Result->bReshapeSkeleton = TypedMorphNode->bReshapeSkeleton;
-				Result->bReshapePhysicsVolumes = TypedMorphNode->bReshapePhysicsVolumes;
+				Result->SetReshapeSkeleton(TypedMorphNode->bReshapeSkeleton);
+				Result->SetReshapePhysicsVolumes(TypedMorphNode->bReshapePhysicsVolumes);
 				{
 					const UEdGraphPin* ConnectedPin = FollowInputPin(*TypedMorphNode->MeshPin());
 					const UEdGraphPin* SourceMeshPin = ConnectedPin ? FindMeshBaseSource(*ConnectedPin, false) : nullptr;
@@ -3584,7 +3582,7 @@ mu::NodeMeshPtr GenerateMorphMesh(const UEdGraphPin* Pin,
 
 						for (const FName BoneName : BonesToDeform)
 						{
-							Result->BonesToDeform.Add(GenerationContext.CompilationContext.GetBoneUnique(BoneName));
+							Result->AddBoneToDeform(GenerationContext.CompilationContext.GetBoneUnique(BoneName));
 						}
 					}
 
@@ -3600,7 +3598,7 @@ mu::NodeMeshPtr GenerateMorphMesh(const UEdGraphPin* Pin,
 	
 						for (const FName& PhysicsBoneName : PhysicsToDeform)
 						{
-							Result->PhysicsToDeform.Add(GenerationContext.CompilationContext.GetBoneUnique(PhysicsBoneName));
+							Result->AddPhysicsBodyToDeform(GenerationContext.CompilationContext.GetBoneUnique(PhysicsBoneName));
 						}
 					}
 						
@@ -3657,7 +3655,7 @@ mu::Ptr<mu::NodeMesh> GenerateMutableSourceMesh(const UEdGraphPin* Pin,
 	
 	if (const UCustomizableObjectNodeSkeletalMesh* TypedNodeSkel = Cast<UCustomizableObjectNodeSkeletalMesh>(Node))
 	{
-		mu::Ptr<mu::NodeMeshConstant> MeshNode = new mu::NodeMeshConstant();
+		mu::NodeMeshConstantPtr MeshNode = new mu::NodeMeshConstant();
 		Result = MeshNode;
 
 		if (TypedNodeSkel->SkeletalMesh.IsNull())
@@ -3741,7 +3739,7 @@ mu::Ptr<mu::NodeMesh> GenerateMutableSourceMesh(const UEdGraphPin* Pin,
 																LODIndex, SectionIndex, MeshUniqueTags, SurfaceMetadataId, GenerationContext, TypedNodeSkel, nullptr, bIsPassthrough);
 			if (MutableMesh)
 			{
-				MeshNode->Value = MutableMesh;
+				MeshNode->SetValue(MutableMesh);
 
 				if (GenerationContext.CompilationContext.Options.bSkinWeightProfilesEnabled && LODIndex>=0)
 				{
@@ -3812,7 +3810,7 @@ mu::Ptr<mu::NodeMesh> GenerateMutableSourceMesh(const UEdGraphPin* Pin,
 				{
 					NumLayouts = ImportedModel->LODModels[LODIndexLayout].NumTexCoords;
 				}
-				MeshNode->Layouts.SetNum(NumLayouts);
+				MeshNode->SetLayoutCount(NumLayouts);
 				
 				const FLayoutGenerationFlags& LayoutFlags = GenerationContext.LayoutGenerationFlags.Last();
 
@@ -3822,12 +3820,12 @@ mu::Ptr<mu::NodeMesh> GenerateMutableSourceMesh(const UEdGraphPin* Pin,
 					if (!LayoutFlags.TexturePinModes.IsValidIndex(LayoutIndex) ||
 						LayoutFlags.TexturePinModes[LayoutIndex] != EPinMode::Mutable)
 					{
-						MeshNode->Layouts[LayoutIndex] = CreateDefaultLayout();
+						MeshNode->SetLayout(LayoutIndex, CreateDefaultLayout());
 					
 						// Keep packing strategy if possible, Overlay can be valid with EPinMode == Passthrough	
 						if (Layouts.IsValidIndex(LayoutIndex) && Layouts[LayoutIndex])
 						{
-							MeshNode->Layouts[LayoutIndex]->Strategy = ConvertLayoutStrategy(Layouts[LayoutIndex]->PackingStrategy);
+							MeshNode->GetLayout(LayoutIndex)->Strategy = ConvertLayoutStrategy(Layouts[LayoutIndex]->PackingStrategy);
 						}
 
 						// Ignore layout
@@ -3841,7 +3839,7 @@ mu::Ptr<mu::NodeMesh> GenerateMutableSourceMesh(const UEdGraphPin* Pin,
 						mu::Ptr<mu::NodeLayout> LayoutNode = CreateMutableLayoutNode(GenerationContext, Layout, bLinkedToExtendMaterial, bWasEmpty); // TODO PERE: Figure out
 						LayoutNode->SetMessageContext(Node);
 
-						MeshNode->Layouts[LayoutIndex] = LayoutNode;
+						MeshNode->SetLayout(LayoutIndex, LayoutNode);
 					}
 				}
 			}
@@ -3876,7 +3874,7 @@ mu::Ptr<mu::NodeMesh> GenerateMutableSourceMesh(const UEdGraphPin* Pin,
 			return {};
 		}
 
-		mu::Ptr<mu::NodeMeshConstant> MeshNode = new mu::NodeMeshConstant();
+		mu::NodeMeshConstantPtr MeshNode = new mu::NodeMeshConstant();
 		Result = MeshNode;
 
 		if (TypedNodeStatic->StaticMesh)
@@ -3909,10 +3907,10 @@ mu::Ptr<mu::NodeMesh> GenerateMutableSourceMesh(const UEdGraphPin* Pin,
 														  FString(), 0, GenerationContext, TypedNodeStatic, nullptr, bIsPassthrough);
 			if (MutableMesh)
 			{
-				MeshNode->Value = MutableMesh;
+				MeshNode->SetValue(MutableMesh);
 
 				// Layouts
-				MeshNode->Layouts.SetNum(1);
+				MeshNode->SetLayoutCount(1);
 
 				mu::Ptr<mu::NodeLayout> LayoutNode;
 
@@ -3935,7 +3933,7 @@ mu::Ptr<mu::NodeMesh> GenerateMutableSourceMesh(const UEdGraphPin* Pin,
 					LayoutNode = CreateDefaultLayout();
 				}
 
-				MeshNode->Layouts[0] = LayoutNode;
+				MeshNode->SetLayout(0, LayoutNode);
 				LayoutNode->SetMessageContext(Node);  // We need it here because we create multiple nodes.
 
 
@@ -4050,7 +4048,7 @@ mu::Ptr<mu::NodeMesh> GenerateMutableSourceMesh(const UEdGraphPin* Pin,
 			// Check Switch Parameter arity preconditions.
 			if (const UEdGraphPin* EnumPin = FollowInputPin(*SwitchParameter))
 			{
-				mu::Ptr<mu::NodeScalar> SwitchParam = GenerateMutableSourceFloat(EnumPin, GenerationContext);
+				mu::NodeScalarPtr SwitchParam = GenerateMutableSourceFloat(EnumPin, GenerationContext);
 
 				// Switch Param not generated
 				if (!SwitchParam)
@@ -4076,15 +4074,15 @@ mu::Ptr<mu::NodeMesh> GenerateMutableSourceMesh(const UEdGraphPin* Pin,
 				const int32 NumSwitchOptions = TypedNodeMeshSwitch->GetNumElements();
 
 				mu::NodeScalarEnumParameter* EnumParameter = static_cast<mu::NodeScalarEnumParameter*>(SwitchParam.get());
-				if (NumSwitchOptions != EnumParameter->Options.Num())
+				if (NumSwitchOptions != EnumParameter->GetValueCount())
 				{
 					const FText Message = LOCTEXT("MismatchedSwitch", "Switch enum and switch node have different number of options. Please refresh the switch node to make sure the outcomes are labeled properly.");
 					GenerationContext.Log(Message, Node);
 				}
 
-				mu::Ptr<mu::NodeMeshSwitch> SwitchNode = new mu::NodeMeshSwitch;
-				SwitchNode->Parameter = SwitchParam;
-				SwitchNode->Options.SetNum(NumSwitchOptions);
+				mu::NodeMeshSwitchPtr SwitchNode = new mu::NodeMeshSwitch;
+				SwitchNode->SetParameter(SwitchParam);
+				SwitchNode->SetOptionCount(NumSwitchOptions);
 
 				for (int SelectorIndex = 0; SelectorIndex < NumSwitchOptions; ++SelectorIndex)
 				{
@@ -4092,7 +4090,7 @@ mu::Ptr<mu::NodeMesh> GenerateMutableSourceMesh(const UEdGraphPin* Pin,
 					{
 						FMutableGraphMeshGenerationData ChildMeshData;
 						Result = GenerateMutableSourceMesh(ConnectedPin, GenerationContext, ChildMeshData, SurfaceMetadataId, false, bOnlyConnectedLOD);
-						SwitchNode->Options[SelectorIndex] = Result;
+						SwitchNode->SetOption(SelectorIndex, Result);
 						MeshData.Combine(ChildMeshData);
 					}
 				}
@@ -4110,7 +4108,7 @@ mu::Ptr<mu::NodeMesh> GenerateMutableSourceMesh(const UEdGraphPin* Pin,
 
 	else if (const UCustomizableObjectNodeMeshVariation* TypedNodeMeshVar = Cast<const UCustomizableObjectNodeMeshVariation>(Node))
 	{
-		mu::Ptr<mu::NodeMeshVariation> MeshNode = new mu::NodeMeshVariation();
+		mu::NodeMeshVariationPtr MeshNode = new mu::NodeMeshVariation();
 		Result = MeshNode;
 
 		if (const UEdGraphPin* ConnectedPin = FollowInputPin(*TypedNodeMeshVar->DefaultPin()))
@@ -4119,7 +4117,7 @@ mu::Ptr<mu::NodeMesh> GenerateMutableSourceMesh(const UEdGraphPin* Pin,
 			mu::NodeMeshPtr ChildNode = GenerateMutableSourceMesh(ConnectedPin, GenerationContext, ChildMeshData, SurfaceMetadataId, false, bOnlyConnectedLOD);
 			if (ChildNode)
 			{
-				MeshNode->DefaultMesh = ChildNode;
+				MeshNode->SetDefaultMesh(ChildNode.get());
 				MeshData.Combine(ChildMeshData);
 			}
 			else
@@ -4129,18 +4127,18 @@ mu::Ptr<mu::NodeMesh> GenerateMutableSourceMesh(const UEdGraphPin* Pin,
 		}
 
 		const int32 NumVariations = TypedNodeMeshVar->GetNumVariations();
-		MeshNode->Variations.SetNum(NumVariations);
+		MeshNode->SetVariationCount(NumVariations);
 		for (int VariationIndex = 0; VariationIndex < NumVariations; ++VariationIndex)
 		{
 			const UEdGraphPin* VariationPin = TypedNodeMeshVar->VariationPin(VariationIndex);
 			if (!VariationPin) continue;
 
-			MeshNode->Variations[VariationIndex].Tag = StringCast<ANSICHAR>(*TypedNodeMeshVar->GetVariation(VariationIndex).Tag).Get();
+			MeshNode->SetVariationTag(VariationIndex, StringCast<ANSICHAR>(*TypedNodeMeshVar->GetVariation(VariationIndex).Tag).Get());
 			if (const UEdGraphPin* ConnectedPin = FollowInputPin(*VariationPin))
 			{
 				FMutableGraphMeshGenerationData VariationMeshData;
 				mu::NodeMeshPtr ChildNode = GenerateMutableSourceMesh(ConnectedPin, GenerationContext, VariationMeshData, SurfaceMetadataId, false, bOnlyConnectedLOD);
-				MeshNode->Variations[VariationIndex].Mesh = ChildNode;
+				MeshNode->SetVariationMesh(VariationIndex, ChildNode.get());
 				MeshData.Combine(VariationMeshData);
 			}
 		}
@@ -4157,7 +4155,7 @@ mu::Ptr<mu::NodeMesh> GenerateMutableSourceMesh(const UEdGraphPin* Pin,
 			mu::Ptr<mu::NodeMesh> ChildNode = GenerateMutableSourceMesh(ConnectedPin, GenerationContext, ChildMeshData, SurfaceMetadataId, false, bOnlyConnectedLOD);
 			if (ChildNode)
 			{
-				MeshNode->BaseMesh = ChildNode;
+				MeshNode->SetBaseMesh(ChildNode.get());
 				MeshData.Combine(ChildMeshData);
 			}
 			else
@@ -4172,11 +4170,11 @@ mu::Ptr<mu::NodeMesh> GenerateMutableSourceMesh(const UEdGraphPin* Pin,
 	
 		{
 
-			MeshNode->bReshapeVertices = TypedNodeReshape->bReshapeVertices;
-			MeshNode->bRecomputeNormals = TypedNodeReshape->bRecomputeNormals;
-			MeshNode->bApplyLaplacian = TypedNodeReshape->bApplyLaplacianSmoothing;
-			MeshNode->bReshapeSkeleton = TypedNodeReshape->bReshapeSkeleton;
-			MeshNode->bReshapePhysicsVolumes = TypedNodeReshape->bReshapePhysicsVolumes;
+			MeshNode->SetReshapeVertices(TypedNodeReshape->bReshapeVertices);
+			MeshNode->SetRecomputeNormals(TypedNodeReshape->bRecomputeNormals);
+			MeshNode->SetApplyLaplacian(TypedNodeReshape->bApplyLaplacianSmoothing);
+			MeshNode->SetReshapeSkeleton(TypedNodeReshape->bReshapeSkeleton);
+			MeshNode->SetReshapePhysicsVolumes(TypedNodeReshape->bReshapePhysicsVolumes);
 
 			EMeshReshapeVertexColorChannelUsage ChannelUsages[4] =
 			{
@@ -4224,10 +4222,11 @@ mu::Ptr<mu::NodeMesh> GenerateMutableSourceMesh(const UEdGraphPin* Pin,
 				};
 			};
 
-			MeshNode->ColorRChannelUsage = ConvertColorUsage(ChannelUsages[0]);
-			MeshNode->ColorGChannelUsage = ConvertColorUsage(ChannelUsages[1]);
-			MeshNode->ColorBChannelUsage = ConvertColorUsage(ChannelUsages[2]);
-			MeshNode->ColorAChannelUsage = ConvertColorUsage(ChannelUsages[3]);
+			MeshNode->SetColorUsages(
+				ConvertColorUsage(ChannelUsages[0]),
+				ConvertColorUsage(ChannelUsages[1]),
+				ConvertColorUsage(ChannelUsages[2]),
+				ConvertColorUsage(ChannelUsages[3]));
 				
 			const UEdGraphPin* ConnectedPin = FollowInputPin(*TypedNodeReshape->BaseMeshPin());
 			const UEdGraphPin* SourceMeshPin = ConnectedPin ? FindMeshBaseSource(*ConnectedPin, false) : nullptr;
@@ -4245,7 +4244,7 @@ mu::Ptr<mu::NodeMesh> GenerateMutableSourceMesh(const UEdGraphPin* Pin,
 				
 				for (const FName& BoneName : BonesToDeform)
 				{
-					MeshNode->BonesToDeform.Add(GenerationContext.CompilationContext.GetBoneUnique(BoneName));
+					MeshNode->AddBoneToDeform(GenerationContext.CompilationContext.GetBoneUnique(BoneName));
 				}
 			}
 
@@ -4260,7 +4259,7 @@ mu::Ptr<mu::NodeMesh> GenerateMutableSourceMesh(const UEdGraphPin* Pin,
 
 				for (const FName& PhysicsBoneName : PhysicsToDeform)
 				{
-					MeshNode->PhysicsToDeform.Add(GenerationContext.CompilationContext.GetBoneUnique(PhysicsBoneName));
+					MeshNode->AddPhysicsBodyToDeform(GenerationContext.CompilationContext.GetBoneUnique(PhysicsBoneName));
 				}	
 			}
 			
@@ -4286,7 +4285,7 @@ mu::Ptr<mu::NodeMesh> GenerateMutableSourceMesh(const UEdGraphPin* Pin,
 			if (ChildNode)
 			{
 				BaseShapeTriangleCount = ChildMeshData.MaxNumTriangles == ChildMeshData.MinNumTriangles ? ChildMeshData.MaxNumTriangles : -1;
-				MeshNode->BaseShape = ChildNode;
+				MeshNode->SetBaseShape(ChildNode.get());	
 			}
 			else
 			{
@@ -4303,7 +4302,7 @@ mu::Ptr<mu::NodeMesh> GenerateMutableSourceMesh(const UEdGraphPin* Pin,
 			if (ChildNode)
 			{
 				TargetShapeTriangleCount = ChildMeshData.MaxNumTriangles == ChildMeshData.MinNumTriangles ? ChildMeshData.MaxNumTriangles : -1;
-				MeshNode->TargetShape = ChildNode;
+				MeshNode->SetTargetShape(ChildNode.get());
 			}
 			else
 			{
@@ -4347,7 +4346,7 @@ mu::Ptr<mu::NodeMesh> GenerateMutableSourceMesh(const UEdGraphPin* Pin,
 					TArray<FName> ArrayBoneName;
 					TArray<FTransform> ArrayTransform;
 					UCustomizableObjectNodeAnimationPose::StaticRetrievePoseInformation(TypedNode->PoseAsset, GenerationContext.GetCurrentComponentInfo()->RefSkeletalMesh, ArrayBoneName, ArrayTransform);
-					mu::Ptr<mu::NodeMeshApplyPose> NodeMeshApplyPose = CreateNodeMeshApplyPose(GenerationContext, InputMeshNode, ArrayBoneName, ArrayTransform);
+					mu::NodeMeshApplyPosePtr NodeMeshApplyPose = CreateNodeMeshApplyPose(GenerationContext, InputMeshNode, ArrayBoneName, ArrayTransform);
 
 					if (NodeMeshApplyPose)
 					{
@@ -4365,11 +4364,11 @@ mu::Ptr<mu::NodeMesh> GenerateMutableSourceMesh(const UEdGraphPin* Pin,
 				{
 					if (const UCustomizableObjectNodeTable* TypedNodeTable = Cast<UCustomizableObjectNodeTable>(TablePosePin->GetOwningNode()))
 					{
-						mu::Ptr<mu::NodeMeshApplyPose> NodeMeshApplyPose = new mu::NodeMeshApplyPose();
+						mu::NodeMeshApplyPosePtr NodeMeshApplyPose = new mu::NodeMeshApplyPose();
 						mu::Ptr<mu::NodeMesh> MeshTableNode = GenerateMutableSourceMesh(TablePosePin, GenerationContext, MeshData, SurfaceMetadataId, false, bOnlyConnectedLOD);
 
-						NodeMeshApplyPose->Base = InputMeshNode;
-						NodeMeshApplyPose->Pose = MeshTableNode;
+						NodeMeshApplyPose->SetBase(InputMeshNode);
+						NodeMeshApplyPose->SetPose(MeshTableNode);
 
 						Result = NodeMeshApplyPose;
 					}
@@ -4394,7 +4393,7 @@ mu::Ptr<mu::NodeMesh> GenerateMutableSourceMesh(const UEdGraphPin* Pin,
 		const FString TableName = TypedNodeTable->Table ? GetNameSafe(TypedNodeTable->Table).ToLower() : GetNameSafe(TypedNodeTable->Structure).ToLower();
 		const uint32 TableId = CityHash32(reinterpret_cast<const char*>(*TableName), TableName.Len() * sizeof(FString::ElementType));
 
-		mu::Ptr<mu::NodeMeshConstant> EmptyNode = new mu::NodeMeshConstant();
+		mu::NodeMeshConstantPtr EmptyNode = new mu::NodeMeshConstant();
 		Result = EmptyNode;
 		bool bSuccess = true;
 
@@ -4436,7 +4435,7 @@ mu::Ptr<mu::NodeMesh> GenerateMutableSourceMesh(const UEdGraphPin* Pin,
 
 				if (Table)
 				{
-					mu::Ptr<mu::NodeMeshTable> MeshTableNode = new mu::NodeMeshTable();
+					mu::NodeMeshTablePtr MeshTableNode = new mu::NodeMeshTable();
 					
 					int32 LODIndexConnected = -1; // LOD which the pin is connected to
 					int32 SectionIndexConnected = -1;
@@ -4479,10 +4478,10 @@ mu::Ptr<mu::NodeMesh> GenerateMutableSourceMesh(const UEdGraphPin* Pin,
 						Result = MeshTableNode;
 
 						MeshTableNode->Table = Table;
-						MeshTableNode->ColumnName = MutableColumnName;
-						MeshTableNode->ParameterName = TypedNodeTable->ParameterName;
-						MeshTableNode->bNoneOption = TypedNodeTable->bAddNoneOption;
-						MeshTableNode->DefaultRowName = TypedNodeTable->DefaultRowName.ToString();
+						MeshTableNode->SetColumn(MutableColumnName);
+						MeshTableNode->SetParameterName(TypedNodeTable->ParameterName);
+						MeshTableNode->SetNoneOption(TypedNodeTable->bAddNoneOption);
+						MeshTableNode->SetDefaultRowName(TypedNodeTable->DefaultRowName.ToString());
 						MeshTableNode->SourceDataDescriptor.SourceId = TableId;
 						MeshTableNode->SourceDataDescriptor.OptionalMaxLODSize = 0;
 
@@ -4514,7 +4513,7 @@ mu::Ptr<mu::NodeMesh> GenerateMutableSourceMesh(const UEdGraphPin* Pin,
 							}
 
 							TArray<UCustomizableObjectLayout*> Layouts = TypedNodeTable->GetLayouts(Pin);
-							MeshTableNode->Layouts.SetNum(Layouts.Num());
+							MeshTableNode->SetLayoutCount(Layouts.Num());
 
 							if (Layouts.Num())
 							{
@@ -4526,12 +4525,12 @@ mu::Ptr<mu::NodeMesh> GenerateMutableSourceMesh(const UEdGraphPin* Pin,
 									if (!LayoutFlags.TexturePinModes.IsValidIndex(LayoutIndex) ||
 										LayoutFlags.TexturePinModes[LayoutIndex] != EPinMode::Mutable)
 									{
-										MeshTableNode->Layouts[LayoutIndex] = CreateDefaultLayout();
+										MeshTableNode->SetLayout(LayoutIndex, CreateDefaultLayout());
 
 										// Keep packing strategy if possible, Overlay can be valid with EPinMode == Passthrough	
 										if (Layouts[LayoutIndex])
 										{
-											MeshTableNode->Layouts[LayoutIndex]->Strategy = ConvertLayoutStrategy(Layouts[LayoutIndex]->PackingStrategy);
+											MeshTableNode->GetLayout(LayoutIndex)->Strategy = ConvertLayoutStrategy(Layouts[LayoutIndex]->PackingStrategy);
 										}
 
 										// Ignore layouts
@@ -4548,7 +4547,7 @@ mu::Ptr<mu::NodeMesh> GenerateMutableSourceMesh(const UEdGraphPin* Pin,
 										GenerationContext.Log(FText::FromString(msg), Node, EMessageSeverity::Warning);
 									}
 
-									MeshTableNode->Layouts[LayoutIndex] = LayoutNode;
+									MeshTableNode->SetLayout(LayoutIndex, LayoutNode);
 								}
 							}
 
@@ -4572,36 +4571,7 @@ mu::Ptr<mu::NodeMesh> GenerateMutableSourceMesh(const UEdGraphPin* Pin,
 			GenerationContext.Log(LOCTEXT("ImageTableError", "Couldn't find the data table of the node."), Node);
 		}
 	}
-
-	else if (const UCustomizableObjectNodeMeshParameter* TypedNodeParam = Cast<UCustomizableObjectNodeMeshParameter>(Node))
-	{
-		mu::Ptr<mu::NodeMeshParameter> ParameterNode = new mu::NodeMeshParameter();
-
-		ParameterNode->Name = TypedNodeParam->ParameterName;
-		ParameterNode->UID = GenerationContext.GetNodeIdUnique(Node).ToString();
-
-		if (TypedNodeParam->DefaultValue)
-		{
-			ParameterNode->DefaultValue = FName(TypedNodeParam->DefaultValue->GetPathName());
-		}
-
-		GenerationContext.ParameterUIDataMap.Add(TypedNodeParam->ParameterName, FMutableParameterData(
-			TypedNodeParam->ParamUIMetadata,
-			EMutableParameterType::Texture));
-
-		// For constants we also have:
-		// animbp instance?
-		// gameplay tags?
-		// layouts?
-		// morphs?
-
-		// Force the format?
-		//mu::NodeMeshFormatPtr FormatNode = new mu::NodeMeshFormat();
-		//FormatNode->SetSource(ParameterNode);
-
-		Result = ParameterNode;
-	}
-
+	
 	else
 	{
 		GenerationContext.Log(LOCTEXT("UnimplementedMeshNode", "Mesh node type not implemented yet."), Node);
@@ -4623,5 +4593,4 @@ mu::Ptr<mu::NodeMesh> GenerateMutableSourceMesh(const UEdGraphPin* Pin,
 
 	return Result;
 }
-
 #undef LOCTEXT_NAMESPACE
diff --git a/Engine/Plugins/Mutable/Source/CustomizableObjectEditor/Private/MuCOE/GenerateMutableSource/GenerateMutableSourceModifier.cpp b/Engine/Plugins/Mutable/Source/CustomizableObjectEditor/Private/MuCOE/GenerateMutableSource/GenerateMutableSourceModifier.cpp
--- a/Engine/Plugins/Mutable/Source/CustomizableObjectEditor/Private/MuCOE/GenerateMutableSource/GenerateMutableSourceModifier.cpp
+++ b/Engine/Plugins/Mutable/Source/CustomizableObjectEditor/Private/MuCOE/GenerateMutableSource/GenerateMutableSourceModifier.cpp
@@ -172,10 +172,10 @@ mu::Ptr<mu::NodeModifier> GenerateMutableSourceModifier(const UEdGraphPin * Pin,
 
 			if (FMatrix Matrix = TypedNodeClipMesh->Transform.ToMatrixWithScale(); Matrix != FMatrix::Identity)
 			{
-				mu::Ptr<mu::NodeMeshTransform> TransformMesh = new mu::NodeMeshTransform();
-				TransformMesh->Source = ClipMesh;
+				mu::NodeMeshTransformPtr TransformMesh = new mu::NodeMeshTransform();
+				TransformMesh->SetSource(ClipMesh.get());
 
-				TransformMesh->Transform = FMatrix44f(Matrix);
+				TransformMesh->SetTransform(FMatrix44f(Matrix));
 				ClipMesh = TransformMesh;
 			}
 
@@ -302,9 +302,9 @@ mu::Ptr<mu::NodeModifier> GenerateMutableSourceModifier(const UEdGraphPin * Pin,
 				mu::Ptr<mu::NodeMesh> MeshPtr = AddMeshNode;
 
 				mu::Ptr<mu::NodeMeshFormat> MeshFormat = new mu::NodeMeshFormat();
-				SetSurfaceFormat(GenerationContext.CompilationContext, MeshFormat->VertexBuffers, MeshFormat->IndexBuffers, MeshData);
+				SetSurfaceFormat(GenerationContext.CompilationContext, MeshFormat->GetVertexBuffers(), MeshFormat->GetIndexBuffers(), MeshData);
 
-				MeshFormat->Source = MeshPtr;
+				MeshFormat->SetSource(MeshPtr.get());
 
 				SurfNode->LODs[LODIndex].MeshAdd = MeshFormat;
 			}
@@ -573,7 +573,7 @@ mu::Ptr<mu::NodeModifier> GenerateMutableSourceModifier(const UEdGraphPin * Pin,
 
 			if (validStaticFactor)
 			{
-				mu::Ptr<mu::NodeScalar> FactorNode = GenerateMutableSourceFloat(ConnectedPin, GenerationContext);
+				mu::NodeScalarPtr FactorNode = GenerateMutableSourceFloat(ConnectedPin, GenerationContext);
 				SurfNode->MorphFactor = FactorNode;
 			}
 		}
@@ -609,10 +609,10 @@ mu::Ptr<mu::NodeModifier> GenerateMutableSourceModifier(const UEdGraphPin * Pin,
 
 			if (FMatrix Matrix = TypedNodeTransformMesh->BoundingMeshTransform.ToMatrixWithScale(); Matrix != FMatrix::Identity)
 			{
-				mu::Ptr<mu::NodeMeshTransform> TransformMesh = new mu::NodeMeshTransform();
-				TransformMesh->Source = BoundingMesh;
+				mu::NodeMeshTransformPtr TransformMesh = new mu::NodeMeshTransform();
+				TransformMesh->SetSource(BoundingMesh.get());
 
-				TransformMesh->Transform = FMatrix44f(Matrix);
+				TransformMesh->SetTransform(FMatrix44f(Matrix));
 				BoundingMesh = TransformMesh;
 			}
 
diff --git a/Engine/Plugins/Mutable/Source/CustomizableObjectEditor/Private/MuCOE/GenerateMutableSource/GenerateMutableSourceProjector.cpp b/Engine/Plugins/Mutable/Source/CustomizableObjectEditor/Private/MuCOE/GenerateMutableSource/GenerateMutableSourceProjector.cpp
--- a/Engine/Plugins/Mutable/Source/CustomizableObjectEditor/Private/MuCOE/GenerateMutableSource/GenerateMutableSourceProjector.cpp
+++ b/Engine/Plugins/Mutable/Source/CustomizableObjectEditor/Private/MuCOE/GenerateMutableSource/GenerateMutableSourceProjector.cpp
@@ -11,7 +11,7 @@
 #define LOCTEXT_NAMESPACE "CustomizableObjectEditor"
 
 
-mu::Ptr<mu::NodeProjector> GenerateMutableSourceProjector(const UEdGraphPin* Pin, FMutableGraphGenerationContext& GenerationContext)
+mu::NodeProjectorPtr GenerateMutableSourceProjector(const UEdGraphPin* Pin, FMutableGraphGenerationContext& GenerationContext)
 {
 	check(Pin)
 	RETURN_ON_CYCLE(*Pin, GenerationContext)
@@ -26,11 +26,11 @@ mu::Ptr<mu::NodeProjector> GenerateMutableSourceProjector(const UEdGraphPin* Pin
 		return static_cast<mu::NodeProjector*>(Generated->Node.get());
 	}
 
-	mu::Ptr<mu::NodeProjector> Result;
+	mu::NodeProjectorPtr Result;
 	
 	if (const UCustomizableObjectNodeProjectorConstant* TypedNodeConst = Cast<UCustomizableObjectNodeProjectorConstant>(Node))
 	{
-		mu::Ptr<mu::NodeProjectorConstant> ProjectorNode = new mu::NodeProjectorConstant();
+		mu::NodeProjectorConstantPtr ProjectorNode = new mu::NodeProjectorConstant();
 		Result = ProjectorNode;
 
 		const mu::EProjectorType ProjectorType = ProjectorUtils::GetEquivalentProjectorType(TypedNodeConst->Value.ProjectionType);
@@ -80,7 +80,7 @@ mu::Ptr<mu::NodeProjector> GenerateMutableSourceProjector(const UEdGraphPin* Pin
 
 	else if (const UCustomizableObjectNodeProjectorParameter* TypedNodeParam = Cast<UCustomizableObjectNodeProjectorParameter>(Node))
 	{
-		mu::Ptr<mu::NodeProjectorParameter> ProjectorNode = new mu::NodeProjectorParameter();
+		mu::NodeProjectorParameterPtr ProjectorNode = new mu::NodeProjectorParameter();
 		Result = ProjectorNode;
 
 		ProjectorNode->SetName(TypedNodeParam->ParameterName);
diff --git a/Engine/Plugins/Mutable/Source/CustomizableObjectEditor/Private/MuCOE/GenerateMutableSource/GenerateMutableSourceProjector.h b/Engine/Plugins/Mutable/Source/CustomizableObjectEditor/Private/MuCOE/GenerateMutableSource/GenerateMutableSourceProjector.h
--- a/Engine/Plugins/Mutable/Source/CustomizableObjectEditor/Private/MuCOE/GenerateMutableSource/GenerateMutableSourceProjector.h
+++ b/Engine/Plugins/Mutable/Source/CustomizableObjectEditor/Private/MuCOE/GenerateMutableSource/GenerateMutableSourceProjector.h
@@ -8,4 +8,4 @@ class UEdGraphPin;
 struct FMutableGraphGenerationContext;
 
 /** Convert a CustomizableObject Source Graph into a mutable source graph. */
-mu::Ptr<mu::NodeProjector> GenerateMutableSourceProjector(const UEdGraphPin* Pin, FMutableGraphGenerationContext& GenerationContext);
+mu::NodeProjectorPtr GenerateMutableSourceProjector(const UEdGraphPin* Pin, FMutableGraphGenerationContext& GenerationContext);
diff --git a/Engine/Plugins/Mutable/Source/CustomizableObjectEditor/Private/MuCOE/GenerateMutableSource/GenerateMutableSourceSurface.cpp b/Engine/Plugins/Mutable/Source/CustomizableObjectEditor/Private/MuCOE/GenerateMutableSource/GenerateMutableSourceSurface.cpp
--- a/Engine/Plugins/Mutable/Source/CustomizableObjectEditor/Private/MuCOE/GenerateMutableSource/GenerateMutableSourceSurface.cpp
+++ b/Engine/Plugins/Mutable/Source/CustomizableObjectEditor/Private/MuCOE/GenerateMutableSource/GenerateMutableSourceSurface.cpp
@@ -475,10 +475,9 @@ mu::Ptr<mu::NodeSurface> GenerateMutableSourceSurface(const UEdGraphPin * Pin, F
 								});
 
 
-								mu::Ptr<mu::NodeImageMipmap> MipmapImage = new mu::NodeImageMipmap();
-								MipmapImage->Source = LastImage;
-								MipmapImage->Settings.FilterType = MipGenerationFilterType;
-								MipmapImage->Settings.AddressMode = mu::EAddressMode::None;
+								mu::NodeImageMipmapPtr MipmapImage = new mu::NodeImageMipmap();
+								MipmapImage->SetSource(LastImage.get());
+								MipmapImage->SetMipmapGenerationSettings(MipGenerationFilterType, mu::EAddressMode::None);
 
 								MipmapImage->SetMessageContext(Node);
 								LastImage = MipmapImage;
@@ -487,9 +486,9 @@ mu::Ptr<mu::NodeSurface> GenerateMutableSourceSurface(const UEdGraphPin * Pin, F
 							// Apply composite image. This needs to be computed after mipmaps generation. 	
 							if (ReferenceTexture && ReferenceTexture->GetCompositeTexture() && ReferenceTexture->CompositeTextureMode != CTM_Disabled)
 							{
-								mu::Ptr<mu::NodeImageNormalComposite> CompositedImage = new mu::NodeImageNormalComposite();
-								CompositedImage->Base = LastImage;
-								CompositedImage->Power = ReferenceTexture->CompositePower;
+								mu::NodeImageNormalCompositePtr CompositedImage = new mu::NodeImageNormalComposite();
+								CompositedImage->SetBase(LastImage.get());
+								CompositedImage->SetPower(ReferenceTexture->CompositePower);
 
 								mu::ECompositeImageMode CompositeImageMode = [CompositeTextureMode = ReferenceTexture->CompositeTextureMode]() -> mu::ECompositeImageMode
 								{
@@ -504,7 +503,7 @@ mu::Ptr<mu::NodeSurface> GenerateMutableSourceSurface(const UEdGraphPin * Pin, F
 									}
 								}();
 
-								CompositedImage->Mode = CompositeImageMode;
+								CompositedImage->SetMode(CompositeImageMode);
 
 								mu::Ptr<mu::NodeImageConstant> CompositeNormalImage = new mu::NodeImageConstant();
 
@@ -517,13 +516,12 @@ mu::Ptr<mu::NodeSurface> GenerateMutableSourceSurface(const UEdGraphPin * Pin, F
 									TSharedPtr<mu::FImage> ImageConstant = GenerateImageConstant(ReferenceCompositeNormalTexture, GenerationContext, false);
 									CompositeNormalImage->SetValue(ImageConstant);
 
-									mu::Ptr<mu::NodeImageMipmap> NormalCompositeMipmapImage = new mu::NodeImageMipmap();
+									mu::NodeImageMipmapPtr NormalCompositeMipmapImage = new mu::NodeImageMipmap();
 									const uint32 MipsToSkip = ComputeLODBiasForTexture(GenerationContext, *ReferenceCompositeNormalTexture, ReferenceTexture);
-									NormalCompositeMipmapImage->Source = ResizeTextureByNumMips(CompositeNormalImage, MipsToSkip);
-									NormalCompositeMipmapImage->Settings.FilterType = mu::EMipmapFilterType::SimpleAverage;
-									NormalCompositeMipmapImage->Settings.AddressMode = mu::EAddressMode::None;
+									NormalCompositeMipmapImage->SetSource(ResizeTextureByNumMips(CompositeNormalImage, MipsToSkip));
+									NormalCompositeMipmapImage->SetMipmapGenerationSettings(mu::EMipmapFilterType::SimpleAverage, mu::EAddressMode::None);
 
-									CompositedImage->Normal = NormalCompositeMipmapImage;
+									CompositedImage->SetNormal(NormalCompositeMipmapImage);
 
 									CompositeNormalImage->SourceDataDescriptor.OptionalMaxLODSize = 0;
 									if (GenerationContext.CompilationContext.Options.TargetPlatform)
@@ -544,9 +542,9 @@ mu::Ptr<mu::NodeSurface> GenerateMutableSourceSurface(const UEdGraphPin * Pin, F
 							}
 
 							mu::Ptr<mu::NodeImage> FormatSource = LastImage;
-							mu::Ptr<mu::NodeImageFormat> FormatImage = new mu::NodeImageFormat();
-							FormatImage->Source = LastImage.get();
-							FormatImage->Format = mu::EImageFormat::RGBA_UByte;
+							mu::NodeImageFormatPtr FormatImage = new mu::NodeImageFormat();
+							FormatImage->SetSource(LastImage.get());
+							FormatImage->SetFormat(mu::EImageFormat::RGBA_UByte);
 							FormatImage->SetMessageContext(Node);
 							LastImage = FormatImage;
 
@@ -616,19 +614,20 @@ mu::Ptr<mu::NodeSurface> GenerateMutableSourceSurface(const UEdGraphPin * Pin, F
 
 											if (bIsNormalMapFormat)
 											{
+
 												// Insert a channel swizzle
 												mu::Ptr<mu::NodeImageSwizzle> Swizzle = new mu::NodeImageSwizzle;
 												Swizzle->SetFormat( mu::EImageFormat::RGBA_UByte );
-												Swizzle->Sources[0] = FormatSource;
-												Swizzle->Sources[1] = FormatSource;
-												Swizzle->Sources[2] = FormatSource;
-												Swizzle->Sources[3] = FormatSource;
-												Swizzle->SourceChannels[0] = 0;
-												Swizzle->SourceChannels[1] = 0;
-												Swizzle->SourceChannels[2] = 0;
-												Swizzle->SourceChannels[3] = 1;
-
-												FormatImage->Source = Swizzle;
+												Swizzle->SetSource(0, FormatSource);
+												Swizzle->SetSource(1, FormatSource);
+												Swizzle->SetSource(2, FormatSource);
+												Swizzle->SetSource(3, FormatSource);
+												Swizzle->SetSourceChannel(0, 0);
+												Swizzle->SetSourceChannel(1, 0);
+												Swizzle->SetSourceChannel(2, 0);
+												Swizzle->SetSourceChannel(3, 1);
+
+												FormatImage->SetSource( Swizzle.get() );
 											}
 										}
 									}
@@ -651,8 +650,7 @@ mu::Ptr<mu::NodeSurface> GenerateMutableSourceSurface(const UEdGraphPin * Pin, F
 										MutableFormatIfAlpha = mu::EImageFormat::RGBA_UByte;
 									}
 
-									FormatImage->Format = MutableFormat;
-									FormatImage->FormatIfAlpha = MutableFormatIfAlpha;
+									FormatImage->SetFormat(MutableFormat, MutableFormatIfAlpha);
 								}
 							}
 
@@ -753,7 +751,7 @@ mu::Ptr<mu::NodeSurface> GenerateMutableSourceSurface(const UEdGraphPin * Pin, F
 			{
 				if (const UEdGraphPin* ConnectedPin = FollowInputPin(*ScalarPin))
 				{
-					mu::Ptr<mu::NodeScalar> ScalarNode = GenerateMutableSourceFloat(ConnectedPin, GenerationContext);
+					mu::NodeScalarPtr ScalarNode = GenerateMutableSourceFloat(ConnectedPin, GenerationContext);
 
 					// Encoding material layer in mutable name
 					if (const int32 LayerIndex = TypedNodeMat->GetParameterLayerIndex(EMaterialParameterType::Scalar, ScalarIndex); LayerIndex != INDEX_NONE)
@@ -789,7 +787,7 @@ mu::Ptr<mu::NodeSurface> GenerateMutableSourceSurface(const UEdGraphPin * Pin, F
 					MaterialName += MaterialPin->PinId.ToString();
 
 					GenerationContext.CurrentMaterialTableParameterId = MaterialName;
-					mu::Ptr<mu::NodeScalar> ScalarNode = GenerateMutableSourceFloat(ConnectedPin, GenerationContext);
+					mu::NodeScalarPtr ScalarNode = GenerateMutableSourceFloat(ConnectedPin, GenerationContext);
 
 					SurfNode->Scalars[MaterialIndex].Scalar = ScalarNode;
 					SurfNode->Scalars[MaterialIndex].Name = MaterialName;
@@ -797,8 +795,8 @@ mu::Ptr<mu::NodeSurface> GenerateMutableSourceSurface(const UEdGraphPin * Pin, F
 			}
 			else
 			{
-				mu::Ptr<mu::NodeScalarConstant> ScalarNode = new mu::NodeScalarConstant();
-				ScalarNode->Value = GenerationContext.CurrentReferencedMaterialIndex;
+				mu::NodeScalarConstantPtr ScalarNode = new mu::NodeScalarConstant();
+				ScalarNode->SetValue(GenerationContext.CurrentReferencedMaterialIndex);
 
 				SurfNode->Scalars[MaterialIndex].Scalar = ScalarNode;
 				SurfNode->Scalars[MaterialIndex].Name = MaterialName;
@@ -840,11 +838,10 @@ mu::Ptr<mu::NodeSurface> GenerateMutableSourceSurface(const UEdGraphPin * Pin, F
 
 					if (!ProjectorInfo || !ProjectorInfo->bIsAlternateResolutionResized)
 					{
-						mu::Ptr<mu::NodeImageResize> NodeImageResize = new mu::NodeImageResize;
-						NodeImageResize->bRelative = true;
-						NodeImageResize->SizeX = *ResolutionFactor;
-						NodeImageResize->SizeY = *ResolutionFactor;
-						NodeImageResize->Base = SurfNode2->Images[ImageIndex].Image;
+						mu::NodeImageResizePtr NodeImageResize = new mu::NodeImageResize;
+						NodeImageResize->SetRelative(true);
+						NodeImageResize->SetSize(*ResolutionFactor, *ResolutionFactor);
+						NodeImageResize->SetBase(SurfNode2->Images[ImageIndex].Image);
 
 						SurfNode2->Images[ImageIndex].Image = NodeImageResize;
 
@@ -942,7 +939,7 @@ mu::Ptr<mu::NodeSurface> GenerateMutableSourceSurface(const UEdGraphPin * Pin, F
 			// Check Switch Parameter arity preconditions.
 			if (const UEdGraphPin* EnumPin = FollowInputPin(*SwitchParameter))
 			{
-				mu::Ptr<mu::NodeScalar> SwitchParam = GenerateMutableSourceFloat(EnumPin, GenerationContext);
+				mu::NodeScalarPtr SwitchParam = GenerateMutableSourceFloat(EnumPin, GenerationContext);
 
 				// Switch Param not generated
 				if (!SwitchParam)
@@ -968,7 +965,7 @@ mu::Ptr<mu::NodeSurface> GenerateMutableSourceSurface(const UEdGraphPin * Pin, F
 				const int32 NumSwitchOptions = TypedNodeSwitch->GetNumElements();
 
 				mu::NodeScalarEnumParameter* EnumParameter = static_cast<mu::NodeScalarEnumParameter*>(SwitchParam.get());
-				if (NumSwitchOptions != EnumParameter->Options.Num())
+				if (NumSwitchOptions != EnumParameter->GetValueCount())
 				{
 					const FText Message = LOCTEXT("MismatchedSwitch", "Switch enum and switch node have different number of options. Please refresh the switch node to make sure the outcomes are labeled properly.");
 					GenerationContext.Log(Message, Node);
diff --git a/Engine/Plugins/Mutable/Source/CustomizableObjectEditor/Private/MuCOE/GraphTraversal.cpp b/Engine/Plugins/Mutable/Source/CustomizableObjectEditor/Private/MuCOE/GraphTraversal.cpp
--- a/Engine/Plugins/Mutable/Source/CustomizableObjectEditor/Private/MuCOE/GraphTraversal.cpp
+++ b/Engine/Plugins/Mutable/Source/CustomizableObjectEditor/Private/MuCOE/GraphTraversal.cpp
@@ -18,7 +18,6 @@
 #include "MuCOE/Nodes/CustomizableObjectNodeMeshReshape.h"
 #include "MuCOE/Nodes/CustomizableObjectNodeMeshSwitch.h"
 #include "MuCOE/Nodes/CustomizableObjectNodeMeshVariation.h"
-#include "MuCOE/Nodes/CustomizableObjectNodeMeshParameter.h"
 #include "MuCOE/Nodes/CustomizableObjectNodeObject.h"
 #include "MuCOE/Nodes/CustomizableObjectNodeSkeletalMesh.h"
 #include "MuCOE/Nodes/CustomizableObjectNodeStaticMesh.h"
@@ -493,7 +492,7 @@ UCustomizableObjectNodeObject* GraphTraversal::GetFullGraphRootNode(const UCusto
 const UEdGraphPin* FindMeshBaseSource(const UEdGraphPin& Pin, const bool bOnlyLookForStaticMesh)
 {
 	check(Pin.Direction == EGPD_Output);
-	check(Pin.PinType.PinCategory == UEdGraphSchema_CustomizableObject::PC_Mesh
+	check(Pin.PinType.PinCategory == UEdGraphSchema_CustomizableObject::PC_Mesh 
 		||
 		Pin.PinType.PinCategory == UEdGraphSchema_CustomizableObject::PC_PassThroughMesh
 		||
@@ -501,7 +500,7 @@ const UEdGraphPin* FindMeshBaseSource(const UEdGraphPin& Pin, const bool bOnlyLo
 		||
 		Pin.PinType.PinCategory == UEdGraphSchema_CustomizableObject::PC_Modifier
 	);
-
+	
 	const UEdGraphNode* Node = Pin.GetOwningNode();
 	check(Node);
 
@@ -590,7 +589,7 @@ const UEdGraphPin* FindMeshBaseSource(const UEdGraphPin& Pin, const bool bOnlyLo
 			return FindMeshBaseSource(*ConnectedPin, bOnlyLookForStaticMesh);
 		}
 	}
-
+	
 	else if (const UCustomizableObjectNodeMeshMorphStackDefinition* TypedNodeMorphStackDef = Cast<UCustomizableObjectNodeMeshMorphStackDefinition>(Node))
 	{
 		if (const UEdGraphPin* ConnectedPin = FollowInputPin(*TypedNodeMorphStackDef->GetMeshPin()))
@@ -614,7 +613,7 @@ const UEdGraphPin* FindMeshBaseSource(const UEdGraphPin& Pin, const bool bOnlyLo
 			return FindMeshBaseSource(*ConnectedPin, bOnlyLookForStaticMesh);
 		}
 	}
-
+	
 	else if (Cast<UCustomizableObjectNodeTable>(Node))
 	{
 		if (!bOnlyLookForStaticMesh)
@@ -631,14 +630,6 @@ const UEdGraphPin* FindMeshBaseSource(const UEdGraphPin& Pin, const bool bOnlyLo
 		}
 	}
 
-	else if (Cast<UCustomizableObjectNodeMeshParameter>(Node))
-	{
-		if (!bOnlyLookForStaticMesh)
-		{
-			return &Pin;
-		}
-	}
-
 	else
 	{
 		unimplemented(); // Case missing.
diff --git a/Engine/Plugins/Mutable/Source/CustomizableObjectEditor/Private/MuCOE/Nodes/CustomizableObjectNodeExtensionDataSwitch.cpp b/Engine/Plugins/Mutable/Source/CustomizableObjectEditor/Private/MuCOE/Nodes/CustomizableObjectNodeExtensionDataSwitch.cpp
--- a/Engine/Plugins/Mutable/Source/CustomizableObjectEditor/Private/MuCOE/Nodes/CustomizableObjectNodeExtensionDataSwitch.cpp
+++ b/Engine/Plugins/Mutable/Source/CustomizableObjectEditor/Private/MuCOE/Nodes/CustomizableObjectNodeExtensionDataSwitch.cpp
@@ -31,7 +31,7 @@ mu::Ptr<mu::NodeExtensionData> UCustomizableObjectNodeExtensionDataSwitch::Gener
 	const int32 NumParameters = FollowInputPinArray(*SwitchParameter()).Num();
 
 	const UEdGraphPin* EnumPin = FollowInputPin(*SwitchParameter());
-	mu::Ptr<mu::NodeScalar> SwitchParam = EnumPin ? GenerateMutableSourceFloat(EnumPin, InCompilerInterface.GenerationContext) : nullptr;
+	mu::NodeScalarPtr SwitchParam = EnumPin ? GenerateMutableSourceFloat(EnumPin, InCompilerInterface.GenerationContext) : nullptr;
 
 	const int32 NumSwitchOptions = GetNumElements();
 
diff --git a/Engine/Plugins/Mutable/Source/CustomizableObjectEditor/Private/MuCOE/SMutableGraphViewer.cpp b/Engine/Plugins/Mutable/Source/CustomizableObjectEditor/Private/MuCOE/SMutableGraphViewer.cpp
--- a/Engine/Plugins/Mutable/Source/CustomizableObjectEditor/Private/MuCOE/SMutableGraphViewer.cpp
+++ b/Engine/Plugins/Mutable/Source/CustomizableObjectEditor/Private/MuCOE/SMutableGraphViewer.cpp
@@ -59,6 +59,28 @@
 #include "MuT/NodeScalarSwitch.h"
 #include "MuT/NodeScalarTable.h"
 #include "MuT/NodeObjectNew.h"
+
+#include "MuT/NodeObjectGroupPrivate.h"
+#include "MuT/NodeImageFormatPrivate.h"
+#include "MuT/NodeMeshFormatPrivate.h"
+#include "MuT/NodeMeshConstantPrivate.h"
+#include "MuT/NodeImageSwitchPrivate.h"
+#include "MuT/NodeImageLayerColourPrivate.h"
+#include "MuT/NodeImageLayerPrivate.h"
+#include "MuT/NodeImageMipmapPrivate.h"
+#include "MuT/NodeImageResizePrivate.h"
+#include "MuT/NodeImageInterpolatePrivate.h"
+#include "MuT/NodeImagePlainColourPrivate.h"
+#include "MuT/NodeImageProjectPrivate.h"
+#include "MuT/NodeMeshMorphPrivate.h"
+#include "MuT/NodeScalarEnumParameterPrivate.h"
+#include "MuT/NodeScalarParameterPrivate.h"
+#include "MuT/NodeMeshMakeMorphPrivate.h"
+#include "MuT/NodeScalarCurvePrivate.h"
+#include "MuT/NodeProjectorPrivate.h"
+#include "MuT/NodeImageInvertPrivate.h"
+#include "MuT/NodeImageSwizzlePrivate.h"
+#include "MuT/NodeImageMultiLayerPrivate.h"
 #include "MuT/NodeModifierMeshTransformInMesh.h"
 #include "Widgets/MutableExpanderArrow.h"
 
@@ -100,7 +122,6 @@ static const char* MutableNodeNames[] =
 	"MeshVariation",
 	"MeshReshape",
 	"MeshClipDeform",
-	"MeshParameter",
 
 	"Image",
 	"ImageConstant",
@@ -206,7 +227,7 @@ static const char* MutableNodeNames[] =
 	"MatrixParameter",
 };
 
-static_assert(UE_ARRAY_COUNT(MutableNodeNames) == SIZE_T(mu::Node::EType::Count));
+static_assert(sizeof(MutableNodeNames) == SIZE_T(mu::Node::EType::Count) * sizeof(char*));
 
 
 class SMutableGraphTreeRow : public STableRow<TSharedPtr<FMutableGraphTreeElement>>
@@ -444,9 +465,10 @@ void SMutableGraphViewer::GetChildrenForInfo(TSharedPtr<FMutableGraphTreeElement
 	else if (ParentNode->GetType() == mu::NodeObjectGroup::GetStaticType())
 	{
 		mu::NodeObjectGroup* ObjectGroup = StaticCast<mu::NodeObjectGroup*>(ParentNode);
-		for (int32 l = 0; l < ObjectGroup->Children.Num(); ++l)
+		mu::NodeObjectGroup::Private* Private = ObjectGroup->GetPrivate();
+		for (int32 l = 0; l < Private->m_children.Num(); ++l)
 		{
-			AddChildFunc(ObjectGroup->Children[l].get(), TEXT("CHILD"));
+			AddChildFunc(Private->m_children[l].get(), TEXT("CHILD"));
 		}
 	}
 
@@ -576,22 +598,25 @@ void SMutableGraphViewer::GetChildrenForInfo(TSharedPtr<FMutableGraphTreeElement
 	else if (ParentNode->GetType() == mu::NodeMeshConstant::GetStaticType())
 	{
 		mu::NodeMeshConstant* MeshConstantVar = StaticCast<mu::NodeMeshConstant*>(ParentNode);
-		for (int32 LayoutIndex = 0; LayoutIndex < MeshConstantVar->Layouts.Num(); LayoutIndex++)
+		mu::NodeMeshConstant::Private* Private = MeshConstantVar->GetPrivate();
+		for (int32 LayoutIndex = 0; LayoutIndex < Private->Layouts.Num(); LayoutIndex++)
 		{
-			AddChildFunc(MeshConstantVar->Layouts[LayoutIndex].get(), FString::Printf(TEXT("LAYOUT [%d]"), LayoutIndex));
+			AddChildFunc(Private->Layouts[LayoutIndex].get(), FString::Printf(TEXT("LAYOUT [%d]"), LayoutIndex));
 		}
 	}
 
 	else if (ParentNode->GetType() == mu::NodeImageFormat::GetStaticType())
 	{
 		mu::NodeImageFormat* ImageFormatVar = StaticCast<mu::NodeImageFormat*>(ParentNode);
-		AddChildFunc(ImageFormatVar->Source.get(), FString::Printf(TEXT("SOURCE IMAGE")));
+		mu::NodeImageFormat::Private* Private = ImageFormatVar->GetPrivate();
+		AddChildFunc(Private->m_source.get(), FString::Printf(TEXT("SOURCE IMAGE")));
 	}
 
 	else if (ParentNode->GetType() == mu::NodeMeshFormat::GetStaticType())
 	{
 		mu::NodeMeshFormat* MeshFormatVar = StaticCast<mu::NodeMeshFormat*>(ParentNode);
-		AddChildFunc(MeshFormatVar->Source.get(), FString::Printf(TEXT("SOURCE MESH")));
+		mu::NodeMeshFormat::Private* Private = MeshFormatVar->GetPrivate();
+		AddChildFunc(Private->Source.get(), FString::Printf(TEXT("SOURCE MESH")));
 	}
 
 	else if (ParentNode->GetType() == mu::NodeModifierMeshClipMorphPlane::GetStaticType())
@@ -628,65 +653,73 @@ void SMutableGraphViewer::GetChildrenForInfo(TSharedPtr<FMutableGraphTreeElement
 	else if (ParentNode->GetType() == mu::NodeImageSwitch::GetStaticType())
 	{
 		mu::NodeImageSwitch* ImageSwitchVar = StaticCast<mu::NodeImageSwitch*>(ParentNode);
-		AddChildFunc(ImageSwitchVar->Parameter.get(), FString::Printf(TEXT("PARAM")));
-		for (int32 OptionIndex = 0; OptionIndex < ImageSwitchVar->Options.Num(); OptionIndex++)
+		mu::NodeImageSwitch::Private* Private = ImageSwitchVar->GetPrivate();
+		AddChildFunc(Private->m_pParameter.get(), FString::Printf(TEXT("PARAM")));
+		for (int32 OptionIndex = 0; OptionIndex < Private->Options.Num(); OptionIndex++)
 		{
-			AddChildFunc(ImageSwitchVar->Options[OptionIndex].get(), FString::Printf(TEXT("OPTION [%d]"), OptionIndex));
+			AddChildFunc(Private->Options[OptionIndex].get(), FString::Printf(TEXT("OPTION [%d]"), OptionIndex));
 		}
 	}
 
 	else if (ParentNode->GetType() == mu::NodeImageMipmap::GetStaticType())
 	{
 		mu::NodeImageMipmap* ImageMipMapVar = StaticCast<mu::NodeImageMipmap*>(ParentNode);
-		AddChildFunc(ImageMipMapVar->Source.get(), FString::Printf(TEXT("SOURCE")));
-		AddChildFunc(ImageMipMapVar->Factor.get(), FString::Printf(TEXT("FACTOR")));
+		mu::NodeImageMipmap::Private* Private = ImageMipMapVar->GetPrivate();
+		AddChildFunc(Private->m_pSource.get(), FString::Printf(TEXT("SOURCE")));
+		AddChildFunc(Private->m_pFactor.get(), FString::Printf(TEXT("FACTOR")));
 	}
 
 	else if (ParentNode->GetType() == mu::NodeImageLayer::GetStaticType())
 	{
 		mu::NodeImageLayer* ImageLayerVar = StaticCast<mu::NodeImageLayer*>(ParentNode);
-		AddChildFunc(ImageLayerVar->Base.get(), FString::Printf(TEXT("BASE")));
-		AddChildFunc(ImageLayerVar->Mask.get(), FString::Printf(TEXT("MASK")));
-		AddChildFunc(ImageLayerVar->Blended.get(), FString::Printf(TEXT("BLEND")));
+		mu::NodeImageLayer::Private* Private = ImageLayerVar->GetPrivate();
+		AddChildFunc(Private->m_pBase.get(), FString::Printf(TEXT("BASE")));
+		AddChildFunc(Private->m_pMask.get(), FString::Printf(TEXT("MASK")));
+		AddChildFunc(Private->m_pBlended.get(), FString::Printf(TEXT("BLEND")));
 	}
 	
 	else if (ParentNode->GetType() == mu::NodeImageLayerColour::GetStaticType())
 	{
 		mu::NodeImageLayerColour* ImageLayerColourVar = StaticCast<mu::NodeImageLayerColour*>(ParentNode);
-		AddChildFunc(ImageLayerColourVar->Base.get(), FString::Printf(TEXT("BASE")));
-		AddChildFunc(ImageLayerColourVar->Mask.get(), FString::Printf(TEXT("MASK")));
-		AddChildFunc(ImageLayerColourVar->Colour.get(), FString::Printf(TEXT("COLOR")));
+		mu::NodeImageLayerColour::Private* Private = ImageLayerColourVar->GetPrivate();
+		AddChildFunc(Private->m_pBase.get(), FString::Printf(TEXT("BASE")));
+		AddChildFunc(Private->m_pMask.get(), FString::Printf(TEXT("MASK")));
+		AddChildFunc(Private->m_pColour.get(), FString::Printf(TEXT("COLOR")));
 	}
 
 	else if (ParentNode->GetType() == mu::NodeImageResize::GetStaticType())
 	{
 		mu::NodeImageResize* ImageResizeVar = StaticCast<mu::NodeImageResize*>(ParentNode);
-		AddChildFunc(ImageResizeVar->Base.get(), FString::Printf(TEXT("BASE")));
+		mu::NodeImageResize::Private* Private = ImageResizeVar->GetPrivate();
+		AddChildFunc(Private->m_pBase.get(), FString::Printf(TEXT("BASE")));
 	}
 
 	else if (ParentNode->GetType() == mu::NodeMeshMorph::GetStaticType())
 	{
 		mu::NodeMeshMorph* MeshMorphVar = StaticCast<mu::NodeMeshMorph*>(ParentNode);
-		AddChildFunc(MeshMorphVar->Base.get(), FString::Printf(TEXT("BASE")));
-		AddChildFunc(MeshMorphVar->Morph.get(), FString::Printf(TEXT("MORPH")));
-		AddChildFunc(MeshMorphVar->Factor.get(), FString::Printf(TEXT("FACTOR")));
+		mu::NodeMeshMorph::Private* Private = MeshMorphVar->GetPrivate();
+		AddChildFunc(Private->Base.get(), FString::Printf(TEXT("BASE")));
+		AddChildFunc(Private->Morph.get(), FString::Printf(TEXT("MORPH")));
+		AddChildFunc(Private->Factor.get(), FString::Printf(TEXT("FACTOR")));
 	}
 
 	else if (ParentNode->GetType() == mu::NodeImageProject::GetStaticType())
 	{
 		mu::NodeImageProject* ImageProjectVar = StaticCast<mu::NodeImageProject*>(ParentNode);
-		AddChildFunc(ImageProjectVar->Projector.get(), FString::Printf(TEXT("PROJECTOR")));
-		AddChildFunc(ImageProjectVar->Mesh.get(), FString::Printf(TEXT("MESH")));
-		AddChildFunc(ImageProjectVar->Image.get(), FString::Printf(TEXT("IMAGE")));
-		AddChildFunc(ImageProjectVar->Mask.get(), FString::Printf(TEXT("MASK")));
-		AddChildFunc(ImageProjectVar->AngleFadeStart.get(), FString::Printf(TEXT("FADE START ANGLE")));
-		AddChildFunc(ImageProjectVar->AngleFadeEnd.get(), FString::Printf(TEXT("FADE END ANGLE")));
+		mu::NodeImageProject::Private* Private = ImageProjectVar->GetPrivate();
+		AddChildFunc(Private->m_pProjector.get(), FString::Printf(TEXT("PROJECTOR")));
+		AddChildFunc(Private->m_pMesh.get(), FString::Printf(TEXT("MESH")));
+		AddChildFunc(Private->m_pImage.get(), FString::Printf(TEXT("IMAGE")));
+		AddChildFunc(Private->m_pMask.get(), FString::Printf(TEXT("MASK")));
+		AddChildFunc(Private->m_pAngleFadeStart.get(), FString::Printf(TEXT("FADE START ANGLE")));
+		AddChildFunc(Private->m_pAngleFadeEnd.get(), FString::Printf(TEXT("FADE END ANGLE")));
 	}
 
 	else if (ParentNode->GetType() == mu::NodeImagePlainColour::GetStaticType())
 	{
 		mu::NodeImagePlainColour* ImagePlainColourVar = StaticCast<mu::NodeImagePlainColour*>(ParentNode);
-		AddChildFunc(ImagePlainColourVar->Colour.get(), FString::Printf(TEXT("COLOR")));
+        mu::NodeImagePlainColour::Private* Private = ImagePlainColourVar->GetPrivate();
+		AddChildFunc(Private->m_pColour.get(), FString::Printf(TEXT("COLOR")));
 	}
 
 	else if (ParentNode->GetType() == mu::NodeLayout::GetStaticType())
@@ -697,9 +730,10 @@ void SMutableGraphViewer::GetChildrenForInfo(TSharedPtr<FMutableGraphTreeElement
 	else if (ParentNode->GetType() == mu::NodeScalarEnumParameter::GetStaticType())
 	{
 		mu::NodeScalarEnumParameter* ScalarEnumParameterVar = StaticCast<mu::NodeScalarEnumParameter*>(ParentNode);
-		for (int32 RangeIndex = 0; RangeIndex < ScalarEnumParameterVar->Ranges.Num(); RangeIndex++)
+		mu::NodeScalarEnumParameter::Private* Private = ScalarEnumParameterVar->GetPrivate();
+		for (int32 RangeIndex = 0; RangeIndex < Private->Ranges.Num(); RangeIndex++)
 		{
-			AddChildFunc(ScalarEnumParameterVar->Ranges[RangeIndex].get(), FString::Printf(TEXT("RANGE [%d]"), RangeIndex));
+			AddChildFunc(Private->Ranges[RangeIndex].get(), FString::Printf(TEXT("RANGE [%d]"), RangeIndex));
 		}
 	}
 
@@ -720,10 +754,11 @@ void SMutableGraphViewer::GetChildrenForInfo(TSharedPtr<FMutableGraphTreeElement
 	else if (ParentNode->GetType() == mu::NodeImageInterpolate::GetStaticType())
 	{
 		mu::NodeImageInterpolate* ImageInterpolateVar = StaticCast<mu::NodeImageInterpolate*>(ParentNode);
-		AddChildFunc(ImageInterpolateVar->Factor.get(), FString::Printf(TEXT("FACTOR")));
-		for (int32 TargetIndex = 0; TargetIndex < ImageInterpolateVar->Targets.Num(); TargetIndex++)
+		mu::NodeImageInterpolate::Private* Private = ImageInterpolateVar->GetPrivate();
+		AddChildFunc(Private->m_pFactor.get(), FString::Printf(TEXT("FACTOR")));
+		for (int32 TargetIndex = 0; TargetIndex < Private->m_targets.Num(); TargetIndex++)
 		{
-			AddChildFunc(ImageInterpolateVar->Targets[TargetIndex].get(), FString::Printf(TEXT("TARGET [%d]"), TargetIndex));
+			AddChildFunc(Private->m_targets[TargetIndex].get(), FString::Printf(TEXT("TARGET [%d]"), TargetIndex));
 		}
 	}
 	
@@ -735,9 +770,10 @@ void SMutableGraphViewer::GetChildrenForInfo(TSharedPtr<FMutableGraphTreeElement
 	else if (ParentNode->GetType() == mu::NodeScalarParameter::GetStaticType())
 	{
 		mu::NodeScalarParameter* ScalarParameterVar = StaticCast<mu::NodeScalarParameter*>(ParentNode);
-		for (int32 RangeIndex = 0; RangeIndex < ScalarParameterVar->Ranges.Num(); RangeIndex++)
+		mu::NodeScalarParameter::Private* Private = ScalarParameterVar->GetPrivate();
+		for (int32 RangeIndex = 0; RangeIndex < Private->Ranges.Num(); RangeIndex++)
 		{
-			AddChildFunc(ScalarParameterVar->Ranges[RangeIndex].get(), FString::Printf(TEXT("RANGE [%d]"), RangeIndex));
+			AddChildFunc(Private->Ranges[RangeIndex].get(), FString::Printf(TEXT("RANGE [%d]"), RangeIndex));
 		}
 	}
 
@@ -771,16 +807,18 @@ void SMutableGraphViewer::GetChildrenForInfo(TSharedPtr<FMutableGraphTreeElement
 	else if (ParentNode->GetType() == mu::NodeMeshMakeMorph::GetStaticType())
 	{
 		mu::NodeMeshMakeMorph* MeshMakeMorphVar = StaticCast<mu::NodeMeshMakeMorph*>(ParentNode);
-		AddChildFunc(MeshMakeMorphVar->Base.get(), FString::Printf(TEXT("BASE")));
-		AddChildFunc(MeshMakeMorphVar->Target.get(), FString::Printf(TEXT("TARGET")));
+		mu::NodeMeshMakeMorph::Private* Private = MeshMakeMorphVar->GetPrivate();
+		AddChildFunc(Private->m_pBase.get(), FString::Printf(TEXT("BASE")));
+		AddChildFunc(Private->m_pTarget.get(), FString::Printf(TEXT("TARGET")));
 	}
 	
 	else if (ParentNode->GetType() == mu::NodeProjectorParameter::GetStaticType())
 	{
 		mu::NodeProjectorParameter* ProjectorParameterVar = StaticCast<mu::NodeProjectorParameter*>(ParentNode);
-		for (int32 RangeIndex = 0; RangeIndex < ProjectorParameterVar->Ranges.Num(); RangeIndex++)
+		mu::NodeProjectorParameter::Private* Private = ProjectorParameterVar->GetPrivate();
+		for (int32 RangeIndex = 0; RangeIndex < Private->Ranges.Num(); RangeIndex++)
 		{
-			AddChildFunc(ProjectorParameterVar->Ranges[RangeIndex].get(), FString::Printf(TEXT("RANGE [%d]"), RangeIndex));
+			AddChildFunc(Private->Ranges[RangeIndex].get(), FString::Printf(TEXT("RANGE [%d]"), RangeIndex));
 		}
 	}
 
@@ -802,25 +840,28 @@ void SMutableGraphViewer::GetChildrenForInfo(TSharedPtr<FMutableGraphTreeElement
 	else if (ParentNode->GetType() == mu::NodeImageSwizzle::GetStaticType())
 	{
 		mu::NodeImageSwizzle* ImageSwizzleVar = StaticCast<mu::NodeImageSwizzle*>(ParentNode);
-		for (int32 SourceIndex = 0; SourceIndex < ImageSwizzleVar->Sources.Num(); ++SourceIndex)
+		mu::NodeImageSwizzle::Private* Private = ImageSwizzleVar->GetPrivate();
+		for (int32 SourceIndex = 0; SourceIndex < Private->m_sources.Num(); ++SourceIndex)
 		{
-			AddChildFunc(ImageSwizzleVar->Sources[SourceIndex].get(), FString::Printf(TEXT("SOURCE [%d]"), SourceIndex));
+			AddChildFunc(Private->m_sources[SourceIndex].get(), FString::Printf(TEXT("SOURCE [%d]"), SourceIndex));
 		}
 	}
 	
 	else if (ParentNode->GetType() == mu::NodeImageInvert::GetStaticType())
 	{
 		mu::NodeImageInvert* ImageInvertVar = StaticCast<mu::NodeImageInvert*>(ParentNode);
-		AddChildFunc(ImageInvertVar->Base.get(), TEXT("BASE"));
+		mu::NodeImageInvert::Private* Private = ImageInvertVar->GetPrivate();
+		AddChildFunc(Private->m_pBase.get(), TEXT("BASE"));
 	}
 
 	else if (ParentNode->GetType() == mu::NodeImageMultiLayer::GetStaticType())
 	{
 		mu::NodeImageMultiLayer* ImageMultilayerVar = StaticCast<mu::NodeImageMultiLayer*>(ParentNode);
-		AddChildFunc(ImageMultilayerVar->Base.get(), FString::Printf(TEXT("BASE")));
-		AddChildFunc(ImageMultilayerVar->Mask.get(), FString::Printf(TEXT("MASK")));
-		AddChildFunc(ImageMultilayerVar->Blended.get(), FString::Printf(TEXT("BLEND")));
-		AddChildFunc(ImageMultilayerVar->Range.get(), FString::Printf(TEXT("RANGE")));
+		mu::NodeImageMultiLayer::Private* Private = ImageMultilayerVar->GetPrivate();
+		AddChildFunc(Private->m_pBase.get(), FString::Printf(TEXT("BASE")));
+		AddChildFunc(Private->m_pMask.get(), FString::Printf(TEXT("MASK")));
+		AddChildFunc(Private->m_pBlended.get(), FString::Printf(TEXT("BLEND")));
+		AddChildFunc(Private->m_pRange.get(), FString::Printf(TEXT("RANGE")));
 	}
 	
 	else if (ParentNode->GetType() == mu::NodeImageTable::GetStaticType())
diff --git a/Engine/Plugins/Mutable/Source/MutableRuntime/Internal/MuR/Operations.h b/Engine/Plugins/Mutable/Source/MutableRuntime/Internal/MuR/Operations.h
--- a/Engine/Plugins/Mutable/Source/MutableRuntime/Internal/MuR/Operations.h
+++ b/Engine/Plugins/Mutable/Source/MutableRuntime/Internal/MuR/Operations.h
@@ -46,9 +46,8 @@ namespace mu
         SC_PARAMETER,
         CO_PARAMETER,
         PR_PARAMETER,
-		IM_PARAMETER,
-		ME_PARAMETER,
-		ST_PARAMETER,
+        IM_PARAMETER,
+        ST_PARAMETER,
     	MA_PARAMETER,
 
 		//! A referenced, but opaque engine resource
diff --git a/Engine/Plugins/Mutable/Source/MutableRuntime/Private/MuR/CodeVisitor.cpp b/Engine/Plugins/Mutable/Source/MutableRuntime/Private/MuR/CodeVisitor.cpp
--- a/Engine/Plugins/Mutable/Source/MutableRuntime/Private/MuR/CodeVisitor.cpp
+++ b/Engine/Plugins/Mutable/Source/MutableRuntime/Private/MuR/CodeVisitor.cpp
@@ -55,7 +55,6 @@ namespace mu
 				case EOpType::CO_PARAMETER:
 				case EOpType::PR_PARAMETER:
 				case EOpType::IM_PARAMETER:
-				case EOpType::ME_PARAMETER:
 				case EOpType::ST_PARAMETER:
 					CurrentParams[Program.GetOpArgs<OP::ParameterArgs>(at).variable]++;
 					break;
diff --git a/Engine/Plugins/Mutable/Source/MutableRuntime/Private/MuR/Operations.cpp b/Engine/Plugins/Mutable/Source/MutableRuntime/Private/MuR/Operations.cpp
--- a/Engine/Plugins/Mutable/Source/MutableRuntime/Private/MuR/Operations.cpp
+++ b/Engine/Plugins/Mutable/Source/MutableRuntime/Private/MuR/Operations.cpp
@@ -35,7 +35,6 @@ namespace mu
 		{ EDataType::Color,			},	// CO_PARAMETER
 		{ EDataType::Projector,		},	// PR_PARAMETER
 		{ EDataType::Image,			},	// IM_PARAMETER
-		{ EDataType::Mesh,			},	// ME_PARAMETER
 		{ EDataType::String,		},	// ST_PARAMETER
 		{ EDataType::Matrix,		},	// MA_PARAMETER
 
@@ -173,9 +172,8 @@ namespace mu
         case EOpType::SC_PARAMETER:
         case EOpType::CO_PARAMETER:
         case EOpType::PR_PARAMETER:
-		case EOpType::IM_PARAMETER:
-		case EOpType::ME_PARAMETER:
-		case EOpType::MA_PARAMETER:
+        case EOpType::IM_PARAMETER:
+        case EOpType::MA_PARAMETER:
 		case EOpType::IM_REFERENCE:
 		case EOpType::ME_REFERENCE:
 			break;
diff --git a/Engine/Plugins/Mutable/Source/MutableRuntime/Private/MuR/System.cpp b/Engine/Plugins/Mutable/Source/MutableRuntime/Private/MuR/System.cpp
--- a/Engine/Plugins/Mutable/Source/MutableRuntime/Private/MuR/System.cpp
+++ b/Engine/Plugins/Mutable/Source/MutableRuntime/Private/MuR/System.cpp
@@ -716,9 +716,8 @@ namespace mu
             case EOpType::SC_PARAMETER:
             case EOpType::CO_PARAMETER:
             case EOpType::PR_PARAMETER:
-			case EOpType::IM_PARAMETER:
-			case EOpType::ME_PARAMETER:
-			case EOpType::MA_PARAMETER:
+            case EOpType::IM_PARAMETER:
+            case EOpType::MA_PARAMETER:
             {
 				OP::ParameterArgs args = Program.GetOpArgs<OP::ParameterArgs>(at);
 				OP::ADDRESS ParamIndex = args.variable;
diff --git a/Engine/Plugins/Mutable/Source/MutableTools/Internal/MuT/NodeImageFormatPrivate.h b/Engine/Plugins/Mutable/Source/MutableTools/Internal/MuT/NodeImageFormatPrivate.h
new file mode 100644
--- /dev/null
+++ b/Engine/Plugins/Mutable/Source/MutableTools/Internal/MuT/NodeImageFormatPrivate.h
@@ -0,0 +1,27 @@
+// Copyright Epic Games, Inc. All Rights Reserved.
+
+#pragma once
+
+#include "MuT/NodeImageFormat.h"
+#include "MuT/NodeImagePrivate.h"
+#include "MuR/ImagePrivate.h"
+
+
+namespace mu
+{
+
+
+	class NodeImageFormat::Private : public NodeImage::Private
+	{
+	public:
+
+		static FNodeType s_type;
+
+		EImageFormat m_format = EImageFormat::None;
+		EImageFormat m_formatIfAlpha = EImageFormat::None;
+        NodeImagePtr m_source;
+
+	};
+
+
+}
diff --git a/Engine/Plugins/Mutable/Source/MutableTools/Internal/MuT/NodePrivate.h b/Engine/Plugins/Mutable/Source/MutableTools/Internal/MuT/NodePrivate.h
new file mode 100644
--- /dev/null
+++ b/Engine/Plugins/Mutable/Source/MutableTools/Internal/MuT/NodePrivate.h
@@ -0,0 +1,54 @@
+// Copyright Epic Games, Inc. All Rights Reserved.
+
+#pragma once
+
+#include "MuT/Node.h"
+#include "MuR/Operations.h"
+
+
+namespace mu
+{
+
+    class Node::Private
+    {
+    public:
+
+		/** Force a virtual destructor. */
+		virtual ~Private() = default; 
+
+		/** Generic pointer to the node owning this private. */
+		const Node* m_pNode = nullptr;
+    };
+
+
+
+#define MUTABLE_IMPLEMENT_NODE( N )				\
+    N::N()										\
+    {											\
+        m_pD = new Private();					\
+		m_pD->m_pNode = this;					\
+	}											\
+                                                \
+    N::~N()										\
+    {											\
+        check( m_pD );							\
+        delete m_pD;							\
+        m_pD = nullptr;							\
+    }											\
+                                                \
+    N::Private* N::GetPrivate() const			\
+    {											\
+        return m_pD;							\
+    }											\
+                                                \
+    const FNodeType* N::GetType() const			\
+    {											\
+        return GetStaticType();					\
+    }											\
+                                                \
+    const FNodeType* N::GetStaticType()			\
+    {											\
+        return &Private::s_type;				\
+    }
+
+}
diff --git a/Engine/Plugins/Mutable/Source/MutableTools/Internal/MuT/NodeScalarEnumParameterPrivate.h b/Engine/Plugins/Mutable/Source/MutableTools/Internal/MuT/NodeScalarEnumParameterPrivate.h
new file mode 100644
--- /dev/null
+++ b/Engine/Plugins/Mutable/Source/MutableTools/Internal/MuT/NodeScalarEnumParameterPrivate.h
@@ -0,0 +1,38 @@
+// Copyright Epic Games, Inc. All Rights Reserved.
+
+#pragma once
+
+#include "MuT/NodeScalarPrivate.h"
+
+#include "MuT/NodeScalarEnumParameter.h"
+#include "MuT/NodeImage.h"
+#include "MuR/ParametersPrivate.h"
+
+
+namespace mu
+{
+
+
+	class NodeScalarEnumParameter::Private : public NodeScalar::Private
+	{
+	public:
+
+		static FNodeType s_type;
+
+		FString Name;
+		FString UID;
+		int32 DefaultValue = 0;
+
+		struct FOption
+		{
+			FString name;
+			float value;
+		};
+
+		TArray<FOption> Options;
+
+		TArray<Ptr<NodeRange>> Ranges;
+	};
+
+}
+
diff --git a/Engine/Plugins/Mutable/Source/MutableTools/Internal/MuT/NodeScalarParameterPrivate.h b/Engine/Plugins/Mutable/Source/MutableTools/Internal/MuT/NodeScalarParameterPrivate.h
new file mode 100644
--- /dev/null
+++ b/Engine/Plugins/Mutable/Source/MutableTools/Internal/MuT/NodeScalarParameterPrivate.h
@@ -0,0 +1,30 @@
+// Copyright Epic Games, Inc. All Rights Reserved.
+
+#pragma once
+
+#include "MuT/NodeScalarPrivate.h"
+
+#include "MuT/NodeScalarParameter.h"
+#include "MuT/NodeImage.h"
+#include "MuT/NodeRange.h"
+#include "MuR/ParametersPrivate.h"
+
+
+namespace mu
+{
+
+
+	class NodeScalarParameter::Private : public NodeScalar::Private
+	{
+	public:
+
+		static FNodeType s_type;
+
+		float DefaultValue = 0.0f;
+		FString Name;
+		FString UID;
+
+		TArray<Ptr<NodeRange>> Ranges;
+	};
+
+}
diff --git a/Engine/Plugins/Mutable/Source/MutableTools/Private/MuT/CodeGenerator.cpp b/Engine/Plugins/Mutable/Source/MutableTools/Private/MuT/CodeGenerator.cpp
--- a/Engine/Plugins/Mutable/Source/MutableTools/Private/MuT/CodeGenerator.cpp
+++ b/Engine/Plugins/Mutable/Source/MutableTools/Private/MuT/CodeGenerator.cpp
@@ -63,17 +63,19 @@
 #include "MuT/NodeImage.h"
 #include "MuT/NodeImageConstant.h"
 #include "MuT/NodeImageFormat.h"
+#include "MuT/NodeImageFormatPrivate.h"
 #include "MuT/NodeImageMipmap.h"
-#include "MuT/NodeImageSwizzle.h"
+#include "MuT/NodeImageMipmapPrivate.h"
+#include "MuT/NodeImageSwizzlePrivate.h"
 #include "MuT/NodeMatrixConstant.h"
 #include "MuT/NodeMesh.h"
 #include "MuT/NodeMeshClipMorphPlane.h"
 #include "MuT/NodeMeshClipWithMesh.h"
+#include "MuT/NodeMeshConstantPrivate.h"
 #include "MuT/NodeMeshFormat.h"
 #include "MuT/NodeMeshFragment.h"
 #include "MuT/NodeMeshMorph.h"
 #include "MuT/NodeMeshReshape.h"
-#include "MuT/NodeMeshConstant.h"
 #include "MuT/NodeModifier.h"
 #include "MuT/NodeModifierMeshClipDeform.h"
 #include "MuT/NodeModifierMeshClipMorphPlane.h"
@@ -82,7 +84,9 @@
 #include "MuT/NodeModifierMeshTransformInMesh.h"
 #include "MuT/NodeModifierSurfaceEdit.h"
 #include "MuT/NodeObject.h"
+#include "MuT/NodeObjectGroupPrivate.h"
 #include "MuT/NodeObjectNew.h"
+#include "MuT/NodePrivate.h"
 #include "MuT/NodeRange.h"
 #include "MuT/NodeRangeFromScalar.h"
 #include "MuT/NodeScalar.h"
@@ -308,10 +312,10 @@ namespace mu
 			const NodeRangeFromScalar* FromScalar = static_cast<const NodeRangeFromScalar*>(Untyped.get());
 
 			Result = FRangeGenerationResult();
-			Result.rangeName = FromScalar->Name;
+			Result.rangeName = FromScalar->GetName();
 
 			FScalarGenerationResult ChildResult;
-			GenerateScalar(ChildResult, Options, FromScalar->Size);
+			GenerateScalar(ChildResult, Options, FromScalar->GetSize());
 			Result.sizeOp = ChildResult.op;
 		}
 		else
@@ -336,9 +340,9 @@ namespace mu
 		{
 			ParamName = CacheKey.Table->GetName();
 		}
-		Result->Name = ParamName;
+		Result->SetName(ParamName);
 
-		Result->DefaultValue = 0;
+		Result->SetDefaultValueIndex(0);
 
 		int32 CurrentRow = 0;
 		int32 RowCount = CacheKey.Table->GetPrivate()->Rows.Num();
@@ -346,13 +350,12 @@ namespace mu
 
 		if (bAddNoneOption)
 		{
-			Result->Options.SetNum(RowCount);
-			Result->Options[CurrentRow].Value = -1;
-			Result->Options[CurrentRow].Name = "None";
+			Result->SetValueCount(RowCount);
+			Result->SetValue(CurrentRow, -1, "None");
 		}
 		else
 		{
-			Result->Options.SetNum(RowCount-1);
+			Result->SetValueCount(RowCount-1);
 		}
 
 		// Add the possible values
@@ -380,20 +383,19 @@ namespace mu
 					ValueName = CacheKey.Table->GetPrivate()->Rows[RowIndex].Values[NameCol].String;
 				}
 
-				Result->Options[CurrentRow].Value = RowIndex;
-				Result->Options[CurrentRow].Name = ValueName;
+				Result->SetValue(CurrentRow, RowIndex, ValueName);
 
 
 				// Set the first row or the selected row as the default one.
 				if(RowIndex == StarCount || ValueName == DefaultRowName)
 				{
-					Result->DefaultValue = RowIndex;
+					Result->SetDefaultValueIndex(RowIndex);
 				}
 
 				// Set the selected row as default (if exists)
 				if (ValueName == DefaultRowName)
 				{
-					Result->DefaultValue = RowIndex;
+					Result->SetDefaultValueIndex(RowIndex);
 				}
 
 				++CurrentRow;
@@ -1262,26 +1264,26 @@ namespace mu
 							{
 								NodeImageMipmap* tm = static_cast<NodeImageMipmap*>(pImageNode.get());
 								if (!mipmapNode) mipmapNode = tm;
-								pImageNode = tm->Source;
+								pImageNode = tm->GetSource();
 							}
 							else if (pImageNode->GetType() == NodeImageFormat::GetStaticType())
 							{
 								NodeImageFormat* tf = static_cast<NodeImageFormat*>(pImageNode.get());
 								if (!formatNode) formatNode = tf;
-								pImageNode = tf->Source;
+								pImageNode = tf->GetSource();
 							}
 							else if (pImageNode->GetType() == NodeImageSwizzle::GetStaticType())
 							{
 								NodeImageSwizzle* ts = static_cast<NodeImageSwizzle*>(pImageNode.get());
 
-								if (!ts->Sources.IsEmpty())
+								if (!ts->GetPrivate()->m_sources.IsEmpty())
 								{
-									NodeImage* Source = ts->Sources[0].get();
+									NodeImage* Source = ts->GetSource(0).get();
 
 									bool bAllSourcesAreTheSame = true;
-									for (int32 SourceIndex = 1; SourceIndex < ts->Sources.Num(); ++SourceIndex)
+									for (int32 SourceIndex = 1; SourceIndex < ts->GetPrivate()->m_sources.Num(); ++SourceIndex)
 									{
-										bAllSourcesAreTheSame = bAllSourcesAreTheSame && (Source == ts->Sources[SourceIndex]);
+										bAllSourcesAreTheSame = bAllSourcesAreTheSame && (Source == ts->GetSource(SourceIndex));
 									}
 
 									if (!swizzleNode && bAllSourcesAreTheSame)
@@ -1349,15 +1351,15 @@ namespace mu
 								if (swizzleNode)
 								{
 									Ptr<ASTOpImageSwizzle> fop = new ASTOpImageSwizzle();
-									fop->Format = swizzleNode->NewFormat;
+									fop->Format = swizzleNode->GetPrivate()->m_format;
 									fop->Sources[0] = imageAd;
 									fop->Sources[1] = imageAd;
 									fop->Sources[2] = imageAd;
 									fop->Sources[3] = imageAd;
-									fop->SourceChannels[0] = swizzleNode->SourceChannels[0];
-									fop->SourceChannels[1] = swizzleNode->SourceChannels[1];
-									fop->SourceChannels[2] = swizzleNode->SourceChannels[2];
-									fop->SourceChannels[3] = swizzleNode->SourceChannels[3];
+									fop->SourceChannels[0] = swizzleNode->GetPrivate()->m_sourceChannels[0];
+									fop->SourceChannels[1] = swizzleNode->GetPrivate()->m_sourceChannels[1];
+									fop->SourceChannels[2] = swizzleNode->GetPrivate()->m_sourceChannels[2];
+									fop->SourceChannels[3] = swizzleNode->GetPrivate()->m_sourceChannels[3];
 									check(fop->Format != EImageFormat::None);
 									imageAd = fop;
 								}
@@ -1369,16 +1371,16 @@ namespace mu
 									op->Source = imageAd;
 									op->BlockLevels = 0;
 
-									op->AddressMode = mipmapNode->Settings.AddressMode;
-									op->FilterType = mipmapNode->Settings.FilterType;
+									op->AddressMode = mipmapNode->GetPrivate()->m_settings.AddressMode;
+									op->FilterType = mipmapNode->GetPrivate()->m_settings.FilterType;
 									imageAd = op;
 								}
 
 								if (formatNode)
 								{
 									Ptr<ASTOpImagePixelFormat> fop = new ASTOpImagePixelFormat();
-									fop->Format = formatNode->Format;
-									fop->FormatIfAlpha = formatNode->FormatIfAlpha;
+									fop->Format = formatNode->GetPrivate()->m_format;
+									fop->FormatIfAlpha = formatNode->GetPrivate()->m_formatIfAlpha;
 									fop->Source = imageAd;
 									check(fop->Format != EImageFormat::None);
 									imageAd = fop;
@@ -1419,10 +1421,10 @@ namespace mu
 									FLayoutBlockDesc LayoutBlockDesc;
 									if (formatNode)
 									{
-										LayoutBlockDesc.FinalFormat = formatNode->FormatIfAlpha;
+										LayoutBlockDesc.FinalFormat = formatNode->GetPrivate()->m_formatIfAlpha;
 										if (LayoutBlockDesc.FinalFormat == EImageFormat::None)
 										{
-											LayoutBlockDesc.FinalFormat = formatNode->Format;
+											LayoutBlockDesc.FinalFormat = formatNode->GetPrivate()->m_format;
 										}
 									}
 
@@ -1521,12 +1523,12 @@ namespace mu
 									if (swizzleNode)
 									{
 										Ptr<ASTOpImageSwizzle> fop = new ASTOpImageSwizzle();
-										fop->Format = swizzleNode->NewFormat;
+										fop->Format = swizzleNode->GetPrivate()->m_format;
 
-										for (int32 ChannelIndex = 0; ChannelIndex < swizzleNode->SourceChannels.Num(); ++ChannelIndex)
+										for (int32 ChannelIndex = 0; ChannelIndex < swizzleNode->GetPrivate()->m_sourceChannels.Num(); ++ChannelIndex)
 										{
 											fop->Sources[ChannelIndex] = imageAd;
-											fop->SourceChannels[ChannelIndex] = swizzleNode->SourceChannels[ChannelIndex];
+											fop->SourceChannels[ChannelIndex] = swizzleNode->GetPrivate()->m_sourceChannels[ChannelIndex];
 										}
 										check(fop->Format != EImageFormat::None);
 										imageAd = fop;
@@ -1568,8 +1570,8 @@ namespace mu
 
 										if (mipmapNode)
 										{
-											mop->AddressMode = mipmapNode->Settings.AddressMode;
-											mop->FilterType = mipmapNode->Settings.FilterType;
+											mop->AddressMode = mipmapNode->GetPrivate()->m_settings.AddressMode;
+											mop->FilterType = mipmapNode->GetPrivate()->m_settings.FilterType;
 										}
 
 										imageAd = mop;
@@ -1578,8 +1580,8 @@ namespace mu
 									if (formatNode)
 									{
 										Ptr<ASTOpImagePixelFormat> fop = new ASTOpImagePixelFormat();
-										fop->Format = formatNode->Format;
-										fop->FormatIfAlpha = formatNode->FormatIfAlpha;
+										fop->Format = formatNode->GetPrivate()->m_format;
+										fop->FormatIfAlpha = formatNode->GetPrivate()->m_formatIfAlpha;
 										fop->Source = imageAd;
 										check(fop->Format != EImageFormat::None);
 										imageAd = fop;
@@ -1632,7 +1634,7 @@ namespace mu
 					{
 						// MUTABLE_CPUPROFILER_SCOPE(SurfaceScalar);
 
-						if (Ptr<NodeScalar> ScalarNode = SurfaceNode->Scalars[t].Scalar)
+						if (NodeScalarPtr pScalarNode = SurfaceNode->Scalars[t].Scalar)
 						{
 							Ptr<ASTOpInstanceAdd> op = new ASTOpInstanceAdd();
 							op->type = EOpType::IN_ADDSCALAR;
@@ -1640,7 +1642,7 @@ namespace mu
 
 							// Scalar
 							FScalarGenerationResult ScalarResult;
-							GenerateScalar(ScalarResult, Options, ScalarNode);
+							GenerateScalar(ScalarResult, Options, pScalarNode);
 							op->value = ScalarResult.op;
 
 							// Name
@@ -1654,14 +1656,14 @@ namespace mu
 					//------------------------------------------------------------------------
 					for (int32 t = 0; t < SurfaceNode->Strings.Num(); ++t)
 					{
-						if (Ptr<NodeString> StringNode = SurfaceNode->Strings[t].String)
+						if (NodeStringPtr pStringNode = SurfaceNode->Strings[t].String)
 						{
 							Ptr<ASTOpInstanceAdd> op = new ASTOpInstanceAdd();
 							op->type = EOpType::IN_ADDSTRING;
 							op->instance = LastSurfOp;
 
 							FStringGenerationResult StringResult;
-							GenerateString(StringResult, Options, StringNode);
+							GenerateString(StringResult, Options, pStringNode);
 							op->value = StringResult.op;
 
 							// Name
@@ -1995,14 +1997,16 @@ namespace mu
     }
 
 
-	void CodeGenerator::GenerateObject_Group(const FObjectGenerationOptions& Options, FObjectGenerationResult& OutResult, const NodeObjectGroup* Node)
+	void CodeGenerator::GenerateObject_Group(const FObjectGenerationOptions& Options, FObjectGenerationResult& OutResult, const NodeObjectGroup* InNode)
 	{
+		const NodeObjectGroup::Private& node = *InNode->GetPrivate();
+
 		TArray<FString> UsedNames;
 
         // Parse the child objects first, which will accumulate operations in the patching lists
-        for ( int32 ChildIndex=0; ChildIndex < Node->Children.Num(); ++ChildIndex)
+        for ( int32 ChildIndex=0; ChildIndex <node.m_children.Num(); ++ChildIndex)
         {
-            if ( const NodeObject* pChildNode = Node->Children[ChildIndex].get() )
+            if ( const NodeObject* pChildNode = node.m_children[ChildIndex].get() )
             {
 				// Look for the child condition in the first pass
                 Ptr<ASTOp> ConditionOp;
@@ -2030,7 +2034,7 @@ namespace mu
 				if (UsedNames.Contains(ChildName))
 				{
 					FString Msg = FString::Printf(TEXT("Object group has more than one children with the same name [%s]."), *ChildName );
-					ErrorLog->Add(Msg, ELMT_WARNING, Node->GetMessageContext());
+					ErrorLog->Add(Msg, ELMT_WARNING, InNode->GetMessageContext());
 				}
 				else
 				{
@@ -2050,7 +2054,7 @@ namespace mu
 
         // Create a constant node
         Ptr<NodeBoolConstant> pNode = new NodeBoolConstant();
-        pNode->Value = Value;
+        pNode->SetValue(Value);
 
 		FBoolGenerationResult ChildResult;
 		FGenericGenerationOptions Options;
@@ -2418,7 +2422,7 @@ namespace mu
 				TSharedPtr< TSharedPtr<FMesh> > ResolveMorphedMesh = MakeShared<TSharedPtr<FMesh>>();
 				UE::Tasks::FTask TargetMeshTask;
 				{
-					FMesh* OriginalMesh = OriginalMeshNode->Value.Get();
+					FMesh* OriginalMesh = OriginalMeshNode->GetValue().Get();
 					check(OriginalMesh->IsReference());
 					uint32 OriginalMeshID = OriginalMesh->GetReferencedMesh();
 
@@ -2441,7 +2445,7 @@ namespace mu
 				else
 				{
 					Ptr<NodeScalarConstant> auxNode = new NodeScalarConstant();
-					auxNode->Value = 1.0f;
+					auxNode->SetValue(1.0f);
 
 					FScalarGenerationResult ChildResult;
 					GenerateScalar(ChildResult, StaticOptions, auxNode);
@@ -2582,7 +2586,7 @@ namespace mu
 				TSharedPtr< TSharedPtr<FMesh> > ResolveOriginalMesh = MakeShared<TSharedPtr<FMesh>>();
 				UE::Tasks::FTask TargetMeshTask;
 				{
-					FMesh* OriginalMesh = OriginalMeshNode->Value.Get();
+					FMesh* OriginalMesh = OriginalMeshNode->GetValue().Get();
 					check(OriginalMesh->IsReference());
 					uint32 OriginalMeshID = OriginalMesh->GetReferencedMesh();
 
diff --git a/Engine/Plugins/Mutable/Source/MutableTools/Private/MuT/CodeGenerator.h b/Engine/Plugins/Mutable/Source/MutableTools/Private/MuT/CodeGenerator.h
--- a/Engine/Plugins/Mutable/Source/MutableTools/Private/MuT/CodeGenerator.h
+++ b/Engine/Plugins/Mutable/Source/MutableTools/Private/MuT/CodeGenerator.h
@@ -87,7 +87,6 @@ namespace mu
 	class NodeMeshSwitch;
 	class NodeMeshTransform;
 	class NodeMeshVariation;
-	class NodeMeshParameter;
 	class NodeRange;
 	class NodeScalarArithmeticOperation;
 	class NodeScalarConstant;
@@ -658,7 +657,6 @@ namespace mu
 		FMeshTask GenerateMesh_Table(const FMeshGenerationStaticOptions& StaticOptions, FMeshOptionsTask, const NodeMeshTable*);
 		FMeshTask GenerateMesh_Reshape(const FMeshGenerationStaticOptions& StaticOptions, FMeshOptionsTask, const NodeMeshReshape*);
 		FMeshTask GenerateMesh_ClipDeform(const FMeshGenerationStaticOptions& StaticOptions, FMeshOptionsTask, const NodeMeshClipDeform*);
-		FMeshTask GenerateMesh_Parameter(const FMeshGenerationStaticOptions& StaticOptions, FMeshOptionsTask, const NodeMeshParameter*);
 
 		//-----------------------------------------------------------------------------------------
 		void PrepareMeshForLayout(const FGeneratedLayout&, TSharedPtr<FMesh>,
@@ -1024,14 +1022,14 @@ namespace mu
 
         // Create the switch node
         Ptr<NODE_SWITCH_TYPE> SwitchNode = new NODE_SWITCH_TYPE;
-		SwitchNode->Parameter = Variable;
-		SwitchNode->Options.SetNum(NumRows);
+		SwitchNode->SetParameter( Variable );
+		SwitchNode->SetOptionCount(NumRows);
 
 		for (int32 RowIndex = 0; RowIndex < NumRows; ++RowIndex)
         {
             check(RowIndex <= 0xFFFF);
             Ptr<NODE_TYPE> BranchNode = GenerateOption( node, ColIndex, RowIndex, ErrorLog.Get() );
-			SwitchNode->Options[RowIndex] = BranchNode;
+			SwitchNode->SetOption(RowIndex, BranchNode);
         }
 
 		OutResult = SwitchNode;
diff --git a/Engine/Plugins/Mutable/Source/MutableTools/Private/MuT/CodeGenerator_Bool.cpp b/Engine/Plugins/Mutable/Source/MutableTools/Private/MuT/CodeGenerator_Bool.cpp
--- a/Engine/Plugins/Mutable/Source/MutableTools/Private/MuT/CodeGenerator_Bool.cpp
+++ b/Engine/Plugins/Mutable/Source/MutableTools/Private/MuT/CodeGenerator_Bool.cpp
@@ -15,6 +15,7 @@
 #include "MuT/ASTOpBoolNot.h"
 #include "MuT/CodeGenerator.h"
 #include "MuT/NodeBool.h"
+#include "MuT/NodeBoolPrivate.h"
 #include "MuT/NodeRange.h"
 
 
@@ -77,14 +78,18 @@ namespace mu
 	}
 
 
-	void CodeGenerator::GenerateBool_Constant(FBoolGenerationResult& result, const FGenericGenerationOptions& Options, const Ptr<const NodeBoolConstant>& Node)
+	void CodeGenerator::GenerateBool_Constant(FBoolGenerationResult& result, const FGenericGenerationOptions& Options, const Ptr<const NodeBoolConstant>& Typed)
 	{
-		result.op = new ASTOpConstantBool(Node->Value);
+		const NodeBoolConstant::Private& node = *Typed->GetPrivate();
+
+		result.op = new ASTOpConstantBool(node.Value);
 	}
 
 
-	void CodeGenerator::GenerateBool_Parameter(FBoolGenerationResult& result, const FGenericGenerationOptions& Options, const Ptr<const NodeBoolParameter>& Node)
+	void CodeGenerator::GenerateBool_Parameter(FBoolGenerationResult& result, const FGenericGenerationOptions& Options, const Ptr<const NodeBoolParameter>& Typed)
 	{
+		const NodeBoolParameter::Private& node = *Typed->GetPrivate();
+
 		Ptr<ASTOpParameter> op;
 
 		// This can run in any thread.
@@ -92,21 +97,21 @@ namespace mu
 		{
 			UE::TUniqueLock Lock(FirstPass.ParameterNodes.Mutex);
 
-			Found = FirstPass.ParameterNodes.Map.Find(Node);
+			Found = FirstPass.ParameterNodes.Map.Find(node.m_pNode);
 			if (!Found)
 			{
 				FParameterDesc param;
-				param.Name = Node->Name;
-				bool bParseOk = FGuid::Parse(Node->UID, param.UID);
+				param.Name = node.Name;
+				bool bParseOk = FGuid::Parse(node.UID, param.UID);
 				check(bParseOk);
 				param.Type = EParameterType::Bool;
-				param.DefaultValue.Set<FParamBoolType>(Node->DefaultValue);
+				param.DefaultValue.Set<FParamBoolType>(node.DefaultValue);
 
 				op = new ASTOpParameter();
 				op->Type = EOpType::BO_PARAMETER;
 				op->Parameter = param;
 
-				FirstPass.ParameterNodes.Map.Add(Node, op);
+				FirstPass.ParameterNodes.Map.Add(node.m_pNode, op);
 			}
 			else
 			{
@@ -117,10 +122,10 @@ namespace mu
 		if (!Found)
 		{
 			// Generate the code for the ranges
-			for (int32 a = 0; a < Node->Ranges.Num(); ++a)
+			for (int32 a = 0; a < node.Ranges.Num(); ++a)
 			{
 				FRangeGenerationResult RangeResult;
-				GenerateRange(RangeResult, Options, Node->Ranges[a]);
+				GenerateRange(RangeResult, Options, node.Ranges[a]);
 				op->Ranges.Emplace(op.get(), RangeResult.sizeOp, RangeResult.rangeName, RangeResult.rangeUID);
 			}
 		}
@@ -129,27 +134,31 @@ namespace mu
 	}
 
 
-	void CodeGenerator::GenerateBool_Not(FBoolGenerationResult& result, const FGenericGenerationOptions& Options, const Ptr<const NodeBoolNot>& Node)
+	void CodeGenerator::GenerateBool_Not(FBoolGenerationResult& result, const FGenericGenerationOptions& Options, const Ptr<const NodeBoolNot>& Typed)
 	{
+		const NodeBoolNot::Private& node = *Typed->GetPrivate();
+
 		Ptr<ASTOpBoolNot> op = new ASTOpBoolNot();
 
 		FBoolGenerationResult ChildResult;
-		GenerateBool(ChildResult, Options, Node->Source);
+		GenerateBool(ChildResult, Options, node.m_pSource);
 		op->A = ChildResult.op;
 
 		result.op = op;
 	}
 
 
-	void CodeGenerator::GenerateBool_And(FBoolGenerationResult& result, const FGenericGenerationOptions& Options, const Ptr<const NodeBoolAnd>& Node)
+	void CodeGenerator::GenerateBool_And(FBoolGenerationResult& result, const FGenericGenerationOptions& Options, const Ptr<const NodeBoolAnd>& Typed)
 	{
+		const NodeBoolAnd::Private& node = *Typed->GetPrivate();
+
 		Ptr<ASTOpBoolAnd> op = new ASTOpBoolAnd;
 
 		FBoolGenerationResult ChildResult;
-		GenerateBool(ChildResult, Options, Node->A);
+		GenerateBool(ChildResult, Options, node.m_pA);
 		op->A = ChildResult.op;
 
-		GenerateBool(ChildResult, Options, Node->B);
+		GenerateBool(ChildResult, Options, node.m_pB);
 		op->B = ChildResult.op;
 
 		result.op = op;
diff --git a/Engine/Plugins/Mutable/Source/MutableTools/Private/MuT/CodeGenerator_Color.cpp b/Engine/Plugins/Mutable/Source/MutableTools/Private/MuT/CodeGenerator_Color.cpp
--- a/Engine/Plugins/Mutable/Source/MutableTools/Private/MuT/CodeGenerator_Color.cpp
+++ b/Engine/Plugins/Mutable/Source/MutableTools/Private/MuT/CodeGenerator_Color.cpp
@@ -348,9 +348,9 @@ namespace mu
 		else
 		{
 			// Set a constant 0.5 value
-			Ptr<NodeScalarConstant> Node = new NodeScalarConstant();
-			Node->Value = 0.5f;
-			GenerateScalar(ChildResult, Options, Node);
+			NodeScalarConstantPtr pNode = new NodeScalarConstant();
+			pNode->SetValue(0.5f);
+			GenerateScalar(ChildResult, Options, pNode);
 			op->X = ChildResult.op;
 		}
 
@@ -364,9 +364,9 @@ namespace mu
 		else
 		{
 			// Set a constant 0.5 value
-			Ptr<NodeScalarConstant> Node = new NodeScalarConstant();
-			Node->Value = 0.5f;
-			GenerateScalar(ChildResult, Options, Node);
+			NodeScalarConstantPtr pNode = new NodeScalarConstant();
+			pNode->SetValue(0.5f);
+			GenerateScalar(ChildResult, Options, pNode);
 			op->Y = ChildResult.op;
 		}
 
@@ -394,9 +394,9 @@ namespace mu
 		}
 		else
 		{
-			Ptr<NodeScalarConstant> Node = new NodeScalarConstant();
-			Node->Value  = 1.0f;
-			GenerateScalar(ChildResult, Options, Node);
+			NodeScalarConstantPtr pNode = new NodeScalarConstant();
+			pNode->SetValue(1.0f);
+			GenerateScalar(ChildResult, Options, pNode);
 			op->V[0] = ChildResult.op;
 		}
 
@@ -408,9 +408,9 @@ namespace mu
 		}
 		else
 		{
-			Ptr<NodeScalarConstant> Node = new NodeScalarConstant();
-			Node->Value = 1.0f;
-			GenerateScalar(ChildResult, Options, Node);
+			NodeScalarConstantPtr pNode = new NodeScalarConstant();
+			pNode->SetValue(1.0f);
+			GenerateScalar(ChildResult, Options, pNode);
 			op->V[1] = ChildResult.op;
 		}
 
@@ -422,9 +422,9 @@ namespace mu
 		}
 		else
 		{
-			Ptr<NodeScalarConstant> Node = new NodeScalarConstant();
-			Node->Value = 1.0f;
-			GenerateScalar(ChildResult, Options, Node);
+			NodeScalarConstantPtr pNode = new NodeScalarConstant();
+			pNode->SetValue(1.0f);
+			GenerateScalar(ChildResult, Options, pNode);
 			op->V[2] = ChildResult.op;
 		}
 
@@ -436,9 +436,9 @@ namespace mu
 		}
 		else
 		{
-			Ptr<NodeScalarConstant> Node = new NodeScalarConstant();
-			Node->Value = 1.0f;
-			GenerateScalar(ChildResult, Options, Node);
+			NodeScalarConstantPtr pNode = new NodeScalarConstant();
+			pNode->SetValue(1.0f);
+			GenerateScalar(ChildResult, Options, pNode);
 			op->V[3] = ChildResult.op;
 		}
 
diff --git a/Engine/Plugins/Mutable/Source/MutableTools/Private/MuT/CodeGenerator_FirstPass.cpp b/Engine/Plugins/Mutable/Source/MutableTools/Private/MuT/CodeGenerator_FirstPass.cpp
--- a/Engine/Plugins/Mutable/Source/MutableTools/Private/MuT/CodeGenerator_FirstPass.cpp
+++ b/Engine/Plugins/Mutable/Source/MutableTools/Private/MuT/CodeGenerator_FirstPass.cpp
@@ -24,12 +24,14 @@
 #include "MuT/NodeModifierMeshClipWithUVMask.h"
 #include "MuT/NodeModifierSurfaceEdit.h"
 #include "MuT/NodeObject.h"
+#include "MuT/NodeObjectGroupPrivate.h"
 #include "MuT/NodeObjectNew.h"
+#include "MuT/NodePrivate.h"
 #include "MuT/NodeSurface.h"
 #include "MuT/NodeSurfaceNew.h"
 #include "MuT/NodeSurfaceVariation.h"
 #include "MuT/NodeSurfaceSwitch.h"
-#include "MuT/NodeScalarEnumParameter.h"
+#include "MuT/NodeScalarEnumParameterPrivate.h"
 
 namespace mu
 {
@@ -609,29 +611,30 @@ namespace mu
 
 
 	//---------------------------------------------------------------------------------------------
-	void FirstPassGenerator::Generate_ObjectGroup(const NodeObjectGroup* Node)
+	void FirstPassGenerator::Generate_ObjectGroup(const NodeObjectGroup* InNode)
 	{
+		const NodeObjectGroup::Private* Private = InNode->GetPrivate();
+
 		// Prepare the enumeration parameter if necessary
         Ptr<ASTOpParameter> enumOp;
-        if (Node->Type==NodeObjectGroup::CS_ALWAYS_ONE ||
-			Node->Type==NodeObjectGroup::CS_ONE_OR_NONE )
+        if (Private->Type==NodeObjectGroup::CS_ALWAYS_ONE ||
+			Private->Type==NodeObjectGroup::CS_ONE_OR_NONE )
         {
 			UE::TUniqueLock Lock(ParameterNodes.Mutex);
-			Ptr<ASTOpParameter>* Found = ParameterNodes.Map.Find(Node);
+			Ptr<ASTOpParameter>* Found = ParameterNodes.Map.Find(InNode);
 
 			if (!Found)
 			{
 				Ptr<ASTOpParameter> op = new ASTOpParameter();
 				op->Type = EOpType::NU_PARAMETER;
-				op->Parameter.DefaultValue.Set<FParamIntType>(Node->DefaultValue);
 
-				op->Parameter.Name = Node->Name;
-				bool bParseOk = FGuid::Parse(Node->Uid, op->Parameter.UID);
+				op->Parameter.Name = Private->Name;
+				bool bParseOk = FGuid::Parse(Private->Uid, op->Parameter.UID);
 				check(bParseOk);
 				op->Parameter.Type = EParameterType::Int;
-				op->Parameter.DefaultValue.Set<FParamIntType>(Node->DefaultValue);
+				op->Parameter.DefaultValue.Set<FParamIntType>(Private->DefaultValue);
 
-				if (Node->Type == NodeObjectGroup::CS_ONE_OR_NONE)
+				if (Private->Type == NodeObjectGroup::CS_ONE_OR_NONE)
 				{
 					FParameterDesc::FIntValueDesc nullValue;
 					nullValue.Value = -1;
@@ -639,7 +642,7 @@ namespace mu
 					op->Parameter.PossibleValues.Add(nullValue);
 				}
 
-				ParameterNodes.Map.Add(Node, op);
+				ParameterNodes.Map.Add(InNode, op);
 
 				enumOp = op;
 			}
@@ -651,20 +654,21 @@ namespace mu
 
 
         // Parse the child objects
-		for ( int32 t=0; t< Node->Children.Num(); ++t )
+		for ( int32 t=0; t< Private->m_children.Num(); ++t )
         {
-            if ( mu::Ptr<const NodeObject> ChildNode = Node->Children[t] )
+            if ( mu::Ptr<const NodeObject> ChildNode = Private->m_children[t] )
             {
                 // Overwrite the implicit condition
                 Ptr<ASTOp> paramOp = 0;
-                switch (Node->Type )
+                switch (Private->Type )
                 {
                     case NodeObjectGroup::CS_TOGGLE_EACH:
                     {           
 						if (ChildNode->GetType() == NodeObjectGroup::GetStaticType())
 						{
-							FString Msg = FString::Printf(TEXT("The Group Node [%s] has type Toggle and its direct child is a Group node, which is not allowed. Change the type or add a Child Object node in between them."), *Node->Name);
-							ErrorLog->Add(Msg, ELMT_ERROR, Node->GetMessageContext());
+							FString Msg = FString::Printf(TEXT("The Group Node [%s] has type Toggle and its direct child is a Group node, which is not allowed. Change the type or add a Child Object node in between them."),
+								*Private->Name);
+							ErrorLog->Add(Msg, ELMT_ERROR, InNode->GetMessageContext());
 						}
 						else
 						{
