[Backout] - CL40275509

diff --git a/Engine/Plugins/Mutable/Source/MutableTools/Private/MuT/CodeGenerator_Scalar.cpp b/Engine/Plugins/Mutable/Source/MutableTools/Private/MuT/CodeGenerator_Scalar.cpp
--- a/Engine/Plugins/Mutable/Source/MutableTools/Private/MuT/CodeGenerator_Scalar.cpp
+++ b/Engine/Plugins/Mutable/Source/MutableTools/Private/MuT/CodeGenerator_Scalar.cpp
@@ -26,13 +26,19 @@
 #include "MuT/NodeRange.h"
 #include "MuT/NodeScalar.h"
 #include "MuT/NodeScalarArithmeticOperation.h"
+#include "MuT/NodeScalarArithmeticOperationPrivate.h"
 #include "MuT/NodeScalarConstant.h"
+#include "MuT/NodeScalarConstantPrivate.h"
 #include "MuT/NodeScalarCurve.h"
+#include "MuT/NodeScalarCurvePrivate.h"
 #include "MuT/NodeScalarEnumParameter.h"
+#include "MuT/NodeScalarEnumParameterPrivate.h"
 #include "MuT/NodeScalarParameter.h"
+#include "MuT/NodeScalarParameterPrivate.h"
 #include "MuT/NodeScalarSwitch.h"
 #include "MuT/NodeScalarTable.h"
 #include "MuT/NodeScalarVariation.h"
+#include "MuT/NodeScalarVariationPrivate.h"
 #include "MuT/Table.h"
 #include "MuT/TablePrivate.h"
 
@@ -117,36 +123,43 @@ namespace mu
 	}
 
 
-	void CodeGenerator::GenerateScalar_Constant(FScalarGenerationResult& result, const FGenericGenerationOptions& Options, const Ptr<const NodeScalarConstant>& Node)
+	//-------------------------------------------------------------------------------------------------
+	void CodeGenerator::GenerateScalar_Constant(FScalarGenerationResult& result, const FGenericGenerationOptions& Options, const Ptr<const NodeScalarConstant>& Typed)
 	{
+		const NodeScalarConstant::Private& node = *Typed->GetPrivate();
+
 		Ptr<ASTOpConstantScalar> op = new ASTOpConstantScalar;
-		op->Value = Node->Value;
+		op->Value = node.Value;
+
 		result.op = op;
 	}
 
 
-	void CodeGenerator::GenerateScalar_Parameter(FScalarGenerationResult& result, const FGenericGenerationOptions& Options, const Ptr<const NodeScalarParameter>& Node)
+	//-------------------------------------------------------------------------------------------------
+	void CodeGenerator::GenerateScalar_Parameter(FScalarGenerationResult& result, const FGenericGenerationOptions& Options, const Ptr<const NodeScalarParameter>& Typed)
 	{
+		const NodeScalarParameter::Private& node = *Typed->GetPrivate();
+
 		Ptr<ASTOpParameter> op;
 
 		Ptr<ASTOpParameter>* Found = nullptr;		
 		{
 			UE::TUniqueLock Lock(FirstPass.ParameterNodes.Mutex);
-			Found = FirstPass.ParameterNodes.Map.Find(Node);
+			Found = FirstPass.ParameterNodes.Map.Find(node.m_pNode);
 			if (!Found)
 			{
 				FParameterDesc param;
-				param.Name = Node->Name;
-				bool bParseOk = FGuid::Parse(Node->UID, param.UID);
+				param.Name = node.Name;
+				bool bParseOk = FGuid::Parse(node.UID, param.UID);
 				check(bParseOk);
 				param.Type = EParameterType::Float;
-				param.DefaultValue.Set<FParamFloatType>(Node->DefaultValue);
+				param.DefaultValue.Set<FParamFloatType>(node.DefaultValue);
 
 				op = new ASTOpParameter();
 				op->Type = EOpType::SC_PARAMETER;
 				op->Parameter = param;
 
-				FirstPass.ParameterNodes.Map.Add(Node, op);
+				FirstPass.ParameterNodes.Map.Add(node.m_pNode, op);
 			}
 			else
 			{
@@ -157,10 +170,10 @@ namespace mu
 		if (!Found)
 		{
 			// Generate the code for the ranges
-			for (int32 a = 0; a < Node->Ranges.Num(); ++a)
+			for (int32 a = 0; a < node.Ranges.Num(); ++a)
 			{
 				FRangeGenerationResult rangeResult;
-				GenerateRange(rangeResult, Options, Node->Ranges[a]);
+				GenerateRange(rangeResult, Options, node.Ranges[a]);
 				op->Ranges.Emplace(op.get(), rangeResult.sizeOp, rangeResult.rangeName, rangeResult.rangeUID);
 			}
 		}
@@ -169,38 +182,41 @@ namespace mu
 	}
 
 
-	void CodeGenerator::GenerateScalar_EnumParameter(FScalarGenerationResult& result, const FGenericGenerationOptions& Options, const Ptr<const NodeScalarEnumParameter>& Node)
+	//-------------------------------------------------------------------------------------------------
+	void CodeGenerator::GenerateScalar_EnumParameter(FScalarGenerationResult& result, const FGenericGenerationOptions& Options, const Ptr<const NodeScalarEnumParameter>& Typed)
 	{
+		const NodeScalarEnumParameter::Private& node = *Typed->GetPrivate();
+
 		Ptr<ASTOpParameter> op;
 
 		Ptr<ASTOpParameter>* Found = nullptr;
 		{
 			UE::TUniqueLock Lock(FirstPass.ParameterNodes.Mutex);
-			Found = FirstPass.ParameterNodes.Map.Find(Node);
+			Found = FirstPass.ParameterNodes.Map.Find(node.m_pNode);
 			if (!Found)
 			{
 				FParameterDesc param;
-				param.Name = Node->Name;
-				if (!Node->UID.IsEmpty())
+				param.Name = node.Name;
+				if (!node.UID.IsEmpty())
 				{
-					bool bParseOk = FGuid::Parse(Node->UID, param.UID);
+					bool bParseOk = FGuid::Parse(node.UID, param.UID);
 					check(bParseOk);
 				}
 				param.Type = EParameterType::Int;
-				param.DefaultValue.Set<FParamIntType>(Node->DefaultValue);
+				param.DefaultValue.Set<FParamIntType>(node.DefaultValue);
 
-				param.PossibleValues.SetNum(Node->Options.Num());
-				for (int32 i = 0; i < Node->Options.Num(); ++i)
+				param.PossibleValues.SetNum(node.Options.Num());
+				for (int32 i = 0; i < node.Options.Num(); ++i)
 				{
-					param.PossibleValues[i].Value = (int16)Node->Options[i].Value;
-					param.PossibleValues[i].Name = Node->Options[i].Name;
+					param.PossibleValues[i].Value = (int16_t)node.Options[i].value;
+					param.PossibleValues[i].Name = node.Options[i].name;
 				}
 
 				op = new ASTOpParameter();
 				op->Type = EOpType::NU_PARAMETER;
 				op->Parameter = param;
 
-				FirstPass.ParameterNodes.Map.Add(Node, op);
+				FirstPass.ParameterNodes.Map.Add(node.m_pNode, op);
 			}
 			else
 			{
@@ -211,11 +227,11 @@ namespace mu
 		if (!Found)
 		{
 			// Generate the code for the ranges
-			for (int32 RangeIndex = 0; RangeIndex < Node->Ranges.Num(); ++RangeIndex)
+			for (int32 a = 0; a < node.Ranges.Num(); ++a)
 			{
-				FRangeGenerationResult RangeResult;
-				GenerateRange(RangeResult, Options, Node->Ranges[RangeIndex]);
-				op->Ranges.Emplace(op.get(), RangeResult.sizeOp, RangeResult.rangeName, RangeResult.rangeUID);
+				FRangeGenerationResult rangeResult;
+				GenerateRange(rangeResult, Options, node.Ranges[a]);
+				op->Ranges.Emplace(op.get(), rangeResult.sizeOp, rangeResult.rangeName, rangeResult.rangeUID);
 			}
 		}
 
@@ -223,6 +239,7 @@ namespace mu
 	}
 
 
+	//-------------------------------------------------------------------------------------------------
 	void CodeGenerator::GenerateScalar_Switch(FScalarGenerationResult& result, const FGenericGenerationOptions& Options, const Ptr<const NodeScalarSwitch>& Typed)
 	{
 		const NodeScalarSwitch& node = *Typed;
@@ -275,25 +292,28 @@ namespace mu
 	}
 
 
-	void CodeGenerator::GenerateScalar_Variation(FScalarGenerationResult& result, const FGenericGenerationOptions& Options, const Ptr<const NodeScalarVariation>& Node)
+	//-------------------------------------------------------------------------------------------------
+	void CodeGenerator::GenerateScalar_Variation(FScalarGenerationResult& result, const FGenericGenerationOptions& Options, const Ptr<const NodeScalarVariation>& Typed)
 	{
+		const NodeScalarVariation::Private& node = *Typed->GetPrivate();
+
 		Ptr<ASTOp> op;
 
 		// Default case
-		if (Node->DefaultScalar)
+		if (node.m_defaultScalar)
 		{
 			FMeshGenerationResult branchResults;
 
 			FScalarGenerationResult ChildResult;
-			GenerateScalar(ChildResult, Options, Node->DefaultScalar);
+			GenerateScalar(ChildResult, Options, node.m_defaultScalar);
 			op = ChildResult.op;
 		}
 
 		// Process variations in reverse order, since conditionals are built bottom-up.
-		for (int32 t = Node->Variations.Num() - 1; t >= 0; --t)
+		for (int32 t = node.m_variations.Num() - 1; t >= 0; --t)
 		{
 			int32 tagIndex = -1;
-			const FString& tag = Node->Variations[t].Tag;
+			const FString& tag = node.m_variations[t].m_tag;
 			for (int32 i = 0; i < FirstPass.Tags.Num(); ++i)
 			{
 				if (FirstPass.Tags[i].Tag == tag)
@@ -306,21 +326,21 @@ namespace mu
 			{
 				FString Msg = FString::Printf(TEXT("Unknown tag found in image variation [%s]."), *tag);
 
-				ErrorLog->Add(Msg, ELMT_WARNING, Node->GetMessageContext());
+				ErrorLog->Add(Msg, ELMT_WARNING, Typed->GetMessageContext());
 				continue;
 			}
 
 			Ptr<ASTOp> variationOp;
-			if (Node->Variations[t].Scalar)
+			if (node.m_variations[t].m_scalar)
 			{
 				FScalarGenerationResult ChildResult;
-				GenerateScalar(ChildResult, Options, Node->Variations[t].Scalar);
+				GenerateScalar(ChildResult, Options, node.m_variations[t].m_scalar);
 				variationOp = ChildResult.op;
 			}
 			else
 			{
 				// This argument is required
-				variationOp = GenerateMissingScalarCode(TEXT("Variation option"), 0.0f, Node->GetMessageContext());
+				variationOp = GenerateMissingScalarCode(TEXT("Variation option"), 0.0f, Typed->GetMessageContext());
 			}
 
 
@@ -337,6 +357,7 @@ namespace mu
 	}
 
 
+	//-------------------------------------------------------------------------------------------------
 	void CodeGenerator::GenerateScalar_Curve(FScalarGenerationResult& result, const FGenericGenerationOptions& Options, const Ptr<const NodeScalarCurve>& Typed)
 	{
 		Ptr<ASTOpScalarCurve> op = new ASTOpScalarCurve();
@@ -358,11 +379,14 @@ namespace mu
 	}
 
 
-	void CodeGenerator::GenerateScalar_Arithmetic(FScalarGenerationResult& result, const FGenericGenerationOptions& Options, const Ptr<const NodeScalarArithmeticOperation>& Node)
+	//-------------------------------------------------------------------------------------------------
+	void CodeGenerator::GenerateScalar_Arithmetic(FScalarGenerationResult& result, const FGenericGenerationOptions& Options, const Ptr<const NodeScalarArithmeticOperation>& Typed)
 	{
+		const NodeScalarArithmeticOperation::Private& node = *Typed->GetPrivate();
+
 		Ptr<ASTOpScalarArithmetic> op = new ASTOpScalarArithmetic;
 
-		switch (Node->Operation)
+		switch (node.m_operation)
 		{
 		case NodeScalarArithmeticOperation::AO_ADD: op->Operation = OP::ArithmeticArgs::ADD; break;
 		case NodeScalarArithmeticOperation::AO_SUBTRACT: op->Operation = OP::ArithmeticArgs::SUBTRACT; break;
@@ -375,7 +399,7 @@ namespace mu
 		}
 
 		// A
-		if (NodeScalar* pA = Node->A.get())
+		if (NodeScalar* pA = node.m_pA.get())
 		{
 			FScalarGenerationResult ParamResult;
 			GenerateScalar(ParamResult, Options, pA);
@@ -384,11 +408,11 @@ namespace mu
 		}
 		else
 		{
-			op->A = GenerateMissingScalarCode( TEXT("ScalarArithmetic A"), 1.0f, Node->GetMessageContext() );
+			op->A = GenerateMissingScalarCode( TEXT("ScalarArithmetic A"), 1.0f, Typed->GetMessageContext() );
 		}
 
 		// B
-		if (NodeScalar* pB = Node->B.get())
+		if (NodeScalar* pB = node.m_pB.get())
 		{
 			FScalarGenerationResult ParamResult;
 			GenerateScalar(ParamResult, Options, pB);
@@ -397,23 +421,27 @@ namespace mu
 		}
 		else
 		{
-			op->B =  GenerateMissingScalarCode( TEXT("ScalarArithmetic B"), 1.0f, Node->GetMessageContext() );
+			op->B =  GenerateMissingScalarCode( TEXT("ScalarArithmetic B"), 1.0f, Typed->GetMessageContext() );
 		}
 
 		result.op = op;
 	}
 
 
-	void CodeGenerator::GenerateScalar_Table(FScalarGenerationResult& result, const FGenericGenerationOptions& Options, const Ptr<const NodeScalarTable>& Node)
+	//-------------------------------------------------------------------------------------------------
+	void CodeGenerator::GenerateScalar_Table(FScalarGenerationResult& result, const FGenericGenerationOptions& Options, const Ptr<const NodeScalarTable>& Typed)
 	{
-		Ptr<ASTOp> Op = GenerateTableSwitch<NodeScalarTable, ETableColumnType::Scalar, EOpType::SC_SWITCH>(*Node,
-			[this,&Options](const NodeScalarTable& Node, int32 colIndex, int32 row, FErrorLog*)
+		const NodeScalarTable& node = *Typed;
+
+		Ptr<ASTOp> Op = GenerateTableSwitch<NodeScalarTable, ETableColumnType::Scalar, EOpType::SC_SWITCH>(node,
+			[this,&Options](const NodeScalarTable& node, int32 colIndex, int32 row, FErrorLog*)
 			{
-				Ptr<NodeScalarConstant> Cell = new NodeScalarConstant();
-				Cell->Value = Node.Table->GetPrivate()->Rows[row].Values[colIndex].Scalar;
+				NodeScalarConstantPtr pCell = new NodeScalarConstant();
+				float scalar = node.Table->GetPrivate()->Rows[row].Values[colIndex].Scalar;
+				pCell->SetValue(scalar);
 
 				FScalarGenerationResult ParamResult;
-				GenerateScalar(ParamResult, Options, Cell);
+				GenerateScalar(ParamResult, Options, pCell);
 
 				return ParamResult.op;
 			});
@@ -422,6 +450,7 @@ namespace mu
 	}
 
 
+	//---------------------------------------------------------------------------------------------
 	mu::Ptr<ASTOp> CodeGenerator::GenerateMissingScalarCode(const TCHAR* strWhere, float value, const void* errorContext)
 	{
 		// Log a warning
@@ -429,12 +458,12 @@ namespace mu
 		ErrorLog->Add(Msg, ELMT_ERROR, errorContext);
 
 		// Create a constant node
-		Ptr<NodeScalarConstant> Node = new NodeScalarConstant();
-		Node->Value = value;
+		NodeScalarConstantPtr pNode = new NodeScalarConstant();
+		pNode->SetValue(value);
 
 		FGenericGenerationOptions Options;
 		FScalarGenerationResult ParamResult;
-		GenerateScalar(ParamResult, Options, Node.get());
+		GenerateScalar(ParamResult, Options, pNode.get());
 
 		return ParamResult.op;
 	}
diff --git a/Engine/Plugins/Mutable/Source/MutableTools/Private/MuT/CodeGenerator_String.cpp b/Engine/Plugins/Mutable/Source/MutableTools/Private/MuT/CodeGenerator_String.cpp
--- a/Engine/Plugins/Mutable/Source/MutableTools/Private/MuT/CodeGenerator_String.cpp
+++ b/Engine/Plugins/Mutable/Source/MutableTools/Private/MuT/CodeGenerator_String.cpp
@@ -16,11 +16,15 @@
 #include "MuT/NodeRange.h"
 #include "MuT/NodeString.h"
 #include "MuT/NodeStringConstant.h"
+#include "MuT/NodeStringConstantPrivate.h"
 #include "MuT/NodeStringParameter.h"
+#include "MuT/NodeStringParameterPrivate.h"
 
 namespace mu
 {
 
+
+	//-------------------------------------------------------------------------------------------------
 	void CodeGenerator::GenerateString(FStringGenerationResult& result, const FGenericGenerationOptions& Options, const Ptr<const NodeString>& Untyped)
 	{
 		if (!Untyped)
@@ -65,38 +69,43 @@ namespace mu
 	}
 
 
-	void CodeGenerator::GenerateString_Constant(FStringGenerationResult& result, const FGenericGenerationOptions& Options, const Ptr<const NodeStringConstant>& Node)
+	//-------------------------------------------------------------------------------------------------
+	void CodeGenerator::GenerateString_Constant(FStringGenerationResult& result, const FGenericGenerationOptions& Options, const Ptr<const NodeStringConstant>& Typed)
 	{
+		const NodeStringConstant::Private& node = *Typed->GetPrivate();
+
 		Ptr<ASTOpConstantString> op = new ASTOpConstantString();
-		op->value = Node->Value;
+		op->value = node.Value;
 
 		result.op = op;
 	}
 
 
-	void CodeGenerator::GenerateString_Parameter(FStringGenerationResult& result, const FGenericGenerationOptions& Options, const Ptr<const NodeStringParameter>& Node)
+	//-------------------------------------------------------------------------------------------------
+	void CodeGenerator::GenerateString_Parameter(FStringGenerationResult& result, const FGenericGenerationOptions& Options, const Ptr<const NodeStringParameter>& Typed)
 	{
+		const NodeStringParameter::Private& node = *Typed->GetPrivate();
+
 		Ptr<ASTOpParameter> op;
 
 		Ptr<ASTOpParameter>* Found = nullptr;
 		{
 			UE::TUniqueLock Lock(FirstPass.ParameterNodes.Mutex);
-			Found = FirstPass.ParameterNodes.Map.Find(Node);
+			Found = FirstPass.ParameterNodes.Map.Find(node.m_pNode);
 			if (!Found)
 			{
 				FParameterDesc param;
-				param.Name = Node->Name;
-				bool bParseOk = FGuid::Parse(Node->UID, param.UID);
+				param.Name = node.Name;
+				bool bParseOk = FGuid::Parse(node.UID, param.UID);
 				check(bParseOk);
 				param.Type = EParameterType::Float;
-				param.DefaultValue.Set<FParamStringType>(Node->DefaultValue);
-				param.DefaultValue.Set<FParamStringType>(Node->DefaultValue);
+				param.DefaultValue.Set<FParamStringType>(node.DefaultValue);
 
 				op = new ASTOpParameter();
 				op->Type = EOpType::ST_PARAMETER;
 				op->Parameter = param;
 
-				FirstPass.ParameterNodes.Map.Add(Node, op);
+				FirstPass.ParameterNodes.Map.Add(node.m_pNode, op);
 			}
 			else
 			{
@@ -107,10 +116,10 @@ namespace mu
 		if (!Found)
 		{
 			// Generate the code for the ranges
-			for (int32 a = 0; a < Node->Ranges.Num(); ++a)
+			for (int32 a = 0; a < node.Ranges.Num(); ++a)
 			{
 				FRangeGenerationResult rangeResult;
-				GenerateRange(rangeResult, Options, Node->Ranges[a]);
+				GenerateRange(rangeResult, Options, node.Ranges[a]);
 				op->Ranges.Emplace(op.get(), rangeResult.sizeOp, rangeResult.rangeName, rangeResult.rangeUID);
 			}
 		}
diff --git a/Engine/Plugins/Mutable/Source/MutableTools/Private/MuT/CodeOptimiser_Semantic.cpp b/Engine/Plugins/Mutable/Source/MutableTools/Private/MuT/CodeOptimiser_Semantic.cpp
--- a/Engine/Plugins/Mutable/Source/MutableTools/Private/MuT/CodeOptimiser_Semantic.cpp
+++ b/Engine/Plugins/Mutable/Source/MutableTools/Private/MuT/CodeOptimiser_Semantic.cpp
@@ -1,6 +1,14 @@
 // Copyright Epic Games, Inc. All Rights Reserved.
 
-#include "MuT/NodeImageFormat.h"
+#include "MuT/NodeImageFormatPrivate.h"
+#include "HAL/PlatformMath.h"
+#include "Misc/AssertionMacros.h"
+#include "MuR/Image.h"
+#include "MuR/MutableMath.h"
+#include "MuR/MutableTrace.h"
+#include "MuR/Operations.h"
+#include "MuR/Ptr.h"
+#include "MuR/RefCounted.h"
 #include "MuT/AST.h"
 #include "MuT/ASTOpConditional.h"
 #include "MuT/ASTOpConstantBool.h"
@@ -20,15 +28,7 @@
 #include "MuT/ASTOpMeshAddTags.h"
 #include "MuT/ASTOpSwitch.h"
 #include "MuT/CodeOptimiser.h"
-#include "MuR/Image.h"
-#include "MuR/MutableMath.h"
-#include "MuR/MutableTrace.h"
-#include "MuR/Operations.h"
-#include "MuR/Ptr.h"
-#include "MuR/RefCounted.h"
 #include "MuT/Table.h"
-#include "HAL/PlatformMath.h"
-#include "Misc/AssertionMacros.h"
 
 
 namespace mu
diff --git a/Engine/Plugins/Mutable/Source/MutableTools/Private/MuT/CompilerPrivate.h b/Engine/Plugins/Mutable/Source/MutableTools/Private/MuT/CompilerPrivate.h
--- a/Engine/Plugins/Mutable/Source/MutableTools/Private/MuT/CompilerPrivate.h
+++ b/Engine/Plugins/Mutable/Source/MutableTools/Private/MuT/CompilerPrivate.h
@@ -6,6 +6,7 @@
 
 #include "MuT/AST.h"
 #include "MuT/ErrorLog.h"
+#include "MuT/NodeObjectPrivate.h"
 #include "MuR/Operations.h"
 #include "Templates/SharedPointer.h"
 
diff --git a/Engine/Plugins/Mutable/Source/MutableTools/Private/MuT/Node.cpp b/Engine/Plugins/Mutable/Source/MutableTools/Private/MuT/Node.cpp
--- a/Engine/Plugins/Mutable/Source/MutableTools/Private/MuT/Node.cpp
+++ b/Engine/Plugins/Mutable/Source/MutableTools/Private/MuT/Node.cpp
@@ -2,14 +2,11 @@
 
 
 #include "MuT/Node.h"
+#include "MuT/NodePrivate.h"
 #include "Misc/AssertionMacros.h"
 
-#include "MuT/NodeObject.h"
-#include "MuT/NodeObjectNew.h"
-#include "MuT/NodeObjectGroup.h"
 #include "MuT/NodeComponent.h"
 #include "MuT/NodeComponentNew.h"
-#include "MuT/NodeComponentEdit.h"
 #include "MuT/NodeComponentSwitch.h"
 #include "MuT/NodeComponentVariation.h"
 #include "MuT/NodeLOD.h"
@@ -19,30 +16,7 @@
 #include "MuT/NodeSurfaceVariation.h"
 #include "MuT/NodeExtensionData.h"
 #include "MuT/NodeExtensionDataConstant.h"
-#include "MuT/NodeImageConstant.h"
 #include "MuT/NodeImageTable.h"
-#include "MuT/NodeImageParameter.h"
-#include "MuT/NodeImageFormat.h"
-#include "MuT/NodeImageBinarise.h"
-#include "MuT/NodeImageConditional.h"
-#include "MuT/NodeImageInterpolate.h"
-#include "MuT/NodeImageInvert.h"
-#include "MuT/NodeImageLayer.h"
-#include "MuT/NodeImageLayerColour.h"
-#include "MuT/NodeImageLuminance.h"
-#include "MuT/NodeImageMipmap.h"
-#include "MuT/NodeImageMultiLayer.h"
-#include "MuT/NodeImageNormalComposite.h"
-#include "MuT/NodeImagePlainColour.h"
-#include "MuT/NodeImageProject.h"
-#include "MuT/NodeImageResize.h"
-#include "MuT/NodeImageSaturate.h"
-#include "MuT/NodeImageSwitch.h"
-#include "MuT/NodeImageSwizzle.h"
-#include "MuT/NodeImageTransform.h"
-#include "MuT/NodeImageVariation.h"
-#include "MuT/NodeImageColourMap.h"
-#include "MuT/NodeBool.h"
 #include "MuT/NodeColour.h"
 #include "MuT/NodeColourSwitch.h"
 #include "MuT/NodeColourVariation.h"
@@ -53,28 +27,10 @@
 #include "MuT/NodeColourFromScalars.h"
 #include "MuT/NodeColourTable.h"
 #include "MuT/NodeScalarSwitch.h"
-#include "MuT/NodeScalarConstant.h"
-#include "MuT/NodeScalarVariation.h"
-#include "MuT/NodeScalarArithmeticOperation.h"
-#include "MuT/NodeScalarParameter.h"
-#include "MuT/NodeScalarEnumParameter.h"
-#include "MuT/NodeScalarTable.h"
-#include "MuT/NodeScalarCurve.h"
-#include "MuT/NodeMeshConstant.h"
 #include "MuT/NodeMeshFragment.h"
 #include "MuT/NodeMeshClipMorphPlane.h"
 #include "MuT/NodeMeshClipDeform.h"
 #include "MuT/NodeMeshClipWithMesh.h"
-#include "MuT/NodeMeshParameter.h"
-#include "MuT/NodeMeshMakeMorph.h"
-#include "MuT/NodeMeshApplyPose.h"
-#include "MuT/NodeMeshTransform.h"
-#include "MuT/NodeMeshSwitch.h"
-#include "MuT/NodeMeshReshape.h"
-#include "MuT/NodeMeshMorph.h"
-#include "MuT/NodeMeshFormat.h"
-#include "MuT/NodeMeshVariation.h"
-#include "MuT/NodeMeshTable.h"
 #include "MuT/NodeModifier.h"
 #include "MuT/NodeModifierMeshClipDeform.h"
 #include "MuT/NodeModifierMeshClipMorphPlane.h"
@@ -85,13 +41,6 @@
 #include "MuT/NodeMatrix.h"
 #include "MuT/NodeMatrixConstant.h"
 #include "MuT/NodeMatrixParameter.h"
-#include "MuT/NodeString.h"
-#include "MuT/NodeStringConstant.h"
-#include "MuT/NodeStringParameter.h"
-#include "MuT/NodeBool.h"
-#include "MuT/NodeProjector.h"
-#include "MuT/NodeRange.h"
-#include "MuT/NodeRangeFromScalar.h"
 
 namespace mu
 {
@@ -99,31 +48,12 @@ namespace mu
 	// Static initialisation
 	FNodeType Node::StaticType = FNodeType(Node::EType::Node, nullptr);
 
-	FNodeType NodeObject::StaticType = FNodeType(Node::EType::Object, Node::GetStaticType());
-	FNodeType NodeObjectNew::StaticType = FNodeType(Node::EType::ObjectNew, NodeObject::GetStaticType());
-	FNodeType NodeObjectGroup::StaticType = FNodeType(Node::EType::ObjectGroup, NodeObject::GetStaticType());
-
 	FNodeType NodeComponent::StaticType = FNodeType(Node::EType::Component, Node::GetStaticType());
 	FNodeType NodeComponentNew::StaticType = FNodeType(Node::EType::ComponentNew, NodeComponent::GetStaticType());
 	FNodeType NodeComponentSwitch::StaticType = FNodeType(Node::EType::ComponentSwitch, NodeComponent::GetStaticType());
 	FNodeType NodeComponentVariation::StaticType = FNodeType(Node::EType::ComponentVariation, NodeComponent::GetStaticType());
-	FNodeType NodeComponentEdit::StaticType = FNodeType(Node::EType::ComponentEdit, NodeComponent::GetStaticType());
-
-	FNodeType NodeBool::StaticType = FNodeType(Node::EType::Bool, Node::GetStaticType());
-	FNodeType NodeBoolConstant::StaticType = FNodeType(Node::EType::BoolConstant, NodeBool::GetStaticType());
-	FNodeType NodeBoolParameter::StaticType = FNodeType(Node::EType::BoolParameter, NodeBool::GetStaticType());
-	FNodeType NodeBoolNot::StaticType = FNodeType(Node::EType::BoolNot, NodeBool::GetStaticType());
-	FNodeType NodeBoolAnd::StaticType = FNodeType(Node::EType::BoolAnd, NodeBool::GetStaticType());
 
-	FNodeType NodeScalar::StaticType = FNodeType(Node::EType::Scalar, Node::GetStaticType());
 	FNodeType NodeScalarSwitch::StaticType = FNodeType(Node::EType::ScalarSwitch, NodeScalar::GetStaticType());
-	FNodeType NodeScalarConstant::StaticType = FNodeType(Node::EType::ScalarConstant, NodeScalar::GetStaticType());
-	FNodeType NodeScalarParameter::StaticType = FNodeType(Node::EType::ScalarParameter, NodeScalar::GetStaticType());
-	FNodeType NodeScalarVariation::StaticType = FNodeType(Node::EType::ScalarVariation, NodeScalar::GetStaticType());
-	FNodeType NodeScalarArithmeticOperation::StaticType = FNodeType(Node::EType::ScalarArithmeticOperation, NodeScalar::GetStaticType());
-	FNodeType NodeScalarEnumParameter::StaticType = FNodeType(Node::EType::ScalarEnumParameter, NodeScalar::GetStaticType());
-	FNodeType NodeScalarTable::StaticType = FNodeType(Node::EType::ScalarTable, NodeScalar::GetStaticType());
-	FNodeType NodeScalarCurve::StaticType = FNodeType(Node::EType::ScalarCurve, NodeScalar::GetStaticType());
 
 	FNodeType NodeSurface::StaticType = FNodeType(Node::EType::Surface, Node::GetStaticType());
 	FNodeType NodeSurfaceNew::StaticType = FNodeType(Node::EType::SurfaceNew, NodeSurface::GetStaticType());
@@ -133,31 +63,7 @@ namespace mu
 	FNodeType NodeLOD::StaticType = FNodeType(Node::EType::LOD, Node::GetStaticType());
 	FNodeType NodeExtensionData::StaticType = FNodeType(Node::EType::ExtensionData, Node::GetStaticType());
 	FNodeType NodeExtensionDataConstant::StaticType = FNodeType(Node::EType::ExtensionDataConstant, NodeExtensionData::GetStaticType());
-
-	FNodeType NodeImage::StaticType = FNodeType(Node::EType::Image, Node::GetStaticType());
-	FNodeType NodeImageConstant::StaticType = FNodeType(Node::EType::ImageConstant, NodeImage::GetStaticType());
 	FNodeType NodeImageTable::StaticType = FNodeType(Node::EType::ImageTable, NodeImage::GetStaticType());
-	FNodeType NodeImageParameter::StaticType = FNodeType(Node::EType::ImageParameter, NodeImage::GetStaticType());
-	FNodeType NodeImageFormat::StaticType = FNodeType(Node::EType::ImageFormat, NodeImage::GetStaticType());
-	FNodeType NodeImageBinarise::StaticType = FNodeType(Node::EType::ImageBinarise, NodeImage::GetStaticType());
-	FNodeType NodeImageConditional::StaticType = FNodeType(Node::EType::ImageConditional, NodeImage::GetStaticType());
-	FNodeType NodeImageInterpolate::StaticType = FNodeType(Node::EType::ImageInterpolate, NodeImage::GetStaticType());
-	FNodeType NodeImageInvert::StaticType = FNodeType(Node::EType::ImageInvert, NodeImage::GetStaticType());
-	FNodeType NodeImageLayer::StaticType = FNodeType(Node::EType::ImageLayer, NodeImage::GetStaticType());
-	FNodeType NodeImageLayerColour::StaticType = FNodeType(Node::EType::ImageLayerColour, NodeImage::GetStaticType());
-	FNodeType NodeImageLuminance::StaticType = FNodeType(Node::EType::ImageLuminance, NodeImage::GetStaticType());
-	FNodeType NodeImageMipmap::StaticType = FNodeType(Node::EType::ImageMipmap, NodeImage::GetStaticType());
-	FNodeType NodeImageMultiLayer::StaticType = FNodeType(Node::EType::ImageMultiLayer, NodeImage::GetStaticType());
-	FNodeType NodeImageNormalComposite::StaticType = FNodeType(Node::EType::ImageNormalComposite, NodeImage::GetStaticType());
-	FNodeType NodeImagePlainColour::StaticType = FNodeType(Node::EType::ImagePlainColour, NodeImage::GetStaticType());
-	FNodeType NodeImageProject::StaticType = FNodeType(Node::EType::ImageProject, NodeImage::GetStaticType());
-	FNodeType NodeImageResize::StaticType = FNodeType(Node::EType::ImageResize, NodeImage::GetStaticType());
-	FNodeType NodeImageSaturate::StaticType = FNodeType(Node::EType::ImageSaturate, NodeImage::GetStaticType());
-	FNodeType NodeImageSwitch::StaticType = FNodeType(Node::EType::ImageSwitch, NodeImage::GetStaticType());
-	FNodeType NodeImageSwizzle::StaticType = FNodeType(Node::EType::ImageSwizzle, NodeImage::GetStaticType());
-	FNodeType NodeImageTransform::StaticType = FNodeType(Node::EType::ImageTransform, NodeImage::GetStaticType());
-	FNodeType NodeImageVariation::StaticType = FNodeType(Node::EType::ImageVariation, NodeImage::GetStaticType());
-	FNodeType NodeImageColourMap::StaticType = FNodeType(Node::EType::ImageColorMap, NodeImage::GetStaticType());
 
 	FNodeType NodeColour::StaticType = FNodeType(Node::EType::Color, Node::GetStaticType());
 	FNodeType NodeColourConstant::StaticType = FNodeType(Node::EType::ColorConstant, NodeColour::GetStaticType());
@@ -169,22 +75,10 @@ namespace mu
 	FNodeType NodeColourSampleImage::StaticType = FNodeType(Node::EType::ColorSampleImage, NodeColour::GetStaticType());
 	FNodeType NodeColourFromScalars::StaticType = FNodeType(Node::EType::ColorFromScalars, NodeColour::GetStaticType());
 
-	FNodeType NodeMesh::StaticType = FNodeType(Node::EType::Mesh, Node::GetStaticType());
-	FNodeType NodeMeshConstant::StaticType = FNodeType(Node::EType::MeshConstant, NodeMesh::GetStaticType());
 	FNodeType NodeMeshFragment::StaticType = FNodeType(Node::EType::MeshFragment, NodeMesh::GetStaticType());
 	FNodeType NodeMeshClipMorphPlane::StaticType = FNodeType(Node::EType::MeshClipMorphPlane, NodeMesh::GetStaticType());
 	FNodeType NodeMeshClipDeform::StaticType = FNodeType(Node::EType::MeshClipDeform, NodeMesh::GetStaticType());
 	FNodeType NodeMeshClipWithMesh::StaticType = FNodeType(Node::EType::MeshClipWithMesh, NodeMesh::GetStaticType());
-	FNodeType NodeMeshParameter::StaticType = FNodeType(Node::EType::MeshParameter, NodeMesh::GetStaticType());
-	FNodeType NodeMeshMakeMorph::StaticType = FNodeType(Node::EType::MeshMakeMorph, NodeMesh::GetStaticType());
-	FNodeType NodeMeshApplyPose::StaticType = FNodeType(Node::EType::MeshApplyPose, NodeMesh::GetStaticType());
-	FNodeType NodeMeshTransform::StaticType = FNodeType(Node::EType::MeshTransform, NodeMesh::GetStaticType());
-	FNodeType NodeMeshSwitch::StaticType = FNodeType(Node::EType::MeshSwitch, NodeMesh::GetStaticType());
-	FNodeType NodeMeshReshape::StaticType = FNodeType(Node::EType::MeshReshape, NodeMesh::GetStaticType());
-	FNodeType NodeMeshMorph::StaticType = FNodeType(Node::EType::MeshMorph, NodeMesh::GetStaticType());
-	FNodeType NodeMeshFormat::StaticType = FNodeType(Node::EType::MeshFormat, NodeMesh::GetStaticType());
-	FNodeType NodeMeshVariation::StaticType = FNodeType(Node::EType::MeshVariation, NodeMesh::GetStaticType());
-	FNodeType NodeMeshTable::StaticType = FNodeType(Node::EType::MeshTable, NodeMesh::GetStaticType());
 
 	FNodeType NodeModifier::StaticType = FNodeType(Node::EType::Modifier, Node::GetStaticType());
 	FNodeType NodeModifierMeshClipDeform::StaticType = FNodeType(Node::EType::ModifierMeshClipDeform, NodeModifier::GetStaticType());
@@ -198,18 +92,6 @@ namespace mu
 	FNodeType NodeMatrixConstant::StaticType = FNodeType(Node::EType::MatrixConstant, Node::GetStaticType());
 	FNodeType NodeMatrixParameter::StaticType = FNodeType(Node::EType::MatrixParameter, Node::GetStaticType());
 
-	FNodeType NodeString::StaticType = FNodeType(Node::EType::String, Node::GetStaticType());
-	FNodeType NodeStringConstant::StaticType = FNodeType(Node::EType::StringConstant, NodeString::GetStaticType());
-	FNodeType NodeStringParameter::StaticType = FNodeType(Node::EType::StringParameter, NodeString::GetStaticType());
-
-	FNodeType NodeProjector::StaticType = FNodeType(Node::EType::Projector, Node::GetStaticType());
-	FNodeType NodeProjectorConstant::StaticType = FNodeType(Node::EType::ProjectorConstant, NodeProjector::GetStaticType());
-	FNodeType NodeProjectorParameter::StaticType = FNodeType(Node::EType::ProjectorParameter, NodeProjector::GetStaticType());
-
-	FNodeType NodeRange::StaticType = FNodeType(Node::EType::Range, Node::GetStaticType());
-	FNodeType NodeRangeFromScalar::StaticType = FNodeType(Node::EType::RangeFromScalar, NodeRange::GetStaticType());
-
-
 	FNodeType::FNodeType()
 	{
 		Type = Node::EType::None;
diff --git a/Engine/Plugins/Mutable/Source/MutableTools/Private/MuT/NodeBool.cpp b/Engine/Plugins/Mutable/Source/MutableTools/Private/MuT/NodeBool.cpp
new file mode 100644
--- /dev/null
+++ b/Engine/Plugins/Mutable/Source/MutableTools/Private/MuT/NodeBool.cpp
@@ -0,0 +1,31 @@
+// Copyright Epic Games, Inc. All Rights Reserved.
+
+
+#include "MuT/NodeBool.h"
+
+#include "Misc/AssertionMacros.h"
+
+
+namespace mu
+{
+
+	static FNodeType s_nodeBoolType = FNodeType(Node::EType::Bool, Node::GetStaticType() );
+
+
+	//---------------------------------------------------------------------------------------------
+	const FNodeType* NodeBool::GetType() const
+	{
+		return GetStaticType();
+	}
+
+
+	//---------------------------------------------------------------------------------------------
+	const FNodeType* NodeBool::GetStaticType()
+	{
+		return &s_nodeBoolType;
+	}
+
+
+}
+
+
diff --git a/Engine/Plugins/Mutable/Source/MutableTools/Private/MuT/NodeBoolAnd.cpp b/Engine/Plugins/Mutable/Source/MutableTools/Private/MuT/NodeBoolAnd.cpp
new file mode 100644
--- /dev/null
+++ b/Engine/Plugins/Mutable/Source/MutableTools/Private/MuT/NodeBoolAnd.cpp
@@ -0,0 +1,58 @@
+// Copyright Epic Games, Inc. All Rights Reserved.
+
+
+#include "Misc/AssertionMacros.h"
+#include "MuT/Node.h"
+#include "MuT/NodeBool.h"
+#include "MuT/NodeBoolPrivate.h"
+#include "MuT/NodePrivate.h"
+
+
+namespace mu
+{
+
+	//---------------------------------------------------------------------------------------------
+	// Static initialisation
+	//---------------------------------------------------------------------------------------------
+	FNodeType NodeBoolAnd::Private::s_type = FNodeType(Node::EType::BoolParameter, NodeBool::GetStaticType() );
+
+
+	//---------------------------------------------------------------------------------------------
+	//!
+	//---------------------------------------------------------------------------------------------
+
+	MUTABLE_IMPLEMENT_NODE( NodeBoolAnd );
+
+
+	//---------------------------------------------------------------------------------------------
+	// Own Interface
+	//---------------------------------------------------------------------------------------------
+	Ptr<NodeBool> NodeBoolAnd::GetA() const
+	{
+		return m_pD->m_pA;
+	}
+
+
+	//---------------------------------------------------------------------------------------------
+	void NodeBoolAnd::SetA(Ptr<NodeBool> p )
+	{
+		m_pD->m_pA = p;
+	}
+
+
+	//---------------------------------------------------------------------------------------------
+	Ptr<NodeBool> NodeBoolAnd::GetB() const
+	{
+		return m_pD->m_pB;
+	}
+
+
+	//---------------------------------------------------------------------------------------------
+	void NodeBoolAnd::SetB(Ptr<NodeBool> p )
+	{
+		m_pD->m_pB = p;
+	}
+
+}
+
+
diff --git a/Engine/Plugins/Mutable/Source/MutableTools/Private/MuT/NodeBoolConstant.cpp b/Engine/Plugins/Mutable/Source/MutableTools/Private/MuT/NodeBoolConstant.cpp
new file mode 100644
--- /dev/null
+++ b/Engine/Plugins/Mutable/Source/MutableTools/Private/MuT/NodeBoolConstant.cpp
@@ -0,0 +1,43 @@
+// Copyright Epic Games, Inc. All Rights Reserved.
+
+
+#include "Misc/AssertionMacros.h"
+#include "MuT/Node.h"
+#include "MuT/NodeBool.h"
+#include "MuT/NodeBoolPrivate.h"
+#include "MuT/NodePrivate.h"
+
+
+
+namespace mu
+{
+
+	//---------------------------------------------------------------------------------------------
+	// Static initialisation
+	//---------------------------------------------------------------------------------------------
+	FNodeType NodeBoolConstant::Private::s_type = FNodeType(Node::EType::BoolConstant, NodeBool::GetStaticType() );
+
+
+	//---------------------------------------------------------------------------------------------
+	//!
+	//---------------------------------------------------------------------------------------------
+
+	MUTABLE_IMPLEMENT_NODE( NodeBoolConstant );
+
+
+	//---------------------------------------------------------------------------------------------
+	// Own Interface
+	//---------------------------------------------------------------------------------------------
+	bool NodeBoolConstant::GetValue() const
+	{
+		return m_pD->Value;
+	}
+
+	//---------------------------------------------------------------------------------------------
+	void NodeBoolConstant::SetValue( bool v )
+	{
+		m_pD->Value = v;
+	}
+
+
+}
diff --git a/Engine/Plugins/Mutable/Source/MutableTools/Private/MuT/NodeBoolNot.cpp b/Engine/Plugins/Mutable/Source/MutableTools/Private/MuT/NodeBoolNot.cpp
new file mode 100644
--- /dev/null
+++ b/Engine/Plugins/Mutable/Source/MutableTools/Private/MuT/NodeBoolNot.cpp
@@ -0,0 +1,44 @@
+// Copyright Epic Games, Inc. All Rights Reserved.
+
+
+#include "Misc/AssertionMacros.h"
+#include "MuT/Node.h"
+#include "MuT/NodeBool.h"
+#include "MuT/NodeBoolPrivate.h"
+#include "MuT/NodePrivate.h"
+
+
+namespace mu
+{
+
+	//---------------------------------------------------------------------------------------------
+	// Static initialisation
+	//---------------------------------------------------------------------------------------------
+	FNodeType NodeBoolNot::Private::s_type = FNodeType(Node::EType::BoolParameter, NodeBool::GetStaticType() );
+
+
+	//---------------------------------------------------------------------------------------------
+	//!
+	//---------------------------------------------------------------------------------------------
+
+	MUTABLE_IMPLEMENT_NODE( NodeBoolNot );
+
+
+	//---------------------------------------------------------------------------------------------
+	// Own Interface
+	//---------------------------------------------------------------------------------------------
+	Ptr<NodeBool> NodeBoolNot::GetInput() const
+	{
+		return m_pD->m_pSource;
+	}
+
+
+	//---------------------------------------------------------------------------------------------
+	void NodeBoolNot::SetInput(Ptr<NodeBool> p )
+	{
+		m_pD->m_pSource = p;
+	}
+
+}
+
+
diff --git a/Engine/Plugins/Mutable/Source/MutableTools/Private/MuT/NodeBoolParameter.cpp b/Engine/Plugins/Mutable/Source/MutableTools/Private/MuT/NodeBoolParameter.cpp
new file mode 100644
--- /dev/null
+++ b/Engine/Plugins/Mutable/Source/MutableTools/Private/MuT/NodeBoolParameter.cpp
@@ -0,0 +1,66 @@
+// Copyright Epic Games, Inc. All Rights Reserved.
+
+
+#include "Misc/AssertionMacros.h"
+#include "MuR/Ptr.h"
+#include "MuR/RefCounted.h"
+#include "MuT/Node.h"
+#include "MuT/NodeBool.h"
+#include "MuT/NodeBoolPrivate.h"
+#include "MuT/NodePrivate.h"
+#include "MuT/NodeRange.h"
+
+
+namespace mu
+{
+
+	//---------------------------------------------------------------------------------------------
+	// Static initialisation
+	//---------------------------------------------------------------------------------------------
+	FNodeType NodeBoolParameter::Private::s_type = FNodeType(Node::EType::BoolParameter, NodeBool::GetStaticType() );
+
+
+	//---------------------------------------------------------------------------------------------
+	//!
+	//---------------------------------------------------------------------------------------------
+
+	MUTABLE_IMPLEMENT_NODE( NodeBoolParameter );
+
+
+	//---------------------------------------------------------------------------------------------
+	// Own Interface
+	//---------------------------------------------------------------------------------------------
+	void NodeBoolParameter::SetName( const FString& strName )
+	{
+		m_pD->Name = strName;
+	}
+
+
+	//---------------------------------------------------------------------------------------------
+	void NodeBoolParameter::SetDefaultValue( bool v )
+	{
+		m_pD->DefaultValue = v;
+	}
+
+
+    //---------------------------------------------------------------------------------------------
+    void NodeBoolParameter::SetRangeCount( int i )
+    {
+        check(i>=0);
+        m_pD->Ranges.SetNum(i);
+    }
+
+
+    //---------------------------------------------------------------------------------------------
+    void NodeBoolParameter::SetRange( int i, NodeRangePtr pRange )
+    {
+        check( i>=0 && i<int(m_pD->Ranges.Num()) );
+        if ( i>=0 && i<int(m_pD->Ranges.Num()) )
+        {
+            m_pD->Ranges[i] = pRange;
+        }
+    }
+
+}
+
+
diff --git a/Engine/Plugins/Mutable/Source/MutableTools/Private/MuT/NodeBoolPrivate.h b/Engine/Plugins/Mutable/Source/MutableTools/Private/MuT/NodeBoolPrivate.h
new file mode 100644
--- /dev/null
+++ b/Engine/Plugins/Mutable/Source/MutableTools/Private/MuT/NodeBoolPrivate.h
@@ -0,0 +1,64 @@
+// Copyright Epic Games, Inc. All Rights Reserved.
+
+#pragma once
+
+#include "MuT/NodePrivate.h"
+
+#include "MuT/NodeBool.h"
+#include "MuT/NodeImage.h"
+#include "MuT/NodeRange.h"
+
+namespace mu
+{
+
+
+	class NodeBool::Private : public Node::Private
+	{
+	};
+
+
+	class NodeBoolConstant::Private : public NodeBool::Private
+	{
+	public:
+
+		static FNodeType s_type;
+
+		bool Value;
+	};
+
+
+	class NodeBoolParameter::Private : public NodeBool::Private
+	{
+	public:
+
+		static FNodeType s_type;
+
+		bool DefaultValue;
+		FString Name;
+		FString UID;
+
+        TArray<Ptr<NodeRange>> Ranges;
+	};
+
+
+	class NodeBoolNot::Private : public NodeBool::Private
+	{
+	public:
+
+		static FNodeType s_type;
+
+		Ptr<NodeBool> m_pSource;
+	};
+
+
+	class NodeBoolAnd::Private : public NodeBool::Private
+	{
+	public:
+
+		static FNodeType s_type;
+
+		Ptr<NodeBool> m_pA;
+		Ptr<NodeBool> m_pB;
+	};
+
+}
diff --git a/Engine/Plugins/Mutable/Source/MutableTools/Private/MuT/NodeComponentEdit.cpp b/Engine/Plugins/Mutable/Source/MutableTools/Private/MuT/NodeComponentEdit.cpp
--- a/Engine/Plugins/Mutable/Source/MutableTools/Private/MuT/NodeComponentEdit.cpp
+++ b/Engine/Plugins/Mutable/Source/MutableTools/Private/MuT/NodeComponentEdit.cpp
@@ -1,21 +1,48 @@
 // Copyright Epic Games, Inc. All Rights Reserved.
 
+
 #include "MuT/NodeComponentEdit.h"
 
+#include "Misc/AssertionMacros.h"
 #include "MuT/NodeComponentNew.h"
+#include "MuT/NodePrivate.h"
 #include "MuT/NodeSurface.h"
 
+
 namespace mu
 {
 
+	FNodeType NodeComponentEdit::StaticType = FNodeType(Node::EType::ComponentEdit, NodeComponent::GetStaticType() );
+
+
+	//---------------------------------------------------------------------------------------------
+	// Own Interface
+	//---------------------------------------------------------------------------------------------
+    void NodeComponentEdit::SetParent( NodeComponent* p )
+	{
+		m_pParent = p;
+	}
+
+
+	//---------------------------------------------------------------------------------------------
+    NodeComponent* NodeComponentEdit::GetParent() const
+	{
+        return m_pParent.get();
+	}
+
+
+	//---------------------------------------------------------------------------------------------
 	const NodeComponentNew* NodeComponentEdit::GetParentComponentNew() const
 	{
-		if (Parent)
+		const NodeComponentNew* Parent = nullptr;
+		if (m_pParent)
 		{
-			return Parent->GetParentComponentNew();
+			Parent = m_pParent->GetParentComponentNew();
 		}
 
-		return nullptr;
+		return Parent;
 	}
 
 }
+
+
diff --git a/Engine/Plugins/Mutable/Source/MutableTools/Private/MuT/NodeImage.cpp b/Engine/Plugins/Mutable/Source/MutableTools/Private/MuT/NodeImage.cpp
new file mode 100644
--- /dev/null
+++ b/Engine/Plugins/Mutable/Source/MutableTools/Private/MuT/NodeImage.cpp
@@ -0,0 +1,32 @@
+// Copyright Epic Games, Inc. All Rights Reserved.
+
+#include "MuT/NodeImage.h"
+
+#include "Misc/AssertionMacros.h"
+
+namespace mu
+{
+
+	//---------------------------------------------------------------------------------------------
+	// Static initialisation
+	//---------------------------------------------------------------------------------------------
+	static FNodeType s_nodeImageType = 	FNodeType( Node::EType::Image, Node::GetStaticType() );
+
+
+	//---------------------------------------------------------------------------------------------
+	const FNodeType* NodeImage::GetType() const
+	{
+		return GetStaticType();
+	}
+
+
+	//---------------------------------------------------------------------------------------------
+	const FNodeType* NodeImage::GetStaticType()
+	{
+		return &s_nodeImageType;
+	}
+
+
+}
+
+
diff --git a/Engine/Plugins/Mutable/Source/MutableTools/Private/MuT/NodeImageBinarise.cpp b/Engine/Plugins/Mutable/Source/MutableTools/Private/MuT/NodeImageBinarise.cpp
new file mode 100644
--- /dev/null
+++ b/Engine/Plugins/Mutable/Source/MutableTools/Private/MuT/NodeImageBinarise.cpp
@@ -0,0 +1,55 @@
+// Copyright Epic Games, Inc. All Rights Reserved.
+
+
+#include "MuT/NodeImageBinarise.h"
+
+#include "Misc/AssertionMacros.h"
+#include "MuT/NodeImageBinarisePrivate.h"
+#include "MuT/NodePrivate.h"
+#include "MuT/NodeScalar.h"
+
+
+namespace mu
+{
+
+	//---------------------------------------------------------------------------------------------
+	// Static initialisation
+	//---------------------------------------------------------------------------------------------
+	FNodeType NodeImageBinarise::Private::s_type = FNodeType(Node::EType::ImageBinarise, NodeImage::GetStaticType() );
+
+	//---------------------------------------------------------------------------------------------
+
+	MUTABLE_IMPLEMENT_NODE( NodeImageBinarise )
+
+
+	//---------------------------------------------------------------------------------------------
+	// Own Interface
+	//---------------------------------------------------------------------------------------------
+	NodeImagePtr NodeImageBinarise::GetBase() const
+	{
+		return m_pD->m_pBase;
+	}
+
+
+	//---------------------------------------------------------------------------------------------
+	void NodeImageBinarise::SetBase( NodeImagePtr pNode )
+	{
+		m_pD->m_pBase = pNode;
+	}
+
+
+	//---------------------------------------------------------------------------------------------
+	NodeScalarPtr NodeImageBinarise::GetThreshold() const
+	{
+		return m_pD->m_pThreshold;
+	}
+
+
+	//---------------------------------------------------------------------------------------------
+	void NodeImageBinarise::SetThreshold( NodeScalarPtr pNode )
+	{
+		m_pD->m_pThreshold = pNode;
+	}
+
+}
+
diff --git a/Engine/Plugins/Mutable/Source/MutableTools/Private/MuT/NodeImageBinarisePrivate.h b/Engine/Plugins/Mutable/Source/MutableTools/Private/MuT/NodeImageBinarisePrivate.h
new file mode 100644
--- /dev/null
+++ b/Engine/Plugins/Mutable/Source/MutableTools/Private/MuT/NodeImageBinarisePrivate.h
@@ -0,0 +1,25 @@
+// Copyright Epic Games, Inc. All Rights Reserved.
+
+#pragma once
+
+#include "MuT/NodeImagePrivate.h"
+#include "MuT/NodeImageBinarise.h"
+#include "MuT/NodeScalar.h"
+
+
+namespace mu
+{
+
+
+	class NodeImageBinarise::Private : public NodeImage::Private
+	{
+	public:
+
+		static FNodeType s_type;
+
+		NodeImagePtr m_pBase;
+		NodeScalarPtr m_pThreshold;
+	};
+
+
+}
diff --git a/Engine/Plugins/Mutable/Source/MutableTools/Private/MuT/NodeImageColourMap.cpp b/Engine/Plugins/Mutable/Source/MutableTools/Private/MuT/NodeImageColourMap.cpp
new file mode 100644
--- /dev/null
+++ b/Engine/Plugins/Mutable/Source/MutableTools/Private/MuT/NodeImageColourMap.cpp
@@ -0,0 +1,69 @@
+// Copyright Epic Games, Inc. All Rights Reserved.
+
+
+#include "MuT/NodeImageColourMap.h"
+
+#include "Misc/AssertionMacros.h"
+#include "MuT/NodeImageColourMapPrivate.h"
+#include "MuT/NodePrivate.h"
+
+
+namespace mu
+{
+
+	//---------------------------------------------------------------------------------------------
+	// Static initialisation
+	//---------------------------------------------------------------------------------------------
+	FNodeType NodeImageColourMap::Private::s_type = FNodeType(Node::EType::ImageColorMap, NodeImage::GetStaticType() );
+
+
+	//---------------------------------------------------------------------------------------------
+	//!
+	//---------------------------------------------------------------------------------------------
+
+	MUTABLE_IMPLEMENT_NODE( NodeImageColourMap )
+
+
+	//---------------------------------------------------------------------------------------------
+	// Own Interface
+	//---------------------------------------------------------------------------------------------
+	NodeImagePtr NodeImageColourMap::GetBase() const
+	{
+		return m_pD->m_pBase.get();
+	}
+
+
+	//---------------------------------------------------------------------------------------------
+	void NodeImageColourMap::SetBase( NodeImagePtr pNode )
+	{
+		m_pD->m_pBase = pNode;
+	}
+
+
+	//---------------------------------------------------------------------------------------------
+	NodeImagePtr NodeImageColourMap::GetMap() const
+	{
+		return m_pD->m_pMap.get();
+	}
+
+
+	//---------------------------------------------------------------------------------------------
+	void NodeImageColourMap::SetMap( NodeImagePtr pNode )
+	{
+		m_pD->m_pMap = pNode;
+	}
+
+
+	//---------------------------------------------------------------------------------------------
+	NodeImagePtr NodeImageColourMap::GetMask() const
+	{
+		return m_pD->m_pMask.get();
+	}
+
+	//---------------------------------------------------------------------------------------------
+	void NodeImageColourMap::SetMask( NodeImagePtr pNode )
+	{
+		m_pD->m_pMask = pNode;
+	}
+
+}
diff --git a/Engine/Plugins/Mutable/Source/MutableTools/Private/MuT/NodeImageColourMapPrivate.h b/Engine/Plugins/Mutable/Source/MutableTools/Private/MuT/NodeImageColourMapPrivate.h
new file mode 100644
--- /dev/null
+++ b/Engine/Plugins/Mutable/Source/MutableTools/Private/MuT/NodeImageColourMapPrivate.h
@@ -0,0 +1,25 @@
+// Copyright Epic Games, Inc. All Rights Reserved.
+
+#pragma once
+
+#include "MuT/NodeImagePrivate.h"
+#include "MuT/NodeImageColourMap.h"
+
+namespace mu
+{
+
+	class NodeImageColourMap::Private : public NodeImage::Private
+	{
+	public:
+
+		static FNodeType s_type;
+
+		NodeImagePtr m_pBase;
+		NodeImagePtr m_pMask;
+		NodeImagePtr m_pMap;
+
+	};
+
+
+}
+
diff --git a/Engine/Plugins/Mutable/Source/MutableTools/Private/MuT/NodeImageConditional.cpp b/Engine/Plugins/Mutable/Source/MutableTools/Private/MuT/NodeImageConditional.cpp
new file mode 100644
--- /dev/null
+++ b/Engine/Plugins/Mutable/Source/MutableTools/Private/MuT/NodeImageConditional.cpp
@@ -0,0 +1,75 @@
+// Copyright Epic Games, Inc. All Rights Reserved.
+
+
+#include "MuT/NodeImageConditional.h"
+
+#include "Misc/AssertionMacros.h"
+#include "MuT/NodeBool.h"
+#include "MuT/NodeImageConditionalPrivate.h"
+#include "MuT/NodePrivate.h"
+
+
+namespace mu
+{
+
+	//---------------------------------------------------------------------------------------------
+	// Static initialisation
+	//---------------------------------------------------------------------------------------------
+    FNodeType NodeImageConditional::Private::s_type = FNodeType(Node::EType::ImageConditional, NodeImage::GetStaticType() );
+
+
+	//---------------------------------------------------------------------------------------------
+	//!
+	//---------------------------------------------------------------------------------------------
+
+    MUTABLE_IMPLEMENT_NODE( NodeImageConditional )
+
+
+	//---------------------------------------------------------------------------------------------
+	// Own Interface
+	//---------------------------------------------------------------------------------------------
+    NodeBool* NodeImageConditional::GetParameter() const
+	{
+        return m_pD->Parameter.get();
+	}
+
+
+	//---------------------------------------------------------------------------------------------
+    void NodeImageConditional::SetParameter( NodeBool* pNode )
+	{
+        m_pD->Parameter = pNode;
+	}
+
+
+    //---------------------------------------------------------------------------------------------
+    NodeImage* NodeImageConditional::GetOptionTrue() const
+    {
+        return m_pD->m_true.get();
+    }
+
+
+    //---------------------------------------------------------------------------------------------
+    void NodeImageConditional::SetOptionTrue( NodeImage* pNode )
+    {
+        m_pD->m_true = pNode;
+    }
+
+
+    //---------------------------------------------------------------------------------------------
+    NodeImage* NodeImageConditional::GetOptionFalse() const
+    {
+        return m_pD->m_false.get();
+    }
+
+
+    //---------------------------------------------------------------------------------------------
+    void NodeImageConditional::SetOptionFalse( NodeImage* pNode )
+    {
+        m_pD->m_false = pNode;
+    }
+
+
+
+}
+
+
diff --git a/Engine/Plugins/Mutable/Source/MutableTools/Private/MuT/NodeImageConditionalPrivate.h b/Engine/Plugins/Mutable/Source/MutableTools/Private/MuT/NodeImageConditionalPrivate.h
new file mode 100644
--- /dev/null
+++ b/Engine/Plugins/Mutable/Source/MutableTools/Private/MuT/NodeImageConditionalPrivate.h
@@ -0,0 +1,28 @@
+// Copyright Epic Games, Inc. All Rights Reserved.
+
+
+#pragma once
+
+#include "MuT/NodeImageConditional.h"
+#include "MuT/NodeImagePrivate.h"
+#include "MuT/NodeBool.h"
+
+namespace mu
+{
+
+
+    class NodeImageConditional::Private : public NodeImage::Private
+	{
+	public:
+
+		static FNodeType s_type;
+
+        Ptr<NodeBool> Parameter;
+        NodeImagePtr m_true;
+        NodeImagePtr m_false;
+
+	};
+
+
+}
+
diff --git a/Engine/Plugins/Mutable/Source/MutableTools/Private/MuT/NodeImageConstant.cpp b/Engine/Plugins/Mutable/Source/MutableTools/Private/MuT/NodeImageConstant.cpp
--- a/Engine/Plugins/Mutable/Source/MutableTools/Private/MuT/NodeImageConstant.cpp
+++ b/Engine/Plugins/Mutable/Source/MutableTools/Private/MuT/NodeImageConstant.cpp
@@ -6,17 +6,26 @@
 #include "Misc/AssertionMacros.h"
 #include "MuR/Image.h"
 #include "MuR/Serialisation.h"
+#include "MuT/NodePrivate.h"
 
 
 namespace mu
 {
 
+	//---------------------------------------------------------------------------------------------
+	// Static initialisation
+	//---------------------------------------------------------------------------------------------
+	FNodeType NodeImageConstant::StaticType = FNodeType(Node::EType::ImageConstant, NodeImage::GetStaticType());
+
+
+	//---------------------------------------------------------------------------------------------
 	void NodeImageConstant::SetValue(TSharedPtr<const FImage> Value)
 	{
 		Proxy = new TResourceProxyMemory<FImage>(Value);
 	}
 
 
+	//---------------------------------------------------------------------------------------------
 	void NodeImageConstant::SetValue(Ptr<TResourceProxy<FImage>> InProxy)
 	{
 		Proxy = InProxy;
diff --git a/Engine/Plugins/Mutable/Source/MutableTools/Private/MuT/NodeImageFormat.cpp b/Engine/Plugins/Mutable/Source/MutableTools/Private/MuT/NodeImageFormat.cpp
new file mode 100644
--- /dev/null
+++ b/Engine/Plugins/Mutable/Source/MutableTools/Private/MuT/NodeImageFormat.cpp
@@ -0,0 +1,57 @@
+// Copyright Epic Games, Inc. All Rights Reserved.
+
+
+#include "MuT/NodeImageFormat.h"
+
+#include "Misc/AssertionMacros.h"
+#include "MuT/NodeImageFormatPrivate.h"
+#include "MuT/NodePrivate.h"
+
+
+namespace mu
+{
+
+	//---------------------------------------------------------------------------------------------
+	// Static initialisation
+	//---------------------------------------------------------------------------------------------
+	FNodeType NodeImageFormat::Private::s_type = FNodeType(Node::EType::ImageFormat, NodeImage::GetStaticType() );
+
+
+	//---------------------------------------------------------------------------------------------
+	//!
+	//---------------------------------------------------------------------------------------------
+
+	MUTABLE_IMPLEMENT_NODE( NodeImageFormat )
+
+
+	//---------------------------------------------------------------------------------------------
+	// Own Interface
+	//---------------------------------------------------------------------------------------------
+	EImageFormat NodeImageFormat::GetFormat() const
+	{
+		return m_pD->m_format;
+	}
+
+
+	//---------------------------------------------------------------------------------------------
+    void NodeImageFormat::SetFormat(EImageFormat format, EImageFormat formatIfAlpha )
+	{
+		m_pD->m_format = format;
+        m_pD->m_formatIfAlpha = formatIfAlpha;
+	}
+
+
+	//---------------------------------------------------------------------------------------------
+	NodeImagePtr NodeImageFormat::GetSource() const
+	{
+		return m_pD->m_source.get();
+	}
+
+
+	//---------------------------------------------------------------------------------------------
+	void NodeImageFormat::SetSource(NodeImagePtr pNode )
+	{
+		m_pD->m_source = pNode;
+	}
+
+}
diff --git a/Engine/Plugins/Mutable/Source/MutableTools/Private/MuT/NodeImageInterpolate.cpp b/Engine/Plugins/Mutable/Source/MutableTools/Private/MuT/NodeImageInterpolate.cpp
new file mode 100644
--- /dev/null
+++ b/Engine/Plugins/Mutable/Source/MutableTools/Private/MuT/NodeImageInterpolate.cpp
@@ -0,0 +1,77 @@
+// Copyright Epic Games, Inc. All Rights Reserved.
+
+
+#include "MuT/NodeImageInterpolate.h"
+
+#include "Misc/AssertionMacros.h"
+#include "MuT/NodeImageInterpolatePrivate.h"
+#include "MuT/NodePrivate.h"
+#include "MuT/NodeScalar.h"
+
+
+namespace mu
+{
+
+	//---------------------------------------------------------------------------------------------
+	// Static initialisation
+	//---------------------------------------------------------------------------------------------
+	FNodeType NodeImageInterpolate::Private::s_type = FNodeType(Node::EType::ImageInterpolate, NodeImage::GetStaticType() );
+
+
+	//---------------------------------------------------------------------------------------------
+	//!
+	//---------------------------------------------------------------------------------------------
+
+	MUTABLE_IMPLEMENT_NODE( NodeImageInterpolate )
+
+
+	//---------------------------------------------------------------------------------------------
+	// Own Interface
+	//---------------------------------------------------------------------------------------------
+	NodeScalarPtr NodeImageInterpolate::GetFactor() const
+	{
+		return m_pD->m_pFactor.get();
+	}
+
+
+	//---------------------------------------------------------------------------------------------
+	void NodeImageInterpolate::SetFactor( NodeScalarPtr pNode )
+	{
+		m_pD->m_pFactor = pNode;
+	}
+
+
+	//---------------------------------------------------------------------------------------------
+	NodeImagePtr NodeImageInterpolate::GetTarget( int t ) const
+	{
+		check( t>=0 && t<(int)m_pD->m_targets.Num() );
+		return m_pD->m_targets[t].get();
+	}
+
+
+	//---------------------------------------------------------------------------------------------
+	void NodeImageInterpolate::SetTarget( int t, NodeImagePtr pNode )
+	{
+		check( t>=0 && t<(int)m_pD->m_targets.Num() );
+		m_pD->m_targets[t] = pNode;
+	}
+
+
+	//---------------------------------------------------------------------------------------------
+	void NodeImageInterpolate::SetTargetCount( int t )
+	{
+		m_pD->m_targets.SetNum(t);
+	}
+
+
+	//---------------------------------------------------------------------------------------------
+	int NodeImageInterpolate::GetTargetCount() const
+	{
+		return int(m_pD->m_targets.Num());
+	}
+
+
+
+}
+
+
diff --git a/Engine/Plugins/Mutable/Source/MutableTools/Private/MuT/NodeImageInvert.cpp b/Engine/Plugins/Mutable/Source/MutableTools/Private/MuT/NodeImageInvert.cpp
new file mode 100644
--- /dev/null
+++ b/Engine/Plugins/Mutable/Source/MutableTools/Private/MuT/NodeImageInvert.cpp
@@ -0,0 +1,36 @@
+// Copyright Epic Games, Inc. All Rights Reserved.
+
+
+#include "MuT/NodeImageInvert.h"
+
+#include "Misc/AssertionMacros.h"
+#include "MuT/NodeImageInvertPrivate.h"
+#include "MuT/NodePrivate.h"
+
+
+namespace mu
+{
+
+	//---------------------------------------------------------------------------------------------
+	// Static initialisation
+	//---------------------------------------------------------------------------------------------
+	FNodeType NodeImageInvert::Private::s_type = FNodeType(Node::EType::ImageInvert, NodeImage::GetStaticType());
+
+	//---------------------------------------------------------------------------------------------
+
+	MUTABLE_IMPLEMENT_NODE( NodeImageInvert )
+
+	//---------------------------------------------------------------------------------------------
+	// Own Interface
+	//---------------------------------------------------------------------------------------------
+	NodeImagePtr NodeImageInvert::GetBase()const
+	{
+		return m_pD->m_pBase;
+	}
+
+	void NodeImageInvert::SetBase(NodeImagePtr pNode)
+	{
+		m_pD->m_pBase = pNode;
+	}
+	
+}
diff --git a/Engine/Plugins/Mutable/Source/MutableTools/Private/MuT/NodeImageLayer.cpp b/Engine/Plugins/Mutable/Source/MutableTools/Private/MuT/NodeImageLayer.cpp
new file mode 100644
--- /dev/null
+++ b/Engine/Plugins/Mutable/Source/MutableTools/Private/MuT/NodeImageLayer.cpp
@@ -0,0 +1,85 @@
+// Copyright Epic Games, Inc. All Rights Reserved.
+
+
+#include "MuT/NodeImageLayer.h"
+
+#include "Misc/AssertionMacros.h"
+#include "MuT/NodeImageLayerPrivate.h"
+#include "MuT/NodePrivate.h"
+
+
+namespace mu
+{
+
+	//---------------------------------------------------------------------------------------------
+	// Static initialisation
+	//---------------------------------------------------------------------------------------------
+	FNodeType NodeImageLayer::Private::s_type = FNodeType(Node::EType::ImageLayer, NodeImage::GetStaticType() );
+
+
+	//---------------------------------------------------------------------------------------------
+	//!
+	//---------------------------------------------------------------------------------------------
+
+	MUTABLE_IMPLEMENT_NODE( NodeImageLayer )
+
+
+	//---------------------------------------------------------------------------------------------
+	// Own Interface
+	//---------------------------------------------------------------------------------------------
+	NodeImagePtr NodeImageLayer::GetBase() const
+	{
+		return m_pD->m_pBase.get();
+	}
+
+
+	//---------------------------------------------------------------------------------------------
+	void NodeImageLayer::SetBase( NodeImagePtr pNode )
+	{
+		m_pD->m_pBase = pNode;
+	}
+
+
+	//---------------------------------------------------------------------------------------------
+	NodeImagePtr NodeImageLayer::GetMask() const
+	{
+		return m_pD->m_pMask.get();
+	}
+
+
+	//---------------------------------------------------------------------------------------------
+	void NodeImageLayer::SetMask( NodeImagePtr pNode )
+	{
+		m_pD->m_pMask = pNode;
+	}
+
+
+	//---------------------------------------------------------------------------------------------
+	NodeImagePtr NodeImageLayer::GetBlended() const
+	{
+		return m_pD->m_pBlended.get();
+	}
+
+
+	//---------------------------------------------------------------------------------------------
+	void NodeImageLayer::SetBlended( NodeImagePtr pNode )
+	{
+		m_pD->m_pBlended = pNode;
+	}
+
+
+	//---------------------------------------------------------------------------------------------
+	EBlendType NodeImageLayer::GetBlendType() const
+	{
+		return m_pD->Type;
+	}
+
+
+	//---------------------------------------------------------------------------------------------
+	void NodeImageLayer::SetBlendType(EBlendType t)
+	{
+		m_pD->Type = t;
+	}
+
+
+}
diff --git a/Engine/Plugins/Mutable/Source/MutableTools/Private/MuT/NodeImageLayerColour.cpp b/Engine/Plugins/Mutable/Source/MutableTools/Private/MuT/NodeImageLayerColour.cpp
new file mode 100644
--- /dev/null
+++ b/Engine/Plugins/Mutable/Source/MutableTools/Private/MuT/NodeImageLayerColour.cpp
@@ -0,0 +1,81 @@
+// Copyright Epic Games, Inc. All Rights Reserved.
+
+
+#include "MuT/NodeImageLayerColour.h"
+
+#include "Misc/AssertionMacros.h"
+#include "MuT/NodeColour.h"
+#include "MuT/NodeImageLayerColourPrivate.h"
+#include "MuT/NodePrivate.h"
+
+
+namespace mu
+{
+
+	//---------------------------------------------------------------------------------------------
+	// Static initialisation
+	//---------------------------------------------------------------------------------------------
+	FNodeType NodeImageLayerColour::Private::s_type = FNodeType(Node::EType::ImageLayerColour, NodeImage::GetStaticType() );
+
+
+	//---------------------------------------------------------------------------------------------
+	//!
+	//---------------------------------------------------------------------------------------------
+
+	MUTABLE_IMPLEMENT_NODE( NodeImageLayerColour );
+
+
+	//---------------------------------------------------------------------------------------------
+	// Own Interface
+	//---------------------------------------------------------------------------------------------
+	NodeImagePtr NodeImageLayerColour::GetBase() const
+	{
+		return m_pD->m_pBase.get();
+	}
+
+	//---------------------------------------------------------------------------------------------
+	void NodeImageLayerColour::SetBase( NodeImagePtr pNode )
+	{
+		m_pD->m_pBase = pNode;
+	}
+
+	//---------------------------------------------------------------------------------------------
+	NodeImagePtr NodeImageLayerColour::GetMask() const
+	{
+		return m_pD->m_pMask.get();
+	}
+
+	//---------------------------------------------------------------------------------------------
+	void NodeImageLayerColour::SetMask( NodeImagePtr pNode )
+	{
+		m_pD->m_pMask = pNode;
+	}
+
+	//---------------------------------------------------------------------------------------------
+	NodeColourPtr NodeImageLayerColour::GetColour() const
+	{
+		return m_pD->m_pColour.get();
+	}
+
+	//---------------------------------------------------------------------------------------------
+	void NodeImageLayerColour::SetColour( NodeColourPtr pNode )
+	{
+		m_pD->m_pColour = pNode;
+	}
+
+
+	//---------------------------------------------------------------------------------------------
+	EBlendType NodeImageLayerColour::GetBlendType() const
+	{
+		return m_pD->Type;
+	}
+
+
+	//---------------------------------------------------------------------------------------------
+	void NodeImageLayerColour::SetBlendType(EBlendType t)
+	{
+		m_pD->Type = t;
+	}
+
+}
+
diff --git a/Engine/Plugins/Mutable/Source/MutableTools/Private/MuT/NodeImageLuminance.cpp b/Engine/Plugins/Mutable/Source/MutableTools/Private/MuT/NodeImageLuminance.cpp
new file mode 100644
--- /dev/null
+++ b/Engine/Plugins/Mutable/Source/MutableTools/Private/MuT/NodeImageLuminance.cpp
@@ -0,0 +1,44 @@
+// Copyright Epic Games, Inc. All Rights Reserved.
+
+
+#include "MuT/NodeImageLuminance.h"
+
+#include "Misc/AssertionMacros.h"
+#include "MuT/NodeImageLuminancePrivate.h"
+#include "MuT/NodePrivate.h"
+
+
+namespace mu
+{
+
+	//---------------------------------------------------------------------------------------------
+	// Static initialisation
+	//---------------------------------------------------------------------------------------------
+	FNodeType NodeImageLuminance::Private::s_type = FNodeType(Node::EType::ImageLuminance, NodeImage::GetStaticType() );
+
+
+	//---------------------------------------------------------------------------------------------
+	//!
+	//---------------------------------------------------------------------------------------------
+
+	MUTABLE_IMPLEMENT_NODE( NodeImageLuminance )
+
+
+	//---------------------------------------------------------------------------------------------
+	// Own Interface
+	//---------------------------------------------------------------------------------------------
+	NodeImagePtr NodeImageLuminance::GetSource() const
+	{
+		return m_pD->m_pSource.get();
+	}
+
+
+	//---------------------------------------------------------------------------------------------
+	void NodeImageLuminance::SetSource( NodeImagePtr pNode )
+	{
+		m_pD->m_pSource = pNode;
+	}
+
+
+
+}
diff --git a/Engine/Plugins/Mutable/Source/MutableTools/Private/MuT/NodeImageLuminancePrivate.h b/Engine/Plugins/Mutable/Source/MutableTools/Private/MuT/NodeImageLuminancePrivate.h
new file mode 100644
--- /dev/null
+++ b/Engine/Plugins/Mutable/Source/MutableTools/Private/MuT/NodeImageLuminancePrivate.h
@@ -0,0 +1,22 @@
+// Copyright Epic Games, Inc. All Rights Reserved.
+
+#pragma once
+
+#include "MuT/NodeImagePrivate.h"
+#include "MuT/NodeImageLuminance.h"
+
+
+namespace mu
+{
+
+	class NodeImageLuminance::Private : public NodeImage::Private
+	{
+	public:
+
+		static FNodeType s_type;
+
+		NodeImagePtr m_pSource;
+	};
+
+
+}
diff --git a/Engine/Plugins/Mutable/Source/MutableTools/Private/MuT/NodeImageMipmap.cpp b/Engine/Plugins/Mutable/Source/MutableTools/Private/MuT/NodeImageMipmap.cpp
new file mode 100644
--- /dev/null
+++ b/Engine/Plugins/Mutable/Source/MutableTools/Private/MuT/NodeImageMipmap.cpp
@@ -0,0 +1,49 @@
+// Copyright Epic Games, Inc. All Rights Reserved.
+
+
+#include "MuT/NodeImageMipmap.h"
+
+#include "Misc/AssertionMacros.h"
+#include "MuR/ImagePrivate.h"
+#include "MuT/NodeImageMipmapPrivate.h"
+#include "MuT/NodePrivate.h"
+
+
+namespace mu
+{
+
+	//---------------------------------------------------------------------------------------------
+	// Static initialisation
+	//---------------------------------------------------------------------------------------------
+    FNodeType NodeImageMipmap::Private::s_type = FNodeType(Node::EType::ImageMipmap, NodeImage::GetStaticType() );
+
+
+	//---------------------------------------------------------------------------------------------
+	//!
+	//---------------------------------------------------------------------------------------------
+
+    MUTABLE_IMPLEMENT_NODE( NodeImageMipmap )
+
+
+	//---------------------------------------------------------------------------------------------
+	// Own Interface
+	//---------------------------------------------------------------------------------------------
+    NodeImagePtr NodeImageMipmap::GetSource() const
+	{
+		return m_pD->m_pSource.get();
+	}
+
+
+	//---------------------------------------------------------------------------------------------
+    void NodeImageMipmap::SetSource( NodeImagePtr pNode )
+	{
+		m_pD->m_pSource = pNode;
+	}
+
+	//---------------------------------------------------------------------------------------------
+	void NodeImageMipmap::SetMipmapGenerationSettings(EMipmapFilterType FilterType, EAddressMode AddressMode)
+	{
+		m_pD->m_settings = FMipmapGenerationSettings{ FilterType, AddressMode };
+	}
+
+}
diff --git a/Engine/Plugins/Mutable/Source/MutableTools/Private/MuT/NodeImageMultiLayer.cpp b/Engine/Plugins/Mutable/Source/MutableTools/Private/MuT/NodeImageMultiLayer.cpp
new file mode 100644
--- /dev/null
+++ b/Engine/Plugins/Mutable/Source/MutableTools/Private/MuT/NodeImageMultiLayer.cpp
@@ -0,0 +1,99 @@
+// Copyright Epic Games, Inc. All Rights Reserved.
+
+
+#include "MuT/NodeImageMultiLayer.h"
+
+#include "Misc/AssertionMacros.h"
+#include "MuT/NodeImageMultiLayerPrivate.h"
+#include "MuT/NodePrivate.h"
+
+
+namespace mu
+{
+
+	//---------------------------------------------------------------------------------------------
+	// Static initialisation
+	//---------------------------------------------------------------------------------------------
+    FNodeType NodeImageMultiLayer::Private::s_type = FNodeType(Node::EType::ImageMultiLayer, NodeImage::GetStaticType() );
+
+
+	//---------------------------------------------------------------------------------------------
+	//!
+	//---------------------------------------------------------------------------------------------
+
+    MUTABLE_IMPLEMENT_NODE( NodeImageMultiLayer )
+
+
+	//---------------------------------------------------------------------------------------------
+	// Own Interface
+	//---------------------------------------------------------------------------------------------
+    NodeImagePtr NodeImageMultiLayer::GetBase() const
+	{
+		return m_pD->m_pBase.get();
+	}
+
+
+	//---------------------------------------------------------------------------------------------
+    void NodeImageMultiLayer::SetBase( NodeImagePtr pNode )
+	{
+		m_pD->m_pBase = pNode;
+	}
+
+
+	//---------------------------------------------------------------------------------------------
+    NodeImagePtr NodeImageMultiLayer::GetMask() const
+	{
+		return m_pD->m_pMask.get();
+	}
+
+
+	//---------------------------------------------------------------------------------------------
+    void NodeImageMultiLayer::SetMask( NodeImagePtr pNode )
+	{
+		m_pD->m_pMask = pNode;
+	}
+
+
+    //---------------------------------------------------------------------------------------------
+    NodeImagePtr NodeImageMultiLayer::GetBlended() const
+    {
+        return m_pD->m_pBlended.get();
+    }
+
+
+    //---------------------------------------------------------------------------------------------
+    void NodeImageMultiLayer::SetBlended( NodeImagePtr pNode )
+    {
+        m_pD->m_pBlended = pNode;
+    }
+
+
+    //---------------------------------------------------------------------------------------------
+    NodeRangePtr NodeImageMultiLayer::GetRange() const
+    {
+        return m_pD->m_pRange.get();
+    }
+
+
+    //---------------------------------------------------------------------------------------------
+    void NodeImageMultiLayer::SetRange( NodeRangePtr pNode )
+    {
+        m_pD->m_pRange = pNode;
+    }
+
+
+	//---------------------------------------------------------------------------------------------
+	EBlendType NodeImageMultiLayer::GetBlendType() const
+	{
+		return m_pD->Type;
+	}
+
+
+	//---------------------------------------------------------------------------------------------
+    void NodeImageMultiLayer::SetBlendType(EBlendType t)
+	{
+		m_pD->Type = t;
+	}
+
+
+}
diff --git a/Engine/Plugins/Mutable/Source/MutableTools/Private/MuT/NodeImageNormalComposite.cpp b/Engine/Plugins/Mutable/Source/MutableTools/Private/MuT/NodeImageNormalComposite.cpp
new file mode 100644
--- /dev/null
+++ b/Engine/Plugins/Mutable/Source/MutableTools/Private/MuT/NodeImageNormalComposite.cpp
@@ -0,0 +1,79 @@
+// Copyright Epic Games, Inc. All Rights Reserved.
+
+
+#include "MuT/NodeImageNormalComposite.h"
+
+#include "Misc/AssertionMacros.h"
+#include "MuT/NodeImageNormalCompositePrivate.h"
+#include "MuT/NodePrivate.h"
+#include "MuR/Image.h"
+
+
+namespace mu
+{
+
+	//---------------------------------------------------------------------------------------------
+	// Static initialisation
+	//---------------------------------------------------------------------------------------------
+	FNodeType NodeImageNormalComposite::Private::s_type = FNodeType(Node::EType::ImageNormalComposite, NodeImage::GetStaticType() );
+
+
+	//---------------------------------------------------------------------------------------------
+	//!
+	//---------------------------------------------------------------------------------------------
+
+	MUTABLE_IMPLEMENT_NODE( NodeImageNormalComposite )
+
+
+	//---------------------------------------------------------------------------------------------
+	// Own Interface
+	//---------------------------------------------------------------------------------------------
+	NodeImagePtr NodeImageNormalComposite::GetBase() const
+	{
+		return m_pD->m_pBase.get();
+	}
+
+	//---------------------------------------------------------------------------------------------
+	void NodeImageNormalComposite::SetBase( NodeImagePtr pNode )
+	{
+		m_pD->m_pBase = pNode;
+	}
+
+	//---------------------------------------------------------------------------------------------
+	NodeImagePtr NodeImageNormalComposite::GetNormal() const
+	{
+		return m_pD->m_pNormal.get();
+	}
+
+	//---------------------------------------------------------------------------------------------
+	void NodeImageNormalComposite::SetNormal( NodeImagePtr pNode )
+	{
+		m_pD->m_pNormal = pNode;
+	}
+
+	//---------------------------------------------------------------------------------------------
+	float NodeImageNormalComposite::GetPower() const
+	{
+		return m_pD->m_power;
+	}
+
+	
+	//---------------------------------------------------------------------------------------------
+	void NodeImageNormalComposite::SetPower( float power )
+	{
+		m_pD->m_power = power;
+	}
+
+
+	//---------------------------------------------------------------------------------------------
+	ECompositeImageMode NodeImageNormalComposite::GetMode() const
+	{
+		return m_pD->m_mode;
+	}
+
+	//---------------------------------------------------------------------------------------------
+	void NodeImageNormalComposite::SetMode( ECompositeImageMode mode )
+	{
+		m_pD->m_mode = mode;
+	}
+}
diff --git a/Engine/Plugins/Mutable/Source/MutableTools/Private/MuT/NodeImageNormalCompositePrivate.h b/Engine/Plugins/Mutable/Source/MutableTools/Private/MuT/NodeImageNormalCompositePrivate.h
new file mode 100644
--- /dev/null
+++ b/Engine/Plugins/Mutable/Source/MutableTools/Private/MuT/NodeImageNormalCompositePrivate.h
@@ -0,0 +1,25 @@
+// Copyright Epic Games, Inc. All Rights Reserved.
+
+#pragma once
+
+#include "MuT/NodeImagePrivate.h"
+#include "MuT/NodeImageNormalComposite.h"
+
+
+namespace mu
+{
+
+	class NodeImageNormalComposite::Private : public NodeImage::Private
+	{
+	public:
+
+		static FNodeType s_type;
+
+		NodeImagePtr m_pBase;
+		NodeImagePtr m_pNormal;
+		float m_power;
+		ECompositeImageMode m_mode;
+
+	};
+
+}
diff --git a/Engine/Plugins/Mutable/Source/MutableTools/Private/MuT/NodeImageParameter.cpp b/Engine/Plugins/Mutable/Source/MutableTools/Private/MuT/NodeImageParameter.cpp
new file mode 100644
--- /dev/null
+++ b/Engine/Plugins/Mutable/Source/MutableTools/Private/MuT/NodeImageParameter.cpp
@@ -0,0 +1,70 @@
+// Copyright Epic Games, Inc. All Rights Reserved.
+
+
+#include "MuT/NodeImageParameter.h"
+
+#include "Misc/AssertionMacros.h"
+#include "MuT/NodeImageParameterPrivate.h"
+#include "MuT/NodePrivate.h"
+
+
+namespace mu
+{
+
+	//---------------------------------------------------------------------------------------------
+	// Static initialisation
+	//---------------------------------------------------------------------------------------------
+    FNodeType NodeImageParameter::Private::s_type = FNodeType(Node::EType::ImageParameter, NodeImage::GetStaticType() );
+
+
+	//---------------------------------------------------------------------------------------------
+	//!
+	//---------------------------------------------------------------------------------------------
+
+    MUTABLE_IMPLEMENT_NODE( NodeImageParameter )
+
+
+	//---------------------------------------------------------------------------------------------
+	// Own Interface
+	//---------------------------------------------------------------------------------------------
+    void NodeImageParameter::SetName( const FString& Name )
+	{
+		m_pD->Name = Name;
+	}
+
+
+	//---------------------------------------------------------------------------------------------
+	void NodeImageParameter::SetUid(const FString& Uid)
+	{
+		m_pD->UID = Uid;
+	}
+
+	
+	//---------------------------------------------------------------------------------------------
+	void NodeImageParameter::SetDefaultValue(FName Value)
+	{
+    	m_pD->DefaultValue = Value;
+	}
+
+
+	//---------------------------------------------------------------------------------------------
+	void NodeImageParameter::SetRangeCount(int i)
+	{
+		check(i >= 0);
+		m_pD->Ranges.SetNum(i);
+	}
+
+
+	//---------------------------------------------------------------------------------------------
+	void NodeImageParameter::SetRange(int i, NodeRangePtr pRange)
+	{
+		check(i >= 0 && i<int(m_pD->Ranges.Num()));
+		if (i >= 0 && i<int(m_pD->Ranges.Num()))
+		{
+			m_pD->Ranges[i] = pRange;
+		}
+	}
+
+}
+
+
diff --git a/Engine/Plugins/Mutable/Source/MutableTools/Private/MuT/NodeImageParameterPrivate.h b/Engine/Plugins/Mutable/Source/MutableTools/Private/MuT/NodeImageParameterPrivate.h
new file mode 100644
--- /dev/null
+++ b/Engine/Plugins/Mutable/Source/MutableTools/Private/MuT/NodeImageParameterPrivate.h
@@ -0,0 +1,26 @@
+// Copyright Epic Games, Inc. All Rights Reserved.
+
+#pragma once
+
+#include "MuT/NodeImageParameter.h"
+#include "MuT/NodeImagePrivate.h"
+#include "MuT/NodeRange.h"
+
+
+namespace mu
+{
+
+    class NodeImageParameter::Private : public NodeImage::Private
+	{
+	public:
+
+		static FNodeType s_type;
+
+    	FName DefaultValue;
+		FString Name;
+		FString UID;
+
+		TArray<Ptr<NodeRange>> Ranges;
+	};
+
+}
diff --git a/Engine/Plugins/Mutable/Source/MutableTools/Private/MuT/NodeImagePlainColour.cpp b/Engine/Plugins/Mutable/Source/MutableTools/Private/MuT/NodeImagePlainColour.cpp
new file mode 100644
--- /dev/null
+++ b/Engine/Plugins/Mutable/Source/MutableTools/Private/MuT/NodeImagePlainColour.cpp
@@ -0,0 +1,50 @@
+// Copyright Epic Games, Inc. All Rights Reserved.
+
+
+#include "MuT/NodeImagePlainColour.h"
+
+#include "Misc/AssertionMacros.h"
+#include "MuT/NodeColour.h"
+#include "MuT/NodeImagePlainColourPrivate.h"
+#include "MuT/NodePrivate.h"
+
+
+namespace mu
+{
+
+	//---------------------------------------------------------------------------------------------
+	// Static initialisation
+	//---------------------------------------------------------------------------------------------
+	FNodeType NodeImagePlainColour::Private::s_type = FNodeType(Node::EType::ImagePlainColour, NodeImage::GetStaticType() );
+
+
+	//---------------------------------------------------------------------------------------------
+
+	MUTABLE_IMPLEMENT_NODE( NodeImagePlainColour )
+
+
+	//---------------------------------------------------------------------------------------------
+	// Own Interface
+	//---------------------------------------------------------------------------------------------
+	NodeColourPtr NodeImagePlainColour::GetColour() const
+	{
+		return m_pD->m_pColour;
+	}
+
+
+	//---------------------------------------------------------------------------------------------
+	void NodeImagePlainColour::SetColour( NodeColourPtr pNode )
+	{
+		m_pD->m_pColour = pNode;
+	}
+
+
+	//---------------------------------------------------------------------------------------------
+	void NodeImagePlainColour::SetSize( int x, int y )
+	{
+		m_pD->m_sizeX = x;
+		m_pD->m_sizeY = y;
+	}
+
+}
+
diff --git a/Engine/Plugins/Mutable/Source/MutableTools/Private/MuT/NodeImageResize.cpp b/Engine/Plugins/Mutable/Source/MutableTools/Private/MuT/NodeImageResize.cpp
new file mode 100644
--- /dev/null
+++ b/Engine/Plugins/Mutable/Source/MutableTools/Private/MuT/NodeImageResize.cpp
@@ -0,0 +1,54 @@
+// Copyright Epic Games, Inc. All Rights Reserved.
+
+
+#include "MuT/NodeImageResize.h"
+
+#include "Misc/AssertionMacros.h"
+#include "MuT/NodeImageResizePrivate.h"
+#include "MuT/NodePrivate.h"
+
+
+namespace mu
+{
+
+	//---------------------------------------------------------------------------------------------
+	// Static initialisation
+	//---------------------------------------------------------------------------------------------
+	FNodeType NodeImageResize::Private::s_type = FNodeType(Node::EType::ImageResize, NodeImage::GetStaticType() );
+
+	//---------------------------------------------------------------------------------------------
+
+	MUTABLE_IMPLEMENT_NODE( NodeImageResize )
+
+
+	//---------------------------------------------------------------------------------------------
+	// Own Interface
+	//---------------------------------------------------------------------------------------------
+	NodeImagePtr NodeImageResize::GetBase() const
+	{
+		return m_pD->m_pBase;
+	}
+
+
+	//---------------------------------------------------------------------------------------------
+	void NodeImageResize::SetBase( NodeImagePtr pNode )
+	{
+		m_pD->m_pBase = pNode;
+	}
+
+
+	//---------------------------------------------------------------------------------------------
+	void NodeImageResize::SetRelative( bool rel )
+	{
+		m_pD->m_relative = rel;
+	}
+
+
+	//---------------------------------------------------------------------------------------------
+	void NodeImageResize::SetSize( float x, float y )
+	{
+		m_pD->m_sizeX = x;
+		m_pD->m_sizeY = y;
+	}
+
+}
diff --git a/Engine/Plugins/Mutable/Source/MutableTools/Private/MuT/NodeImageSaturate.cpp b/Engine/Plugins/Mutable/Source/MutableTools/Private/MuT/NodeImageSaturate.cpp
new file mode 100644
--- /dev/null
+++ b/Engine/Plugins/Mutable/Source/MutableTools/Private/MuT/NodeImageSaturate.cpp
@@ -0,0 +1,60 @@
+// Copyright Epic Games, Inc. All Rights Reserved.
+
+
+#include "MuT/NodeImageSaturate.h"
+
+#include "Misc/AssertionMacros.h"
+#include "MuT/NodeImageSaturatePrivate.h"
+#include "MuT/NodePrivate.h"
+#include "MuT/NodeScalar.h"
+
+
+namespace mu
+{
+
+	//---------------------------------------------------------------------------------------------
+	// Static initialisation
+	//---------------------------------------------------------------------------------------------
+	FNodeType NodeImageSaturate::Private::s_type = FNodeType(Node::EType::ImageSaturate, NodeImage::GetStaticType() );
+
+
+	//---------------------------------------------------------------------------------------------
+	//!
+	//---------------------------------------------------------------------------------------------
+
+	MUTABLE_IMPLEMENT_NODE( NodeImageSaturate )
+
+
+	//---------------------------------------------------------------------------------------------
+	// Own Interface
+	//---------------------------------------------------------------------------------------------
+	NodeScalarPtr NodeImageSaturate::GetFactor() const
+	{
+		return m_pD->m_pFactor.get();
+	}
+
+
+	//---------------------------------------------------------------------------------------------
+	void NodeImageSaturate::SetFactor( NodeScalarPtr pNode )
+	{
+		m_pD->m_pFactor = pNode;
+	}
+
+
+	//---------------------------------------------------------------------------------------------
+	NodeImagePtr NodeImageSaturate::GetSource() const
+	{
+		return m_pD->m_pSource.get();
+	}
+
+
+	//---------------------------------------------------------------------------------------------
+	void NodeImageSaturate::SetSource( NodeImagePtr pNode )
+	{
+		m_pD->m_pSource = pNode;
+	}
+
+
+
+}
+
diff --git a/Engine/Plugins/Mutable/Source/MutableTools/Private/MuT/NodeImageSaturatePrivate.h b/Engine/Plugins/Mutable/Source/MutableTools/Private/MuT/NodeImageSaturatePrivate.h
new file mode 100644
--- /dev/null
+++ b/Engine/Plugins/Mutable/Source/MutableTools/Private/MuT/NodeImageSaturatePrivate.h
@@ -0,0 +1,23 @@
+// Copyright Epic Games, Inc. All Rights Reserved.
+
+#pragma once
+
+#include "MuT/NodeImagePrivate.h"
+#include "MuT/NodeImageSaturate.h"
+#include "MuT/NodeScalar.h"
+
+namespace mu
+{
+
+	class NodeImageSaturate::Private : public NodeImage::Private
+	{
+	public:
+
+		static FNodeType s_type;
+
+		NodeImagePtr m_pSource;
+		NodeScalarPtr m_pFactor;
+	};
+
+
+}
diff --git a/Engine/Plugins/Mutable/Source/MutableTools/Private/MuT/NodeImageSwitch.cpp b/Engine/Plugins/Mutable/Source/MutableTools/Private/MuT/NodeImageSwitch.cpp
new file mode 100644
--- /dev/null
+++ b/Engine/Plugins/Mutable/Source/MutableTools/Private/MuT/NodeImageSwitch.cpp
@@ -0,0 +1,77 @@
+// Copyright Epic Games, Inc. All Rights Reserved.
+
+
+#include "MuT/NodeImageSwitch.h"
+
+#include "Misc/AssertionMacros.h"
+#include "MuT/NodeImageSwitchPrivate.h"
+#include "MuT/NodePrivate.h"
+#include "MuT/NodeScalar.h"
+
+
+namespace mu
+{
+
+	//---------------------------------------------------------------------------------------------
+	// Static initialisation
+	//---------------------------------------------------------------------------------------------
+	FNodeType NodeImageSwitch::Private::s_type = FNodeType(Node::EType::ImageSwitch, NodeImage::GetStaticType() );
+
+
+	//---------------------------------------------------------------------------------------------
+	//!
+	//---------------------------------------------------------------------------------------------
+
+	MUTABLE_IMPLEMENT_NODE( NodeImageSwitch )
+
+
+	//---------------------------------------------------------------------------------------------
+	// Own Interface
+	//---------------------------------------------------------------------------------------------
+	NodeScalarPtr NodeImageSwitch::GetParameter() const
+	{
+		return m_pD->m_pParameter.get();
+	}
+
+
+	//---------------------------------------------------------------------------------------------
+	void NodeImageSwitch::SetParameter( NodeScalarPtr pNode )
+	{
+		m_pD->m_pParameter = pNode;
+	}
+
+
+	//---------------------------------------------------------------------------------------------
+	void NodeImageSwitch::SetOptionCount( int t )
+	{
+		m_pD->Options.SetNum(t);
+	}
+
+
+	//---------------------------------------------------------------------------------------------
+	int NodeImageSwitch::GetOptionCount() const 
+	{
+		return int(m_pD->Options.Num());
+	}
+
+
+	//---------------------------------------------------------------------------------------------
+	NodeImagePtr NodeImageSwitch::GetOption( int t ) const
+	{
+		check( t>=0 && t<(int)m_pD->Options.Num() );
+		return m_pD->Options[t].get();
+	}
+
+
+	//---------------------------------------------------------------------------------------------
+	void NodeImageSwitch::SetOption( int t, NodeImagePtr pNode )
+	{
+		check( t>=0 && t<(int)m_pD->Options.Num() );
+		m_pD->Options[t] = pNode;
+	}
+
+
+
+}
+
+
diff --git a/Engine/Plugins/Mutable/Source/MutableTools/Private/MuT/NodeImageSwizzle.cpp b/Engine/Plugins/Mutable/Source/MutableTools/Private/MuT/NodeImageSwizzle.cpp
new file mode 100644
--- /dev/null
+++ b/Engine/Plugins/Mutable/Source/MutableTools/Private/MuT/NodeImageSwizzle.cpp
@@ -0,0 +1,104 @@
+// Copyright Epic Games, Inc. All Rights Reserved.
+
+
+#include "MuT/NodeImageSwizzle.h"
+
+#include "Misc/AssertionMacros.h"
+#include "MuR/ImagePrivate.h"
+#include "MuT/NodeImageSwizzlePrivate.h"
+#include "MuT/NodePrivate.h"
+
+
+namespace mu
+{
+
+	//---------------------------------------------------------------------------------------------
+	// Static initialisation
+	//---------------------------------------------------------------------------------------------
+	FNodeType NodeImageSwizzle::Private::s_type = FNodeType(Node::EType::ImageSwizzle, NodeImage::GetStaticType() );
+
+
+	//---------------------------------------------------------------------------------------------
+	//!
+	//---------------------------------------------------------------------------------------------
+
+	MUTABLE_IMPLEMENT_NODE( NodeImageSwizzle )
+
+
+	//---------------------------------------------------------------------------------------------
+	// Own Interface
+	//---------------------------------------------------------------------------------------------
+	EImageFormat NodeImageSwizzle::GetFormat() const
+	{
+		return m_pD->m_format;
+	}
+
+
+	//---------------------------------------------------------------------------------------------
+	void NodeImageSwizzle::SetFormat(EImageFormat format )
+	{
+		m_pD->m_format = format;
+
+		int32 channelCount = GetImageFormatData( format ).Channels;
+		m_pD->m_sources.SetNum( channelCount );
+		m_pD->m_sourceChannels.SetNum( channelCount );
+	}
+
+
+	//---------------------------------------------------------------------------------------------
+	NodeImagePtr NodeImageSwizzle::GetSource( int32 t ) const
+	{
+		if (m_pD->m_sources.IsValidIndex(t))
+		{
+			return m_pD->m_sources[t].get();
+		}
+		ensure(false);
+		return nullptr;
+	}
+
+
+	//---------------------------------------------------------------------------------------------
+	void NodeImageSwizzle::SetSource( int32 t, NodeImagePtr pNode )
+	{
+		if (m_pD->m_sources.IsValidIndex(t))
+		{
+			m_pD->m_sources[t] = pNode;
+		}
+		else
+		{
+			ensure(false);
+		}
+	}
+
+
+	//---------------------------------------------------------------------------------------------
+	int32 NodeImageSwizzle::GetSourceChannel( int32 t ) const
+	{
+		if (m_pD->m_sourceChannels.IsValidIndex(t))
+		{
+			return m_pD->m_sourceChannels[t];
+		}
+		else
+		{
+			ensure(false);
+			return 0;
+		}
+	}
+
+
+	//---------------------------------------------------------------------------------------------
+	void NodeImageSwizzle::SetSourceChannel( int32 OutputChannel, int32 SourceChannel )
+	{
+		if (m_pD->m_sourceChannels.IsValidIndex(OutputChannel))
+		{
+			m_pD->m_sourceChannels[OutputChannel] = SourceChannel;
+		}
+		else
+		{
+			ensure(false);
+		}
+	}
+
+
+
+}
diff --git a/Engine/Plugins/Mutable/Source/MutableTools/Private/MuT/NodeImageTransform.cpp b/Engine/Plugins/Mutable/Source/MutableTools/Private/MuT/NodeImageTransform.cpp
new file mode 100644
--- /dev/null
+++ b/Engine/Plugins/Mutable/Source/MutableTools/Private/MuT/NodeImageTransform.cpp
@@ -0,0 +1,147 @@
+// Copyright Epic Games, Inc. All Rights Reserved.
+
+
+#include "MuT/NodeImageTransform.h"
+
+#include "Misc/AssertionMacros.h"
+#include "MuT/NodeImageTransformPrivate.h"
+#include "MuT/NodePrivate.h"
+#include "MuT/NodeScalar.h"
+
+
+namespace mu
+{
+
+	//---------------------------------------------------------------------------------------------
+	// Static initialisation
+	//---------------------------------------------------------------------------------------------
+	FNodeType NodeImageTransform::Private::s_type = FNodeType(Node::EType::ImageTransform, NodeImage::GetStaticType() );
+
+	//---------------------------------------------------------------------------------------------
+
+	MUTABLE_IMPLEMENT_NODE( NodeImageTransform )
+
+
+	//---------------------------------------------------------------------------------------------
+	// Own Interface
+	//---------------------------------------------------------------------------------------------
+	NodeImagePtr NodeImageTransform::GetBase() const
+	{
+		return m_pD->m_pBase;
+	}
+
+	//---------------------------------------------------------------------------------------------
+	NodeScalarPtr NodeImageTransform::GetOffsetX() const
+	{
+		return m_pD->m_pOffsetX;
+	}
+
+	//---------------------------------------------------------------------------------------------
+	NodeScalarPtr NodeImageTransform::GetOffsetY() const
+	{
+		return m_pD->m_pOffsetY;
+	}
+
+	//---------------------------------------------------------------------------------------------
+	NodeScalarPtr NodeImageTransform::GetScaleX() const
+	{
+		return m_pD->m_pScaleX;
+	}
+
+	//---------------------------------------------------------------------------------------------
+	NodeScalarPtr NodeImageTransform::GetScaleY() const
+	{
+		return m_pD->m_pScaleY;
+	}
+
+	//---------------------------------------------------------------------------------------------
+	NodeScalarPtr NodeImageTransform::GetRotation() const
+	{
+		return m_pD->m_pRotation;
+	}
+
+	//---------------------------------------------------------------------------------------------
+	void NodeImageTransform::SetBase( NodeImagePtr pNode )
+	{
+		m_pD->m_pBase = pNode;
+	}
+
+	//---------------------------------------------------------------------------------------------
+	void NodeImageTransform::SetOffsetX( NodeScalarPtr pNode )
+	{
+		m_pD->m_pOffsetX = pNode;
+	}
+
+	//---------------------------------------------------------------------------------------------
+	void NodeImageTransform::SetOffsetY( NodeScalarPtr pNode )
+	{
+		m_pD->m_pOffsetY = pNode;
+	}
+
+	//---------------------------------------------------------------------------------------------
+	void NodeImageTransform::SetScaleX( NodeScalarPtr pNode )
+	{
+		m_pD->m_pScaleX = pNode;
+	}
+
+	//---------------------------------------------------------------------------------------------
+	void NodeImageTransform::SetScaleY( NodeScalarPtr pNode )
+	{
+		m_pD->m_pScaleY = pNode;
+	}
+
+	//---------------------------------------------------------------------------------------------
+	void NodeImageTransform::SetRotation( NodeScalarPtr pNode )
+	{
+		m_pD->m_pRotation = pNode;
+	}
+
+	//---------------------------------------------------------------------------------------------
+	void NodeImageTransform::SetAddressMode(EAddressMode AddressMode)
+	{
+		m_pD->AddressMode = AddressMode;
+	}
+
+	//---------------------------------------------------------------------------------------------
+	EAddressMode NodeImageTransform::GetAddressMode() const
+	{
+		return m_pD->AddressMode;
+	}
+
+	//---------------------------------------------------------------------------------------------
+	uint16 NodeImageTransform::GetSizeX() const
+	{
+		return m_pD->SizeX;
+	}
+
+	//---------------------------------------------------------------------------------------------
+	void NodeImageTransform::SetSizeX(uint16 SizeX)
+	{
+		m_pD->SizeX = SizeX;
+	}
+
+	//---------------------------------------------------------------------------------------------
+	uint16 NodeImageTransform::GetSizeY() const
+	{
+		return m_pD->SizeY;
+	}
+
+	//---------------------------------------------------------------------------------------------
+	void NodeImageTransform::SetSizeY(uint16 SizeY)
+	{
+		m_pD->SizeY = SizeY;
+	}
+
+	//---------------------------------------------------------------------------------------------
+	bool NodeImageTransform::GetKeepAspectRatio() const
+	{
+		return m_pD->bKeepAspectRatio;
+	}
+
+	//---------------------------------------------------------------------------------------------
+	void NodeImageTransform::SetKeepAspectRatio(bool bKeepAspectRatio)
+	{
+		m_pD->bKeepAspectRatio = bKeepAspectRatio;
+	}
+}
+
diff --git a/Engine/Plugins/Mutable/Source/MutableTools/Private/MuT/NodeImageTransformPrivate.h b/Engine/Plugins/Mutable/Source/MutableTools/Private/MuT/NodeImageTransformPrivate.h
new file mode 100644
--- /dev/null
+++ b/Engine/Plugins/Mutable/Source/MutableTools/Private/MuT/NodeImageTransformPrivate.h
@@ -0,0 +1,39 @@
+// Copyright Epic Games, Inc. All Rights Reserved.
+
+
+#pragma once
+
+#include "MuR/ImagePrivate.h"
+#include "MuT/NodeImagePrivate.h"
+#include "MuT/NodeImageTransform.h"
+#include "MuT/NodeScalar.h"
+
+namespace mu
+{
+
+
+	class NodeImageTransform::Private : public NodeImage::Private
+	{
+	public:
+
+		static FNodeType s_type;
+
+		NodeImagePtr m_pBase;
+		NodeScalarPtr m_pOffsetX;
+		NodeScalarPtr m_pOffsetY;
+		NodeScalarPtr m_pScaleX;
+		NodeScalarPtr m_pScaleY;
+		NodeScalarPtr m_pRotation;
+
+		EAddressMode AddressMode = EAddressMode::Wrap;
+		uint32 SizeX = 0;
+		uint32 SizeY = 0;
+		
+		bool bKeepAspectRatio = false;
+		uint8 UnusedPadding[sizeof(NodeImagePtr) - sizeof(bool)] = {0}; 
+		static_assert(sizeof(NodeImagePtr) - sizeof(bool) >= 1);
+
+	};
+
+
+}
