[Backout] - CL40275509

diff --git a/Engine/Plugins/Mutable/Source/MutableTools/Private/MuT/NodeImageVariation.cpp b/Engine/Plugins/Mutable/Source/MutableTools/Private/MuT/NodeImageVariation.cpp
new file mode 100644
--- /dev/null
+++ b/Engine/Plugins/Mutable/Source/MutableTools/Private/MuT/NodeImageVariation.cpp
@@ -0,0 +1,65 @@
+// Copyright Epic Games, Inc. All Rights Reserved.
+
+
+#include "MuT/NodeImageVariation.h"
+
+#include "Misc/AssertionMacros.h"
+#include "MuT/NodeImageVariationPrivate.h"
+#include "MuT/NodePrivate.h"
+
+
+namespace mu
+{
+
+    //---------------------------------------------------------------------------------------------
+    // Static initialisation
+    //---------------------------------------------------------------------------------------------
+    FNodeType NodeImageVariation::Private::s_type = FNodeType(Node::EType::ImageVariation, NodeImage::GetStaticType() );
+
+
+    //---------------------------------------------------------------------------------------------
+    //---------------------------------------------------------------------------------------------
+    //---------------------------------------------------------------------------------------------
+
+    MUTABLE_IMPLEMENT_NODE( NodeImageVariation )
+
+
+    //---------------------------------------------------------------------------------------------
+    // Own Interface
+    //---------------------------------------------------------------------------------------------
+    void NodeImageVariation::SetDefaultImage( NodeImage* p ) 
+	{ 
+		m_pD->m_defaultImage = p; 
+	}
+
+
+	int NodeImageVariation::GetVariationCount() const 
+	{ 
+		return int( m_pD->m_variations.Num() ); 
+	}
+
+
+    void NodeImageVariation::SetVariationCount( int num )
+    {
+        check( num >= 0 );
+        m_pD->m_variations.SetNum( num );
+    }
+
+
+	void NodeImageVariation::SetVariationTag( int index, const FString& Tag )
+    {
+        check( index >= 0 && index < (int)m_pD->m_variations.Num() );
+
+		m_pD->m_variations[index].m_tag = Tag;
+    }
+
+
+    void NodeImageVariation::SetVariationImage( int index, NodeImage* pNode )
+    {
+        check( index >= 0 && index < (int)m_pD->m_variations.Num() );
+
+        m_pD->m_variations[index].m_image = pNode;
+    }
+
+
+} // namespace mu
diff --git a/Engine/Plugins/Mutable/Source/MutableTools/Private/MuT/NodeImageVariationPrivate.h b/Engine/Plugins/Mutable/Source/MutableTools/Private/MuT/NodeImageVariationPrivate.h
new file mode 100644
--- /dev/null
+++ b/Engine/Plugins/Mutable/Source/MutableTools/Private/MuT/NodeImageVariationPrivate.h
@@ -0,0 +1,33 @@
+// Copyright Epic Games, Inc. All Rights Reserved.
+
+#pragma once
+
+#include "MuT/NodePrivate.h"
+#include "MuT/NodeImageVariation.h"
+
+namespace mu
+{
+
+    //---------------------------------------------------------------------------------------------
+    //!
+    //---------------------------------------------------------------------------------------------
+    class NodeImageVariation::Private : public Node::Private
+    {
+    public:
+        Private() {}
+
+        static FNodeType s_type;
+
+        NodeImagePtr m_defaultImage;
+
+        struct FVariation
+        {
+            NodeImagePtr m_image;
+			FString m_tag;
+        };
+
+		TArray<FVariation> m_variations;
+    };
+
+} // namespace mu
+
diff --git a/Engine/Plugins/Mutable/Source/MutableTools/Private/MuT/NodeMesh.cpp b/Engine/Plugins/Mutable/Source/MutableTools/Private/MuT/NodeMesh.cpp
--- a/Engine/Plugins/Mutable/Source/MutableTools/Private/MuT/NodeMesh.cpp
+++ b/Engine/Plugins/Mutable/Source/MutableTools/Private/MuT/NodeMesh.cpp
@@ -11,6 +11,20 @@
 namespace mu
 {
 
+	static FNodeType s_nodeMeshType = FNodeType(Node::EType::Mesh , Node::GetStaticType());
+
+
+	const FNodeType* NodeMesh::GetType() const
+	{
+		return GetStaticType();
+	}
+
+
+	const FNodeType* NodeMesh::GetStaticType()
+	{
+		return &s_nodeMeshType;
+	}
+
 
 	void MeshCreateCollapsedVertexMap(const mu::FMesh* Mesh, TArray<int32>& CollapsedVertices)
 	{
diff --git a/Engine/Plugins/Mutable/Source/MutableTools/Private/MuT/NodeMeshApplyPose.cpp b/Engine/Plugins/Mutable/Source/MutableTools/Private/MuT/NodeMeshApplyPose.cpp
new file mode 100644
--- /dev/null
+++ b/Engine/Plugins/Mutable/Source/MutableTools/Private/MuT/NodeMeshApplyPose.cpp
@@ -0,0 +1,60 @@
+// Copyright Epic Games, Inc. All Rights Reserved.
+
+#include "MuT/NodeMeshApplyPose.h"
+
+#include "Misc/AssertionMacros.h"
+#include "MuT/NodeLayout.h"
+#include "MuT/NodeMeshApplyPosePrivate.h"
+#include "MuT/NodeMeshPrivate.h"
+#include "MuT/NodePrivate.h"
+
+
+namespace mu
+{
+
+	//---------------------------------------------------------------------------------------------
+	// Static initialisation
+	//---------------------------------------------------------------------------------------------
+    FNodeType NodeMeshApplyPose::Private::s_type = FNodeType(Node::EType::MeshApplyPose, NodeMesh::GetStaticType() );
+
+
+	//---------------------------------------------------------------------------------------------
+	//!
+	//---------------------------------------------------------------------------------------------
+
+    MUTABLE_IMPLEMENT_NODE( NodeMeshApplyPose )
+
+
+	//---------------------------------------------------------------------------------------------
+	// Own Interface
+	//---------------------------------------------------------------------------------------------
+    NodeMeshPtr NodeMeshApplyPose::GetBase() const
+    {
+        return m_pD->m_pBase.get();
+    }
+
+
+    //---------------------------------------------------------------------------------------------
+    void NodeMeshApplyPose::SetBase( NodeMeshPtr pNode )
+    {
+        m_pD->m_pBase = pNode;
+    }
+
+
+    //---------------------------------------------------------------------------------------------
+    NodeMeshPtr NodeMeshApplyPose::GetPose() const
+    {
+        return m_pD->m_pPose.get();
+    }
+
+
+    //---------------------------------------------------------------------------------------------
+    void NodeMeshApplyPose::SetPose( NodeMeshPtr pNode )
+    {
+        m_pD->m_pPose = pNode;
+    }
+
+
+}
+
+
diff --git a/Engine/Plugins/Mutable/Source/MutableTools/Private/MuT/NodeMeshApplyPosePrivate.h b/Engine/Plugins/Mutable/Source/MutableTools/Private/MuT/NodeMeshApplyPosePrivate.h
new file mode 100644
--- /dev/null
+++ b/Engine/Plugins/Mutable/Source/MutableTools/Private/MuT/NodeMeshApplyPosePrivate.h
@@ -0,0 +1,26 @@
+// Copyright Epic Games, Inc. All Rights Reserved.
+
+#pragma once
+
+#include "MuT/NodeMeshApplyPose.h"
+#include "MuT/NodeScalar.h"
+#include "MuT/NodeMeshPrivate.h"
+
+namespace mu
+{
+
+
+    class NodeMeshApplyPose::Private : public NodeMesh::Private
+	{
+	public:
+
+		static FNodeType s_type;
+
+        NodeMeshPtr m_pBase;
+        NodeMeshPtr m_pPose;
+
+	};
+
+
+}
+
diff --git a/Engine/Plugins/Mutable/Source/MutableTools/Private/MuT/NodeMeshClipMorphPlane.cpp b/Engine/Plugins/Mutable/Source/MutableTools/Private/MuT/NodeMeshClipMorphPlane.cpp
--- a/Engine/Plugins/Mutable/Source/MutableTools/Private/MuT/NodeMeshClipMorphPlane.cpp
+++ b/Engine/Plugins/Mutable/Source/MutableTools/Private/MuT/NodeMeshClipMorphPlane.cpp
@@ -5,6 +5,8 @@
 #include "Misc/AssertionMacros.h"
 #include "MuR/MutableMath.h"
 #include "MuT/NodeLayout.h"
+#include "MuT/NodeMeshPrivate.h"
+#include "MuT/NodePrivate.h"
 
 
 namespace mu
diff --git a/Engine/Plugins/Mutable/Source/MutableTools/Private/MuT/NodeMeshConstant.cpp b/Engine/Plugins/Mutable/Source/MutableTools/Private/MuT/NodeMeshConstant.cpp
new file mode 100644
--- /dev/null
+++ b/Engine/Plugins/Mutable/Source/MutableTools/Private/MuT/NodeMeshConstant.cpp
@@ -0,0 +1,76 @@
+// Copyright Epic Games, Inc. All Rights Reserved.
+
+
+#include "MuT/NodeMeshConstant.h"
+
+#include "Misc/AssertionMacros.h"
+#include "MuT/NodeLayout.h"
+#include "MuT/NodeMeshConstantPrivate.h"
+#include "MuT/NodePrivate.h"
+#include "MuR/Mesh.h"
+
+
+namespace mu
+{
+	FNodeType NodeMeshConstant::Private::s_type = FNodeType(Node::EType::MeshConstant, NodeMesh::GetStaticType() );
+
+
+	MUTABLE_IMPLEMENT_NODE( NodeMeshConstant );
+
+
+	TSharedPtr<FMesh> NodeMeshConstant::GetValue() const
+	{
+		return m_pD->Value;
+	}
+
+
+	void NodeMeshConstant::SetValue( TSharedPtr<FMesh> pValue )
+	{
+		m_pD->Value = pValue;
+
+        if (m_pD->Value)
+        {
+            // Ensure mesh is well formed
+            m_pD->Value->EnsureSurfaceData();
+        }
+    }
+
+
+	int32 NodeMeshConstant::GetLayoutCount() const
+	{
+		return m_pD->Layouts.Num();
+	}
+
+
+	void NodeMeshConstant::SetLayoutCount( int32 num )
+	{
+		check( num >=0 );
+		m_pD->Layouts.SetNum( num );
+	}
+
+
+	Ptr<NodeLayout> NodeMeshConstant::GetLayout( int32 index ) const
+	{
+		check( index >=0 && index < m_pD->Layouts.Num() );
+
+		Ptr<NodeLayout> pResult;
+
+		if (index >= 0 && index < m_pD->Layouts.Num())
+		{
+			pResult = m_pD->Layouts[index];
+		}
+		return pResult;
+	}
+
+
+	void NodeMeshConstant::SetLayout( int32 index, Ptr<NodeLayout> pLayout )
+	{
+		check( index >=0 && index < m_pD->Layouts.Num() );
+
+		m_pD->Layouts[ index ] = pLayout;
+	}
+
+
+}
+
+
diff --git a/Engine/Plugins/Mutable/Source/MutableTools/Private/MuT/NodeMeshFormat.cpp b/Engine/Plugins/Mutable/Source/MutableTools/Private/MuT/NodeMeshFormat.cpp
new file mode 100644
--- /dev/null
+++ b/Engine/Plugins/Mutable/Source/MutableTools/Private/MuT/NodeMeshFormat.cpp
@@ -0,0 +1,64 @@
+// Copyright Epic Games, Inc. All Rights Reserved.
+
+
+#include "MuT/NodeMeshFormat.h"
+
+#include "Misc/AssertionMacros.h"
+#include "MuT/NodeLayout.h"
+#include "MuT/NodeMeshFormatPrivate.h"
+#include "MuT/NodeMeshPrivate.h"
+#include "MuT/NodePrivate.h"
+
+
+namespace mu
+{
+class FMeshBufferSet;
+
+	//---------------------------------------------------------------------------------------------
+	// Static initialisation
+	//---------------------------------------------------------------------------------------------
+	FNodeType NodeMeshFormat::Private::s_type = FNodeType(Node::EType::MeshFormat, NodeMesh::GetStaticType() );
+
+
+	//---------------------------------------------------------------------------------------------
+	//!
+	//---------------------------------------------------------------------------------------------
+
+	MUTABLE_IMPLEMENT_NODE( NodeMeshFormat )
+
+
+	//---------------------------------------------------------------------------------------------
+	// Own Interface
+	//---------------------------------------------------------------------------------------------
+	NodeMeshPtr NodeMeshFormat::GetSource() const
+	{
+		return m_pD->Source;
+	}
+
+
+	void NodeMeshFormat::SetSource( NodeMesh* pValue )
+	{
+		m_pD->Source = pValue;
+	}
+
+
+	FMeshBufferSet& NodeMeshFormat::GetVertexBuffers()
+	{
+		return m_pD->VertexBuffers;
+	}
+
+
+	FMeshBufferSet& NodeMeshFormat::GetIndexBuffers()
+	{
+		return m_pD->IndexBuffers;
+	}
+
+
+	void NodeMeshFormat::SetOptimizeBuffers(bool bEnable)
+	{
+		m_pD->bOptimizeBuffers = bEnable;
+	}
+
+}
+
+
diff --git a/Engine/Plugins/Mutable/Source/MutableTools/Private/MuT/NodeMeshMakeMorph.cpp b/Engine/Plugins/Mutable/Source/MutableTools/Private/MuT/NodeMeshMakeMorph.cpp
new file mode 100644
--- /dev/null
+++ b/Engine/Plugins/Mutable/Source/MutableTools/Private/MuT/NodeMeshMakeMorph.cpp
@@ -0,0 +1,71 @@
+// Copyright Epic Games, Inc. All Rights Reserved.
+
+
+#include "MuT/NodeMeshMakeMorph.h"
+
+#include "Misc/AssertionMacros.h"
+#include "MuT/NodeLayout.h"
+#include "MuT/NodeMeshMakeMorphPrivate.h"
+#include "MuT/NodeMeshPrivate.h"
+#include "MuT/NodePrivate.h"
+
+
+namespace mu
+{
+
+	//---------------------------------------------------------------------------------------------
+	// Static initialisation
+	//---------------------------------------------------------------------------------------------
+    FNodeType NodeMeshMakeMorph::Private::s_type = FNodeType(Node::EType::MeshMakeMorph, NodeMesh::GetStaticType() );
+
+
+	//---------------------------------------------------------------------------------------------
+	//!
+	//---------------------------------------------------------------------------------------------
+
+    MUTABLE_IMPLEMENT_NODE( NodeMeshMakeMorph )
+
+
+	//---------------------------------------------------------------------------------------------
+	// Own Interface
+	//---------------------------------------------------------------------------------------------
+    NodeMeshPtr NodeMeshMakeMorph::GetBase() const
+	{
+		return m_pD->m_pBase;
+	}
+
+
+	//---------------------------------------------------------------------------------------------
+    void NodeMeshMakeMorph::SetBase( NodeMesh* p )
+	{
+		m_pD->m_pBase = p;
+	}
+
+
+	//---------------------------------------------------------------------------------------------
+    NodeMeshPtr NodeMeshMakeMorph::GetTarget() const
+	{
+		return m_pD->m_pTarget;
+	}
+
+
+	//---------------------------------------------------------------------------------------------
+    void NodeMeshMakeMorph::SetTarget( NodeMesh* p )
+	{
+		m_pD->m_pTarget = p;
+	}
+
+	//---------------------------------------------------------------------------------------------
+	void NodeMeshMakeMorph::SetOnlyPositionAndNormal(bool bInOnlyPositionAndNormals)
+	{
+		m_pD->bOnlyPositionAndNormal = bInOnlyPositionAndNormals;
+	}
+
+
+	//---------------------------------------------------------------------------------------------
+	bool NodeMeshMakeMorph::GetOnlyPositionAndNormal() const
+	{
+		return m_pD->bOnlyPositionAndNormal;
+	}
+
+}
diff --git a/Engine/Plugins/Mutable/Source/MutableTools/Private/MuT/NodeMeshMorph.cpp b/Engine/Plugins/Mutable/Source/MutableTools/Private/MuT/NodeMeshMorph.cpp
new file mode 100644
--- /dev/null
+++ b/Engine/Plugins/Mutable/Source/MutableTools/Private/MuT/NodeMeshMorph.cpp
@@ -0,0 +1,99 @@
+// Copyright Epic Games, Inc. All Rights Reserved.
+
+#include "MuT/NodeMeshMorph.h"
+
+#include "Containers/Array.h"
+#include "HAL/PlatformCrt.h"
+#include "Misc/AssertionMacros.h"
+#include "MuT/NodeLayout.h"
+#include "MuT/NodeMeshMorphPrivate.h"
+#include "MuT/NodeMeshPrivate.h"
+#include "MuT/NodePrivate.h"
+#include "MuT/NodeScalar.h"
+
+
+namespace mu
+{
+
+	//---------------------------------------------------------------------------------------------
+	// Static initialisation
+	//---------------------------------------------------------------------------------------------
+	FNodeType NodeMeshMorph::Private::s_type = FNodeType(Node::EType::MeshMorph, NodeMesh::GetStaticType() );
+
+
+	//---------------------------------------------------------------------------------------------
+	//!
+	//---------------------------------------------------------------------------------------------
+
+	MUTABLE_IMPLEMENT_NODE( NodeMeshMorph )
+
+
+	//---------------------------------------------------------------------------------------------
+	// Own Interface
+	//---------------------------------------------------------------------------------------------
+	Ptr<NodeScalar> NodeMeshMorph::GetFactor() const
+	{
+		return m_pD->Factor.get();
+	}
+
+
+	//---------------------------------------------------------------------------------------------
+	void NodeMeshMorph::SetFactor( Ptr<NodeScalar> pNode )
+	{
+		m_pD->Factor = pNode;
+	}
+
+
+	//---------------------------------------------------------------------------------------------
+	Ptr<NodeMesh> NodeMeshMorph::GetBase() const
+	{
+		return m_pD->Base.get();
+	}
+
+
+	//---------------------------------------------------------------------------------------------
+	void NodeMeshMorph::SetBase( Ptr<NodeMesh> pNode )
+	{
+		m_pD->Base = pNode;
+	}
+
+
+	//---------------------------------------------------------------------------------------------
+	NodeMeshPtr NodeMeshMorph::GetMorph() const
+	{
+		return m_pD->Morph.get();
+	}
+
+
+	//---------------------------------------------------------------------------------------------
+	void NodeMeshMorph::SetMorph( Ptr<NodeMesh> pNode )
+	{
+		m_pD->Morph = pNode;
+	}
+
+	//---------------------------------------------------------------------------------------------
+	void NodeMeshMorph::SetReshapeSkeleton(bool bEnable)
+	{
+		m_pD->bReshapeSkeleton = bEnable;
+	}	
+
+	//---------------------------------------------------------------------------------------------
+	void NodeMeshMorph::AddBoneToDeform(const FBoneName& BoneId)
+	{
+		m_pD->BonesToDeform.Emplace(BoneId);
+	}
+
+	void NodeMeshMorph::AddPhysicsBodyToDeform(const FBoneName& BoneId)
+	{
+		m_pD->PhysicsToDeform.Emplace(BoneId);
+	}
+
+	//---------------------------------------------------------------------------------------------
+	void NodeMeshMorph::SetReshapePhysicsVolumes(bool bEnable)
+	{
+		m_pD->bReshapePhysicsVolumes = bEnable;
+	}
+
+}
+
+
diff --git a/Engine/Plugins/Mutable/Source/MutableTools/Private/MuT/NodeMeshReshape.cpp b/Engine/Plugins/Mutable/Source/MutableTools/Private/MuT/NodeMeshReshape.cpp
new file mode 100644
--- /dev/null
+++ b/Engine/Plugins/Mutable/Source/MutableTools/Private/MuT/NodeMeshReshape.cpp
@@ -0,0 +1,128 @@
+// Copyright Epic Games, Inc. All Rights Reserved.
+
+#include "MuT/NodeMeshReshape.h"
+
+#include "Containers/Array.h"
+#include "HAL/PlatformCrt.h"
+#include "Misc/AssertionMacros.h"
+#include "MuT/NodeLayout.h"
+#include "MuT/NodeMeshPrivate.h"
+#include "MuT/NodeMeshReshapePrivate.h"
+#include "MuT/NodePrivate.h"
+
+
+namespace mu
+{
+
+	//---------------------------------------------------------------------------------------------
+	// Static initialisation
+	//---------------------------------------------------------------------------------------------
+	FNodeType NodeMeshReshape::Private::s_type = FNodeType(Node::EType::MeshReshape, NodeMesh::GetStaticType());
+
+
+	//---------------------------------------------------------------------------------------------
+	//!
+	//---------------------------------------------------------------------------------------------
+
+	MUTABLE_IMPLEMENT_NODE(NodeMeshReshape )
+
+
+	//---------------------------------------------------------------------------------------------
+	// Own Interface
+	//---------------------------------------------------------------------------------------------
+	const Ptr<NodeMesh>& NodeMeshReshape::GetBaseMesh() const
+	{
+		return m_pD->BaseMesh;
+	}
+
+
+	//---------------------------------------------------------------------------------------------
+	void NodeMeshReshape::SetBaseMesh(const Ptr<NodeMesh>& pNode)
+	{
+		m_pD->BaseMesh = pNode;
+	}
+
+
+	//---------------------------------------------------------------------------------------------
+	const Ptr<NodeMesh>& NodeMeshReshape::GetBaseShape() const
+	{
+		return m_pD->BaseShape;
+	}
+
+
+	//---------------------------------------------------------------------------------------------
+	void NodeMeshReshape::SetBaseShape(const Ptr<NodeMesh>& pNode)
+	{
+		m_pD->BaseShape = pNode;
+	}
+
+
+	//---------------------------------------------------------------------------------------------
+	const Ptr<NodeMesh>& NodeMeshReshape::GetTargetShape() const
+	{
+		return m_pD->TargetShape;
+	}
+
+
+	//---------------------------------------------------------------------------------------------
+	void NodeMeshReshape::SetTargetShape(const Ptr<NodeMesh>& pNode)
+	{
+		m_pD->TargetShape = pNode;
+	}
+
+	//---------------------------------------------------------------------------------------------
+	void NodeMeshReshape::SetReshapeVertices(bool bEnable)
+	{
+		m_pD->bReshapeVertices = bEnable;
+	}
+
+	//---------------------------------------------------------------------------------------------
+	void NodeMeshReshape::SetRecomputeNormals(bool bEnable)
+	{
+		m_pD->bRecomputeNormals = bEnable;
+	}
+
+
+	//---------------------------------------------------------------------------------------------
+	void NodeMeshReshape::SetApplyLaplacian(bool bEnable)
+	{
+		m_pD->bApplyLaplacian = bEnable;
+	}
+
+
+	//---------------------------------------------------------------------------------------------
+	void NodeMeshReshape::SetReshapeSkeleton(bool bEnable)
+	{
+		m_pD->bReshapeSkeleton = bEnable;
+	}
+	
+
+	//---------------------------------------------------------------------------------------------
+	void NodeMeshReshape::SetColorUsages(EVertexColorUsage R, EVertexColorUsage G, EVertexColorUsage B, EVertexColorUsage A)
+	{	
+		m_pD->ColorRChannelUsage = R;
+		m_pD->ColorGChannelUsage = G;
+		m_pD->ColorBChannelUsage = B;
+		m_pD->ColorAChannelUsage = A;
+	}	
+
+	//---------------------------------------------------------------------------------------------
+	void NodeMeshReshape::AddBoneToDeform(const FBoneName& BoneId)
+	{
+		m_pD->BonesToDeform.Emplace(BoneId);
+	}
+
+	//---------------------------------------------------------------------------------------------
+	void NodeMeshReshape::AddPhysicsBodyToDeform(const FBoneName& BoneId)
+	{
+		m_pD->PhysicsToDeform.Emplace(BoneId);
+	}
+
+	void NodeMeshReshape::SetReshapePhysicsVolumes(bool bEnable)
+	{
+		m_pD->bReshapePhysicsVolumes = bEnable;
+	}
+
+}
+
+
diff --git a/Engine/Plugins/Mutable/Source/MutableTools/Private/MuT/NodeMeshReshapePrivate.h b/Engine/Plugins/Mutable/Source/MutableTools/Private/MuT/NodeMeshReshapePrivate.h
new file mode 100644
--- /dev/null
+++ b/Engine/Plugins/Mutable/Source/MutableTools/Private/MuT/NodeMeshReshapePrivate.h
@@ -0,0 +1,39 @@
+// Copyright Epic Games, Inc. All Rights Reserved.
+
+#pragma once
+
+#include "MuT/NodeMeshReshape.h"
+#include "MuT/NodeScalar.h"
+#include "MuT/NodeMeshPrivate.h"
+#include "MuR/Skeleton.h"
+
+namespace mu
+{
+
+
+	class NodeMeshReshape::Private : public NodeMesh::Private
+	{
+	public:
+
+		static FNodeType s_type;
+
+		Ptr<NodeMesh> BaseMesh;
+		Ptr<NodeMesh> BaseShape;
+		Ptr<NodeMesh> TargetShape;
+		bool bReshapeVertices = true;
+		bool bRecomputeNormals = false;
+		bool bApplyLaplacian = false;
+		bool bReshapeSkeleton = false;
+		bool bReshapePhysicsVolumes = false;
+
+		EVertexColorUsage ColorRChannelUsage = EVertexColorUsage::None;
+		EVertexColorUsage ColorGChannelUsage = EVertexColorUsage::None;
+		EVertexColorUsage ColorBChannelUsage = EVertexColorUsage::None;
+		EVertexColorUsage ColorAChannelUsage = EVertexColorUsage::None;
+
+		TArray<FBoneName> BonesToDeform;
+		TArray<FBoneName> PhysicsToDeform;
+
+	};
+
+}
diff --git a/Engine/Plugins/Mutable/Source/MutableTools/Private/MuT/NodeMeshSwitch.cpp b/Engine/Plugins/Mutable/Source/MutableTools/Private/MuT/NodeMeshSwitch.cpp
new file mode 100644
--- /dev/null
+++ b/Engine/Plugins/Mutable/Source/MutableTools/Private/MuT/NodeMeshSwitch.cpp
@@ -0,0 +1,69 @@
+// Copyright Epic Games, Inc. All Rights Reserved.
+
+#include "MuT/NodeMeshSwitch.h"
+
+#include "Misc/AssertionMacros.h"
+#include "MuT/NodeLayout.h"
+#include "MuT/NodeMeshPrivate.h"
+#include "MuT/NodeMeshSwitchPrivate.h"
+#include "MuT/NodePrivate.h"
+#include "MuT/NodeScalar.h"
+
+
+namespace mu
+{
+
+	//---------------------------------------------------------------------------------------------
+	// Static initialisation
+	//---------------------------------------------------------------------------------------------
+	FNodeType NodeMeshSwitch::Private::s_type = FNodeType(Node::EType::MeshSwitch, NodeMesh::GetStaticType() );
+
+
+	//---------------------------------------------------------------------------------------------
+	//!
+	//---------------------------------------------------------------------------------------------
+
+	MUTABLE_IMPLEMENT_NODE( NodeMeshSwitch )
+
+
+	//---------------------------------------------------------------------------------------------
+	// Own Interface
+	//---------------------------------------------------------------------------------------------
+	NodeScalarPtr NodeMeshSwitch::GetParameter() const
+	{
+		return m_pD->m_pParameter.get();
+	}
+
+
+	//---------------------------------------------------------------------------------------------
+	void NodeMeshSwitch::SetParameter( NodeScalarPtr pNode )
+	{
+		m_pD->m_pParameter = pNode;
+	}
+
+
+	//---------------------------------------------------------------------------------------------
+	void NodeMeshSwitch::SetOptionCount( int t )
+	{
+		m_pD->Options.SetNum(t);
+	}
+
+
+	//---------------------------------------------------------------------------------------------
+	NodeMeshPtr NodeMeshSwitch::GetOption( int t ) const
+	{
+		check( t>=0 && t<m_pD->Options.Num() );
+		return m_pD->Options[t].get();
+	}
+
+
+	//---------------------------------------------------------------------------------------------
+	void NodeMeshSwitch::SetOption( int t, NodeMeshPtr pNode )
+	{
+		check( t>=0 && t<m_pD->Options.Num() );
+		m_pD->Options[t] = pNode;
+	}
+
+}
+
+
diff --git a/Engine/Plugins/Mutable/Source/MutableTools/Private/MuT/NodeMeshSwitchPrivate.h b/Engine/Plugins/Mutable/Source/MutableTools/Private/MuT/NodeMeshSwitchPrivate.h
new file mode 100644
--- /dev/null
+++ b/Engine/Plugins/Mutable/Source/MutableTools/Private/MuT/NodeMeshSwitchPrivate.h
@@ -0,0 +1,26 @@
+// Copyright Epic Games, Inc. All Rights Reserved.
+
+#pragma once
+
+#include "MuT/NodeMeshSwitch.h"
+#include "MuT/NodeScalar.h"
+#include "MuT/NodeMeshPrivate.h"
+
+
+namespace mu
+{
+
+
+	class NodeMeshSwitch::Private : public NodeMesh::Private
+	{
+	public:
+
+		static FNodeType s_type;
+
+		NodeScalarPtr m_pParameter;
+		TArray<NodeMeshPtr> Options;
+
+	};
+
+
+}
diff --git a/Engine/Plugins/Mutable/Source/MutableTools/Private/MuT/NodeMeshTable.cpp b/Engine/Plugins/Mutable/Source/MutableTools/Private/MuT/NodeMeshTable.cpp
new file mode 100644
--- /dev/null
+++ b/Engine/Plugins/Mutable/Source/MutableTools/Private/MuT/NodeMeshTable.cpp
@@ -0,0 +1,88 @@
+// Copyright Epic Games, Inc. All Rights Reserved.
+
+
+#include "MuT/NodeMeshTable.h"
+
+#include "Misc/AssertionMacros.h"
+#include "MuT/NodeLayout.h"
+#include "MuT/NodePrivate.h"
+#include "MuT/Table.h"
+
+
+namespace mu
+{
+
+	//---------------------------------------------------------------------------------------------
+	// Static initialisation
+	//---------------------------------------------------------------------------------------------
+	FNodeType NodeMeshTable::StaticType = FNodeType(Node::EType::MeshTable, NodeMesh::GetStaticType() );
+
+
+	//---------------------------------------------------------------------------------------------
+	// Own Interface
+	//---------------------------------------------------------------------------------------------
+	void NodeMeshTable::SetColumn( const FString& strName )
+	{
+		ColumnName = strName;
+	}
+
+
+	//---------------------------------------------------------------------------------------------
+	void NodeMeshTable::SetParameterName( const FString& strName )
+	{
+		ParameterName = strName;
+	}
+
+
+	//---------------------------------------------------------------------------------------------
+	int NodeMeshTable::GetLayoutCount() const
+	{
+		return Layouts.Num();
+	}
+
+
+	//---------------------------------------------------------------------------------------------
+	void NodeMeshTable::SetLayoutCount( int i )
+	{
+		Layouts.SetNum( i );
+	}
+
+
+	//---------------------------------------------------------------------------------------------
+	Ptr<NodeLayout> NodeMeshTable::GetLayout( int i ) const
+	{
+		Ptr<NodeLayout> pResult;
+
+		if (i >= 0 && i < Layouts.Num())
+		{
+			pResult = Layouts[i];
+		}
+
+		return pResult;
+	}
+
+
+	//---------------------------------------------------------------------------------------------
+	void NodeMeshTable::SetLayout( int i, Ptr<NodeLayout> pLayout )
+	{
+		check( i>=0 && i<GetLayoutCount() );
+		Layouts[i] = pLayout;
+	}
+
+
+	//---------------------------------------------------------------------------------------------
+	void NodeMeshTable::SetNoneOption(bool bAddNoneOption)
+	{
+		bNoneOption = bAddNoneOption;
+	}
+
+
+	//---------------------------------------------------------------------------------------------
+	void NodeMeshTable::SetDefaultRowName(const FString& RowName)
+	{
+		DefaultRowName = RowName;
+	}
+
+}
+
+
diff --git a/Engine/Plugins/Mutable/Source/MutableTools/Private/MuT/NodeMeshTangents.cpp b/Engine/Plugins/Mutable/Source/MutableTools/Private/MuT/NodeMeshTangents.cpp
new file mode 100644
--- /dev/null
+++ b/Engine/Plugins/Mutable/Source/MutableTools/Private/MuT/NodeMeshTangents.cpp
@@ -0,0 +1,44 @@
+// Copyright Epic Games, Inc. All Rights Reserved.
+
+
+#include "MuT/NodeMeshTangents.h"
+
+#include "Misc/AssertionMacros.h"
+#include "MuT/NodeLayout.h"
+#include "MuT/NodeMeshPrivate.h"
+#include "MuT/NodeMeshTangentsPrivate.h"
+#include "MuT/NodePrivate.h"
+
+
+namespace mu
+{
+
+	//---------------------------------------------------------------------------------------------
+	// Static initialisation
+	//---------------------------------------------------------------------------------------------
+	FNodeType NodeMeshTangents::Private::s_type = FNodeType(Node::EType::MeshTangents, NodeMesh::GetStaticType() );
+
+
+	//---------------------------------------------------------------------------------------------
+	//!
+	//---------------------------------------------------------------------------------------------
+
+	MUTABLE_IMPLEMENT_NODE( NodeMeshTangents )
+
+
+	//---------------------------------------------------------------------------------------------
+	// Own Interface
+	//---------------------------------------------------------------------------------------------
+	NodeMeshPtr NodeMeshTangents::GetSource() const
+	{
+		return m_pD->m_pSource;
+	}
+
+
+	//---------------------------------------------------------------------------------------------
+	void NodeMeshTangents::SetSource( NodeMesh* p )
+	{
+		m_pD->m_pSource = p;
+	}
+
+}
diff --git a/Engine/Plugins/Mutable/Source/MutableTools/Private/MuT/NodeMeshTangentsPrivate.h b/Engine/Plugins/Mutable/Source/MutableTools/Private/MuT/NodeMeshTangentsPrivate.h
new file mode 100644
--- /dev/null
+++ b/Engine/Plugins/Mutable/Source/MutableTools/Private/MuT/NodeMeshTangentsPrivate.h
@@ -0,0 +1,24 @@
+// Copyright Epic Games, Inc. All Rights Reserved.
+
+#pragma once
+
+#include "MuT/NodeMeshPrivate.h"
+#include "MuT/NodeMeshTangents.h"
+
+
+namespace mu
+{
+
+
+	class NodeMeshTangents::Private : public NodeMesh::Private
+	{
+	public:
+
+		static FNodeType s_type;
+
+		NodeMeshPtr m_pSource;
+
+	};
+
+
+}
diff --git a/Engine/Plugins/Mutable/Source/MutableTools/Private/MuT/NodeMeshTransform.cpp b/Engine/Plugins/Mutable/Source/MutableTools/Private/MuT/NodeMeshTransform.cpp
new file mode 100644
--- /dev/null
+++ b/Engine/Plugins/Mutable/Source/MutableTools/Private/MuT/NodeMeshTransform.cpp
@@ -0,0 +1,60 @@
+// Copyright Epic Games, Inc. All Rights Reserved.
+
+
+#include "MuT/NodeMeshTransform.h"
+
+#include "HAL/PlatformString.h"
+#include "Misc/AssertionMacros.h"
+#include "MuR/MutableMath.h"
+#include "MuT/NodeLayout.h"
+#include "MuT/NodeMeshPrivate.h"
+#include "MuT/NodeMeshTransformPrivate.h"
+#include "MuT/NodePrivate.h"
+
+
+namespace mu
+{
+
+	//---------------------------------------------------------------------------------------------
+	// Static initialisation
+	//---------------------------------------------------------------------------------------------
+    FNodeType NodeMeshTransform::Private::s_type = FNodeType(Node::EType::MeshTransform, NodeMesh::GetStaticType() );
+
+
+	//---------------------------------------------------------------------------------------------
+	//!
+	//---------------------------------------------------------------------------------------------
+
+    MUTABLE_IMPLEMENT_NODE( NodeMeshTransform )
+
+
+	//---------------------------------------------------------------------------------------------
+	// Own Interface
+	//---------------------------------------------------------------------------------------------
+    NodeMeshPtr NodeMeshTransform::GetSource() const
+	{
+		return m_pD->Source;
+	}
+
+
+    //---------------------------------------------------------------------------------------------
+    void NodeMeshTransform::SetSource( NodeMesh* p )
+    {
+        m_pD->Source = p;
+    }
+
+
+    //---------------------------------------------------------------------------------------------
+    void NodeMeshTransform::SetTransform( const FMatrix44f& Value )
+    {
+        m_pD->Transform = Value;
+    }
+
+
+    //---------------------------------------------------------------------------------------------
+	const FMatrix44f& NodeMeshTransform::GetTransform() const
+    {
+        return m_pD->Transform;
+    }
+
+}
diff --git a/Engine/Plugins/Mutable/Source/MutableTools/Private/MuT/NodeMeshTransformPrivate.h b/Engine/Plugins/Mutable/Source/MutableTools/Private/MuT/NodeMeshTransformPrivate.h
new file mode 100644
--- /dev/null
+++ b/Engine/Plugins/Mutable/Source/MutableTools/Private/MuT/NodeMeshTransformPrivate.h
@@ -0,0 +1,26 @@
+// Copyright Epic Games, Inc. All Rights Reserved.
+
+#pragma once
+
+#include "MuT/NodeMeshPrivate.h"
+#include "MuT/NodeMeshTransform.h"
+#include "MuR/MutableMath.h"
+
+
+namespace mu
+{
+
+
+    class NodeMeshTransform::Private : public NodeMesh::Private
+	{
+	public:
+
+		static FNodeType s_type;
+
+		NodeMeshPtr Source;
+		FMatrix44f Transform;
+
+	};
+
+
+}
diff --git a/Engine/Plugins/Mutable/Source/MutableTools/Private/MuT/NodeMeshVariation.cpp b/Engine/Plugins/Mutable/Source/MutableTools/Private/MuT/NodeMeshVariation.cpp
new file mode 100644
--- /dev/null
+++ b/Engine/Plugins/Mutable/Source/MutableTools/Private/MuT/NodeMeshVariation.cpp
@@ -0,0 +1,75 @@
+// Copyright Epic Games, Inc. All Rights Reserved.
+
+#include "MuT/NodeMeshVariation.h"
+
+#include "Misc/AssertionMacros.h"
+#include "MuT/NodeMeshVariationPrivate.h"
+#include "MuT/NodePrivate.h"
+
+
+namespace mu
+{
+
+    //---------------------------------------------------------------------------------------------
+    // Static initialisation
+    //---------------------------------------------------------------------------------------------
+    FNodeType NodeMeshVariation::Private::s_type = FNodeType(Node::EType::MeshVariation, NodeMesh::GetStaticType() );
+
+
+    //---------------------------------------------------------------------------------------------
+    //---------------------------------------------------------------------------------------------
+    //---------------------------------------------------------------------------------------------
+
+    MUTABLE_IMPLEMENT_NODE( NodeMeshVariation )
+
+
+    //---------------------------------------------------------------------------------------------
+    // Own Interface
+    //---------------------------------------------------------------------------------------------
+    void NodeMeshVariation::SetDefaultMesh( NodeMesh* p ) 
+	{ 
+		m_pD->m_defaultMesh = p; 
+	}
+
+
+    //---------------------------------------------------------------------------------------------
+    int NodeMeshVariation::GetVariationCount() const 
+	{ 
+		return  m_pD->m_variations.Num(); 
+	}
+
+
+    //---------------------------------------------------------------------------------------------
+    void NodeMeshVariation::SetVariationCount( int num )
+    {
+        check( num >= 0 );
+        m_pD->m_variations.SetNum( num );
+    }
+
+    //---------------------------------------------------------------------------------------------
+    void NodeMeshVariation::SetVariationTag( int index, const char* strTag )
+    {
+        check( index >= 0 && index < m_pD->m_variations.Num() );
+        check( strTag );
+
+        if ( strTag )
+        {
+            m_pD->m_variations[index].m_tag = strTag;
+        }
+        else
+        {
+            m_pD->m_variations[index].m_tag = "";
+        }
+    }
+
+
+    //---------------------------------------------------------------------------------------------
+    void NodeMeshVariation::SetVariationMesh( int index, NodeMesh* pNode )
+    {
+        check( index >= 0 && index < m_pD->m_variations.Num() );
+
+        m_pD->m_variations[index].m_mesh = pNode;
+    }
+
+
+} // namespace mu
diff --git a/Engine/Plugins/Mutable/Source/MutableTools/Private/MuT/NodeMeshVariationPrivate.h b/Engine/Plugins/Mutable/Source/MutableTools/Private/MuT/NodeMeshVariationPrivate.h
new file mode 100644
--- /dev/null
+++ b/Engine/Plugins/Mutable/Source/MutableTools/Private/MuT/NodeMeshVariationPrivate.h
@@ -0,0 +1,32 @@
+// Copyright Epic Games, Inc. All Rights Reserved.
+
+#pragma once
+
+#include "MuT/NodePrivate.h"
+#include "MuT/NodeMeshVariation.h"
+
+namespace mu
+{
+
+    //---------------------------------------------------------------------------------------------
+    //!
+    //---------------------------------------------------------------------------------------------
+    class NodeMeshVariation::Private : public Node::Private
+    {
+    public:
+        Private() {}
+
+        static FNodeType s_type;
+
+        Ptr<NodeMesh> m_defaultMesh;
+
+        struct FVariation
+        {
+			Ptr<NodeMesh> m_mesh;
+            FString m_tag;
+        };
+
+        TArray<FVariation> m_variations;
+    };
+
+} // namespace mu
diff --git a/Engine/Plugins/Mutable/Source/MutableTools/Private/MuT/NodeModifierMeshClipMorphPlane.cpp b/Engine/Plugins/Mutable/Source/MutableTools/Private/MuT/NodeModifierMeshClipMorphPlane.cpp
--- a/Engine/Plugins/Mutable/Source/MutableTools/Private/MuT/NodeModifierMeshClipMorphPlane.cpp
+++ b/Engine/Plugins/Mutable/Source/MutableTools/Private/MuT/NodeModifierMeshClipMorphPlane.cpp
@@ -1,12 +1,19 @@
 // Copyright Epic Games, Inc. All Rights Reserved.
 
+
 #include "MuT/NodeModifierMeshClipMorphPlane.h"
 
+#include "Misc/AssertionMacros.h"
 #include "MuR/MutableMath.h"
+#include "MuT/NodePrivate.h"
+
 
 namespace mu
 {
 
+	//---------------------------------------------------------------------------------------------
+	// Own Interface
+	//---------------------------------------------------------------------------------------------
     void NodeModifierMeshClipMorphPlane::SetPlane(FVector3f Center, FVector3f Normal)
 	{
 		Parameters.Origin = Center;
@@ -14,13 +21,14 @@ namespace mu
 	}
 
 
+	//---------------------------------------------------------------------------------------------
 	void NodeModifierMeshClipMorphPlane::SetParams(float dist, float factor)
 	{
 		Parameters.DistanceToPlane = dist;
 		Parameters.LinearityFactor = factor;
 	}
 
-
+	//---------------------------------------------------------------------------------------------
 	void NodeModifierMeshClipMorphPlane::SetMorphEllipse(float radius1, float radius2, float rotation)
 	{
 		Parameters.Radius1 = radius1;
@@ -28,7 +36,7 @@ namespace mu
 		Parameters.Rotation = rotation;
 	}
 
-
+	//---------------------------------------------------------------------------------------------
 	void NodeModifierMeshClipMorphPlane::SetVertexSelectionBox(float centerX, float centerY, float centerZ, float radiusX, float radiusY, float radiusZ)
 	{
 		Parameters.VertexSelectionType = EClipVertexSelectionType::Shape;
@@ -36,7 +44,7 @@ namespace mu
 		Parameters.SelectionBoxRadius = FVector3f(radiusX, radiusY, radiusZ);
 	}
 
-
+	//---------------------------------------------------------------------------------------------
 	void NodeModifierMeshClipMorphPlane::SetVertexSelectionBone(const FBoneName& BoneId, float maxEffectRadius)
 	{
 		Parameters.VertexSelectionType = EClipVertexSelectionType::BoneHierarchy;
@@ -45,3 +53,4 @@ namespace mu
 	}
 
 }
+
diff --git a/Engine/Plugins/Mutable/Source/MutableTools/Private/MuT/NodeObject.cpp b/Engine/Plugins/Mutable/Source/MutableTools/Private/MuT/NodeObject.cpp
new file mode 100644
--- /dev/null
+++ b/Engine/Plugins/Mutable/Source/MutableTools/Private/MuT/NodeObject.cpp
@@ -0,0 +1,30 @@
+// Copyright Epic Games, Inc. All Rights Reserved.
+
+#include "MuT/NodeObject.h"
+
+#include "Misc/AssertionMacros.h"
+
+
+namespace mu
+{
+
+	FNodeType s_nodeObjectType = FNodeType(Node::EType::Object, Node::GetStaticType() );
+
+
+	//---------------------------------------------------------------------------------------------
+	const FNodeType* NodeObject::GetType() const
+	{
+		return GetStaticType();
+	}
+
+
+	//---------------------------------------------------------------------------------------------
+	const FNodeType* NodeObject::GetStaticType()
+	{
+		return &s_nodeObjectType;
+	}
+
+
+}
+
+
diff --git a/Engine/Plugins/Mutable/Source/MutableTools/Private/MuT/NodeObjectGroup.cpp b/Engine/Plugins/Mutable/Source/MutableTools/Private/MuT/NodeObjectGroup.cpp
--- a/Engine/Plugins/Mutable/Source/MutableTools/Private/MuT/NodeObjectGroup.cpp
+++ b/Engine/Plugins/Mutable/Source/MutableTools/Private/MuT/NodeObjectGroup.cpp
@@ -5,32 +5,108 @@
 #include "HAL/PlatformCrt.h"
 #include "Misc/AssertionMacros.h"
 #include "MuT/NodeLayout.h"
+#include "MuT/NodeObjectGroupPrivate.h"
+#include "MuT/NodeObjectPrivate.h"
+#include "MuT/NodePrivate.h"
 
 
 namespace mu
 {
+	MUTABLE_IMPLEMENT_ENUM_SERIALISABLE(NodeObjectGroup::CHILD_SELECTION)
 
+
+	//---------------------------------------------------------------------------------------------
+	// Static initialisation
+	//---------------------------------------------------------------------------------------------
+	FNodeType NodeObjectGroup::Private::s_type = FNodeType(Node::EType::ObjectGroup, NodeObject::GetStaticType() );
+
+
+	//---------------------------------------------------------------------------------------------
+	//!
+	//---------------------------------------------------------------------------------------------
+
+	MUTABLE_IMPLEMENT_NODE( NodeObjectGroup )
+
+
+	//---------------------------------------------------------------------------------------------
+	// Own Interface
+	//---------------------------------------------------------------------------------------------
 	const FString& NodeObjectGroup::GetName() const
 	{
-		return Name;
+		return m_pD->Name;
 	}
 
 
-	void NodeObjectGroup::SetName( const FString& InName )
+	//---------------------------------------------------------------------------------------------
+	void NodeObjectGroup::SetName( const FString& Name )
 	{
-		Name = InName;
+		m_pD->Name = Name;
 	}
 
 
 	const FString& NodeObjectGroup::GetUid() const
 	{
-		return Uid;
+		return m_pD->Uid;
+	}
+
+
+	//---------------------------------------------------------------------------------------------
+	void NodeObjectGroup::SetUid( const FString& Uid )
+	{
+		m_pD->Uid = Uid;
+	}
+
+
+	//---------------------------------------------------------------------------------------------
+	NodeObjectGroup::CHILD_SELECTION NodeObjectGroup::GetSelectionType() const
+	{
+		return m_pD->Type;
+	}
+
+
+	//---------------------------------------------------------------------------------------------
+	void NodeObjectGroup::SetSelectionType( CHILD_SELECTION t )
+	{
+		m_pD->Type = t;
+	}
+
+
+	//---------------------------------------------------------------------------------------------
+	int NodeObjectGroup::GetChildCount() const
+	{
+		return m_pD->m_children.Num();
+	}
+
+
+	//---------------------------------------------------------------------------------------------
+	void NodeObjectGroup::SetChildCount( int num )
+	{
+		check( num >=0 );
+		m_pD->m_children.SetNum( num );
+	}
+
+
+	//---------------------------------------------------------------------------------------------
+	NodeObjectPtr NodeObjectGroup::GetChild( int index ) const
+	{
+		check( index >=0 && index < m_pD->m_children.Num() );
+
+		return m_pD->m_children[ index ].get();
+	}
+
+
+	//---------------------------------------------------------------------------------------------
+	void NodeObjectGroup::SetChild( int index, NodeObjectPtr pObject )
+	{
+		check( index >=0 && index < m_pD->m_children.Num() );
+
+		m_pD->m_children[ index ] = pObject;
 	}
 
 
-	void NodeObjectGroup::SetUid( const FString& InUid )
+	void NodeObjectGroup::SetDefaultValue(int32 Value)
 	{
-		Uid = InUid;
+		m_pD->DefaultValue = Value;
 	}
 
 }
diff --git a/Engine/Plugins/Mutable/Source/MutableTools/Private/MuT/NodeObjectNew.cpp b/Engine/Plugins/Mutable/Source/MutableTools/Private/MuT/NodeObjectNew.cpp
--- a/Engine/Plugins/Mutable/Source/MutableTools/Private/MuT/NodeObjectNew.cpp
+++ b/Engine/Plugins/Mutable/Source/MutableTools/Private/MuT/NodeObjectNew.cpp
@@ -10,12 +10,16 @@
 #include "MuT/NodeLOD.h"
 #include "MuT/NodeMesh.h"
 #include "MuT/NodeMeshConstant.h"
+#include "MuT/NodeMeshPrivate.h"
+#include "MuT/NodePrivate.h"
 #include "MuT/NodeSurface.h"
 #include "MuT/NodeSurfaceNew.h"
 
 
 namespace mu
 {
+	//---------------------------------------------------------------------------------------------
+	FNodeType NodeObjectNew::StaticType = FNodeType(Node::EType::ObjectNew, NodeObject::GetStaticType() );
 
 
 	const FString& NodeObjectNew::GetName() const
@@ -79,6 +83,18 @@ namespace mu
 	}
 
 
+	void NodeObjectNew::RemoveStateParam( int32 s, const FString& param )
+	{
+		check( s>=0 && s<GetStateCount() );
+
+		int32 it = States[s].RuntimeParams.Find( param );
+		if ( it != INDEX_NONE )
+		{
+			States[s].RuntimeParams.RemoveAt( it );
+		}
+	}
+
+
     void NodeObjectNew::SetStateProperties( int32 StateIndex, 
 		ETextureCompressionStrategy TextureCompressionStrategy, 
 		bool bOnlyFirstLOD, 
diff --git a/Engine/Plugins/Mutable/Source/MutableTools/Private/MuT/NodeProjector.cpp b/Engine/Plugins/Mutable/Source/MutableTools/Private/MuT/NodeProjector.cpp
new file mode 100644
--- /dev/null
+++ b/Engine/Plugins/Mutable/Source/MutableTools/Private/MuT/NodeProjector.cpp
@@ -0,0 +1,30 @@
+// Copyright Epic Games, Inc. All Rights Reserved.
+
+#include "MuT/NodeProjector.h"
+
+#include "Misc/AssertionMacros.h"
+
+
+namespace mu
+{
+
+	static FNodeType s_nodeProjectorType = FNodeType(Node::EType::Projector, Node::GetStaticType() );
+
+
+	//---------------------------------------------------------------------------------------------
+	const FNodeType* NodeProjector::GetType() const
+	{
+		return GetStaticType();
+	}
+
+
+	//---------------------------------------------------------------------------------------------
+	const FNodeType* NodeProjector::GetStaticType()
+	{
+		return &s_nodeProjectorType;
+	}
+
+
+}
+
+
diff --git a/Engine/Plugins/Mutable/Source/MutableTools/Private/MuT/NodeProjectorConstant.cpp b/Engine/Plugins/Mutable/Source/MutableTools/Private/MuT/NodeProjectorConstant.cpp
--- a/Engine/Plugins/Mutable/Source/MutableTools/Private/MuT/NodeProjectorConstant.cpp
+++ b/Engine/Plugins/Mutable/Source/MutableTools/Private/MuT/NodeProjectorConstant.cpp
@@ -1,14 +1,34 @@
 // Copyright Epic Games, Inc. All Rights Reserved.
 
+
 #include "Misc/AssertionMacros.h"
 #include "MuR/MutableMath.h"
 #include "MuR/Parameters.h"
 #include "MuT/Node.h"
+#include "MuT/NodePrivate.h"
 #include "MuT/NodeProjector.h"
+#include "MuT/NodeProjectorPrivate.h"
+
 
 namespace mu
 {
 
+	//---------------------------------------------------------------------------------------------
+	// Static initialisation
+	//---------------------------------------------------------------------------------------------
+	FNodeType NodeProjectorConstant::Private::s_type = FNodeType(Node::EType::ProjectorConstant, NodeProjector::GetStaticType() );
+
+
+	//---------------------------------------------------------------------------------------------
+	//!
+	//---------------------------------------------------------------------------------------------
+
+	MUTABLE_IMPLEMENT_NODE( NodeProjectorConstant )
+
+
+	//---------------------------------------------------------------------------------------------
+	// Own Interface
+	//---------------------------------------------------------------------------------------------
     void NodeProjectorConstant::GetValue( EProjectorType* OutType,
 		FVector3f* OutPos,
 		FVector3f* OutDir,
@@ -16,25 +36,29 @@ namespace mu
 		FVector3f* OutScale,
 		float* OutProjectionAngle ) const
 	{
-        if (OutType) *OutType = Type;
-        if (OutPos) *OutPos = Position;
-        if (OutDir) *OutDir = Direction;
-        if (OutUp) *OutUp = Up;
-        if (OutScale) *OutScale = Scale;
-        if (OutProjectionAngle) *OutProjectionAngle = ProjectionAngle;
+        if (OutType) *OutType = m_pD->Type;
+
+        if (OutPos) *OutPos = m_pD->m_position;
+        if (OutDir) *OutDir = m_pD->m_direction;
+        if (OutUp) *OutUp = m_pD->m_up;
+        if (OutScale) *OutScale = m_pD->m_scale;
+
+        if (OutProjectionAngle) *OutProjectionAngle = m_pD->m_projectionAngle;
     }
 
 
-	void NodeProjectorConstant::SetValue( EProjectorType type, FVector3f pos, FVector3f dir, FVector3f up, FVector3f scale, float projectionAngle )
+	//---------------------------------------------------------------------------------------------
+    void NodeProjectorConstant::SetValue( EProjectorType type, FVector3f pos, FVector3f dir, FVector3f up, FVector3f scale, float projectionAngle )
 	{
-        Type = type;
-		Position = pos;
-		Direction = dir;
-		Up = up;
-        Scale = scale;
-        ProjectionAngle = projectionAngle;
+        m_pD->Type = type;
+		m_pD->m_position = pos;
+		m_pD->m_direction = dir;
+		m_pD->m_up = up;
+        m_pD->m_scale = scale;
+        m_pD->m_projectionAngle = projectionAngle;
 	}
 
+
 }
 
 
diff --git a/Engine/Plugins/Mutable/Source/MutableTools/Private/MuT/NodeProjectorParameter.cpp b/Engine/Plugins/Mutable/Source/MutableTools/Private/MuT/NodeProjectorParameter.cpp
--- a/Engine/Plugins/Mutable/Source/MutableTools/Private/MuT/NodeProjectorParameter.cpp
+++ b/Engine/Plugins/Mutable/Source/MutableTools/Private/MuT/NodeProjectorParameter.cpp
@@ -7,30 +7,50 @@
 #include "MuR/Ptr.h"
 #include "MuR/RefCounted.h"
 #include "MuT/Node.h"
+#include "MuT/NodePrivate.h"
 #include "MuT/NodeProjector.h"
+#include "MuT/NodeProjectorPrivate.h"
 #include "MuT/NodeRange.h"
 
 
 namespace mu
 {
-	void NodeProjectorParameter::SetName( const FString& InName )
+
+	//---------------------------------------------------------------------------------------------
+	// Static initialisation
+	//---------------------------------------------------------------------------------------------
+	FNodeType NodeProjectorParameter::Private::s_type = FNodeType(Node::EType::ProjectorParameter, NodeProjector::GetStaticType() );
+
+
+	//---------------------------------------------------------------------------------------------
+	//!
+	//---------------------------------------------------------------------------------------------
+
+	MUTABLE_IMPLEMENT_NODE( NodeProjectorParameter )
+
+
+	//---------------------------------------------------------------------------------------------
+	// Own Interface
+	//---------------------------------------------------------------------------------------------
+	void NodeProjectorParameter::SetName( const FString& strName )
 	{
-		Name = InName;
+		m_pD->Name = strName;
 	}
 
 
 	const FString& NodeProjectorParameter::GetUid() const
 	{
-		return UID;
+		return m_pD->UID;
 	}
 
 
-	void NodeProjectorParameter::SetUid( const FString& InUid)
+	void NodeProjectorParameter::SetUid( const FString& strUid )
 	{
-		UID = InUid;
+		m_pD->UID = strUid;
 	}
 
 
+	//---------------------------------------------------------------------------------------------
     void NodeProjectorParameter::SetDefaultValue( EProjectorType type,
 		FVector3f pos,
 		FVector3f dir,
@@ -38,28 +58,30 @@ namespace mu
 		FVector3f scale,
 		float projectionAngle )
 	{
-        Type = type;
-        Position = pos;
-		Direction = dir;
-		Up = up;
-        Scale = scale;
-        ProjectionAngle = projectionAngle;
+        m_pD->Type = type;
+        m_pD->m_position = pos;
+		m_pD->m_direction = dir;
+		m_pD->m_up = up;
+        m_pD->m_scale = scale;
+        m_pD->m_projectionAngle = projectionAngle;
     }
 
 
-    void NodeProjectorParameter::SetRangeCount( int32 i )
+    //---------------------------------------------------------------------------------------------
+    void NodeProjectorParameter::SetRangeCount( int i )
     {
         check(i>=0);
-        Ranges.SetNum(i);
+        m_pD->Ranges.SetNum(i);
     }
 
 
-    void NodeProjectorParameter::SetRange( int32 i, Ptr<NodeRange> InRange )
+    //---------------------------------------------------------------------------------------------
+    void NodeProjectorParameter::SetRange( int i, NodeRangePtr pRange )
     {
-        check( i>=0 && i<Ranges.Num() );
-        if ( i>=0 && i<Ranges.Num() )
+        check( i>=0 && i<m_pD->Ranges.Num() );
+        if ( i>=0 && i<m_pD->Ranges.Num() )
         {
-            Ranges[i] = InRange;
+            m_pD->Ranges[i] = pRange;
         }
     }
 
diff --git a/Engine/Plugins/Mutable/Source/MutableTools/Private/MuT/NodeRange.cpp b/Engine/Plugins/Mutable/Source/MutableTools/Private/MuT/NodeRange.cpp
new file mode 100644
--- /dev/null
+++ b/Engine/Plugins/Mutable/Source/MutableTools/Private/MuT/NodeRange.cpp
@@ -0,0 +1,33 @@
+// Copyright Epic Games, Inc. All Rights Reserved.
+
+
+#include "MuT/NodeRange.h"
+
+#include "Misc/AssertionMacros.h"
+
+namespace mu
+{
+
+	//---------------------------------------------------------------------------------------------
+	// Static initialisation
+	//---------------------------------------------------------------------------------------------
+    static FNodeType s_nodeRangeType = 	FNodeType(Node::EType::Range, Node::GetStaticType() );
+
+
+	//---------------------------------------------------------------------------------------------
+    const FNodeType* NodeRange::GetType() const
+	{
+		return GetStaticType();
+	}
+
+
+	//---------------------------------------------------------------------------------------------
+    const FNodeType* NodeRange::GetStaticType()
+	{
+        return &s_nodeRangeType;
+	}
+
+
+}
+
+
diff --git a/Engine/Plugins/Mutable/Source/MutableTools/Private/MuT/NodeRangeFromScalar.cpp b/Engine/Plugins/Mutable/Source/MutableTools/Private/MuT/NodeRangeFromScalar.cpp
new file mode 100644
--- /dev/null
+++ b/Engine/Plugins/Mutable/Source/MutableTools/Private/MuT/NodeRangeFromScalar.cpp
@@ -0,0 +1,59 @@
+// Copyright Epic Games, Inc. All Rights Reserved.
+
+
+#include "MuT/NodeRangeFromScalar.h"
+
+#include "Misc/AssertionMacros.h"
+#include "MuR/RefCounted.h"
+#include "MuT/NodePrivate.h"
+#include "MuT/NodeRangeFromScalarPrivate.h"
+#include "MuT/NodeScalar.h"
+
+
+namespace mu
+{
+
+	//---------------------------------------------------------------------------------------------
+	// Static initialisation
+	//---------------------------------------------------------------------------------------------
+    FNodeType NodeRangeFromScalar::Private::s_type = FNodeType(Node::EType::RangeFromScalar, NodeRange::GetStaticType() );
+
+
+	//---------------------------------------------------------------------------------------------
+	//!
+	//---------------------------------------------------------------------------------------------
+
+    MUTABLE_IMPLEMENT_NODE( NodeRangeFromScalar )
+
+
+	//---------------------------------------------------------------------------------------------
+	// Own Interface
+	//---------------------------------------------------------------------------------------------
+	Ptr<NodeScalar> NodeRangeFromScalar::GetSize() const
+	{
+        return m_pD->m_pSize.get();
+	}
+
+
+	//---------------------------------------------------------------------------------------------
+    void NodeRangeFromScalar::SetSize( const Ptr<NodeScalar>& pNode )
+	{
+        m_pD->m_pSize = pNode;
+	}
+
+
+	//---------------------------------------------------------------------------------------------
+    const FString& NodeRangeFromScalar::GetName() const
+	{
+        return m_pD->Name;
+	}
+
+
+	//---------------------------------------------------------------------------------------------
+    void NodeRangeFromScalar::SetName( const FString& strName )
+	{
+        m_pD->Name = strName;
+	}
+
+}
+
diff --git a/Engine/Plugins/Mutable/Source/MutableTools/Private/MuT/NodeRangeFromScalarPrivate.h b/Engine/Plugins/Mutable/Source/MutableTools/Private/MuT/NodeRangeFromScalarPrivate.h
new file mode 100644
--- /dev/null
+++ b/Engine/Plugins/Mutable/Source/MutableTools/Private/MuT/NodeRangeFromScalarPrivate.h
@@ -0,0 +1,26 @@
+// Copyright Epic Games, Inc. All Rights Reserved.
+
+#pragma once
+
+#include "MuT/NodeRangePrivate.h"
+#include "MuT/NodeRangeFromScalar.h"
+#include "MuT/NodeRange.h"
+#include "MuT/NodeScalar.h"
+
+
+namespace mu
+{
+
+
+    class NodeRangeFromScalar::Private : public NodeRange::Private
+	{
+	public:
+
+		static FNodeType s_type;
+
+        NodeScalarPtr m_pSize;
+        FString Name;
+	};
+
+
+}
diff --git a/Engine/Plugins/Mutable/Source/MutableTools/Private/MuT/NodeRangePrivate.h b/Engine/Plugins/Mutable/Source/MutableTools/Private/MuT/NodeRangePrivate.h
new file mode 100644
--- /dev/null
+++ b/Engine/Plugins/Mutable/Source/MutableTools/Private/MuT/NodeRangePrivate.h
@@ -0,0 +1,17 @@
+// Copyright Epic Games, Inc. All Rights Reserved.
+
+#pragma once
+
+#include "MuT/NodePrivate.h"
+#include "MuT/NodeRange.h"
+
+
+namespace mu
+{
+
+
+    class NodeRange::Private : public Node::Private
+	{
+	};
+
+}
diff --git a/Engine/Plugins/Mutable/Source/MutableTools/Private/MuT/NodeScalar.cpp b/Engine/Plugins/Mutable/Source/MutableTools/Private/MuT/NodeScalar.cpp
new file mode 100644
--- /dev/null
+++ b/Engine/Plugins/Mutable/Source/MutableTools/Private/MuT/NodeScalar.cpp
@@ -0,0 +1,31 @@
+// Copyright Epic Games, Inc. All Rights Reserved.
+
+
+#include "MuT/NodeScalar.h"
+
+#include "Misc/AssertionMacros.h"
+
+
+namespace mu
+{
+
+	static FNodeType s_nodeScalarType = FNodeType(Node::EType::Scalar, Node::GetStaticType() );
+
+
+	//---------------------------------------------------------------------------------------------
+	const FNodeType* NodeScalar::GetType() const
+	{
+		return GetStaticType();
+	}
+
+
+	//---------------------------------------------------------------------------------------------
+	const FNodeType* NodeScalar::GetStaticType()
+	{
+		return &s_nodeScalarType;
+	}
+
+
+}
+
+
diff --git a/Engine/Plugins/Mutable/Source/MutableTools/Private/MuT/NodeScalarArithmeticOperation.cpp b/Engine/Plugins/Mutable/Source/MutableTools/Private/MuT/NodeScalarArithmeticOperation.cpp
new file mode 100644
--- /dev/null
+++ b/Engine/Plugins/Mutable/Source/MutableTools/Private/MuT/NodeScalarArithmeticOperation.cpp
@@ -0,0 +1,84 @@
+// Copyright Epic Games, Inc. All Rights Reserved.
+
+
+#include "MuT/NodeScalarArithmeticOperation.h"
+
+#include "Misc/AssertionMacros.h"
+#include "MuT/NodePrivate.h"
+#include "MuT/NodeScalar.h"
+#include "MuT/NodeScalarArithmeticOperationPrivate.h"
+#include "MuR/Serialisation.h"
+
+
+namespace mu
+{
+    MUTABLE_IMPLEMENT_ENUM_SERIALISABLE(NodeScalarArithmeticOperation::OPERATION)
+
+	
+	//---------------------------------------------------------------------------------------------
+	// Static initialisation
+	//---------------------------------------------------------------------------------------------
+    FNodeType NodeScalarArithmeticOperation::Private::s_type = FNodeType(Node::EType::ScalarArithmeticOperation, NodeScalar::GetStaticType() );
+
+
+	//---------------------------------------------------------------------------------------------
+	//!
+	//---------------------------------------------------------------------------------------------
+
+    MUTABLE_IMPLEMENT_NODE( NodeScalarArithmeticOperation )
+
+
+    //---------------------------------------------------------------------------------------------
+    const char* NodeScalarArithmeticOperation::s_opTypeName[] =
+    {
+        "Add",
+        "Subtract",
+        "Multiply",
+        "Divide"
+    };
+
+
+	//---------------------------------------------------------------------------------------------
+	// Own Interface
+	//---------------------------------------------------------------------------------------------
+    NodeScalarArithmeticOperation::OPERATION NodeScalarArithmeticOperation::GetOperation() const
+	{
+		return m_pD->m_operation;
+	}
+
+
+	//---------------------------------------------------------------------------------------------
+    void NodeScalarArithmeticOperation::SetOperation(NodeScalarArithmeticOperation::OPERATION o)
+	{
+		m_pD->m_operation = o;
+	}
+
+
+    NodeScalarPtr NodeScalarArithmeticOperation::GetA() const
+	{
+		return m_pD->m_pA.get();
+	}
+
+
+	//---------------------------------------------------------------------------------------------
+    void NodeScalarArithmeticOperation::SetA(NodeScalarPtr pNode)
+	{
+		m_pD->m_pA = pNode;
+	}
+
+
+	//---------------------------------------------------------------------------------------------
+    NodeScalarPtr NodeScalarArithmeticOperation::GetB() const
+	{
+		return m_pD->m_pB.get();
+	}
+
+
+	//---------------------------------------------------------------------------------------------
+    void NodeScalarArithmeticOperation::SetB( NodeScalarPtr pNode )
+	{
+		m_pD->m_pB = pNode;
+	}
+
+}
+
diff --git a/Engine/Plugins/Mutable/Source/MutableTools/Private/MuT/NodeScalarArithmeticOperationPrivate.h b/Engine/Plugins/Mutable/Source/MutableTools/Private/MuT/NodeScalarArithmeticOperationPrivate.h
new file mode 100644
--- /dev/null
+++ b/Engine/Plugins/Mutable/Source/MutableTools/Private/MuT/NodeScalarArithmeticOperationPrivate.h
@@ -0,0 +1,25 @@
+// Copyright Epic Games, Inc. All Rights Reserved.
+
+#pragma once
+
+#include "MuT/NodeScalarPrivate.h"
+#include "MuT/NodeScalarArithmeticOperation.h"
+
+namespace mu
+{
+    MUTABLE_DEFINE_ENUM_SERIALISABLE(NodeScalarArithmeticOperation::OPERATION)
+
+    class NodeScalarArithmeticOperation::Private : public NodeScalar::Private
+	{
+	public:
+
+		static FNodeType s_type;
+
+        NodeScalarArithmeticOperation::OPERATION m_operation;
+        NodeScalarPtr m_pA;
+        NodeScalarPtr m_pB;
+
+	};
+
+
+}
diff --git a/Engine/Plugins/Mutable/Source/MutableTools/Private/MuT/NodeScalarConstant.cpp b/Engine/Plugins/Mutable/Source/MutableTools/Private/MuT/NodeScalarConstant.cpp
new file mode 100644
--- /dev/null
+++ b/Engine/Plugins/Mutable/Source/MutableTools/Private/MuT/NodeScalarConstant.cpp
@@ -0,0 +1,43 @@
+// Copyright Epic Games, Inc. All Rights Reserved.
+
+
+#include "MuT/NodeScalarConstant.h"
+
+#include "Misc/AssertionMacros.h"
+#include "MuT/NodePrivate.h"
+#include "MuT/NodeScalarConstantPrivate.h"
+
+
+namespace mu
+{
+
+	//---------------------------------------------------------------------------------------------
+	// Static initialisation
+	//---------------------------------------------------------------------------------------------
+	FNodeType NodeScalarConstant::Private::s_type = FNodeType(Node::EType::ScalarConstant, NodeScalar::GetStaticType() );
+
+
+	//---------------------------------------------------------------------------------------------
+	//!
+	//---------------------------------------------------------------------------------------------
+
+	MUTABLE_IMPLEMENT_NODE( NodeScalarConstant )
+
+
+	//---------------------------------------------------------------------------------------------
+	// Own Interface
+	//---------------------------------------------------------------------------------------------
+	float NodeScalarConstant::GetValue() const
+	{
+		return m_pD->Value;
+	}
+
+	//---------------------------------------------------------------------------------------------
+	void NodeScalarConstant::SetValue( float v )
+	{
+		m_pD->Value = v;
+	}
+
+
+}
+
diff --git a/Engine/Plugins/Mutable/Source/MutableTools/Private/MuT/NodeScalarConstantPrivate.h b/Engine/Plugins/Mutable/Source/MutableTools/Private/MuT/NodeScalarConstantPrivate.h
new file mode 100644
--- /dev/null
+++ b/Engine/Plugins/Mutable/Source/MutableTools/Private/MuT/NodeScalarConstantPrivate.h
@@ -0,0 +1,23 @@
+// Copyright Epic Games, Inc. All Rights Reserved.
+
+#pragma once
+
+#include "MuT/NodeScalarPrivate.h"
+
+#include "MuT/NodeScalarConstant.h"
+#include "MuT/NodeImage.h"
+
+namespace mu
+{
+
+	class NodeScalarConstant::Private : public NodeScalar::Private
+	{
+	public:
+
+		static FNodeType s_type;
+
+		float Value;
+	};
+
+
+}
diff --git a/Engine/Plugins/Mutable/Source/MutableTools/Private/MuT/NodeScalarCurve.cpp b/Engine/Plugins/Mutable/Source/MutableTools/Private/MuT/NodeScalarCurve.cpp
new file mode 100644
--- /dev/null
+++ b/Engine/Plugins/Mutable/Source/MutableTools/Private/MuT/NodeScalarCurve.cpp
@@ -0,0 +1,21 @@
+// Copyright Epic Games, Inc. All Rights Reserved.
+
+
+#include "MuT/NodeScalarCurve.h"
+
+#include "Misc/AssertionMacros.h"
+#include "MuR/ParametersPrivate.h"
+#include "MuT/NodePrivate.h"
+#include "MuT/NodeScalar.h"
+#include "MuT/NodeScalarCurvePrivate.h"
+
+
+namespace mu
+{
+
+    FNodeType NodeScalarCurve::Private::s_type = FNodeType(Node::EType::ScalarCurve, NodeScalar::GetStaticType() );
+
+    MUTABLE_IMPLEMENT_NODE( NodeScalarCurve )
+
+}
+
diff --git a/Engine/Plugins/Mutable/Source/MutableTools/Private/MuT/NodeScalarEnumParameter.cpp b/Engine/Plugins/Mutable/Source/MutableTools/Private/MuT/NodeScalarEnumParameter.cpp
new file mode 100644
--- /dev/null
+++ b/Engine/Plugins/Mutable/Source/MutableTools/Private/MuT/NodeScalarEnumParameter.cpp
@@ -0,0 +1,111 @@
+// Copyright Epic Games, Inc. All Rights Reserved.
+
+
+#include "MuT/NodeScalarEnumParameter.h"
+
+#include "Misc/AssertionMacros.h"
+#include "MuT/NodePrivate.h"
+#include "MuT/NodeScalarEnumParameterPrivate.h"
+
+
+namespace mu
+{
+
+	//---------------------------------------------------------------------------------------------
+	// Static initialisation
+	//---------------------------------------------------------------------------------------------
+	FNodeType NodeScalarEnumParameter::Private::s_type = FNodeType(Node::EType::ScalarEnumParameter, NodeScalar::GetStaticType() );
+
+
+	//---------------------------------------------------------------------------------------------
+	//!
+	//---------------------------------------------------------------------------------------------
+
+	MUTABLE_IMPLEMENT_NODE( NodeScalarEnumParameter )
+
+
+	//---------------------------------------------------------------------------------------------
+	const FString& NodeScalarEnumParameter::GetName() const
+	{
+		return m_pD->Name;
+	}
+
+
+	//---------------------------------------------------------------------------------------------
+	void NodeScalarEnumParameter::SetName( const FString& strName )
+	{
+		m_pD->Name = strName;
+	}
+
+
+	const FString& NodeScalarEnumParameter::GetUid() const
+	{
+		return m_pD->UID;
+	}
+
+
+	//---------------------------------------------------------------------------------------------
+	void NodeScalarEnumParameter::SetUid( const FString& strUid )
+	{
+		m_pD->UID = strUid;
+	}
+
+
+	//---------------------------------------------------------------------------------------------
+	int NodeScalarEnumParameter::GetDefaultValueIndex() const
+	{
+		return m_pD->DefaultValue;
+	}
+
+
+	//---------------------------------------------------------------------------------------------
+	void NodeScalarEnumParameter::SetDefaultValueIndex( int i )
+	{
+		m_pD->DefaultValue = i;
+	}
+
+
+	//---------------------------------------------------------------------------------------------
+	void NodeScalarEnumParameter::SetValueCount( int i )
+	{
+		m_pD->Options.SetNum(i);
+	}
+
+
+	//---------------------------------------------------------------------------------------------
+	int NodeScalarEnumParameter::GetValueCount() const
+	{
+		return (int)m_pD->Options.Num();
+	}
+
+
+	//---------------------------------------------------------------------------------------------
+	void NodeScalarEnumParameter::SetValue( int i, float value, const FString& strName )
+	{
+		check( i>=0 && i<(int)m_pD->Options.Num() );
+		m_pD->Options[i].name = strName;
+		m_pD->Options[i].value = value;
+	}
+
+
+    //---------------------------------------------------------------------------------------------
+    void NodeScalarEnumParameter::SetRangeCount( int i )
+    {
+        check(i>=0);
+        m_pD->Ranges.SetNum(i);
+    }
+
+
+    //---------------------------------------------------------------------------------------------
+    void NodeScalarEnumParameter::SetRange( int i, NodeRangePtr pRange )
+    {
+        check( i>=0 && i<int(m_pD->Ranges.Num()) );
+        if ( i>=0 && i<int(m_pD->Ranges.Num()) )
+        {
+            m_pD->Ranges[i] = pRange;
+        }
+    }
+
+}
+
+
diff --git a/Engine/Plugins/Mutable/Source/MutableTools/Private/MuT/NodeScalarParameter.cpp b/Engine/Plugins/Mutable/Source/MutableTools/Private/MuT/NodeScalarParameter.cpp
new file mode 100644
--- /dev/null
+++ b/Engine/Plugins/Mutable/Source/MutableTools/Private/MuT/NodeScalarParameter.cpp
@@ -0,0 +1,71 @@
+// Copyright Epic Games, Inc. All Rights Reserved.
+
+
+#include "MuT/NodeScalarParameter.h"
+
+#include "Misc/AssertionMacros.h"
+#include "MuT/NodeImage.h"
+#include "MuT/NodePrivate.h"
+#include "MuT/NodeRange.h"
+#include "MuT/NodeScalarParameterPrivate.h"
+
+
+namespace mu
+{
+
+	//---------------------------------------------------------------------------------------------
+	// Static initialisation
+	//---------------------------------------------------------------------------------------------
+	FNodeType NodeScalarParameter::Private::s_type = FNodeType(Node::EType::ScalarParameter, NodeScalar::GetStaticType() );
+
+
+	//---------------------------------------------------------------------------------------------
+	//!
+	//---------------------------------------------------------------------------------------------
+
+	MUTABLE_IMPLEMENT_NODE( NodeScalarParameter )
+
+
+	//---------------------------------------------------------------------------------------------
+	void NodeScalarParameter::SetName( const FString& Name )
+	{
+		m_pD->Name = Name;
+	}
+
+
+	//---------------------------------------------------------------------------------------------
+	void NodeScalarParameter::SetUid( const FString& Uid )
+	{
+		m_pD->UID = Uid;
+	}
+
+
+	//---------------------------------------------------------------------------------------------
+	void NodeScalarParameter::SetDefaultValue( float v )
+	{
+		m_pD->DefaultValue = v;
+	}
+
+
+    //---------------------------------------------------------------------------------------------
+    void NodeScalarParameter::SetRangeCount( int i )
+    {
+        check(i>=0);
+        m_pD->Ranges.SetNum(i);
+    }
+
+
+    //---------------------------------------------------------------------------------------------
+    void NodeScalarParameter::SetRange( int i, NodeRangePtr pRange )
+    {
+        check( i>=0 && i<int(m_pD->Ranges.Num()) );
+        if ( i>=0 && i<int(m_pD->Ranges.Num()) )
+        {
+            m_pD->Ranges[i] = pRange;
+        }
+    }
+
+
+}
+
+
diff --git a/Engine/Plugins/Mutable/Source/MutableTools/Private/MuT/NodeScalarTable.cpp b/Engine/Plugins/Mutable/Source/MutableTools/Private/MuT/NodeScalarTable.cpp
new file mode 100644
--- /dev/null
+++ b/Engine/Plugins/Mutable/Source/MutableTools/Private/MuT/NodeScalarTable.cpp
@@ -0,0 +1,51 @@
+// Copyright Epic Games, Inc. All Rights Reserved.
+
+
+#include "MuT/NodeScalarTable.h"
+
+#include "Misc/AssertionMacros.h"
+#include "MuT/NodePrivate.h"
+#include "MuT/Table.h"
+
+
+namespace mu
+{
+
+
+	//---------------------------------------------------------------------------------------------
+	// Static initialisation
+	//---------------------------------------------------------------------------------------------
+	FNodeType NodeScalarTable::StaticType = FNodeType(Node::EType::ScalarTable, NodeScalar::GetStaticType() );
+
+
+
+	//---------------------------------------------------------------------------------------------
+	// Own Interface
+	//---------------------------------------------------------------------------------------------
+	void NodeScalarTable::SetColumn( const FString& strName )
+	{
+		ColumnName = strName;
+	}
+
+
+	//---------------------------------------------------------------------------------------------
+	void NodeScalarTable::SetParameterName( const FString& strName )
+	{
+		ParameterName = strName;
+	}
+
+
+	//---------------------------------------------------------------------------------------------
+	void NodeScalarTable::SetNoneOption(bool bAddNoneOption)
+	{
+		bNoneOption = bAddNoneOption;
+	}
+
+	//---------------------------------------------------------------------------------------------
+	void NodeScalarTable::SetDefaultRowName(const FString& RowName)
+	{
+		DefaultRowName = RowName;
+	}
+}
+
+
diff --git a/Engine/Plugins/Mutable/Source/MutableTools/Private/MuT/NodeScalarVariation.cpp b/Engine/Plugins/Mutable/Source/MutableTools/Private/MuT/NodeScalarVariation.cpp
new file mode 100644
--- /dev/null
+++ b/Engine/Plugins/Mutable/Source/MutableTools/Private/MuT/NodeScalarVariation.cpp
@@ -0,0 +1,59 @@
+// Copyright Epic Games, Inc. All Rights Reserved.
+
+
+#include "MuT/NodeScalarVariation.h"
+
+#include "Misc/AssertionMacros.h"
+#include "MuT/NodePrivate.h"
+#include "MuT/NodeScalarVariationPrivate.h"
+
+
+namespace mu
+{
+
+    //---------------------------------------------------------------------------------------------
+    // Static initialisation
+    //---------------------------------------------------------------------------------------------
+    FNodeType NodeScalarVariation::Private::s_type = FNodeType(Node::EType::ScalarVariation, NodeScalar::GetStaticType() );
+
+
+    //---------------------------------------------------------------------------------------------
+    //---------------------------------------------------------------------------------------------
+    //---------------------------------------------------------------------------------------------
+
+    MUTABLE_IMPLEMENT_NODE( NodeScalarVariation )
+
+
+    //---------------------------------------------------------------------------------------------
+    // Own Interface
+    //---------------------------------------------------------------------------------------------
+    void NodeScalarVariation::SetDefaultScalar( NodeScalar* p ) 
+	{ 
+		m_pD->m_defaultScalar = p; 
+	}
+
+
+	void NodeScalarVariation::SetVariationCount( int32 num )
+    {
+        check( num >= 0 );
+        m_pD->m_variations.SetNum( num );
+    }
+
+
+	void NodeScalarVariation::SetVariationTag( int32 index, const FString& Tag )
+    {
+        check( index >= 0 && index < m_pD->m_variations.Num() );
+
+		m_pD->m_variations[index].m_tag = Tag;
+    }
+
+
+    void NodeScalarVariation::SetVariationScalar( int32 index, NodeScalar* pNode )
+    {
+        check( index >= 0 && index < m_pD->m_variations.Num() );
+
+        m_pD->m_variations[index].m_scalar = pNode;
+    }
+
+
+} // namespace mu
diff --git a/Engine/Plugins/Mutable/Source/MutableTools/Private/MuT/NodeScalarVariationPrivate.h b/Engine/Plugins/Mutable/Source/MutableTools/Private/MuT/NodeScalarVariationPrivate.h
new file mode 100644
--- /dev/null
+++ b/Engine/Plugins/Mutable/Source/MutableTools/Private/MuT/NodeScalarVariationPrivate.h
@@ -0,0 +1,32 @@
+// Copyright Epic Games, Inc. All Rights Reserved.
+
+#pragma once
+
+#include "MuT/NodePrivate.h"
+#include "MuT/NodeScalarVariation.h"
+
+namespace mu
+{
+
+    //---------------------------------------------------------------------------------------------
+    //!
+    //---------------------------------------------------------------------------------------------
+    class NodeScalarVariation::Private : public Node::Private
+    {
+    public:
+
+        static FNodeType s_type;
+
+        NodeScalarPtr m_defaultScalar;
+
+        struct FVariation
+        {
+            NodeScalarPtr m_scalar;
+            FString m_tag;
+        };
+
+        TArray<FVariation> m_variations;
+
+    };
+
+} // namespace mu
diff --git a/Engine/Plugins/Mutable/Source/MutableTools/Private/MuT/NodeString.cpp b/Engine/Plugins/Mutable/Source/MutableTools/Private/MuT/NodeString.cpp
new file mode 100644
--- /dev/null
+++ b/Engine/Plugins/Mutable/Source/MutableTools/Private/MuT/NodeString.cpp
@@ -0,0 +1,34 @@
+// Copyright Epic Games, Inc. All Rights Reserved.
+
+
+#include "MuT/NodeString.h"
+
+#include "Misc/AssertionMacros.h"
+
+
+namespace mu
+{
+
+	//---------------------------------------------------------------------------------------------
+	// Static initialisation
+	//---------------------------------------------------------------------------------------------
+	static FNodeType s_nodeStringType = FNodeType(Node::EType::String, Node::GetStaticType() );
+
+
+	//---------------------------------------------------------------------------------------------
+	const FNodeType* NodeString::GetType() const
+	{
+		return GetStaticType();
+	}
+
+
+	//---------------------------------------------------------------------------------------------
+	const FNodeType* NodeString::GetStaticType()
+	{
+		return &s_nodeStringType;
+	}
+
+
+}
+
+
diff --git a/Engine/Plugins/Mutable/Source/MutableTools/Private/MuT/NodeStringConstant.cpp b/Engine/Plugins/Mutable/Source/MutableTools/Private/MuT/NodeStringConstant.cpp
new file mode 100644
--- /dev/null
+++ b/Engine/Plugins/Mutable/Source/MutableTools/Private/MuT/NodeStringConstant.cpp
@@ -0,0 +1,37 @@
+// Copyright Epic Games, Inc. All Rights Reserved.
+
+
+#include "MuT/NodeStringConstant.h"
+
+#include "Misc/AssertionMacros.h"
+#include "MuT/NodePrivate.h"
+#include "MuT/NodeStringConstantPrivate.h"
+
+
+namespace mu
+{
+
+	//---------------------------------------------------------------------------------------------
+	// Static initialisation
+	//---------------------------------------------------------------------------------------------
+	FNodeType NodeStringConstant::Private::s_type = FNodeType(Node::EType::StringConstant, NodeString::GetStaticType() );
+
+
+	//---------------------------------------------------------------------------------------------
+	//!
+	//---------------------------------------------------------------------------------------------
+
+	MUTABLE_IMPLEMENT_NODE( NodeStringConstant )
+
+
+	//---------------------------------------------------------------------------------------------
+	// Own Interface
+	//---------------------------------------------------------------------------------------------
+	void NodeStringConstant::SetValue( const FString& v )
+	{
+		m_pD->Value = v;
+	}
+
+
+}
+
diff --git a/Engine/Plugins/Mutable/Source/MutableTools/Private/MuT/NodeStringConstantPrivate.h b/Engine/Plugins/Mutable/Source/MutableTools/Private/MuT/NodeStringConstantPrivate.h
new file mode 100644
--- /dev/null
+++ b/Engine/Plugins/Mutable/Source/MutableTools/Private/MuT/NodeStringConstantPrivate.h
@@ -0,0 +1,25 @@
+// Copyright Epic Games, Inc. All Rights Reserved.
+
+#pragma once
+
+#include "MuT/NodeStringPrivate.h"
+
+#include "MuT/NodeStringConstant.h"
+#include "MuT/NodeImage.h"
+
+namespace mu
+{
+
+
+	class NodeStringConstant::Private : public NodeString::Private
+	{
+	public:
+
+		static FNodeType s_type;
+
+		FString Value;
+
+	};
+
+
+}
diff --git a/Engine/Plugins/Mutable/Source/MutableTools/Private/MuT/NodeStringParameter.cpp b/Engine/Plugins/Mutable/Source/MutableTools/Private/MuT/NodeStringParameter.cpp
new file mode 100644
--- /dev/null
+++ b/Engine/Plugins/Mutable/Source/MutableTools/Private/MuT/NodeStringParameter.cpp
@@ -0,0 +1,71 @@
+// Copyright Epic Games, Inc. All Rights Reserved.
+
+
+#include "MuT/NodeStringParameter.h"
+
+#include "Misc/AssertionMacros.h"
+#include "MuT/NodeImage.h"
+#include "MuT/NodePrivate.h"
+#include "MuT/NodeRange.h"
+#include "MuT/NodeStringParameterPrivate.h"
+
+
+namespace mu
+{
+
+	//---------------------------------------------------------------------------------------------
+	// Static initialisation
+	//---------------------------------------------------------------------------------------------
+	FNodeType NodeStringParameter::Private::s_type = FNodeType(Node::EType::StringParameter, NodeString::GetStaticType() );
+
+
+	//---------------------------------------------------------------------------------------------
+	//!
+	//---------------------------------------------------------------------------------------------
+
+	MUTABLE_IMPLEMENT_NODE( NodeStringParameter )
+
+
+	//---------------------------------------------------------------------------------------------
+	void NodeStringParameter::SetName( const FString& Name )
+	{
+		m_pD->Name = Name;
+	}
+
+
+	//---------------------------------------------------------------------------------------------
+	void NodeStringParameter::SetUid( const FString& Uid )
+	{
+		m_pD->UID = Uid;
+	}
+
+
+	//---------------------------------------------------------------------------------------------
+	void NodeStringParameter::SetDefaultValue( const FString& v )
+	{
+		m_pD->DefaultValue = v;
+	}
+
+
+    //---------------------------------------------------------------------------------------------
+    void NodeStringParameter::SetRangeCount( int i )
+    {
+        check(i>=0);
+        m_pD->Ranges.SetNum(i);
+    }
+
+
+    //---------------------------------------------------------------------------------------------
+    void NodeStringParameter::SetRange( int i, NodeRangePtr pRange )
+    {
+        check( i>=0 && i<int(m_pD->Ranges.Num()) );
+        if ( i>=0 && i<int(m_pD->Ranges.Num()) )
+        {
+            m_pD->Ranges[i] = pRange;
+        }
+    }
+
+
+}
+
+
diff --git a/Engine/Plugins/Mutable/Source/MutableTools/Private/MuT/NodeStringParameterPrivate.h b/Engine/Plugins/Mutable/Source/MutableTools/Private/MuT/NodeStringParameterPrivate.h
new file mode 100644
--- /dev/null
+++ b/Engine/Plugins/Mutable/Source/MutableTools/Private/MuT/NodeStringParameterPrivate.h
@@ -0,0 +1,31 @@
+// Copyright Epic Games, Inc. All Rights Reserved.
+
+#pragma once
+
+#include "MuT/NodeStringPrivate.h"
+
+#include "MuT/NodeStringParameter.h"
+#include "MuT/NodeImage.h"
+#include "MuT/NodeRange.h"
+
+
+namespace mu
+{
+
+	class NodeStringParameter::Private : public NodeString::Private
+	{
+	public:
+
+		static FNodeType s_type;
+
+		FString DefaultValue;
+		FString Name;
+		FString UID;
+
+        TArray<Ptr<NodeImage>> m_additionalImages;
+
+		TArray<Ptr<NodeRange>> Ranges;
+
+	};
+
+}
diff --git a/Engine/Plugins/Mutable/Source/MutableTools/Private/MuT/NodeStringPrivate.h b/Engine/Plugins/Mutable/Source/MutableTools/Private/MuT/NodeStringPrivate.h
new file mode 100644
--- /dev/null
+++ b/Engine/Plugins/Mutable/Source/MutableTools/Private/MuT/NodeStringPrivate.h
@@ -0,0 +1,20 @@
+// Copyright Epic Games, Inc. All Rights Reserved.
+
+#pragma once
+
+#include "MuT/NodePrivate.h"
+
+#include "MuT/NodeString.h"
+#include "MuR/ParametersPrivate.h"
+
+
+namespace mu
+{
+
+
+	class NodeString::Private : public Node::Private
+	{
+	};
+
+
+}
diff --git a/Engine/Plugins/Mutable/Source/MutableTools/Public/MuT/Compiler.h b/Engine/Plugins/Mutable/Source/MutableTools/Public/MuT/Compiler.h
--- a/Engine/Plugins/Mutable/Source/MutableTools/Public/MuT/Compiler.h
+++ b/Engine/Plugins/Mutable/Source/MutableTools/Public/MuT/Compiler.h
@@ -118,7 +118,7 @@ namespace mu
 
     protected:
 
-        /** Forbidden. Manage with the Ptr<> template. */
+        //! Forbidden. Manage with the Ptr<> template.
         ~CompilerOptions() override;
 
     private:
@@ -166,7 +166,7 @@ namespace mu
 
     protected:
 
-        /** Forbidden. Manage with the Ptr<> template. */
+        //! Forbidden. Manage with the Ptr<> template.
         ~Compiler();
 
     private:
diff --git a/Engine/Plugins/Mutable/Source/MutableTools/Public/MuT/NodeBool.h b/Engine/Plugins/Mutable/Source/MutableTools/Public/MuT/NodeBool.h
--- a/Engine/Plugins/Mutable/Source/MutableTools/Public/MuT/NodeBool.h
+++ b/Engine/Plugins/Mutable/Source/MutableTools/Public/MuT/NodeBool.h
@@ -6,134 +6,221 @@
 #include "MuR/Ptr.h"
 #include "MuR/RefCounted.h"
 #include "MuT/Node.h"
-#include "MuT/NodeRange.h"
 
 
 namespace mu
 {
+	class NodeRange;
 
-	/** Base class of any node that outputs a Bool value. */
+	//---------------------------------------------------------------------------------------------
+    //! %Base class of any node that outputs a Bool value.
+	//! \ingroup model
+	//---------------------------------------------------------------------------------------------
 	class MUTABLETOOLS_API NodeBool : public Node
 	{
 	public:
 
+		//-----------------------------------------------------------------------------------------
 		// Node Interface
-		virtual const FNodeType* GetType() const override { return &StaticType; }
-		static const FNodeType* GetStaticType() { return &StaticType; }
+		//-----------------------------------------------------------------------------------------
 
-	protected:
+        const FNodeType* GetType() const override;
+		static const FNodeType* GetStaticType();
 
-		/** Forbidden. Manage with the Ptr<> template. */
-		inline ~NodeBool() {}
 
-	private:
+		//-----------------------------------------------------------------------------------------
+		// Interface pattern
+		//-----------------------------------------------------------------------------------------
+		class Private;
 
-		static FNodeType StaticType;
+	protected:
 
+		//! Forbidden. Manage with the Ptr<> template.
+		inline ~NodeBool() {}
 	};
 
 
-	/** Node returning a Bool constant value. */
+	//---------------------------------------------------------------------------------------------
+	//! Node returning a Bool constant value.
+	//! \ingroup model
+	//---------------------------------------------------------------------------------------------
 	class MUTABLETOOLS_API NodeBoolConstant : public NodeBool
 	{
 	public:
 
-		bool Value = false;
+		NodeBoolConstant();
 
-	public:
 
+		//-----------------------------------------------------------------------------------------
 		// Node Interface
-		virtual const FNodeType* GetType() const override { return &StaticType; }
-		static const FNodeType* GetStaticType() { return &StaticType; }
+		//-----------------------------------------------------------------------------------------
+
+        const FNodeType* GetType() const override;
+		static const FNodeType* GetStaticType();
+
+		//-----------------------------------------------------------------------------------------
+		// Own Interface
+		//-----------------------------------------------------------------------------------------
+
+		//! Get the value to be returned by the node.
+		bool GetValue() const;
+
+		//! Set the value to be returned by the node.
+		void SetValue( bool v );
+
+		//-----------------------------------------------------------------------------------------
+		// Interface pattern
+		//-----------------------------------------------------------------------------------------
+		class Private;
+		Private* GetPrivate() const;
 
 	protected:
 
-		/** Forbidden. Manage with the Ptr<> template. */
-		~NodeBoolConstant() {}
+		//! Forbidden.
+		~NodeBoolConstant();
 
 	private:
 
-		static FNodeType StaticType;
+		Private* m_pD;
 
 	};
 
 
-	/** Node that defines a Bool model parameter. */
+	//---------------------------------------------------------------------------------------------
+	//! Node that defines a Bool model parameter.
+	//! \ingroup model
+	//---------------------------------------------------------------------------------------------
 	class MUTABLETOOLS_API NodeBoolParameter : public NodeBool
 	{
 	public:
 
-		bool DefaultValue;
-		FString Name;
-		FString UID;
+		NodeBoolParameter();
 
-		TArray<Ptr<NodeRange>> Ranges;
+		//-----------------------------------------------------------------------------------------
+		// Node Interface
+		//-----------------------------------------------------------------------------------------
 
-	public:
+        const FNodeType* GetType() const override;
+		static const FNodeType* GetStaticType();
 
-		// Node Interface
-		virtual const FNodeType* GetType() const override { return &StaticType; }
-		static const FNodeType* GetStaticType() { return &StaticType; }
+		//-----------------------------------------------------------------------------------------
+		// Own Interface
+		//-----------------------------------------------------------------------------------------
+
+		//! Set the name of the parameter. It will be exposed in the final compiled data.
+		void SetName( const FString& );
+
+		//! Get the default value of the parameter.
+		void SetDefaultValue( bool v );
+
+        //! Set the number of ranges (dimensions) for this parameter.
+        //! By default a parameter has 0 ranges, meaning it only has one value.
+        void SetRangeCount( int i );
+        void SetRange( int i, Ptr<NodeRange> );
+
+		//-----------------------------------------------------------------------------------------
+		// Interface pattern
+		//-----------------------------------------------------------------------------------------
+		class Private;
+		Private* GetPrivate() const;
 
 	protected:
 
-		/** Forbidden. Manage with the Ptr<> template. */
-		~NodeBoolParameter() {}
+		//! Forbidden. 
+		~NodeBoolParameter();
 
 	private:
 
-		static FNodeType StaticType;
+		Private* m_pD;
 
 	};
 
 
-	/** Node that returns the oposite of the input value. */
+	//---------------------------------------------------------------------------------------------
+	//! Node that returns the oposite of the input value.
+	//! \ingroup model
+	//---------------------------------------------------------------------------------------------
 	class MUTABLETOOLS_API NodeBoolNot : public NodeBool
 	{
 	public:
 
-		Ptr<NodeBool> Source;
-
-	public:
+		NodeBoolNot();
 
+		//-----------------------------------------------------------------------------------------
 		// Node Interface
-		virtual const FNodeType* GetType() const override { return &StaticType; }
-		static const FNodeType* GetStaticType() { return &StaticType; }
+		//-----------------------------------------------------------------------------------------
+
+        const FNodeType* GetType() const override;
+		static const FNodeType* GetStaticType();
+
+		//-----------------------------------------------------------------------------------------
+		// Own Interface
+		//-----------------------------------------------------------------------------------------
+		
+		//! Input
+		Ptr<NodeBool> GetInput() const;
+		void SetInput( Ptr<NodeBool> );
+
+		//-----------------------------------------------------------------------------------------
+		// Interface pattern
+		//-----------------------------------------------------------------------------------------
+		class Private;
+		Private* GetPrivate() const;
 
 	protected:
 
-		/** Forbidden. Manage with the Ptr<> template. */
-		~NodeBoolNot() {}
+		//! Forbidden. 
+		~NodeBoolNot();
 
 	private:
 
-		static FNodeType StaticType;
+		Private* m_pD;
 
 	};
 
 
-	/** */
+	//---------------------------------------------------------------------------------------------
+	//! 
+	//! \ingroup model
+	//---------------------------------------------------------------------------------------------
 	class MUTABLETOOLS_API NodeBoolAnd : public NodeBool
 	{
 	public:
 
-		Ptr<NodeBool> A;
-		Ptr<NodeBool> B;
-
-	public:
+		NodeBoolAnd();
 
+		//-----------------------------------------------------------------------------------------
 		// Node Interface
-		virtual const FNodeType* GetType() const override { return &StaticType; }
-		static const FNodeType* GetStaticType() { return &StaticType; }
+		//-----------------------------------------------------------------------------------------
+
+        const FNodeType* GetType() const override;
+		static const FNodeType* GetStaticType();
+
+		//-----------------------------------------------------------------------------------------
+		// Own Interface
+		//-----------------------------------------------------------------------------------------
+		
+		//! Inputs
+		Ptr<NodeBool> GetA() const;
+		void SetA(Ptr<NodeBool>);
+
+		Ptr<NodeBool> GetB() const;
+		void SetB(Ptr<NodeBool>);
+
+		//-----------------------------------------------------------------------------------------
+		// Interface pattern
+		//-----------------------------------------------------------------------------------------
+		class Private;
+		Private* GetPrivate() const;
 
 	protected:
 
-		/** Forbidden. Manage with the Ptr<> template. */
-		~NodeBoolAnd() {}
+		//! Forbidden. 
+		~NodeBoolAnd();
 
 	private:
 
-		static FNodeType StaticType;
+		Private* m_pD;
 
 	};
 
diff --git a/Engine/Plugins/Mutable/Source/MutableTools/Public/MuT/NodeColour.h b/Engine/Plugins/Mutable/Source/MutableTools/Public/MuT/NodeColour.h
--- a/Engine/Plugins/Mutable/Source/MutableTools/Public/MuT/NodeColour.h
+++ b/Engine/Plugins/Mutable/Source/MutableTools/Public/MuT/NodeColour.h
@@ -23,7 +23,7 @@ namespace mu
 
     protected:
 
-        /** Forbidden. Manage with the Ptr<> template. */
+        //! Forbidden. Manage with the Ptr<> template.
         inline ~NodeColour() {}
 
 	private:
diff --git a/Engine/Plugins/Mutable/Source/MutableTools/Public/MuT/NodeColourArithmeticOperation.h b/Engine/Plugins/Mutable/Source/MutableTools/Public/MuT/NodeColourArithmeticOperation.h
--- a/Engine/Plugins/Mutable/Source/MutableTools/Public/MuT/NodeColourArithmeticOperation.h
+++ b/Engine/Plugins/Mutable/Source/MutableTools/Public/MuT/NodeColourArithmeticOperation.h
@@ -37,7 +37,7 @@ namespace mu
 
 	protected:
 
-		/** Forbidden. Manage with the Ptr<> template. */
+		//! Forbidden. Manage with the Ptr<> template.
 		~NodeColourArithmeticOperation() {}
 
 	private:
diff --git a/Engine/Plugins/Mutable/Source/MutableTools/Public/MuT/NodeComponentEdit.h b/Engine/Plugins/Mutable/Source/MutableTools/Public/MuT/NodeComponentEdit.h
--- a/Engine/Plugins/Mutable/Source/MutableTools/Public/MuT/NodeComponentEdit.h
+++ b/Engine/Plugins/Mutable/Source/MutableTools/Public/MuT/NodeComponentEdit.h
@@ -7,37 +7,50 @@
 #include "MuT/Node.h"
 #include "MuT/NodeComponent.h"
 
+
 namespace mu
 {
 
-	/** This node modifies a node of the parent object of the object that this node belongs to.
-	* It allows to extend, cut and morph the parent component's meshes.
-	* It also allows to patch the parent component's textures.
-	*/
+	//! This node modifies a node of the parent object of the object that this node belongs to.
+	//! It allows to extend, cut and morph the parent component's meshes.
+	//! It also allows to patch the parent component's textures.
 	class MUTABLETOOLS_API NodeComponentEdit : public NodeComponent
 	{
 	public:
 
-		Ptr<NodeComponent> Parent;
-
-	public:
+		//-----------------------------------------------------------------------------------------
+        // Node interface
+		//-----------------------------------------------------------------------------------------
 
-		// Node interface
 		virtual const FNodeType* GetType() const override { return GetStaticType(); }
 		static const FNodeType* GetStaticType() { return &StaticType; }
 
+        //-----------------------------------------------------------------------------------------
         // NodeComponent interface
+        //-----------------------------------------------------------------------------------------
 		virtual const class NodeComponentNew* GetParentComponentNew() const override;
 
+		//-----------------------------------------------------------------------------------------
+        // Own interface
+		//-----------------------------------------------------------------------------------------
+
+		//! Set the parent component to modify. It should be a node from the parent objects of this node's object.
+        void SetParent( NodeComponent* );
+        NodeComponent* GetParent() const;
+
 	protected:
 
-		/** Forbidden. Manage with the Ptr<> template. */
+		//! Forbidden. Manage with the Ptr<> template.
 		~NodeComponentEdit() {}
 
 	private:
 
 		static FNodeType StaticType;
 
+		Ptr<NodeComponent> m_pParent;
+
 	};
 
+
+
 }
diff --git a/Engine/Plugins/Mutable/Source/MutableTools/Public/MuT/NodeImage.h b/Engine/Plugins/Mutable/Source/MutableTools/Public/MuT/NodeImage.h
--- a/Engine/Plugins/Mutable/Source/MutableTools/Public/MuT/NodeImage.h
+++ b/Engine/Plugins/Mutable/Source/MutableTools/Public/MuT/NodeImage.h
@@ -116,19 +116,25 @@ namespace mu
 	{
 	public:
 
+		//-----------------------------------------------------------------------------------------
 		// Node Interface
-		virtual const FNodeType* GetType() const override { return &StaticType; }
-		static const FNodeType* GetStaticType() { return &StaticType; }
+		//-----------------------------------------------------------------------------------------
+
+		const FNodeType* GetType() const override;
+		static const FNodeType* GetStaticType();
+
+		//-----------------------------------------------------------------------------------------
+		// Interface pattern
+		//-----------------------------------------------------------------------------------------
+		class Private;
 
 	protected:
 
-		/** Forbidden. Manage with the Ptr<> template. */
+		//! Forbidden. Manage with the Ptr<> template.
 		inline ~NodeImage() {}
 
-	private:
+	};
 
-		static FNodeType StaticType;
 
-	};
 
 }
diff --git a/Engine/Plugins/Mutable/Source/MutableTools/Public/MuT/NodeImageBinarise.h b/Engine/Plugins/Mutable/Source/MutableTools/Public/MuT/NodeImageBinarise.h
--- a/Engine/Plugins/Mutable/Source/MutableTools/Public/MuT/NodeImageBinarise.h
+++ b/Engine/Plugins/Mutable/Source/MutableTools/Public/MuT/NodeImageBinarise.h
@@ -6,34 +6,64 @@
 #include "MuR/RefCounted.h"
 #include "MuT/Node.h"
 #include "MuT/NodeImage.h"
-#include "MuT/NodeScalar.h"
+
 
 namespace mu
 {
 
-	/** Node that multiplies the colors of an image, channel by channel. */
+	// Forward definitions
+	class NodeScalar;
+	typedef Ptr<NodeScalar> NodeScalarPtr;
+	typedef Ptr<const NodeScalar> NodeScalarPtrConst;
+
+	class NodeImageBinarise;
+	typedef Ptr<NodeImageBinarise> NodeImageBinarisePtr;
+	typedef Ptr<const NodeImageBinarise> NodeImageBinarisePtrConst;
+
+
+	//! Node that multiplies the colors of an image, channel by channel.
+	//! \ingroup model
 	class MUTABLETOOLS_API NodeImageBinarise : public NodeImage
 	{
 	public:
 
-		Ptr<NodeImage> Base;
-		Ptr<NodeScalar> Threshold;
-
-	public:
+		NodeImageBinarise();
 
+		//-----------------------------------------------------------------------------------------
 		// Node Interface
-		virtual const FNodeType* GetType() const override { return &StaticType; }
-		static const FNodeType* GetStaticType() { return &StaticType; }
+		//-----------------------------------------------------------------------------------------
+
+        const FNodeType* GetType() const override;
+		static const FNodeType* GetStaticType();
+
+		//-----------------------------------------------------------------------------------------
+		// Own Interface
+		//-----------------------------------------------------------------------------------------
+
+		//! Base image to multiply.
+		NodeImagePtr GetBase() const;
+		void SetBase( NodeImagePtr );
+
+		//! Colour to multiply by the image.
+		NodeScalarPtr GetThreshold() const;
+		void SetThreshold( NodeScalarPtr );
+
+		//-----------------------------------------------------------------------------------------
+		// Interface pattern
+		//-----------------------------------------------------------------------------------------
+		class Private;
+		Private* GetPrivate() const;
 
 	protected:
 
-		/** Forbidden. Manage with the Ptr<> template. */
-		~NodeImageBinarise() {}
+		//! Forbidden. Manage with the Ptr<> template.
+		~NodeImageBinarise();
 
 	private:
 
-		static FNodeType StaticType;
+		Private* m_pD;
 
 	};
 
+
 }
