[Backout] - CL40275509

diff --git a/Engine/Plugins/Mutable/Source/MutableTools/Public/MuT/NodeImageColourMap.h b/Engine/Plugins/Mutable/Source/MutableTools/Public/MuT/NodeImageColourMap.h
--- a/Engine/Plugins/Mutable/Source/MutableTools/Public/MuT/NodeImageColourMap.h
+++ b/Engine/Plugins/Mutable/Source/MutableTools/Public/MuT/NodeImageColourMap.h
@@ -7,35 +7,66 @@
 #include "MuT/Node.h"
 #include "MuT/NodeImage.h"
 
+
 namespace mu
 {
 
-	/** This node changes the colours of a selectd part of the image, applying a colour map from
-	* conteined in another image.
-	*/
+	// Forward definitions
+	class NodeImageColourMap;
+	typedef Ptr<NodeImageColourMap> NodeImageColourMapPtr;
+	typedef Ptr<const NodeImageColourMap> NodeImageColourMapPtrConst;
+
+
+	//! This node changes the colours of a selectd part of the image, applying a colour map from
+	//! conteined in another image.
+	//! \ingroup model
 	class MUTABLETOOLS_API NodeImageColourMap : public NodeImage
 	{
 	public:
 
-		Ptr<NodeImage> Base;
-		Ptr<NodeImage> Mask;
-		Ptr<NodeImage> Map;
+		NodeImageColourMap();
 
-	public:
 
+		//-----------------------------------------------------------------------------------------
 		// Node Interface
-		virtual const FNodeType* GetType() const override { return &StaticType; }
-		static const FNodeType* GetStaticType() { return &StaticType; }
+		//-----------------------------------------------------------------------------------------
+
+        const FNodeType* GetType() const override;
+		static const FNodeType* GetStaticType();
+
+		//-----------------------------------------------------------------------------------------
+		// Own Interface
+		//-----------------------------------------------------------------------------------------
+
+		//! Get the node generating the base image that will have the other one blended on top.
+		NodeImagePtr GetBase() const;
+		void SetBase( NodeImagePtr );
+
+		//! Get the node generating the mask image controlling the weight of the blend effect.
+		//! \todo: make it optional
+		NodeImagePtr GetMask() const;
+		void SetMask( NodeImagePtr );
+
+		//! Get the node generating the image to blend on the base.
+		NodeImagePtr GetMap() const;
+		void SetMap( NodeImagePtr );
+
+		//-----------------------------------------------------------------------------------------
+		// Interface pattern
+		//-----------------------------------------------------------------------------------------
+		class Private;
+		Private* GetPrivate() const;
 
 	protected:
 
-		/** Forbidden. Manage with the Ptr<> template. */
-		~NodeImageColourMap() {}
+		//! Forbidden. Manage with the Ptr<> template.
+		~NodeImageColourMap();
 
 	private:
 
-		static FNodeType StaticType;
+		Private* m_pD;
 
 	};
 
+
 }
diff --git a/Engine/Plugins/Mutable/Source/MutableTools/Public/MuT/NodeImageConditional.h b/Engine/Plugins/Mutable/Source/MutableTools/Public/MuT/NodeImageConditional.h
--- a/Engine/Plugins/Mutable/Source/MutableTools/Public/MuT/NodeImageConditional.h
+++ b/Engine/Plugins/Mutable/Source/MutableTools/Public/MuT/NodeImageConditional.h
@@ -6,35 +6,65 @@
 #include "MuR/RefCounted.h"
 #include "MuT/Node.h"
 #include "MuT/NodeImage.h"
-#include "MuT/NodeBool.h"
+
 
 namespace mu
 {
 
-	/** This node selects an output image from a set of input images based on a parameter. */
+	// Forward definitions
+    class NodeBool;
+    class NodeImageConditional;
+    typedef Ptr<NodeImageConditional> NodeImageConditionalPtr;
+    typedef Ptr<const NodeImageConditional> NodeImageConditionalPtrConst;
+
+
+	//! This node selects an output image from a set of input images based on a parameter.
+	//! \ingroup model
     class MUTABLETOOLS_API NodeImageConditional : public NodeImage
 	{
 	public:
 
-		Ptr<NodeBool> Parameter;
-		Ptr<NodeImage> True;
-		Ptr<NodeImage> False;
-
-	public:
+        NodeImageConditional();
 
+		//-----------------------------------------------------------------------------------------
 		// Node Interface
-		virtual const FNodeType* GetType() const override { return &StaticType; }
-		static const FNodeType* GetStaticType() { return &StaticType; }
+		//-----------------------------------------------------------------------------------------
+
+        const FNodeType* GetType() const override;
+		static const FNodeType* GetStaticType();
+
+		//-----------------------------------------------------------------------------------------
+		// Own Interface
+		//-----------------------------------------------------------------------------------------
+
+		//! Get the node generating the parameter used to select the option.
+        NodeBool* GetParameter() const;
+        void SetParameter( NodeBool* );
+
+        //! Get the node generating the option image when the parameters is true.
+        NodeImage* GetOptionTrue() const;
+        void SetOptionTrue( NodeImage* );
+
+        NodeImage* GetOptionFalse() const;
+        void SetOptionFalse( NodeImage* );
+
+
+		//-----------------------------------------------------------------------------------------
+		// Interface pattern
+		//-----------------------------------------------------------------------------------------
+		class Private;
+		Private* GetPrivate() const;
 
 	protected:
 
-		/** Forbidden. Manage with the Ptr<> template. */
-		~NodeImageConditional() {}
+		//! Forbidden. Manage with the Ptr<> template.
+        ~NodeImageConditional();
 
 	private:
 
-		static FNodeType StaticType;
+		Private* m_pD;
 
 	};
 
+
 }
diff --git a/Engine/Plugins/Mutable/Source/MutableTools/Public/MuT/NodeImageConstant.h b/Engine/Plugins/Mutable/Source/MutableTools/Public/MuT/NodeImageConstant.h
--- a/Engine/Plugins/Mutable/Source/MutableTools/Public/MuT/NodeImageConstant.h
+++ b/Engine/Plugins/Mutable/Source/MutableTools/Public/MuT/NodeImageConstant.h
@@ -8,6 +8,7 @@
 #include "MuT/Node.h"
 #include "MuT/NodeImage.h"
 
+
 namespace mu
 {
 
@@ -26,11 +27,16 @@ namespace mu
 
 	public:
 
+		//-----------------------------------------------------------------------------------------
 		// Node Interface
+		//-----------------------------------------------------------------------------------------
+
 		virtual const FNodeType* GetType() const override { return GetStaticType(); }
 		static const FNodeType* GetStaticType() { return &StaticType; }
 
+		//-----------------------------------------------------------------------------------------
 		// Own Interface
+		//-----------------------------------------------------------------------------------------
 
         //! Set the image to be output by this node
         void SetValue(TSharedPtr<const FImage>);
@@ -40,7 +46,7 @@ namespace mu
 
 	protected:
 
-		/** Forbidden. Manage with the Ptr<> template. */
+		//! Forbidden. Manage with the Ptr<> template.
 		~NodeImageConstant() {}
 
 	private:
diff --git a/Engine/Plugins/Mutable/Source/MutableTools/Public/MuT/NodeImageFormat.h b/Engine/Plugins/Mutable/Source/MutableTools/Public/MuT/NodeImageFormat.h
--- a/Engine/Plugins/Mutable/Source/MutableTools/Public/MuT/NodeImageFormat.h
+++ b/Engine/Plugins/Mutable/Source/MutableTools/Public/MuT/NodeImageFormat.h
@@ -12,29 +12,56 @@
 namespace mu
 {
 
-	/** Node that composes a new image by gathering pixel data from channels in other images. */
+	// Forward definitions
+	class NodeImageFormat;
+	typedef Ptr<NodeImageFormat> NodeImageFormatPtr;
+	typedef Ptr<const NodeImageFormat> NodeImageFormatPtrConst;
+
+
+	//! Node that composes a new image by gathering pixel data from channels in other images.
+	//! \ingroup model
 	class MUTABLETOOLS_API NodeImageFormat : public NodeImage
 	{
 	public:
 
-		EImageFormat Format = EImageFormat::None;
-		EImageFormat FormatIfAlpha = EImageFormat::None;
-		Ptr<NodeImage> Source;
+		NodeImageFormat();
 
-	public:
 
+		//-----------------------------------------------------------------------------------------
 		// Node Interface
-		virtual const FNodeType* GetType() const override { return &StaticType; }
-		static const FNodeType* GetStaticType() { return &StaticType; }
+		//-----------------------------------------------------------------------------------------
+
+		const FNodeType* GetType() const override;
+		static const FNodeType* GetStaticType();
+
+		//-----------------------------------------------------------------------------------------
+		// Own Interface
+		//-----------------------------------------------------------------------------------------
+
+		//!
+		EImageFormat GetFormat() const;
+        void SetFormat(EImageFormat format, EImageFormat formatIfAlpha = EImageFormat::None );
+
+		//!
+		NodeImagePtr GetSource() const;
+
+		void SetSource( NodeImagePtr );
+
+
+		//-----------------------------------------------------------------------------------------
+		// Interface pattern
+		//-----------------------------------------------------------------------------------------
+		class Private;
+		Private* GetPrivate() const;
 
 	protected:
 
-		/** Forbidden. Manage with the Ptr<> template. */
-		~NodeImageFormat() {}
+		//! Forbidden. Manage with the Ptr<> template.
+		~NodeImageFormat();
 
 	private:
 
-		static FNodeType StaticType;
+		Private* m_pD;
 
 	};
 
diff --git a/Engine/Plugins/Mutable/Source/MutableTools/Public/MuT/NodeImageInterpolate.h b/Engine/Plugins/Mutable/Source/MutableTools/Public/MuT/NodeImageInterpolate.h
--- a/Engine/Plugins/Mutable/Source/MutableTools/Public/MuT/NodeImageInterpolate.h
+++ b/Engine/Plugins/Mutable/Source/MutableTools/Public/MuT/NodeImageInterpolate.h
@@ -6,40 +6,76 @@
 #include "MuR/RefCounted.h"
 #include "MuT/Node.h"
 #include "MuT/NodeImage.h"
-#include "MuT/NodeScalar.h"
+
 
 namespace mu
 {
 
-	/** Node that interpolates linearly between several images based on a weight.
-	* The different input images are uniformly distributed in the 0 to 1 range:
-	* - if two images are set, the first one is 0.0 and the second one is 1.0
-	* - if three images are set, the first one is 0.0 the second one is 0.5 and the third one 1.0
-	* regardless of the input slots used: If B and C are set, but not A, B will be at weight 0
-	* and C will be at weight 1
-	*/
+	// Forward definitions
+	class NodeScalar;
+	typedef Ptr<NodeScalar> NodeScalarPtr;
+	typedef Ptr<const NodeScalar> NodeScalarPtrConst;
+
+	class NodeImageInterpolate;
+	typedef Ptr<NodeImageInterpolate> NodeImageInterpolatePtr;
+	typedef Ptr<const NodeImageInterpolate> NodeImageInterpolatePtrConst;
+
+
+	//! Node that interpolates linearly between several images based on a weight.
+	//! The different input images are uniformly distributed in the 0 to 1 range:
+	//! - if two images are set, the first one is 0.0 and the second one is 1.0
+	//! - if three images are set, the first one is 0.0 the second one is 0.5 and the third one 1.0
+	//! regardless of the input slots used: If B and C are set, but not A, B will be at weight 0
+	//! and C will be at weight 1
+	//! \ingroup model
 	class MUTABLETOOLS_API NodeImageInterpolate : public NodeImage
 	{
 	public:
 
-		Ptr<NodeScalar> Factor;
-		TArray<Ptr<NodeImage>> Targets;
+		NodeImageInterpolate();
 
-	public:
 
+		//-----------------------------------------------------------------------------------------
 		// Node Interface
-		virtual const FNodeType* GetType() const override { return &StaticType; }
-		static const FNodeType* GetStaticType() { return &StaticType; }
+		//-----------------------------------------------------------------------------------------
+
+        const FNodeType* GetType() const override;
+		static const FNodeType* GetStaticType();
+
+		//-----------------------------------------------------------------------------------------
+		// Own Interface
+		//-----------------------------------------------------------------------------------------
+
+		//! Get the node generating the weight of the interpolation.
+		NodeScalarPtr GetFactor() const;
+		void SetFactor( NodeScalarPtr );
+
+		//! Set the number of target images. It will keep the currently set targets and initialise
+		//! the new ones as null.
+		void SetTargetCount( int );
+		int GetTargetCount() const;
+
+		//! Get the node generating the t-th target image of the interpolation.
+		NodeImagePtr GetTarget( int t ) const;
+		void SetTarget( int t, NodeImagePtr );
+
+
+		//-----------------------------------------------------------------------------------------
+		// Interface pattern
+		//-----------------------------------------------------------------------------------------
+		class Private;
+		Private* GetPrivate() const;
 
 	protected:
 
-		/** Forbidden. Manage with the Ptr<> template. */
-		~NodeImageInterpolate() {}
+		//! Forbidden. Manage with the Ptr<> template.
+		~NodeImageInterpolate();
 
 	private:
 
-		static FNodeType StaticType;
+		Private* m_pD;
 
 	};
 
+
 }
diff --git a/Engine/Plugins/Mutable/Source/MutableTools/Public/MuT/NodeImageInvert.h b/Engine/Plugins/Mutable/Source/MutableTools/Public/MuT/NodeImageInvert.h
--- a/Engine/Plugins/Mutable/Source/MutableTools/Public/MuT/NodeImageInvert.h
+++ b/Engine/Plugins/Mutable/Source/MutableTools/Public/MuT/NodeImageInvert.h
@@ -7,31 +7,51 @@
 #include "MuT/Node.h"
 #include "MuT/NodeImage.h"
 
+
 namespace mu
 {
 
-	/** Node that inverts the colors of an image, channel by channel. */
+	// Forward definitions
+	class NodeImageInvert;
+	typedef Ptr<NodeImageInvert> NodeImageInvertPtr;
+	typedef Ptr<const NodeImageInvert> NodeImageInvertPtrConst;
+
+	//! Node that inverts the colors of an image, channel by channel
 	class MUTABLETOOLS_API NodeImageInvert : public NodeImage
 	{
 	public:
 
-		Ptr<NodeImage> Base;
+		NodeImageInvert();
 
-	public:
 
+		//-----------------------------------------------------------------------------------------
 		// Node Interface
-		virtual const FNodeType* GetType() const override { return &StaticType; }
-		static const FNodeType* GetStaticType() { return &StaticType; }
+		//-----------------------------------------------------------------------------------------
+
+		const FNodeType* GetType() const override;
+		static const FNodeType* GetStaticType();
+
+		//-----------------------------------------------------------------------------------------
+		// Own Interface
+		//-----------------------------------------------------------------------------------------
+
+		//! Base image to invert.
+		NodeImagePtr GetBase() const;
+		void SetBase(NodeImagePtr);
+
+		//-----------------------------------------------------------------------------------------
+		// Interface pattern
+		//-----------------------------------------------------------------------------------------
+		class Private;
+		Private* GetPrivate() const;
 
 	protected:
 
-		/** Forbidden. Manage with the Ptr<> template. */
-		~NodeImageInvert() {}
+		//! Forbidden. Manage with the Ptr<> template
+		~NodeImageInvert();
 
 	private:
 
-		static FNodeType StaticType;
-
+		Private* m_pD;
 	};
-
 }
diff --git a/Engine/Plugins/Mutable/Source/MutableTools/Public/MuT/NodeImageLayer.h b/Engine/Plugins/Mutable/Source/MutableTools/Public/MuT/NodeImageLayer.h
--- a/Engine/Plugins/Mutable/Source/MutableTools/Public/MuT/NodeImageLayer.h
+++ b/Engine/Plugins/Mutable/Source/MutableTools/Public/MuT/NodeImageLayer.h
@@ -12,30 +12,61 @@
 namespace mu
 {
 
-	/** This node applies a layer blending effect on a base image using a mask and a blended image. */
+	// Forward definitions
+	class NodeImageLayer;
+	typedef Ptr<NodeImageLayer> NodeImageLayerPtr;
+	typedef Ptr<const NodeImageLayer> NodeImageLayerPtrConst;
+
+
+	//! This node applies a layer blending effect on a base image using a mask and a blended image.
+	//! \ingroup model
 	class MUTABLETOOLS_API NodeImageLayer : public NodeImage
 	{
 	public:
 
-		Ptr<NodeImage> Base;
-		Ptr<NodeImage> Mask;
-		Ptr<NodeImage> Blended;
-		EBlendType Type;
-
-	public:
+		NodeImageLayer();
 
+		//-----------------------------------------------------------------------------------------
 		// Node Interface
-		virtual const FNodeType* GetType() const override { return &StaticType; }
-		static const FNodeType* GetStaticType() { return &StaticType; }
+		//-----------------------------------------------------------------------------------------
+
+        const FNodeType* GetType() const override;
+		static const FNodeType* GetStaticType();
+
+		//-----------------------------------------------------------------------------------------
+		// Own Interface
+		//-----------------------------------------------------------------------------------------
+
+		//! Get the node generating the base image that will have the blendeing effect applied.
+		NodeImagePtr GetBase() const;
+		void SetBase( NodeImagePtr );
+
+		//! Get the node generating the mask image controlling the weight of the effect.
+		NodeImagePtr GetMask() const;
+		void SetMask( NodeImagePtr );
+
+		//! Get the image blended on top of the base.
+		NodeImagePtr GetBlended() const;
+		void SetBlended( NodeImagePtr );
+
+		//!
+		EBlendType GetBlendType() const;
+		void SetBlendType(EBlendType);
+
+		//-----------------------------------------------------------------------------------------
+		// Interface pattern
+		//-----------------------------------------------------------------------------------------
+		class Private;
+		Private* GetPrivate() const;
 
 	protected:
 
-		/** Forbidden. Manage with the Ptr<> template. */
-		~NodeImageLayer() {}
+		//! Forbidden. Manage with the Ptr<> template.
+		~NodeImageLayer();
 
 	private:
 
-		static FNodeType StaticType;
+		Private* m_pD;
 
 	};
 
diff --git a/Engine/Plugins/Mutable/Source/MutableTools/Public/MuT/NodeImageLayerColour.h b/Engine/Plugins/Mutable/Source/MutableTools/Public/MuT/NodeImageLayerColour.h
--- a/Engine/Plugins/Mutable/Source/MutableTools/Public/MuT/NodeImageLayerColour.h
+++ b/Engine/Plugins/Mutable/Source/MutableTools/Public/MuT/NodeImageLayerColour.h
@@ -7,36 +7,73 @@
 #include "MuR/RefCounted.h"
 #include "MuT/Node.h"
 #include "MuT/NodeImage.h"
-#include "MuT/NodeColour.h"
 
 namespace mu
 {
 
-	/** This node applies a layer blending effect on a base image using a mask and a colour. */
+	// Forward definitions
+	class NodeColour;
+	typedef Ptr<NodeColour> NodeColourPtr;
+	typedef Ptr<const NodeColour> NodeColourPtrConst;
+
+	class NodeImageLayerColour;
+	typedef Ptr<NodeImageLayerColour> NodeImageLayerColourPtr;
+	typedef Ptr<const NodeImageLayerColour> NodeImageLayerColourPtrConst;
+
+
+	//! This node applies a layer blending effect on a base image using a mask and a colour.
+	//! \ingroup model
 	class MUTABLETOOLS_API NodeImageLayerColour : public NodeImage
 	{
 	public:
 
-		Ptr<NodeImage> Base;
-		Ptr<NodeImage> Mask;
-		Ptr<NodeColour> Colour;
-		EBlendType Type;
-
-	public:
+		NodeImageLayerColour();
 
+		//-----------------------------------------------------------------------------------------
 		// Node Interface
-		virtual const FNodeType* GetType() const override { return &StaticType; }
-		static const FNodeType* GetStaticType() { return &StaticType; }
+		//-----------------------------------------------------------------------------------------
+
+        const FNodeType* GetType() const override;
+		static const FNodeType* GetStaticType();
+
+		//-----------------------------------------------------------------------------------------
+		// Own Interface
+		//-----------------------------------------------------------------------------------------
+
+		//! Get the node generating the base image that will have the blending effect applied.
+		NodeImagePtr GetBase() const;
+		void SetBase( NodeImagePtr );
+
+		//! Get the node generating the mask image controlling the weight of the effect.
+		NodeImagePtr GetMask() const;
+		void SetMask( NodeImagePtr );
+
+		//! Get the node generating the color to blend on the base.
+		NodeColourPtr GetColour() const;
+		void SetColour( NodeColourPtr );
+
+		//!
+		EBlendType GetBlendType() const;
+
+		//!
+		void SetBlendType(EBlendType);
+
+		//-----------------------------------------------------------------------------------------
+		// Interface pattern
+		//-----------------------------------------------------------------------------------------
+		class Private;
+		Private* GetPrivate() const;
 
 	protected:
 
-		/** Forbidden. Manage with the Ptr<> template. */
-		~NodeImageLayerColour() {}
+		//! Forbidden. Manage with the Ptr<> template.
+		~NodeImageLayerColour();
 
 	private:
 
-		static FNodeType StaticType;
+		Private* m_pD;
 
 	};
 
+
 }
diff --git a/Engine/Plugins/Mutable/Source/MutableTools/Public/MuT/NodeImageLuminance.h b/Engine/Plugins/Mutable/Source/MutableTools/Public/MuT/NodeImageLuminance.h
--- a/Engine/Plugins/Mutable/Source/MutableTools/Public/MuT/NodeImageLuminance.h
+++ b/Engine/Plugins/Mutable/Source/MutableTools/Public/MuT/NodeImageLuminance.h
@@ -7,31 +7,55 @@
 #include "MuT/Node.h"
 #include "MuT/NodeImage.h"
 
+
 namespace mu
 {
 
-	/** Calculate the luminance of an image into a new single - channel image. */
+	// Forward definitions
+	class NodeImageLuminance;
+	typedef Ptr<NodeImageLuminance> NodeImageLuminancePtr;
+	typedef Ptr<const NodeImageLuminance> NodeImageLuminancePtrConst;
+
+
+	//! Calculate the luminance of an image into a new single-channel image..
+	//! \ingroup model
 	class MUTABLETOOLS_API NodeImageLuminance : public NodeImage
 	{
 	public:
 
-		Ptr<NodeImage> Source;
-
-	public:
+		NodeImageLuminance();
 
+		//-----------------------------------------------------------------------------------------
 		// Node Interface
-		virtual const FNodeType* GetType() const override { return &StaticType; }
-		static const FNodeType* GetStaticType() { return &StaticType; }
+		//-----------------------------------------------------------------------------------------
+
+        const FNodeType* GetType() const override;
+		static const FNodeType* GetStaticType();
+
+		//-----------------------------------------------------------------------------------------
+		// Own Interface
+		//-----------------------------------------------------------------------------------------
+
+		//! Get the node generating the source image calculate the luminance from.
+		NodeImagePtr GetSource() const;
+		void SetSource( NodeImagePtr );
+
+		//-----------------------------------------------------------------------------------------
+		// Interface pattern
+		//-----------------------------------------------------------------------------------------
+		class Private;
+		Private* GetPrivate() const;
 
 	protected:
 
-		/** Forbidden. Manage with the Ptr<> template. */
-		~NodeImageLuminance() {}
+		//! Forbidden. Manage with the Ptr<> template.
+		~NodeImageLuminance();
 
 	private:
 
-		static FNodeType StaticType;
+		Private* m_pD;
 
 	};
 
+
 }
diff --git a/Engine/Plugins/Mutable/Source/MutableTools/Public/MuT/NodeImageMipmap.h b/Engine/Plugins/Mutable/Source/MutableTools/Public/MuT/NodeImageMipmap.h
--- a/Engine/Plugins/Mutable/Source/MutableTools/Public/MuT/NodeImageMipmap.h
+++ b/Engine/Plugins/Mutable/Source/MutableTools/Public/MuT/NodeImageMipmap.h
@@ -7,36 +7,62 @@
 #include "MuR/RefCounted.h"
 #include "MuT/Node.h"
 #include "MuT/NodeImage.h"
-#include "MuT/NodeScalar.h"
+
 
 namespace mu
 {
 
-	/** Generate mimaps for an image. */
+	// Forward definitions
+	class NodeScalar;
+	typedef Ptr<NodeScalar> NodeScalarPtr;
+	typedef Ptr<const NodeScalar> NodeScalarPtrConst;
+
+    class NodeImageMipmap;
+    typedef Ptr<NodeImageMipmap> NodeImageMipmapPtr;
+    typedef Ptr<const NodeImageMipmap> NodeImageMipmapPtrConst;
+
+
+	//! Generate mimaps for an image.
+	//! \ingroup model
     class MUTABLETOOLS_API NodeImageMipmap : public NodeImage
 	{
 	public:
 
-		Ptr<NodeImage> Source;
-		Ptr<NodeScalar> Factor;
+        NodeImageMipmap();
 
-		FMipmapGenerationSettings Settings;
+		//-----------------------------------------------------------------------------------------
+		// Node Interface
+		//-----------------------------------------------------------------------------------------
 
-	public:
+        const FNodeType* GetType() const override;
+		static const FNodeType* GetStaticType();
 
-		// Node Interface
-		virtual const FNodeType* GetType() const override { return &StaticType; }
-		static const FNodeType* GetStaticType() { return &StaticType; }
+		//-----------------------------------------------------------------------------------------
+		// Own Interface
+		//-----------------------------------------------------------------------------------------
+
+        //! Get the node generating the source image to be mipmapped.
+		NodeImagePtr GetSource() const;
+		void SetSource( NodeImagePtr );
+
+		void SetMipmapGenerationSettings(EMipmapFilterType FilterType, EAddressMode AddressMode);
+
+		//-----------------------------------------------------------------------------------------
+		// Interface pattern
+		//-----------------------------------------------------------------------------------------
+		class Private;
+		Private* GetPrivate() const;
 
 	protected:
 
-		/** Forbidden. Manage with the Ptr<> template. */
-		~NodeImageMipmap() {}
+		//! Forbidden. Manage with the Ptr<> template.
+        ~NodeImageMipmap();
 
 	private:
 
-		static FNodeType StaticType;
+		Private* m_pD;
 
 	};
 
+
 }
diff --git a/Engine/Plugins/Mutable/Source/MutableTools/Public/MuT/NodeImageMultiLayer.h b/Engine/Plugins/Mutable/Source/MutableTools/Public/MuT/NodeImageMultiLayer.h
--- a/Engine/Plugins/Mutable/Source/MutableTools/Public/MuT/NodeImageMultiLayer.h
+++ b/Engine/Plugins/Mutable/Source/MutableTools/Public/MuT/NodeImageMultiLayer.h
@@ -9,37 +9,72 @@
 #include "MuT/NodeImage.h"
 #include "MuT/NodeRange.h"
 
+
 namespace mu
 {
 
-    /** This node applies any numbre of layer blending effect on a base image using a mask and a
-    * blended image. The number of layers depends on a scalar input of this node.
-	*/
+	// Forward definitions
+    class NodeImageMultiLayer;
+    typedef Ptr<NodeImageMultiLayer> NodeImageMultiLayerPtr;
+    typedef Ptr<const NodeImageMultiLayer> NodeImageMultiLayerPtrConst;
+
+
+    //! This node applies any numbre of layer blending effect on a base image using a mask and a
+    //! blended image. The number of layers depends on a scalar input of this node.
+	//! \ingroup model
     class MUTABLETOOLS_API NodeImageMultiLayer : public NodeImage
 	{
 	public:
 
-		Ptr<NodeImage> Base;
-		Ptr<NodeImage> Mask;
-		Ptr<NodeImage> Blended;
-		Ptr<NodeRange> Range;
-		EBlendType Type;
-
-	public:
+        NodeImageMultiLayer();
 
+		//-----------------------------------------------------------------------------------------
 		// Node Interface
-		virtual const FNodeType* GetType() const override { return &StaticType; }
-		static const FNodeType* GetStaticType() { return &StaticType; }
+		//-----------------------------------------------------------------------------------------
+
+        const FNodeType* GetType() const override;
+		static const FNodeType* GetStaticType();
+
+		//-----------------------------------------------------------------------------------------
+		// Own Interface
+		//-----------------------------------------------------------------------------------------
+
+		//! Get the node generating the base image that will have the blendeing effect applied.
+		NodeImagePtr GetBase() const;
+		void SetBase( NodeImagePtr );
+
+		//! Get the node generating the mask image controlling the weight of the effect.
+		NodeImagePtr GetMask() const;
+		void SetMask( NodeImagePtr );
+
+        //! Get the image blended on top of the base.
+        NodeImagePtr GetBlended() const;
+        void SetBlended( NodeImagePtr );
+
+        //! Get the number of layers to apply.
+        NodeRangePtr GetRange() const;
+        void SetRange( NodeRangePtr );
+
+		//!
+		EBlendType GetBlendType() const;
+        void SetBlendType(EBlendType);
+
+		//-----------------------------------------------------------------------------------------
+		// Interface pattern
+		//-----------------------------------------------------------------------------------------
+		class Private;
+		Private* GetPrivate() const;
 
 	protected:
 
-		/** Forbidden. Manage with the Ptr<> template. */
-		~NodeImageMultiLayer() {}
+		//! Forbidden. Manage with the Ptr<> template.
+        ~NodeImageMultiLayer();
 
 	private:
 
-		static FNodeType StaticType;
+		Private* m_pD;
 
 	};
 
+
 }
diff --git a/Engine/Plugins/Mutable/Source/MutableTools/Public/MuT/NodeImageNormalComposite.h b/Engine/Plugins/Mutable/Source/MutableTools/Public/MuT/NodeImageNormalComposite.h
--- a/Engine/Plugins/Mutable/Source/MutableTools/Public/MuT/NodeImageNormalComposite.h
+++ b/Engine/Plugins/Mutable/Source/MutableTools/Public/MuT/NodeImageNormalComposite.h
@@ -11,31 +11,55 @@
 namespace mu
 {
 
-	/** Node that inverts the colors of an image, channel by channel. */
+	// Forward definitions
+	class NodeImageNormalComposite;
+	typedef Ptr<NodeImageNormalComposite> NodeImageNormalCompositePtr;
+	typedef Ptr<const NodeImageNormalComposite> NodeImageNormalCompositePtrConst;
+
+	//! Node that inverts the colors of an image, channel by channel
 	class MUTABLETOOLS_API NodeImageNormalComposite : public NodeImage
 	{
 	public:
 
-		Ptr<NodeImage> Base;
-		Ptr<NodeImage> Normal;
-		float Power = 1.0f;
-		ECompositeImageMode Mode;
-
-	public:
+		NodeImageNormalComposite();
 
+		//-----------------------------------------------------------------------------------------
 		// Node Interface
-		virtual const FNodeType* GetType() const override { return &StaticType; }
-		static const FNodeType* GetStaticType() { return &StaticType; }
+		//-----------------------------------------------------------------------------------------
+
+		const FNodeType* GetType() const override;
+		static const FNodeType* GetStaticType();
+
+		//-----------------------------------------------------------------------------------------
+		// Own Interface
+		//-----------------------------------------------------------------------------------------
+
+		NodeImagePtr GetBase() const;
+		void SetBase(NodeImagePtr);
+
+        NodeImagePtr GetNormal() const;
+        void SetNormal(NodeImagePtr);
+
+		float GetPower() const;
+		void SetPower(float); 
+
+        ECompositeImageMode GetMode() const;
+        void SetMode(ECompositeImageMode);
+
+		//-----------------------------------------------------------------------------------------
+		// Interface pattern
+		//-----------------------------------------------------------------------------------------
+		class Private;
+		Private* GetPrivate() const;
 
 	protected:
 
-		/** Forbidden. Manage with the Ptr<> template. */
-		~NodeImageNormalComposite() {}
+		//! Forbidden. Manage with the Ptr<> template
+		~NodeImageNormalComposite();
 
 	private:
 
-		static FNodeType StaticType;
-
+		Private* m_pD;
 	};
 
 }
diff --git a/Engine/Plugins/Mutable/Source/MutableTools/Public/MuT/NodeImageParameter.h b/Engine/Plugins/Mutable/Source/MutableTools/Public/MuT/NodeImageParameter.h
--- a/Engine/Plugins/Mutable/Source/MutableTools/Public/MuT/NodeImageParameter.h
+++ b/Engine/Plugins/Mutable/Source/MutableTools/Public/MuT/NodeImageParameter.h
@@ -6,42 +6,59 @@
 #include "MuR/RefCounted.h"
 #include "MuT/Node.h"
 #include "MuT/NodeImage.h"
-#include "MuT/NodeRange.h"
+
 
 namespace mu
 {
+	class NodeRange;
 
-    /** */
+    //! Node that defines a Image model parameter.
+	//! \ingroup model
     class MUTABLETOOLS_API NodeImageParameter : public NodeImage
 	{
 	public:
 
-		/** Full asset name for the default value of the parameter. */
-		FName DefaultValue;
-		
-		/** Name of the parameter */
-		FString Name;
+        NodeImageParameter();
 
-		/** User provided ID to identify the parameter. */
-		FString UID;
+		//-----------------------------------------------------------------------------------------
+		// Node Interface
+		//-----------------------------------------------------------------------------------------
 
-		/** Ranges for the parameter in case it is multidimensional. */
-		TArray<Ptr<NodeRange>> Ranges;
+		const FNodeType* GetType() const override;
+		static const FNodeType* GetStaticType();
 
-	public:
+		//-----------------------------------------------------------------------------------------
+		// Own Interface
+		//-----------------------------------------------------------------------------------------
 
-		// Node Interface
-		virtual const FNodeType* GetType() const override { return &StaticType; }
-		static const FNodeType* GetStaticType() { return &StaticType; }
+		//! Set the name of the parameter.
+		void SetName(const FString&);
+
+		//! Set the uid of the parameter.
+		void SetUid(const FString&);
+
+    	//! Get the default value of the parameter.
+    	void SetDefaultValue(FName Value);
+    	
+		//! Set the number of ranges (dimensions) for this parameter.
+		//! By default a parameter has 0 ranges, meaning it only has one value.
+		void SetRangeCount(int Index);
+		void SetRange(int Index, Ptr<NodeRange> Range);
+
+		//-----------------------------------------------------------------------------------------
+		// Interface pattern
+		//-----------------------------------------------------------------------------------------
+		class Private;
+		Private* GetPrivate() const;
 
 	protected:
 
-		/** Forbidden. Manage with the Ptr<> template. */
-		~NodeImageParameter() {}
+		//! Forbidden. Manage with the Ptr<> template.
+        ~NodeImageParameter();
 
 	private:
 
-		static FNodeType StaticType;
+		Private* m_pD;
 
 	};
 
diff --git a/Engine/Plugins/Mutable/Source/MutableTools/Public/MuT/NodeImagePlainColour.h b/Engine/Plugins/Mutable/Source/MutableTools/Public/MuT/NodeImagePlainColour.h
--- a/Engine/Plugins/Mutable/Source/MutableTools/Public/MuT/NodeImagePlainColour.h
+++ b/Engine/Plugins/Mutable/Source/MutableTools/Public/MuT/NodeImagePlainColour.h
@@ -6,37 +6,64 @@
 #include "MuR/RefCounted.h"
 #include "MuT/Node.h"
 #include "MuT/NodeImage.h"
-#include "MuT/NodeColour.h"
-#include "MuR/Image.h"
+
 
 namespace mu
 {
 
-	/** Node that multiplies the colors of an image, channel by channel. */
+	// Forward definitions
+	class NodeColour;
+	typedef Ptr<NodeColour> NodeColourPtr;
+	typedef Ptr<const NodeColour> NodeColourPtrConst;
+
+	class NodeImagePlainColour;
+	typedef Ptr<NodeImagePlainColour> NodeImagePlainColourPtr;
+	typedef Ptr<const NodeImagePlainColour> NodeImagePlainColourPtrConst;
+
+
+	//! Node that multiplies the colors of an image, channel by channel.
+	//! \ingroup model
 	class MUTABLETOOLS_API NodeImagePlainColour : public NodeImage
 	{
 	public:
 
-		Ptr<NodeColour> Colour;
-		int32 SizeX = 4;
-		int32 SizeY = 4;
-		EImageFormat Format = EImageFormat::RGBA_UByte;
+		NodeImagePlainColour();
 
-	public:
 
+		//-----------------------------------------------------------------------------------------
 		// Node Interface
-		virtual const FNodeType* GetType() const override { return &StaticType; }
-		static const FNodeType* GetStaticType() { return &StaticType; }
+		//-----------------------------------------------------------------------------------------
+
+        const FNodeType* GetType() const override;
+		static const FNodeType* GetStaticType();
+
+		//-----------------------------------------------------------------------------------------
+		// Own Interface
+		//-----------------------------------------------------------------------------------------
+
+		//! Colour of the image.
+		NodeColourPtr GetColour() const;
+		void SetColour( NodeColourPtr );
+
+		//! New size or relative factor
+		void SetSize( int x, int y );
+
+		//-----------------------------------------------------------------------------------------
+		// Interface pattern
+		//-----------------------------------------------------------------------------------------
+		class Private;
+		Private* GetPrivate() const;
 
 	protected:
 
-		/** Forbidden. Manage with the Ptr<> template. */
-		~NodeImagePlainColour() {}
+		//! Forbidden. Manage with the Ptr<> template.
+		~NodeImagePlainColour();
 
 	private:
 
-		static FNodeType StaticType;
+		Private* m_pD;
 
 	};
 
+
 }
diff --git a/Engine/Plugins/Mutable/Source/MutableTools/Public/MuT/NodeImageProject.h b/Engine/Plugins/Mutable/Source/MutableTools/Public/MuT/NodeImageProject.h
--- a/Engine/Plugins/Mutable/Source/MutableTools/Public/MuT/NodeImageProject.h
+++ b/Engine/Plugins/Mutable/Source/MutableTools/Public/MuT/NodeImageProject.h
@@ -1,53 +1,98 @@
 // Copyright Epic Games, Inc. All Rights Reserved.
 
+
 #pragma once
 
 #include "MuR/Ptr.h"
 #include "MuR/RefCounted.h"
 #include "MuT/Node.h"
 #include "MuT/NodeImage.h"
-#include "MuT/NodeMesh.h"
-#include "MuT/NodeProjector.h"
-#include "MuT/NodeScalar.h"
 #include "MuR/Image.h"
 
 namespace mu
 {
 
-	/** */
+	// Forward definitions
+	class NodeScalar;
+	class NodeMesh;
+	class NodeProjector;
+	class NodeImageProject;
+
+	//!
+	//! \ingroup model
 	class MUTABLETOOLS_API NodeImageProject : public NodeImage
 	{
 	public:
 
-		Ptr<NodeProjector> Projector;
-		Ptr<NodeMesh> Mesh;
-		Ptr<NodeScalar> AngleFadeStart;
-		Ptr<NodeScalar> AngleFadeEnd;
-		Ptr<NodeImage> Image;
-		Ptr<NodeImage> Mask;
-		FUintVector2 ImageSize;
-		uint8 Layout = 0;
-		bool bIsRGBFadingEnabled = true;
-		bool bIsAlphaFadingEnabled = true;
-		bool bEnableTextureSeamCorrection = true;
-		ESamplingMethod SamplingMethod = ESamplingMethod::Point;
-		EMinFilterMethod MinFilterMethod = EMinFilterMethod::None;
-
-	public:
+		NodeImageProject();
 
+		//-----------------------------------------------------------------------------------------
 		// Node Interface
-		virtual const FNodeType* GetType() const override { return &StaticType; }
-		static const FNodeType* GetStaticType() { return &StaticType; }
+		//-----------------------------------------------------------------------------------------
+
+        const FNodeType* GetType() const override;
+		static const FNodeType* GetStaticType();
+
+		//-----------------------------------------------------------------------------------------
+		// Own Interface
+		//-----------------------------------------------------------------------------------------
+
+		//! Get the node generating the mesh to be used for the projection.
+		void SetMesh( Ptr<NodeMesh> );
+
+		//! Get the node generating the projector to be used for the projection.
+		void SetProjector( Ptr<NodeProjector> );
+
+		//! Get the node generating a mask to filter out the projected images. It is optional.
+		void SetTargetMask( Ptr<NodeImage> );
+
+		/** Eanble or disable the additional operations to correct texture UV seam artifact correction. Default is enabled.*/
+		void SetEnableSeamCorrection(bool bEnabled);
+
+		//! Set the angle-based fading behaviour for this projector. It must be set separately for 
+		//! RGB and Alpha channels of the images to project. By default it is enabled for all channels.
+		//! Disabling the angle fade will still use the AngleFadeEnd below to clamp the projection,
+		//! but it will ignore the start angle, so it will not apply a fading gradient.
+		void SetAngleFadeChannels( bool bFadeRGB, bool bFadeA );
+
+        //! Set the node generating the fading start angle. Only relevant if fading is enabled with 
+		void SetAngleFadeStart( Ptr<NodeScalar> );
+
+		//! Set the node generating the fading end angle
+		void SetAngleFadeEnd( Ptr<NodeScalar> );
+
+		//! Set sampling method.
+		void SetSamplingMethod(ESamplingMethod SamplingMethod);
+		
+		//! Set min filter method
+		void SetMinFilterMethod(EMinFilterMethod MinFilterMethod);
+
+        //! Get the node generating the image to project.
+        void SetImage( Ptr<NodeImage> );
+
+        //! UV layout of the mesh to use for the generated image. Defaults to 0.
+        void SetLayout( uint8  );
+
+		//! Set the size of the image to generate with the projection. If set to 0 (default) a size
+		//! that matches how this node is used will try to be guessed. 
+		void SetImageSize( const FUintVector2& size );
+
+		//-----------------------------------------------------------------------------------------
+		// Interface pattern
+		//-----------------------------------------------------------------------------------------
+		class Private;
+		Private* GetPrivate() const;
 
 	protected:
 
-		/** Forbidden. Manage with the Ptr<> template. */
-		~NodeImageProject() {}
+		//! Forbidden. Manage with the Ptr<> template.
+		~NodeImageProject();
 
 	private:
 
-		static FNodeType StaticType;
+		Private* m_pD;
 
 	};
 
+
 }
diff --git a/Engine/Plugins/Mutable/Source/MutableTools/Public/MuT/NodeImageResize.h b/Engine/Plugins/Mutable/Source/MutableTools/Public/MuT/NodeImageResize.h
--- a/Engine/Plugins/Mutable/Source/MutableTools/Public/MuT/NodeImageResize.h
+++ b/Engine/Plugins/Mutable/Source/MutableTools/Public/MuT/NodeImageResize.h
@@ -7,36 +7,66 @@
 #include "MuT/Node.h"
 #include "MuT/NodeImage.h"
 
+
 namespace mu
 {
 
-	/** Node that multiplies the colors of an image, channel by channel. */
+	// Forward definitions
+	class NodeScalar;
+	typedef Ptr<NodeScalar> NodeScalarPtr;
+	typedef Ptr<const NodeScalar> NodeScalarPtrConst;
+
+	class NodeImageResize;
+	typedef Ptr<NodeImageResize> NodeImageResizePtr;
+	typedef Ptr<const NodeImageResize> NodeImageResizePtrConst;
+
+
+	//! Node that multiplies the colors of an image, channel by channel.
+	//! \ingroup model
 	class MUTABLETOOLS_API NodeImageResize : public NodeImage
 	{
 	public:
 
-		Ptr<NodeImage> Base;
+		NodeImageResize();
 
-		bool bRelative = true;
 
-		float SizeX = 0.5f;
-		float SizeY = 0.5f;
+		//-----------------------------------------------------------------------------------------
+		// Node Interface
+		//-----------------------------------------------------------------------------------------
 
-	public:
+        const FNodeType* GetType() const override;
+		static const FNodeType* GetStaticType();
 
-		// Node Interface
-		virtual const FNodeType* GetType() const override { return &StaticType; }
-		static const FNodeType* GetStaticType() { return &StaticType; }
+		//-----------------------------------------------------------------------------------------
+		// Own Interface
+		//-----------------------------------------------------------------------------------------
+
+		//! Base image to resize.
+		NodeImagePtr GetBase() const;
+		void SetBase( NodeImagePtr );
+
+		//! Is the size a relative factor or an absolute size?
+		void SetRelative( bool );
+
+		//! New size or relative factor
+		void SetSize( float x, float y );
+
+		//-----------------------------------------------------------------------------------------
+		// Interface pattern
+		//-----------------------------------------------------------------------------------------
+		class Private;
+		Private* GetPrivate() const;
 
 	protected:
 
-		/** Forbidden. Manage with the Ptr<> template. */
-		~NodeImageResize() {}
+		//! Forbidden. Manage with the Ptr<> template.
+		~NodeImageResize();
 
 	private:
 
-		static FNodeType StaticType;
+		Private* m_pD;
 
 	};
 
+
 }
diff --git a/Engine/Plugins/Mutable/Source/MutableTools/Public/MuT/NodeImageSaturate.h b/Engine/Plugins/Mutable/Source/MutableTools/Public/MuT/NodeImageSaturate.h
--- a/Engine/Plugins/Mutable/Source/MutableTools/Public/MuT/NodeImageSaturate.h
+++ b/Engine/Plugins/Mutable/Source/MutableTools/Public/MuT/NodeImageSaturate.h
@@ -6,38 +6,68 @@
 #include "MuR/RefCounted.h"
 #include "MuT/Node.h"
 #include "MuT/NodeImage.h"
-#include "MuT/NodeScalar.h"
+
 
 namespace mu
 {
 
-	/** Change the saturation of an image.This node can be used to increase the saturation with a
-	* factor bigger than 1, or to decrease it or desaturate it completely with a factor smaller
-	* than 1 or 0.
-	*/
+	// Forward definitions
+	class NodeScalar;
+	typedef Ptr<NodeScalar> NodeScalarPtr;
+	typedef Ptr<const NodeScalar> NodeScalarPtrConst;
+
+	class NodeImageSaturate;
+	typedef Ptr<NodeImageSaturate> NodeImageSaturatePtr;
+	typedef Ptr<const NodeImageSaturate> NodeImageSaturatePtrConst;
+
+
+	//! Change the saturation of an image. This node can be used to increase the saturation with a
+	//! factor bigger than 1, or to decrease it or desaturate it completely with a factor smaller
+	//! than 1 or 0.
 	class MUTABLETOOLS_API NodeImageSaturate : public NodeImage
 	{
 	public:
 
-		Ptr<NodeImage> Source;
+		NodeImageSaturate();
 
-		Ptr<NodeScalar> Factor;
+		//-----------------------------------------------------------------------------------------
+		// Node Interface
+		//-----------------------------------------------------------------------------------------
 
-	public:
+        const FNodeType* GetType() const override;
+		static const FNodeType* GetStaticType();
 
-		// Node Interface
-		virtual const FNodeType* GetType() const override { return &StaticType; }
-		static const FNodeType* GetStaticType() { return &StaticType; }
+		//-----------------------------------------------------------------------------------------
+		// Own Interface
+		//-----------------------------------------------------------------------------------------
+
+		//! Get the node generating the saturation factor.
+		//! A value of 0 completely desaturates. 1 will leave the same saturation and bigger than 1
+		//! will increase it.
+		NodeScalarPtr GetFactor() const;
+		void SetFactor( NodeScalarPtr );
+
+		//! Get the node generating the source image to be [de]saturated.
+		NodeImagePtr GetSource() const;
+		void SetSource( NodeImagePtr );
+
+		//-----------------------------------------------------------------------------------------
+		// Interface pattern
+		//-----------------------------------------------------------------------------------------
+		class Private;
+
+		Private* GetPrivate() const;
 
 	protected:
 
-		/** Forbidden. Manage with the Ptr<> template. */
-		~NodeImageSaturate() {}
+		//! Forbidden. Manage with the Ptr<> template.
+		~NodeImageSaturate();
 
 	private:
 
-		static FNodeType StaticType;
+		Private* m_pD;
 
 	};
 
+
 }
diff --git a/Engine/Plugins/Mutable/Source/MutableTools/Public/MuT/NodeImageSwitch.h b/Engine/Plugins/Mutable/Source/MutableTools/Public/MuT/NodeImageSwitch.h
--- a/Engine/Plugins/Mutable/Source/MutableTools/Public/MuT/NodeImageSwitch.h
+++ b/Engine/Plugins/Mutable/Source/MutableTools/Public/MuT/NodeImageSwitch.h
@@ -6,35 +6,69 @@
 #include "MuR/RefCounted.h"
 #include "MuT/Node.h"
 #include "MuT/NodeImage.h"
-#include "MuT/NodeScalar.h"
+
 
 namespace mu
 {
 
+	// Forward definitions
+	class NodeScalar;
+	typedef Ptr<NodeScalar> NodeScalarPtr;
+	typedef Ptr<const NodeScalar> NodeScalarPtrConst;
+
+	class NodeImageSwitch;
+	typedef Ptr<NodeImageSwitch> NodeImageSwitchPtr;
+	typedef Ptr<const NodeImageSwitch> NodeImageSwitchPtrConst;
+
+
 	//! This node selects an output image from a set of input images based on a parameter.
 	class MUTABLETOOLS_API NodeImageSwitch : public NodeImage
 	{
 	public:
 
-		Ptr<NodeScalar> Parameter;
+		NodeImageSwitch();
 
-		TArray<Ptr<NodeImage>> Options;
+		//-----------------------------------------------------------------------------------------
+		// Node Interface
+		//-----------------------------------------------------------------------------------------
 
-	public:
+        const FNodeType* GetType() const override;
+		static const FNodeType* GetStaticType();
 
-		// Node Interface
-		virtual const FNodeType* GetType() const override { return &StaticType; }
-		static const FNodeType* GetStaticType() { return &StaticType; }
+		//-----------------------------------------------------------------------------------------
+		// Own Interface
+		//-----------------------------------------------------------------------------------------
+
+		//! Get the node generating the parameter used to select the option.
+		NodeScalarPtr GetParameter() const;
+		void SetParameter( NodeScalarPtr );
+
+		//! Set the number of option images. It will keep the currently set targets and initialise
+		//! the new ones as null.
+		void SetOptionCount( int );
+		int GetOptionCount() const;
+
+		//! Get the node generating the t-th option image.
+		NodeImagePtr GetOption( int t ) const;
+		void SetOption( int t, NodeImagePtr );
+
+
+		//-----------------------------------------------------------------------------------------
+		// Interface pattern
+		//-----------------------------------------------------------------------------------------
+		class Private;
+		Private* GetPrivate() const;
 
 	protected:
 
-		/** Forbidden. Manage with the Ptr<> template. */
-		~NodeImageSwitch() {}
+		//! Forbidden. Manage with the Ptr<> template.
+		~NodeImageSwitch();
 
 	private:
 
-		static FNodeType StaticType;
+		Private* m_pD;
 
 	};
 
+
 }
diff --git a/Engine/Plugins/Mutable/Source/MutableTools/Public/MuT/NodeImageSwizzle.h b/Engine/Plugins/Mutable/Source/MutableTools/Public/MuT/NodeImageSwizzle.h
--- a/Engine/Plugins/Mutable/Source/MutableTools/Public/MuT/NodeImageSwizzle.h
+++ b/Engine/Plugins/Mutable/Source/MutableTools/Public/MuT/NodeImageSwizzle.h
@@ -8,43 +8,62 @@
 #include "MuT/Node.h"
 #include "MuT/NodeImage.h"
 
+
 namespace mu
 {
 
-	/** Node that composes a new image by gathering pixel data from channels in other images. */
+	// Forward definitions
+	class NodeImageSwizzle;
+	typedef Ptr<NodeImageSwizzle> NodeImageSwizzlePtr;
+	typedef Ptr<const NodeImageSwizzle> NodeImageSwizzlePtrConst;
+
+
+	//! Node that composes a new image by gathering pixel data from channels in other images.
+	//! \ingroup model
 	class MUTABLETOOLS_API NodeImageSwizzle : public NodeImage
 	{
 	public:
 
-		EImageFormat NewFormat;
-		TArray<Ptr<NodeImage>> Sources;
-		TArray<int32> SourceChannels;
-
-	public:
+		NodeImageSwizzle();
 
+		//-----------------------------------------------------------------------------------------
 		// Node Interface
-		virtual const FNodeType* GetType() const override { return &StaticType; }
-		static const FNodeType* GetStaticType() { return &StaticType; }
+		//-----------------------------------------------------------------------------------------
 
-		// Own interface 
-		inline void SetFormat(EImageFormat InFormat)
-		{
-			NewFormat = InFormat;
+        const FNodeType* GetType() const override;
+		static const FNodeType* GetStaticType();
 
-			int32 ChannelCount = GetImageFormatData(InFormat).Channels;
-			Sources.SetNum(ChannelCount);
-			SourceChannels.SetNum(ChannelCount);
-		}
+		//-----------------------------------------------------------------------------------------
+		// Own Interface
+		//-----------------------------------------------------------------------------------------
+
+		//!
+		EImageFormat GetFormat() const;
+		void SetFormat(EImageFormat);
+
+		//!
+		NodeImagePtr GetSource( int32 ) const;
+		int GetSourceChannel( int32 ) const;
+
+		void SetSource( int32, NodeImagePtr );
+		void SetSourceChannel( int32 OutputChannel, int32 SourceChannel );
+
+		//-----------------------------------------------------------------------------------------
+		// Interface pattern
+		//-----------------------------------------------------------------------------------------
+		class Private;
+		Private* GetPrivate() const;
 
 	protected:
 
-		/** Forbidden. Manage with the Ptr<> template. */
-		~NodeImageSwizzle() {}
+		//! Forbidden. Manage with the Ptr<> template.
+		~NodeImageSwizzle();
 
 	private:
 
-		static FNodeType StaticType;
+		Private* m_pD;
 
 	};
 
+
 }
diff --git a/Engine/Plugins/Mutable/Source/MutableTools/Public/MuT/NodeImageTransform.h b/Engine/Plugins/Mutable/Source/MutableTools/Public/MuT/NodeImageTransform.h
--- a/Engine/Plugins/Mutable/Source/MutableTools/Public/MuT/NodeImageTransform.h
+++ b/Engine/Plugins/Mutable/Source/MutableTools/Public/MuT/NodeImageTransform.h
@@ -2,48 +2,87 @@
 
 #pragma once
 
+#include "MuR/Ptr.h"
+#include "MuR/RefCounted.h"
 #include "MuT/Node.h"
 #include "MuT/NodeImage.h"
-#include "MuT/NodeScalar.h"
-#include "MuR/ImageTypes.h"
-#include "MuR/Ptr.h"
 
 namespace mu
 {
 
-	/**  */
+	// Forward definitions
+	class NodeScalar;
+	typedef Ptr<NodeScalar> NodeScalarPtr;
+	typedef Ptr<const NodeScalar> NodeScalarPtrConst;
+
+	class NodeImageTransform;
+	typedef Ptr<NodeImageTransform> NodeImageTransformPtr;
+	typedef Ptr<const NodeImageTransform> NodeImageTransformPtrConst;
+
+	enum class EAddressMode;
+
+	//! Node that multiplies the colors of an image, channel by channel.
+	//! \ingroup model
 	class MUTABLETOOLS_API NodeImageTransform : public NodeImage
 	{
 	public:
 
-		Ptr<NodeImage> Base;
-		Ptr<NodeScalar> OffsetX;
-		Ptr<NodeScalar> OffsetY;
-		Ptr<NodeScalar> ScaleX;
-		Ptr<NodeScalar> ScaleY;
-		Ptr<NodeScalar> Rotation;
+		NodeImageTransform();
 
-		EAddressMode AddressMode = EAddressMode::Wrap;
-		uint32 SizeX = 0;
-		uint32 SizeY = 0;
+		//-----------------------------------------------------------------------------------------
+		// Node Interface
+		//-----------------------------------------------------------------------------------------
 
-		bool bKeepAspectRatio = false;
+        const FNodeType* GetType() const override;
+		static const FNodeType* GetStaticType();
 
-	public:
+		//-----------------------------------------------------------------------------------------
+		// Own Interface
+		//-----------------------------------------------------------------------------------------
 
-		// Node Interface
-		virtual const FNodeType* GetType() const override { return &StaticType; }
-		static const FNodeType* GetStaticType() { return &StaticType; }
+		//! Base image to multiply.
+		NodeImagePtr GetBase() const;
+		void SetBase( NodeImagePtr );
+	
+		NodeScalarPtr GetOffsetX() const;
+		NodeScalarPtr GetOffsetY() const;
+		NodeScalarPtr GetScaleX() const;
+		NodeScalarPtr GetScaleY() const;
+		NodeScalarPtr GetRotation() const;
+		void SetOffsetX( NodeScalarPtr pNode );
+		void SetOffsetY( NodeScalarPtr pNode );
+		void SetScaleX( NodeScalarPtr pNode );
+		void SetScaleY( NodeScalarPtr pNode );
+		void SetRotation( NodeScalarPtr pNode );
+
+		EAddressMode GetAddressMode() const;
+		void SetAddressMode(EAddressMode AddressMode);
+
+		uint16 GetSizeX() const;
+		void SetSizeX(uint16 SizeX); 
+
+		uint16 GetSizeY() const;
+		void SetSizeY(uint16 SizeY); 
+
+		bool GetKeepAspectRatio() const;
+		void SetKeepAspectRatio(bool bKeepAspectRatio);
+
+		//-----------------------------------------------------------------------------------------
+		// Interface pattern
+		//-----------------------------------------------------------------------------------------
+		class Private;
+		Private* GetPrivate() const;
 
 	protected:
 
-		/** Forbidden. Manage with the Ptr<> template. */
-		~NodeImageTransform() {}
+		//! Forbidden. Manage with the Ptr<> template.
+		~NodeImageTransform();
 
 	private:
 
-		static FNodeType StaticType;
+		Private* m_pD;
 
 	};
 
+
 }
diff --git a/Engine/Plugins/Mutable/Source/MutableTools/Public/MuT/NodeImageVariation.h b/Engine/Plugins/Mutable/Source/MutableTools/Public/MuT/NodeImageVariation.h
--- a/Engine/Plugins/Mutable/Source/MutableTools/Public/MuT/NodeImageVariation.h
+++ b/Engine/Plugins/Mutable/Source/MutableTools/Public/MuT/NodeImageVariation.h
@@ -7,40 +7,69 @@
 #include "MuT/Node.h"
 #include "MuT/NodeImage.h"
 
+
 namespace mu
 {
 
-    /** */
+    // Forward definitions
+    class NodeImageVariation;
+    typedef Ptr<NodeImageVariation> NodeImageVariationPtr;
+    typedef Ptr<const NodeImageVariation> NodeImageVariationPtrConst;
+
+
+    //!
+    //! \ingroup model
     class MUTABLETOOLS_API NodeImageVariation : public NodeImage
     {
-	public:
+    public:
+        //-----------------------------------------------------------------------------------------
+        // Life cycle
+        //-----------------------------------------------------------------------------------------
 
-		Ptr<NodeImage> DefaultImage;
+        NodeImageVariation();
 
-		struct FVariation
-		{
-			Ptr<NodeImage> Image;
-			FString Tag;
-		};
+        //-----------------------------------------------------------------------------------------
+        // Node Interface
+        //-----------------------------------------------------------------------------------------
 
-		TArray<FVariation> Variations;
+        const FNodeType* GetType() const override;
+        static const FNodeType* GetStaticType();
 
-    public:
+        //-----------------------------------------------------------------------------------------
+        // Own Interface
+        //-----------------------------------------------------------------------------------------
 
-		// Node Interface
-		virtual const FNodeType* GetType() const override { return &StaticType; }
-		static const FNodeType* GetStaticType() { return &StaticType; }
+        //!
+        void SetDefaultImage( NodeImage* Image );
 
-    protected:
+        //! Set the number of tags to consider in this variation
+        void SetVariationCount( int count );
 
-        /** Forbidden. Manage with the Ptr<> template. */
-		~NodeImageVariation() {}
+        //!
+        int GetVariationCount() const;
 
-    private:
+        //! Set the tag or state name that will enable a specific vartiation
+        void SetVariationTag( int index, const FString& Tag );
 
-		static FNodeType StaticType;
-	
-	};
+        //!
+        void SetVariationImage( int index, NodeImage* Image );
+
+        //!}
+
+
+        //-----------------------------------------------------------------------------------------
+        // Interface pattern
+        //-----------------------------------------------------------------------------------------
+        class Private;
+        Private* GetPrivate() const;
+
+    protected:
+        //! Forbidden. Manage with the Ptr<> template.
+        ~NodeImageVariation();
+
+    private:
+        Private* m_pD;
+    };
 
-}
 
+} // namespace mu
diff --git a/Engine/Plugins/Mutable/Source/MutableTools/Public/MuT/NodeMesh.h b/Engine/Plugins/Mutable/Source/MutableTools/Public/MuT/NodeMesh.h
--- a/Engine/Plugins/Mutable/Source/MutableTools/Public/MuT/NodeMesh.h
+++ b/Engine/Plugins/Mutable/Source/MutableTools/Public/MuT/NodeMesh.h
@@ -25,18 +25,16 @@ namespace mu
 	public:
 
 		// Node Interface
-		virtual const FNodeType* GetType() const override { return &StaticType; }
-		static const FNodeType* GetStaticType() { return &StaticType; }
+        const FNodeType* GetType() const override;
+		static const FNodeType* GetStaticType();
+
+		// Interface pattern
+		class Private;
 
 	protected:
 
 		// Forbidden. Manage with the Ptr<> template.
 		inline ~NodeMesh() {}
-
-	private:
-
-		static FNodeType StaticType;
-
 	};
 
 
diff --git a/Engine/Plugins/Mutable/Source/MutableTools/Public/MuT/NodeMeshApplyPose.h b/Engine/Plugins/Mutable/Source/MutableTools/Public/MuT/NodeMeshApplyPose.h
--- a/Engine/Plugins/Mutable/Source/MutableTools/Public/MuT/NodeMeshApplyPose.h
+++ b/Engine/Plugins/Mutable/Source/MutableTools/Public/MuT/NodeMeshApplyPose.h
@@ -7,32 +7,60 @@
 #include "MuT/Node.h"
 #include "MuT/NodeMesh.h"
 
+
 namespace mu
 {
 
-    /** Node that applies a pose to a mesh, baking it into the vertex data. */
+	// Forward definitions
+    class NodeMeshApplyPose;
+    using NodeMeshApplyPosePtr = Ptr<NodeMeshApplyPose>;
+    using NodeMeshApplyPosePtrConst = Ptr<const NodeMeshApplyPose>;
+
+
+    //! Node that applies a pose to a mesh, baking it into the vertex data
     class MUTABLETOOLS_API NodeMeshApplyPose : public NodeMesh
 	{
 	public:
 
-		Ptr<NodeMesh> Base;
-		Ptr<NodeMesh> Pose;
+        NodeMeshApplyPose();
 
-	public:
 
+		//-----------------------------------------------------------------------------------------
 		// Node Interface
-		virtual const FNodeType* GetType() const override { return &StaticType; }
-		static const FNodeType* GetStaticType() { return &StaticType; }
+		//-----------------------------------------------------------------------------------------
+
+		const FNodeType* GetType() const override;
+		static const FNodeType* GetStaticType();
+
+		//-----------------------------------------------------------------------------------------
+		// Own Interface
+		//-----------------------------------------------------------------------------------------
+
+		//! Get the nodes generating the base mesh to be morphed
+		NodeMeshPtr GetBase() const;
+		void SetBase( NodeMeshPtr );
+
+        //! Get the nodes generating the pose to apply. A pose is represented with a mesh object
+        //! with no geometry: only bone matrices and names are relevant.
+        NodeMeshPtr GetPose() const;
+        void SetPose( NodeMeshPtr );
+
+		//-----------------------------------------------------------------------------------------
+		// Interface pattern
+		//-----------------------------------------------------------------------------------------
+		class Private;
+		Private* GetPrivate() const;
 
 	protected:
 
-		/** Forbidden.Manage with the Ptr<> template. */
-		~NodeMeshApplyPose() {}
+		//! Forbidden. Manage with the Ptr<> template.
+        ~NodeMeshApplyPose();
 
 	private:
 
-		static FNodeType StaticType;
+		Private* m_pD;
 
 	};
 
+
 }
diff --git a/Engine/Plugins/Mutable/Source/MutableTools/Public/MuT/NodeMeshConstant.h b/Engine/Plugins/Mutable/Source/MutableTools/Public/MuT/NodeMeshConstant.h
--- a/Engine/Plugins/Mutable/Source/MutableTools/Public/MuT/NodeMeshConstant.h
+++ b/Engine/Plugins/Mutable/Source/MutableTools/Public/MuT/NodeMeshConstant.h
@@ -2,44 +2,77 @@
 
 #pragma once
 
+#include "MuR/Ptr.h"
+#include "MuR/RefCounted.h"
 #include "MuT/Node.h"
 #include "MuT/NodeMesh.h"
 #include "MuT/NodeImage.h"
-#include "MuT/NodeLayout.h"
-#include "MuR/Mesh.h"
 #include "Templates/SharedPointer.h"
 
 namespace mu
 {
 
-	/** Node that outputs a constant mesh.
-	* It allows to define the layouts for the texture channels of the constant mesh
-	*/
+	// Forward definitions
+	class NodeMeshConstant;
+	typedef Ptr<NodeMeshConstant> NodeMeshConstantPtr;
+	typedef Ptr<const NodeMeshConstant> NodeMeshConstantPtrConst;
+
+	class FMesh;
+	class NodeLayout;
+
+	//! Node that outputs a constant mesh.
+	//! It allows to define the layouts for the texture channels of the constant mesh
 	class MUTABLETOOLS_API NodeMeshConstant : public NodeMesh
 	{
 	public:
 
 		FSourceDataDescriptor SourceDataDescriptor;
 
-		TSharedPtr<FMesh> Value;
-
-		TArray<Ptr<NodeLayout>> Layouts;
-
 	public:
 
+		NodeMeshConstant();
+
+		//-----------------------------------------------------------------------------------------
 		// Node Interface
-		virtual const FNodeType* GetType() const override { return &StaticType; }
-		static const FNodeType* GetStaticType() { return &StaticType; }
+		//-----------------------------------------------------------------------------------------
+
+		const FNodeType* GetType() const override;
+		static const FNodeType* GetStaticType();
+
+		//-----------------------------------------------------------------------------------------
+		// Own Interface
+		//-----------------------------------------------------------------------------------------
+
+		//! Get the constant mesh that will be returned.
+		TSharedPtr<FMesh> GetValue() const;
+
+		//! Set the constant mesh that will be returned.
+		void SetValue(TSharedPtr<FMesh>);
+
+		//! Get the number of layouts defined in this mesh.
+		int32 GetLayoutCount() const;
+		void SetLayoutCount( int32 );
+
+		//! Get the node defining a layout of the returned mesh.
+		Ptr<NodeLayout> GetLayout( int32 index ) const;
+		void SetLayout( int32 index, Ptr<NodeLayout>);
+
+		//-----------------------------------------------------------------------------------------
+		// Interface pattern
+		//-----------------------------------------------------------------------------------------
+		class Private;
+		Private* GetPrivate() const;
 
 	protected:
 
-		/** Forbidden. Manage with the Ptr<> template. */
-		~NodeMeshConstant() {}
+		//! Forbidden. Manage with the Ptr<> template.
+		~NodeMeshConstant();
 
 	private:
 
-		static FNodeType StaticType;
+		Private* m_pD;
 
 	};
 
+
 }
diff --git a/Engine/Plugins/Mutable/Source/MutableTools/Public/MuT/NodeMeshFormat.h b/Engine/Plugins/Mutable/Source/MutableTools/Public/MuT/NodeMeshFormat.h
--- a/Engine/Plugins/Mutable/Source/MutableTools/Public/MuT/NodeMeshFormat.h
+++ b/Engine/Plugins/Mutable/Source/MutableTools/Public/MuT/NodeMeshFormat.h
@@ -6,7 +6,6 @@
 #include "MuR/RefCounted.h"
 #include "MuT/Node.h"
 #include "MuT/NodeMesh.h"
-#include "MuR/Mesh.h"
 
 namespace mu
 {
@@ -19,30 +18,55 @@ namespace mu
 	{
 	public:
 
-		/** Source mesh to transform. */
-		Ptr<NodeMesh> Source;
-
-		/** New mesh format.The buffers in the sets have no elements, but they define the formats. */
-		FMeshBufferSet VertexBuffers;
-		FMeshBufferSet IndexBuffers;
-
-		/** */
-		bool bOptimizeBuffers = false;
-
-	public:
+		NodeMeshFormat();
 
+		//-----------------------------------------------------------------------------------------
 		// Node Interface
-		virtual const FNodeType* GetType() const override { return &StaticType; }
-		static const FNodeType* GetStaticType() { return &StaticType; }
+		//-----------------------------------------------------------------------------------------
+
+		const FNodeType* GetType() const override;
+		static const FNodeType* GetStaticType();
+
+		//-----------------------------------------------------------------------------------------
+		// Own Interface
+		//-----------------------------------------------------------------------------------------
+
+		/** Mesh to be reformatted. */
+		Ptr<NodeMesh> GetSource() const;
+		void SetSource( NodeMesh* );
+
+		/** Access the MeshBufferSet that defines the new format for the mesh vertices.These
+		* buffers don't really contain any data (they have 0 elements) but they define the
+		* structure. If this is null, the vertex buffers will not be changed.
+		*/
+		FMeshBufferSet& GetVertexBuffers();
+
+		/** Access the MeshBufferSet that defines the new format for the mesh indices. These
+		* buffers don't really contain any data (they have 0 elements) but they define the
+		* structure. If this is null, the vertex buffers will not be changed.
+		*/
+		FMeshBufferSet& GetIndexBuffers();
+
+		/** Optimize the buffers is possible. This may change the target format to reduce the number
+		* of channels or even the type if possible to minimize size.
+		* By default it is disabled.
+		*/
+		void SetOptimizeBuffers(bool bEnable);
+
+		//-----------------------------------------------------------------------------------------
+		// Interface pattern
+		//-----------------------------------------------------------------------------------------
+		class Private;
+		Private* GetPrivate() const;
 
 	protected:
 
-		/** Forbidden. Manage with the Ptr<> template. */
-		~NodeMeshFormat() {}
+		//! Forbidden. Manage with the Ptr<> template.
+		~NodeMeshFormat();
 
 	private:
 
-		static FNodeType StaticType;
+		Private* m_pD;
 
 	};
 
diff --git a/Engine/Plugins/Mutable/Source/MutableTools/Public/MuT/NodeMeshMakeMorph.h b/Engine/Plugins/Mutable/Source/MutableTools/Public/MuT/NodeMeshMakeMorph.h
--- a/Engine/Plugins/Mutable/Source/MutableTools/Public/MuT/NodeMeshMakeMorph.h
+++ b/Engine/Plugins/Mutable/Source/MutableTools/Public/MuT/NodeMeshMakeMorph.h
@@ -11,31 +11,57 @@
 namespace mu
 {
 
-	/** This node removes from a mesh A all the faces that are also part of a mesh B. */
+	// Forward definitions
+    class NodeMeshMakeMorph;
+    typedef Ptr<NodeMeshMakeMorph> NodeMeshMakeMorphPtr;
+    typedef Ptr<const NodeMeshMakeMorph> NodeMeshMakeMorphPtrConst;
+
+
+	//! This node removes from a mesh A all the faces that are also part of a mesh B.
     class MUTABLETOOLS_API NodeMeshMakeMorph : public NodeMesh
 	{
 	public:
 
-		Ptr<NodeMesh> Base;
-		Ptr<NodeMesh> Target;
+        NodeMeshMakeMorph();
 
-		bool bOnlyPositionAndNormal = false;
+		//-----------------------------------------------------------------------------------------
+		// Node Interface
+		//-----------------------------------------------------------------------------------------
 
-	public:
+		const FNodeType* GetType() const override;
+		static const FNodeType* GetStaticType();
 
-		// Node Interface
-		virtual const FNodeType* GetType() const override { return &StaticType; }
-		static const FNodeType* GetStaticType() { return &StaticType; }
+		//-----------------------------------------------------------------------------------------
+		// Own Interface
+		//-----------------------------------------------------------------------------------------
+
+		//!
+		NodeMeshPtr GetBase() const;
+		void SetBase( NodeMesh* p );
+
+		//!
+		NodeMeshPtr GetTarget() const;
+		void SetTarget( NodeMesh* p );
+
+		bool GetOnlyPositionAndNormal() const;
+		void SetOnlyPositionAndNormal(bool bOnlyPositionAndNormal);
+
+		//-----------------------------------------------------------------------------------------
+		// Interface pattern
+		//-----------------------------------------------------------------------------------------
+		class Private;
+		Private* GetPrivate() const;
 
 	protected:
 
-		/** Forbidden. Manage with the Ptr<> template. */
-		~NodeMeshMakeMorph() {}
+		//! Forbidden. Manage with the Ptr<> template.
+        ~NodeMeshMakeMorph();
 
 	private:
 
-		static FNodeType StaticType;
+		Private* m_pD;
 
 	};
 
+
 }
diff --git a/Engine/Plugins/Mutable/Source/MutableTools/Public/MuT/NodeMeshMorph.h b/Engine/Plugins/Mutable/Source/MutableTools/Public/MuT/NodeMeshMorph.h
--- a/Engine/Plugins/Mutable/Source/MutableTools/Public/MuT/NodeMeshMorph.h
+++ b/Engine/Plugins/Mutable/Source/MutableTools/Public/MuT/NodeMeshMorph.h
@@ -6,42 +6,73 @@
 #include "MuR/RefCounted.h"
 #include "MuT/Node.h"
 #include "MuT/NodeMesh.h"
-#include "MuT/NodeScalar.h"
-#include "MuR/Skeleton.h"
+
 
 namespace mu
 {
 
+	// Forward definitions
+	class NodeScalar;
+	struct FBoneName;
+
 	/** Node that morphs a base mesh with one weighted target. */
 	class MUTABLETOOLS_API NodeMeshMorph : public NodeMesh
 	{
 	public:
 
-		Ptr<NodeScalar> Factor;
-		Ptr<NodeMesh> Base;
-		Ptr<NodeMesh> Morph;
+		NodeMeshMorph();
 
-		bool bReshapeSkeleton = false;
-		bool bReshapePhysicsVolumes = false;
+		//-----------------------------------------------------------------------------------------
+		// Node Interface
+		//-----------------------------------------------------------------------------------------		
 
-		TArray<FBoneName> BonesToDeform;
-		TArray<FBoneName> PhysicsToDeform;
+		const FNodeType* GetType() const override;
+		static const FNodeType* GetStaticType();
 
-	public:
+		//-----------------------------------------------------------------------------------------
+		// Own Interface
+		//-----------------------------------------------------------------------------------------
 
-		// Node Interface
-		virtual const FNodeType* GetType() const override { return &StaticType; }
-		static const FNodeType* GetStaticType() { return &StaticType; }
+		/** Get the node generating the weight used to select(and combine) the morphs to apply. */
+		Ptr<NodeScalar> GetFactor() const;
+		void SetFactor( Ptr<NodeScalar> );
+
+		/** Get the nodes generating the base mesh to be morphed. */
+		Ptr<NodeMesh> GetBase() const;
+		void SetBase( Ptr<NodeMesh> );
+
+		/** Get the nodes generating the several morph targets. */
+		Ptr<NodeMesh> GetMorph() const;
+		void SetMorph( Ptr<NodeMesh> );
+
+		/** Also deform the mesh skeleton. Disabled by default. */
+		void SetReshapeSkeleton(bool);	
+
+		/** Deform Mesh Physics Volumes */
+		void SetReshapePhysicsVolumes(bool);
+
+		/** Sets the number of bones that will be deform */
+		void AddBoneToDeform(const FBoneName& BoneName);	
+
+		/** Add a Physics Body to deform */
+		void AddPhysicsBodyToDeform(const FBoneName& BoneName);
+
+        //-----------------------------------------------------------------------------------------
+		// Interface pattern
+		//-----------------------------------------------------------------------------------------
+		class Private;
+		Private* GetPrivate() const;
 
 	protected:
 
-		/** Forbidden. Manage with the Ptr<> template. */
-		~NodeMeshMorph() {}
+		//! Forbidden. Manage with the Ptr<> template.
+		~NodeMeshMorph();
 
 	private:
 
-		static FNodeType StaticType;
+		Private* m_pD;
 
 	};
 
+
 }
diff --git a/Engine/Plugins/Mutable/Source/MutableTools/Public/MuT/NodeMeshReshape.h b/Engine/Plugins/Mutable/Source/MutableTools/Public/MuT/NodeMeshReshape.h
--- a/Engine/Plugins/Mutable/Source/MutableTools/Public/MuT/NodeMeshReshape.h
+++ b/Engine/Plugins/Mutable/Source/MutableTools/Public/MuT/NodeMeshReshape.h
@@ -6,47 +6,77 @@
 #include "MuT/Node.h"
 #include "MuT/NodeMesh.h"
 #include "MuR/Mesh.h"
-#include "MuR/Skeleton.h"
 
 namespace mu
 {
+	// Forward declarations
+	struct FBoneName;
 
-	/** Node that morphs a base mesh with one or two weighted targets from a sequence. */
+	//! Node that morphs a base mesh with one or two weighted targets from a sequence.
 	class MUTABLETOOLS_API NodeMeshReshape : public NodeMesh
 	{
 	public:
 
-		Ptr<NodeMesh> BaseMesh;
-		Ptr<NodeMesh> BaseShape;
-		Ptr<NodeMesh> TargetShape;
-		bool bReshapeVertices = true;
-		bool bRecomputeNormals = false;
-		bool bApplyLaplacian = false;
-		bool bReshapeSkeleton = false;
-		bool bReshapePhysicsVolumes = false;
+		NodeMeshReshape();
 
-		EVertexColorUsage ColorRChannelUsage = EVertexColorUsage::None;
-		EVertexColorUsage ColorGChannelUsage = EVertexColorUsage::None;
-		EVertexColorUsage ColorBChannelUsage = EVertexColorUsage::None;
-		EVertexColorUsage ColorAChannelUsage = EVertexColorUsage::None;
+		//-----------------------------------------------------------------------------------------
+		// Node Interface
+		//-----------------------------------------------------------------------------------------
 
-		TArray<FBoneName> BonesToDeform;
-		TArray<FBoneName> PhysicsToDeform;
+		const FNodeType* GetType() const override;
+		static const FNodeType* GetStaticType();
 
-	public:
+		//-----------------------------------------------------------------------------------------
+		// Own Interface
+		//-----------------------------------------------------------------------------------------
 
-		// Node Interface
-		virtual const FNodeType* GetType() const override { return &StaticType; }
-		static const FNodeType* GetStaticType() { return &StaticType; }
+		const NodeMeshPtr& GetBaseMesh() const;
+		void SetBaseMesh( const NodeMeshPtr& );
+
+		const NodeMeshPtr& GetBaseShape() const;
+		void SetBaseShape(const NodeMeshPtr&);
+
+		const NodeMeshPtr& GetTargetShape() const;
+		void SetTargetShape(const NodeMeshPtr&);
+
+		void SetReshapeVertices(bool);
+
+		/** Recompute normals after reshaping. This will ignore the reshaped normals. Disabled by default. */
+		void SetRecomputeNormals(bool);
+
+
+		/** Also deform the mesh skeleton. Disabled by default. */
+		void SetReshapeSkeleton(bool);
+	
+		/** Apply Laplacian smoothing to the reshaped mesh.  */
+		void SetApplyLaplacian(bool);
+		
+		/** Set vertex color channel usages for Reshape operations. */
+		void SetColorUsages(EVertexColorUsage R, EVertexColorUsage G, EVertexColorUsage B, EVertexColorUsage A);
+
+		/** Deform Mesh Physics Volumes */
+		void SetReshapePhysicsVolumes(bool);
+
+		/** Sets the number of bones that will be deform */
+		void AddBoneToDeform(const FBoneName& BoneId);
+	
+		/** Add a Physics Body to deform */
+		void AddPhysicsBodyToDeform(const FBoneName& BoneId);
+        
+		//-----------------------------------------------------------------------------------------
+		// Interface pattern
+		//-----------------------------------------------------------------------------------------
+		class Private;
+		Private* GetPrivate() const;
 
 	protected:
 
-		/** Forbidden. Manage with the Ptr<> template. */
-		~NodeMeshReshape() {}
+		//! Forbidden. Manage with the Ptr<> template.
+		~NodeMeshReshape();
 
 	private:
 
-		static FNodeType StaticType;
+		Private* m_pD;
 
 	};
 
diff --git a/Engine/Plugins/Mutable/Source/MutableTools/Public/MuT/NodeMeshSwitch.h b/Engine/Plugins/Mutable/Source/MutableTools/Public/MuT/NodeMeshSwitch.h
--- a/Engine/Plugins/Mutable/Source/MutableTools/Public/MuT/NodeMeshSwitch.h
+++ b/Engine/Plugins/Mutable/Source/MutableTools/Public/MuT/NodeMeshSwitch.h
@@ -3,36 +3,71 @@
 #pragma once
 
 #include "MuR/Ptr.h"
+#include "MuR/RefCounted.h"
+#include "MuT/Node.h"
 #include "MuT/NodeMesh.h"
-#include "MuT/NodeScalar.h"
+
 
 namespace mu
 {
 
-	/** This node selects an output Mesh from a set of input Meshs based on a parameter. */
+	// Forward definitions
+	class NodeScalar;
+	typedef Ptr<NodeScalar> NodeScalarPtr;
+	typedef Ptr<const NodeScalar> NodeScalarPtrConst;
+
+	class NodeMeshSwitch;
+	typedef Ptr<NodeMeshSwitch> NodeMeshSwitchPtr;
+	typedef Ptr<const NodeMeshSwitch> NodeMeshSwitchPtrConst;
+
+
+	//! This node selects an output Mesh from a set of input Meshs based on a parameter.
 	class MUTABLETOOLS_API NodeMeshSwitch : public NodeMesh
 	{
 	public:
 
-		Ptr<NodeScalar> Parameter;
+		NodeMeshSwitch();
 
-		TArray<Ptr<NodeMesh>> Options;
+		//-----------------------------------------------------------------------------------------
+		// Node Interface
+		//-----------------------------------------------------------------------------------------
 
-	public:
+		const FNodeType* GetType() const override;
+		static const FNodeType* GetStaticType();
 
-		// Node Interface
-		virtual const FNodeType* GetType() const override { return &StaticType; }
-		static const FNodeType* GetStaticType() { return &StaticType; }
+		//-----------------------------------------------------------------------------------------
+		// Own Interface
+		//-----------------------------------------------------------------------------------------
+
+		//! Get the node generating the parameter used to select the option.
+		NodeScalarPtr GetParameter() const;
+		void SetParameter( NodeScalarPtr );
+
+		//! Set the number of option Meshs. It will keep the currently set targets and initialise
+		//! the new ones as null.
+		void SetOptionCount( int );
+
+		//! Get the node generating the t-th option Mesh.
+		NodeMeshPtr GetOption( int t ) const;
+		void SetOption( int t, NodeMeshPtr );
+
+
+		//-----------------------------------------------------------------------------------------
+		// Interface pattern
+		//-----------------------------------------------------------------------------------------
+		class Private;
+		Private* GetPrivate() const;
 
 	protected:
 
-		/** Forbidden. Manage with the Ptr<> template. */
-		~NodeMeshSwitch() {}
+		//! Forbidden. Manage with the Ptr<> template.
+		~NodeMeshSwitch();
 
 	private:
 
-		static FNodeType StaticType;
+		Private* m_pD;
 
 	};
 
+
 }
diff --git a/Engine/Plugins/Mutable/Source/MutableTools/Public/MuT/NodeMeshTable.h b/Engine/Plugins/Mutable/Source/MutableTools/Public/MuT/NodeMeshTable.h
--- a/Engine/Plugins/Mutable/Source/MutableTools/Public/MuT/NodeMeshTable.h
+++ b/Engine/Plugins/Mutable/Source/MutableTools/Public/MuT/NodeMeshTable.h
@@ -10,9 +10,16 @@
 #include "MuT/NodeImage.h"
 #include "MuT/Table.h"
 
+
 namespace mu
 {
 
+	// Forward definitions
+	class NodeMeshTable;
+	typedef Ptr<NodeMeshTable> NodeMeshTablePtr;
+	typedef Ptr<const NodeMeshTable> NodeMeshTablePtrConst;
+
+
 	//! This node provides the meshes stored in the column of a table.
 	class MUTABLETOOLS_API NodeMeshTable : public NodeMesh
 	{
@@ -35,9 +42,34 @@ namespace mu
 		virtual const FNodeType* GetType() const override { return GetStaticType(); }
 		static const FNodeType* GetStaticType() { return &StaticType; }
 
+		//-----------------------------------------------------------------------------------------
+		// Own Interface
+		//-----------------------------------------------------------------------------------------
+
+		//! Set the name of the implicit table parameter.
+		void SetParameterName( const FString& strName );
+
+		//!
+		void SetColumn( const FString& strName );
+
+		//! Get the number of layouts defined in the meshes on this column.
+		int GetLayoutCount() const;
+		void SetLayoutCount( int );
+
+		//! Get the node defining a layout of the meshes on this column.
+		Ptr<NodeLayout> GetLayout( int index ) const;
+		void SetLayout( int index, Ptr<NodeLayout>);
+
+		//! Adds the "None" option to the parameter that represents this table column
+		void SetNoneOption(bool bAddOption);
+
+		//! Set the row name to be used as default value
+		void SetDefaultRowName(const FString& RowName);
+
+
 	protected:
 
-		/** Forbidden. Manage with the Ptr<> template. */
+		//! Forbidden. Manage with the Ptr<> template.
 		~NodeMeshTable() {}
 
 	private:
@@ -46,4 +78,5 @@ namespace mu
 
 	};
 
+
 }
diff --git a/Engine/Plugins/Mutable/Source/MutableTools/Public/MuT/NodeMeshTangents.h b/Engine/Plugins/Mutable/Source/MutableTools/Public/MuT/NodeMeshTangents.h
new file mode 100644
--- /dev/null
+++ b/Engine/Plugins/Mutable/Source/MutableTools/Public/MuT/NodeMeshTangents.h
@@ -0,0 +1,60 @@
+// Copyright Epic Games, Inc. All Rights Reserved.
+
+#pragma once
+
+#include "MuR/Ptr.h"
+#include "MuR/RefCounted.h"
+#include "MuT/Node.h"
+#include "MuT/NodeMesh.h"
+
+
+namespace mu
+{
+
+	// Forward definitions
+	class NodeMeshTangents;
+	typedef Ptr<NodeMeshTangents> NodeMeshTangentsPtr;
+	typedef Ptr<const NodeMeshTangents> NodeMeshTangentsPtrConst;
+
+
+	//! This node rebuilds the tangents and binormals of the source mesh.
+	class MUTABLETOOLS_API NodeMeshTangents : public NodeMesh
+	{
+	public:
+
+		NodeMeshTangents();
+
+		//-----------------------------------------------------------------------------------------
+		// Node Interface
+		//-----------------------------------------------------------------------------------------
+
+        const FNodeType* GetType() const override;
+		static const FNodeType* GetStaticType();
+
+		//-----------------------------------------------------------------------------------------
+		// Own Interface
+		//-----------------------------------------------------------------------------------------
+
+		//!
+		NodeMeshPtr GetSource() const;
+		void SetSource( NodeMesh* p );
+
+		//-----------------------------------------------------------------------------------------
+		// Interface pattern
+		//-----------------------------------------------------------------------------------------
+		class Private;
+		Private* GetPrivate() const;
+
+	protected:
+
+		//! Forbidden. Manage with the Ptr<> template.
+		~NodeMeshTangents();
+
+	private:
+
+		Private* m_pD;
+
+	};
+
+
+}
diff --git a/Engine/Plugins/Mutable/Source/MutableTools/Public/MuT/NodeMeshTransform.h b/Engine/Plugins/Mutable/Source/MutableTools/Public/MuT/NodeMeshTransform.h
--- a/Engine/Plugins/Mutable/Source/MutableTools/Public/MuT/NodeMeshTransform.h
+++ b/Engine/Plugins/Mutable/Source/MutableTools/Public/MuT/NodeMeshTransform.h
@@ -6,34 +6,58 @@
 #include "MuR/RefCounted.h"
 #include "MuT/Node.h"
 #include "MuT/NodeMesh.h"
-#include "Math/Matrix.h"
+#include "Math/MathFwd.h"
 
 namespace mu
 {
 
-    /** This node applies a geometric transform represented by a 4x4 matrix to a mesh. */
+	// Forward definitions
+    class NodeMeshTransform;
+    typedef Ptr<NodeMeshTransform> NodeMeshTransformPtr;
+    typedef Ptr<const NodeMeshTransform> NodeMeshTransformPtrConst;
+
+
+    //! This node applies a geometric transform represented by a 4x4 matrix to a mesh
     class MUTABLETOOLS_API NodeMeshTransform : public NodeMesh
 	{
 	public:
 
-		Ptr<NodeMesh> Source;
-		FMatrix44f Transform;
-
-	public:
+        NodeMeshTransform();
 
+		//-----------------------------------------------------------------------------------------
 		// Node Interface
-		virtual const FNodeType* GetType() const override { return &StaticType; }
-		static const FNodeType* GetStaticType() { return &StaticType; }
+		//-----------------------------------------------------------------------------------------
+
+        const FNodeType* GetType() const override;
+		static const FNodeType* GetStaticType();
+
+		//-----------------------------------------------------------------------------------------
+		// Own Interface
+		//-----------------------------------------------------------------------------------------
+
+		//! Source mesh to be re-formatted
+        NodeMeshPtr GetSource() const;
+		void SetSource( NodeMesh* );
+
+        void SetTransform(const FMatrix44f&);
+		const FMatrix44f& GetTransform() const;
+
+		//-----------------------------------------------------------------------------------------
+		// Interface pattern
+		//-----------------------------------------------------------------------------------------
+		class Private;
+		Private* GetPrivate() const;
 
 	protected:
 
-		/** Forbidden. Manage with the Ptr<> template. */
-		~NodeMeshTransform() {}
+		//! Forbidden. Manage with the Ptr<> template.
+        ~NodeMeshTransform();
 
 	private:
 
-		static FNodeType StaticType;
+		Private* m_pD;
 
 	};
 
+
 }
diff --git a/Engine/Plugins/Mutable/Source/MutableTools/Public/MuT/NodeMeshVariation.h b/Engine/Plugins/Mutable/Source/MutableTools/Public/MuT/NodeMeshVariation.h
--- a/Engine/Plugins/Mutable/Source/MutableTools/Public/MuT/NodeMeshVariation.h
+++ b/Engine/Plugins/Mutable/Source/MutableTools/Public/MuT/NodeMeshVariation.h
@@ -7,40 +7,65 @@
 #include "MuT/Node.h"
 #include "MuT/NodeMesh.h"
 
+
 namespace mu
 {
 
-	/** */
-	class MUTABLETOOLS_API NodeMeshVariation : public NodeMesh
+    // Forward definitions
+    class NodeMeshVariation;
+    typedef Ptr<NodeMeshVariation> NodeMeshVariationPtr;
+    typedef Ptr<const NodeMeshVariation> NodeMeshVariationPtrConst;
+
+    //!
+    //! \ingroup model
+    class MUTABLETOOLS_API NodeMeshVariation : public NodeMesh
     {
-	public:
+    public:
 
-		Ptr<NodeMesh> DefaultMesh;
+        NodeMeshVariation();
 
-		struct FVariation
-		{
-			Ptr<NodeMesh> Mesh;
-			FString Tag;
-		};
+        //-----------------------------------------------------------------------------------------
+        // Node Interface
+        //-----------------------------------------------------------------------------------------
 
-		TArray<FVariation> Variations;
+        const FNodeType* GetType() const override;
+        static const FNodeType* GetStaticType();
 
-    public:
+        //-----------------------------------------------------------------------------------------
+        // Own Interface
+        //-----------------------------------------------------------------------------------------
 
-		// Node Interface
-		virtual const FNodeType* GetType() const override { return &StaticType; }
-		static const FNodeType* GetStaticType() { return &StaticType; }
+        //!
+        void SetDefaultMesh( NodeMesh* Mesh );
 
-    protected:
+        //! Set the number of tags to consider in this variation
+        void SetVariationCount( int count );
 
-		/** Forbidden. Manage with the Ptr<> template. */
-		~NodeMeshVariation() {}
+        //!
+        int GetVariationCount() const;
 
-    private:
+        //! Set the tag or state name that will enable a specific vartiation
+        void SetVariationTag( int index, const char* strTag );
 
-		static FNodeType StaticType;
+        //!
+        void SetVariationMesh( int index, NodeMesh* Mesh );
 
+        //!}
+
+
+        //-----------------------------------------------------------------------------------------
+        // Interface pattern
+        //-----------------------------------------------------------------------------------------
+        class Private;
+        Private* GetPrivate() const;
+
+    protected:
+        //! Forbidden. Manage with the Ptr<> template.
+        ~NodeMeshVariation();
+
+    private:
+        Private* m_pD;
     };
 
-}
 
+} // namespace mu
diff --git a/Engine/Plugins/Mutable/Source/MutableTools/Public/MuT/NodeModifierMeshClipMorphPlane.h b/Engine/Plugins/Mutable/Source/MutableTools/Public/MuT/NodeModifierMeshClipMorphPlane.h
--- a/Engine/Plugins/Mutable/Source/MutableTools/Public/MuT/NodeModifierMeshClipMorphPlane.h
+++ b/Engine/Plugins/Mutable/Source/MutableTools/Public/MuT/NodeModifierMeshClipMorphPlane.h
@@ -9,9 +9,12 @@
 #include "MuT/Node.h"
 #include "MuT/NodeModifier.h"
 
+
 namespace mu
 {
 
+	/** */
+
 	struct FBoneName;
 
 	struct FClipMorphPlaneParameters
@@ -74,7 +77,7 @@ namespace mu
 
 	protected:
 
-		/** Forbidden. Manage with the Ptr<> template. */
+		//! Forbidden. Manage with the Ptr<> template.
 		~NodeModifierMeshClipMorphPlane() {}
 
 	private:
diff --git a/Engine/Plugins/Mutable/Source/MutableTools/Public/MuT/NodeObject.h b/Engine/Plugins/Mutable/Source/MutableTools/Public/MuT/NodeObject.h
--- a/Engine/Plugins/Mutable/Source/MutableTools/Public/MuT/NodeObject.h
+++ b/Engine/Plugins/Mutable/Source/MutableTools/Public/MuT/NodeObject.h
@@ -7,19 +7,32 @@
 #include "MuR/RefCounted.h"
 #include "MuT/Node.h"
 
+
 namespace mu
 {
 
-	/** Base class for any node that outputs an object. */
+	// Forward definitions
+	class NodeObject;
+	typedef Ptr<NodeObject> NodeObjectPtr;
+	typedef Ptr<const NodeObject> NodeObjectPtrConst;
+
+
+    //! %Base class for any node that outputs an object.
+	//! \ingroup model
 	class MUTABLETOOLS_API NodeObject : public Node
 	{
 	public:
 
+		//-----------------------------------------------------------------------------------------
 		// Node Interface
-		virtual const FNodeType* GetType() const override { return &StaticType; }
-		static const FNodeType* GetStaticType() { return &StaticType; }
+		//-----------------------------------------------------------------------------------------
+
+		const FNodeType* GetType() const override;
+		static const FNodeType* GetStaticType();
 
+		//-----------------------------------------------------------------------------------------
 		// Own interface
+		//-----------------------------------------------------------------------------------------
 
 		//! Get the name of the object.
 		virtual const FString& GetName() const = 0;
@@ -33,15 +46,17 @@ namespace mu
 		//! Set the uid of the object.
 		virtual void SetUid( const FString& ) = 0;
 
+		//-----------------------------------------------------------------------------------------
+		// Interface pattern
+		//-----------------------------------------------------------------------------------------
+		class Private;
+
 	protected:
 
-		/** Forbidden. Manage with the Ptr<> template. */
+		//! Forbidden. Manage with the Ptr<> template.
 		inline ~NodeObject() {}
 
-	private:
-
-		static FNodeType StaticType;
-
 	};
 
+
 }
diff --git a/Engine/Plugins/Mutable/Source/MutableTools/Public/MuT/NodeObjectGroup.h b/Engine/Plugins/Mutable/Source/MutableTools/Public/MuT/NodeObjectGroup.h
--- a/Engine/Plugins/Mutable/Source/MutableTools/Public/MuT/NodeObjectGroup.h
+++ b/Engine/Plugins/Mutable/Source/MutableTools/Public/MuT/NodeObjectGroup.h
@@ -7,16 +7,42 @@
 #include "MuT/Node.h"
 #include "MuT/NodeObject.h"
 
+
 namespace mu
 {
 
-	/** Node that creates a group of objects and describes how they are selected. */
+	// Forward definitions
+	class NodeObjectGroup;
+	typedef Ptr<NodeObjectGroup> NodeObjectGroupPtr;
+	typedef Ptr<const NodeObjectGroup>NodeObjectGroupPtrConst;
+
+
+	//! Node that creates a group of objects and describes how they are selected.
+	//! \ingroup model
 	class MUTABLETOOLS_API NodeObjectGroup : public NodeObject
 	{
 	public:
 
-		FString Name;
-		FString Uid;
+		NodeObjectGroup();
+
+
+		//-----------------------------------------------------------------------------------------
+		// Node Interface
+		//-----------------------------------------------------------------------------------------
+        const FNodeType* GetType() const override;
+		static const FNodeType* GetStaticType();
+
+		//-----------------------------------------------------------------------------------------
+		// NodeObject Interface
+		//-----------------------------------------------------------------------------------------
+        virtual const FString& GetName() const override;
+		virtual void SetName( const FString&) override;
+		virtual const FString& GetUid() const override;
+		virtual void SetUid( const FString&) override;
+
+		//-----------------------------------------------------------------------------------------
+		// Own Interface
+		//-----------------------------------------------------------------------------------------
 
 		//! Typed of child selection
 		typedef enum
@@ -36,38 +62,46 @@ namespace mu
 			//! A boolean parameter will be generated for every child.
 			CS_TOGGLE_EACH
 
-		} EChildSelection;
+		} CHILD_SELECTION;
 
-		EChildSelection Type;
+		//! Get child selection type
+		CHILD_SELECTION GetSelectionType() const;
 
-		TArray<Ptr<NodeObject>> Children;
+		//! Set the child selection type
+		void SetSelectionType( CHILD_SELECTION );
 
-		//! Default value for CS_ONE_OR_NONE or CS_ALWAYS_ONE groups
+		//! Get the number of child objects
+		int GetChildCount() const;
+
+		//! Set the number of child objects
+		void SetChildCount( int );
+
+		//! Get a child object node
+		NodeObjectPtr GetChild( int index ) const;
+
+		//! Set a child object node
+		void SetChild( int index, NodeObjectPtr );
+
+		//! Set default value for CS_ONE_OR_NONE or CS_ALWAYS_ONE groups
 		//! the value is the index of the child option in the group
 		//! -1 is the value for the None option.
 		//! 0 is the first child whether or not the NONE option is present.
-		int32 DefaultValue;
+		void SetDefaultValue(int32 Value);
 
-	public:
-
-		// Node Interface
-		virtual const FNodeType* GetType() const override { return &StaticType; }
-		static const FNodeType* GetStaticType() { return &StaticType; }
-
-		// NodeObject Interface
-        virtual const FString& GetName() const override;
-		virtual void SetName( const FString&) override;
-		virtual const FString& GetUid() const override;
-		virtual void SetUid( const FString&) override;
+		//-----------------------------------------------------------------------------------------
+		// Interface pattern
+		//-----------------------------------------------------------------------------------------
+		class Private;
+		Private* GetPrivate() const;
 
 	protected:
 
-		/** Forbidden. Manage with the Ptr<> template. */
-		~NodeObjectGroup() {}
+		//! Forbidden. Manage with the Ptr<> template.
+		~NodeObjectGroup();
 
 	private:
 
-		static FNodeType StaticType;
+		Private* m_pD;
 
 	};
 
diff --git a/Engine/Plugins/Mutable/Source/MutableTools/Public/MuT/NodeObjectNew.h b/Engine/Plugins/Mutable/Source/MutableTools/Public/MuT/NodeObjectNew.h
--- a/Engine/Plugins/Mutable/Source/MutableTools/Public/MuT/NodeObjectNew.h
+++ b/Engine/Plugins/Mutable/Source/MutableTools/Public/MuT/NodeObjectNew.h
@@ -49,17 +49,23 @@ namespace mu
 
 	public:
 
+		//-----------------------------------------------------------------------------------------
 		// Node Interface
-		virtual const FNodeType* GetType() const override { return &StaticType; }
+		//-----------------------------------------------------------------------------------------
+		virtual const FNodeType* GetType() const override { return GetStaticType(); }
 		static const FNodeType* GetStaticType() { return &StaticType; }
 
+		//-----------------------------------------------------------------------------------------
 		// NodeObject Interface
+		//-----------------------------------------------------------------------------------------
         virtual const FString& GetName() const override;
 		virtual void SetName( const FString& ) override;
 		virtual const FString& GetUid() const override;
 		virtual void SetUid( const FString& ) override;
 
+		//-----------------------------------------------------------------------------------------
 		// Own Interface
+		//-----------------------------------------------------------------------------------------
 
 		//! Set the number of states that the model can be in.
 		int32 GetStateCount() const;
@@ -74,6 +80,9 @@ namespace mu
 		//! Add a runtime parameter to the state.
 		void AddStateParam( int32 s, const FString& param );
 
+		//! Remove a runtime parameter from a state.
+		void RemoveStateParam( int32 s, const FString& param );
+
         //! Set the optimisation properties of a state
 		void SetStateProperties(int32 StateIndex,
 			ETextureCompressionStrategy TextureCompressionStrategy,
@@ -85,13 +94,14 @@ namespace mu
 
 	protected:
 
-		/** Forbidden. Manage with the Ptr<> template. */
+		//! Forbidden. Manage with the Ptr<> template.
 		~NodeObjectNew() {}
 
 	private:
 
 		static FNodeType StaticType;
 
+
 	};
 
 
