[Backout] - CL40275509

diff --git a/Engine/Plugins/Mutable/Source/MutableTools/Public/MuT/NodeProjector.h b/Engine/Plugins/Mutable/Source/MutableTools/Public/MuT/NodeProjector.h
--- a/Engine/Plugins/Mutable/Source/MutableTools/Public/MuT/NodeProjector.h
+++ b/Engine/Plugins/Mutable/Source/MutableTools/Public/MuT/NodeProjector.h
@@ -7,29 +7,50 @@
 #include "MuR/Ptr.h"
 #include "MuR/RefCounted.h"
 #include "MuT/Node.h"
-#include "MuT/NodeRange.h"
 
 
 namespace mu
 {
 
+	class NodeProjector;
+    using NodeProjectorPtr = Ptr<NodeProjector>;
+    using NodeProjectorPtrConst = Ptr<const NodeProjector>;
+
+	class NodeProjectorConstant;
+    using NodeProjectorConstantPtr = Ptr<NodeProjectorConstant>;
+    using NodeProjectorConstantPtrConst = Ptr<const NodeProjectorConstant>;
+
+	class NodeProjectorParameter;
+    using NodeProjectorParameterPtr = Ptr<NodeProjectorParameter>;
+    using NodeProjectorParameterPtrConst = Ptr<const NodeProjectorParameter>;
+
+    class NodeRange;
+    using NodeRangePtr = Ptr<NodeRange>;
+    using NodeRangePtrConst = Ptr<const NodeRange>;
+
+
     //! Base class of any node that outputs a Projector.
 	class MUTABLETOOLS_API NodeProjector : public Node
 	{
 	public:
 
+		//-----------------------------------------------------------------------------------------
 		// Node Interface
-		virtual const FNodeType* GetType() const override { return &StaticType; }
-		static const FNodeType* GetStaticType() { return &StaticType; }
+		//-----------------------------------------------------------------------------------------
 
-	protected:
+        const FNodeType* GetType() const override;
+		static const FNodeType* GetStaticType();
 
-		/** Forbidden. Manage with the Ptr<> template. */
-		inline ~NodeProjector() {}
 
-	private:
+		//-----------------------------------------------------------------------------------------
+		// Interface pattern
+		//-----------------------------------------------------------------------------------------
+		class Private;
 
-		static FNodeType StaticType;
+	protected:
+
+		//! Forbidden. Manage with the Ptr<> template.
+		inline ~NodeProjector() {}
 
 	};
 
@@ -39,20 +60,18 @@ namespace mu
 	{
 	public:
 
-		EProjectorType Type = EProjectorType::Planar;
-		FVector3f Position;
-		FVector3f Direction;
-		FVector3f Up;
-		FVector3f Scale;
-		float ProjectionAngle = 0.0f;
-
-	public:
+		NodeProjectorConstant();
 
+		//-----------------------------------------------------------------------------------------
 		// Node Interface
-		virtual const FNodeType* GetType() const override { return &StaticType; }
-		static const FNodeType* GetStaticType() { return &StaticType; }
+		//-----------------------------------------------------------------------------------------
+
+        const FNodeType* GetType() const override;
+		static const FNodeType* GetStaticType();
 
+		//-----------------------------------------------------------------------------------------
 		// Own Interface
+		//-----------------------------------------------------------------------------------------
 
 		//! Get the value that this node returns
         void GetValue( EProjectorType* OutType,
@@ -65,14 +84,20 @@ namespace mu
 		//! Set the value to be returned by this node
         void SetValue( EProjectorType type, FVector3f pos, FVector3f dir, FVector3f up, FVector3f scale, float projectionAngle);
 
+		//-----------------------------------------------------------------------------------------
+		// Interface pattern
+		//-----------------------------------------------------------------------------------------
+		class Private;
+		Private* GetPrivate() const;
+
 	protected:
 
-		/** Forbidden. Manage with the Ptr<> template. */
-		~NodeProjectorConstant() {}
+		//! Forbidden. Manage with the Ptr<> template.
+		~NodeProjectorConstant();
 
 	private:
 
-		static FNodeType StaticType;
+		Private* m_pD;
 
 	};
 
@@ -82,29 +107,23 @@ namespace mu
 	{
 	public:
 
-		EProjectorType Type = EProjectorType::Planar;
-		FVector3f Position;
-		FVector3f Direction;
-		FVector3f Up;
-		FVector3f Scale;
-		float ProjectionAngle = 0.0f;
-
-		FString Name;
-		FString UID;
-
-		TArray<Ptr<NodeRange>> Ranges;
-
-	public:
+		NodeProjectorParameter();
 
+		//-----------------------------------------------------------------------------------------
 		// Node Interface
-		virtual const FNodeType* GetType() const override { return &StaticType; }
-		static const FNodeType* GetStaticType() { return &StaticType; }
+		//-----------------------------------------------------------------------------------------
+
+        const FNodeType* GetType() const override;
+		static const FNodeType* GetStaticType();
 
+		//-----------------------------------------------------------------------------------------
 		// Own Interface
+		//-----------------------------------------------------------------------------------------
 
 		//! Set the name of the parameter.
 		void SetName( const FString& );
 
+
 		//! Get the uid of the parameter. It will be exposed in the final compiled data.
 		const FString& GetUid() const;
 		void SetUid( const FString& );
@@ -119,17 +138,23 @@ namespace mu
 
         //! Set the number of ranges (dimensions) for this parameter.
         //! By default a parameter has 0 ranges, meaning it only has one value.
-        void SetRangeCount( int32 i );
-        void SetRange( int32 i, Ptr<NodeRange> Range );
+        void SetRangeCount( int i );
+        void SetRange( int i, NodeRangePtr pRange );
+
+		//-----------------------------------------------------------------------------------------
+		// Interface pattern
+		//-----------------------------------------------------------------------------------------
+		class Private;
+		Private* GetPrivate() const;
 
 	protected:
 
-		/** Forbidden. Manage with the Ptr<> template. */
-		~NodeProjectorParameter() {}
+		//! Forbidden. Manage with the Ptr<> template.
+		~NodeProjectorParameter();
 
 	private:
 
-		static FNodeType StaticType;
+		Private* m_pD;
 
 	};
 
diff --git a/Engine/Plugins/Mutable/Source/MutableTools/Public/MuT/NodeRange.h b/Engine/Plugins/Mutable/Source/MutableTools/Public/MuT/NodeRange.h
--- a/Engine/Plugins/Mutable/Source/MutableTools/Public/MuT/NodeRange.h
+++ b/Engine/Plugins/Mutable/Source/MutableTools/Public/MuT/NodeRange.h
@@ -7,28 +7,41 @@
 #include "MuR/RefCounted.h"
 #include "MuT/Node.h"
 
+
 namespace mu
 {
 
-    /** Base class of any node that outputs a range. */
+    // Forward definitions
+    class NodeRange;
+    typedef Ptr<NodeRange> NodeRangePtr;
+    typedef Ptr<const NodeRange> NodeRangePtrConst;
+
+
+    //! %Base class of any node that outputs a range.
     class MUTABLETOOLS_API NodeRange : public Node
 	{
 	public:
 
+		//-----------------------------------------------------------------------------------------
 		// Node Interface
-		virtual const FNodeType* GetType() const override { return &StaticType; }
-		static const FNodeType* GetStaticType() { return &StaticType; }
+		//-----------------------------------------------------------------------------------------
 
-	protected:
+		const FNodeType* GetType() const override;
+		static const FNodeType* GetStaticType();
 
-		/** Forbidden. Manage with the Ptr<> template. */
-        inline ~NodeRange() {}
 
-	private:
+		//-----------------------------------------------------------------------------------------
+		// Interface pattern
+		//-----------------------------------------------------------------------------------------
+		class Private;
 
-		static FNodeType StaticType;
+	protected:
+
+		//! Forbidden. Manage with the Ptr<> template.
+        inline ~NodeRange() {}
 
 	};
 
+
 }
 
diff --git a/Engine/Plugins/Mutable/Source/MutableTools/Public/MuT/NodeRangeFromScalar.h b/Engine/Plugins/Mutable/Source/MutableTools/Public/MuT/NodeRangeFromScalar.h
--- a/Engine/Plugins/Mutable/Source/MutableTools/Public/MuT/NodeRangeFromScalar.h
+++ b/Engine/Plugins/Mutable/Source/MutableTools/Public/MuT/NodeRangeFromScalar.h
@@ -3,8 +3,8 @@
 #pragma once
 
 #include "MuR/Ptr.h"
+#include "MuT/Node.h"
 #include "MuT/NodeRange.h"
-#include "MuT/NodeScalar.h"
 #include "Containers/UnrealString.h"
 
 namespace mu
@@ -17,25 +17,48 @@ namespace mu
 	{
 	public:
 
-		Ptr<NodeScalar> Size;
+		//-----------------------------------------------------------------------------------------
+		// Life cycle
+		//-----------------------------------------------------------------------------------------
 
-		FString Name;
+        NodeRangeFromScalar();
 
-	public:
 
+		//-----------------------------------------------------------------------------------------
 		// Node Interface
-		virtual const FNodeType* GetType() const override { return &StaticType; }
-		static const FNodeType* GetStaticType() { return &StaticType; }
+		//-----------------------------------------------------------------------------------------
+
+		const FNodeType* GetType() const override;
+		static const FNodeType* GetStaticType();
+
+		//-----------------------------------------------------------------------------------------
+		// Own Interface
+		//-----------------------------------------------------------------------------------------
+
+        //!
+        Ptr<NodeScalar> GetSize() const;
+        void SetSize( const Ptr<NodeScalar>& );
+
+        //!
+		const FString& GetName() const;
+		void SetName(const FString& Name);
+
+		//-----------------------------------------------------------------------------------------
+		// Interface pattern
+		//-----------------------------------------------------------------------------------------
+		class Private;
+		Private* GetPrivate() const;
 
 	protected:
 
-		/** Forbidden. Manage with the Ptr<> template. */
-		~NodeRangeFromScalar() {}
+		//! Forbidden. Manage with the Ptr<> template.
+        ~NodeRangeFromScalar();
 
 	private:
 
-		static FNodeType StaticType;
+		Private* m_pD;
 
 	};
 
+
 }
diff --git a/Engine/Plugins/Mutable/Source/MutableTools/Public/MuT/NodeScalar.h b/Engine/Plugins/Mutable/Source/MutableTools/Public/MuT/NodeScalar.h
--- a/Engine/Plugins/Mutable/Source/MutableTools/Public/MuT/NodeScalar.h
+++ b/Engine/Plugins/Mutable/Source/MutableTools/Public/MuT/NodeScalar.h
@@ -2,29 +2,45 @@
 
 #pragma once
 
+#include "HAL/Platform.h"
+#include "MuR/Ptr.h"
+#include "MuR/RefCounted.h"
 #include "MuT/Node.h"
 
+
 namespace mu
 {
 
-    /** Base class of any node that outputs a scalar value. */
+	// Forward definitions
+	class NodeScalar;
+	typedef Ptr<NodeScalar> NodeScalarPtr;
+	typedef Ptr<const NodeScalar> NodeScalarPtrConst;
+
+
+    //! %Base class of any node that outputs a scalar value.
 	class MUTABLETOOLS_API NodeScalar : public Node
 	{
 	public:
 
+		//-----------------------------------------------------------------------------------------
 		// Node Interface
-		virtual const FNodeType* GetType() const override { return &StaticType; }
-		static const FNodeType* GetStaticType() { return &StaticType; }
+		//-----------------------------------------------------------------------------------------
 
-	protected:
+		const FNodeType* GetType() const override;
+		static const FNodeType* GetStaticType();
 
-		/** Forbidden. Manage with the Ptr<> template. */
-		inline ~NodeScalar() {}
 
-	private:
+		//-----------------------------------------------------------------------------------------
+		// Interface pattern
+		//-----------------------------------------------------------------------------------------
+		class Private;
 
-		static FNodeType StaticType;
+	protected:
+
+		//! Forbidden. Manage with the Ptr<> template.
+		inline ~NodeScalar() {}
 
 	};
 
+
 }
diff --git a/Engine/Plugins/Mutable/Source/MutableTools/Public/MuT/NodeScalarArithmeticOperation.h b/Engine/Plugins/Mutable/Source/MutableTools/Public/MuT/NodeScalarArithmeticOperation.h
--- a/Engine/Plugins/Mutable/Source/MutableTools/Public/MuT/NodeScalarArithmeticOperation.h
+++ b/Engine/Plugins/Mutable/Source/MutableTools/Public/MuT/NodeScalarArithmeticOperation.h
@@ -7,43 +7,80 @@
 #include "MuT/Node.h"
 #include "MuT/NodeScalar.h"
 
+
 namespace mu
 {
 
-    /** Perform an arithmetic operation between two scalars. */
+	// Forward definitions
+    class NodeScalarArithmeticOperation;
+    typedef Ptr<NodeScalarArithmeticOperation> NodeScalarArithmeticOperationPtr;
+    typedef Ptr<const NodeScalarArithmeticOperation> NodeScalarArithmeticOperationPtrConst;
+
+
+    //! Perform an arithmetic operation between two scalars.
+	//! \ingroup model
     class MUTABLETOOLS_API NodeScalarArithmeticOperation : public NodeScalar
 	{
 	public:
 
-		/** Possible arithmetic operations. */
+        NodeScalarArithmeticOperation();
+
+
+		//-----------------------------------------------------------------------------------------
+		// Node Interface
+		//-----------------------------------------------------------------------------------------
+
+        const FNodeType* GetType() const override;
+		static const FNodeType* GetStaticType();
+
+
+		//-----------------------------------------------------------------------------------------
+		// Own Interface
+		//-----------------------------------------------------------------------------------------
+
+		//! Possible arithmetic operations
 		typedef enum
 		{
-			AO_ADD = 0,
+            AO_ADD = 0,
 			AO_SUBTRACT,
 			AO_MULTIPLY,
-			AO_DIVIDE,
-			_AO_COUNT
-		} EOperation;
+            AO_DIVIDE,
+            _AO_COUNT
+		} OPERATION;
 
-		EOperation Operation;
-		Ptr<NodeScalar> A;
-		Ptr<NodeScalar> B;
+        static const char* s_opTypeName[ _AO_COUNT ];
 
-	public:
+		//! Get child selection type
+		OPERATION GetOperation() const;
+
+		//! Set the child selection type
+		void SetOperation(OPERATION);
 
-		// Node Interface
-		virtual const FNodeType* GetType() const override { return &StaticType; }
-		static const FNodeType* GetStaticType() { return &StaticType; }
+
+		//! Get the first operand
+        NodeScalarPtr GetA() const;
+        void SetA(NodeScalarPtr);
+
+		// Get the second operand
+        NodeScalarPtr GetB() const;
+        void SetB(NodeScalarPtr);
+
+		//-----------------------------------------------------------------------------------------
+		// Interface pattern
+		//-----------------------------------------------------------------------------------------
+		class Private;
+		Private* GetPrivate() const;
 
 	protected:
 
-		/** Forbidden. Manage with the Ptr<> template. */
-		~NodeScalarArithmeticOperation() {}
+		//! Forbidden. Manage with the Ptr<> template.
+        ~NodeScalarArithmeticOperation();
 
 	private:
 
-		static FNodeType StaticType;
+		Private* m_pD;
 
 	};
 
+
 }
diff --git a/Engine/Plugins/Mutable/Source/MutableTools/Public/MuT/NodeScalarConstant.h b/Engine/Plugins/Mutable/Source/MutableTools/Public/MuT/NodeScalarConstant.h
--- a/Engine/Plugins/Mutable/Source/MutableTools/Public/MuT/NodeScalarConstant.h
+++ b/Engine/Plugins/Mutable/Source/MutableTools/Public/MuT/NodeScalarConstant.h
@@ -7,31 +7,59 @@
 #include "MuT/Node.h"
 #include "MuT/NodeScalar.h"
 
+
 namespace mu
 {
 
-	/** Node returning a scalar constant value. */
+	// Forward definitions
+	class NodeScalarConstant;
+	typedef Ptr<NodeScalarConstant> NodeScalarConstantPtr;
+	typedef Ptr<const NodeScalarConstant> NodeScalarConstantPtrConst;
+
+
+	//! Node returning a scalar constant value.
+	//! \ingroup model
 	class MUTABLETOOLS_API NodeScalarConstant : public NodeScalar
 	{
 	public:
 
-		float Value;
+		NodeScalarConstant();
 
-	public:
 
+		//-----------------------------------------------------------------------------------------
 		// Node Interface
-		virtual const FNodeType* GetType() const override { return &StaticType; }
-		static const FNodeType* GetStaticType() { return &StaticType; }
+		//-----------------------------------------------------------------------------------------
+
+        const FNodeType* GetType() const override;
+		static const FNodeType* GetStaticType();
+
+		//-----------------------------------------------------------------------------------------
+		// Own Interface
+		//-----------------------------------------------------------------------------------------
+
+		//! Get the value to be returned by the node.
+		float GetValue() const;
+
+		//! Set the value to be returned by the node.
+		void SetValue( float v );
+
+		//-----------------------------------------------------------------------------------------
+		// Interface pattern
+		//-----------------------------------------------------------------------------------------
+		class Private;
+		Private* GetPrivate() const;
 
 	protected:
 
-		/** Forbidden. Manage with the Ptr<> template. */
-		~NodeScalarConstant() {}
+		//! Forbidden. Manage with the Ptr<> template.
+		~NodeScalarConstant();
 
 	private:
 
-		static FNodeType StaticType;
+		Private* m_pD;
 
 	};
 
+
+
 }
diff --git a/Engine/Plugins/Mutable/Source/MutableTools/Public/MuT/NodeScalarCurve.h b/Engine/Plugins/Mutable/Source/MutableTools/Public/MuT/NodeScalarCurve.h
--- a/Engine/Plugins/Mutable/Source/MutableTools/Public/MuT/NodeScalarCurve.h
+++ b/Engine/Plugins/Mutable/Source/MutableTools/Public/MuT/NodeScalarCurve.h
@@ -6,6 +6,7 @@
 #include "MuR/RefCounted.h"
 #include "MuT/Node.h"
 #include "MuT/NodeScalar.h"
+
 #include "Curves/RichCurve.h"
 
 namespace mu
@@ -21,19 +22,35 @@ namespace mu
 
 	public:
 
+        NodeScalarCurve();
+
+		//-----------------------------------------------------------------------------------------
 		// Node Interface
-		virtual const FNodeType* GetType() const override { return &StaticType; }
-		static const FNodeType* GetStaticType() { return &StaticType; }
+		//-----------------------------------------------------------------------------------------
+
+        const FNodeType* GetType() const override;
+		static const FNodeType* GetStaticType();
+
+		//-----------------------------------------------------------------------------------------
+		// Own Interface
+		//-----------------------------------------------------------------------------------------
+
+		//-----------------------------------------------------------------------------------------
+		// Interface pattern
+		//-----------------------------------------------------------------------------------------
+		class Private;
+		Private* GetPrivate() const;
 
 	protected:
 
-		/** Forbidden. Manage with the Ptr<> template. */
-		~NodeScalarCurve() {}
+		//! Forbidden. Manage with the Ptr<> template.
+        ~NodeScalarCurve();
 
 	private:
 
-		static FNodeType StaticType;
+		Private* m_pD;
 
 	};
 
+
 }
diff --git a/Engine/Plugins/Mutable/Source/MutableTools/Public/MuT/NodeScalarEnumParameter.h b/Engine/Plugins/Mutable/Source/MutableTools/Public/MuT/NodeScalarEnumParameter.h
--- a/Engine/Plugins/Mutable/Source/MutableTools/Public/MuT/NodeScalarEnumParameter.h
+++ b/Engine/Plugins/Mutable/Source/MutableTools/Public/MuT/NodeScalarEnumParameter.h
@@ -8,43 +8,83 @@
 #include "MuT/NodeRange.h"
 #include "MuT/NodeScalar.h"
 
+
 namespace mu
 {
 
-	/** Node that defines a scalar model parameter to be selected from a set of named values. */
+	// Forward definitions
+	class NodeImage;
+    using NodeImagePtr = Ptr<NodeImage>;
+    using NodeImagePtrConst = Ptr<const NodeImage>;
+
+	class NodeScalarEnumParameter;
+    using NodeScalarEnumParameterPtr = Ptr<NodeScalarEnumParameter>;
+    using NodeScalarEnumParameterPtrConst = Ptr<const NodeScalarEnumParameter>;
+
+    class NodeRange;
+    using NodeRangePtr = Ptr<NodeRange>;
+    using NodeRangePtrConst = Ptr<const NodeRange>;
+
+
+	//! Node that defines a scalar model parameter to be selected from a set of named values.
 	class MUTABLETOOLS_API NodeScalarEnumParameter : public NodeScalar
 	{
 	public:
 
-		FString Name;
-		FString UID;
-		int32 DefaultValue = 0;
+		NodeScalarEnumParameter();
 
-		struct FOption
-		{
-			FString Name;
-			float Value = 0.0f;
-		};
+		//-----------------------------------------------------------------------------------------
+		// Node Interface
+		//-----------------------------------------------------------------------------------------
+        const FNodeType* GetType() const override;
+		static const FNodeType* GetStaticType();
 
-		TArray<FOption> Options;
+		//-----------------------------------------------------------------------------------------
+		// Own Interface
+		//-----------------------------------------------------------------------------------------
 
-		TArray<Ptr<NodeRange>> Ranges;
+		//! Get the name of the parameter. It will be exposed in the final compiled data.
+		const FString& GetName() const;
+		void SetName( const FString&);
 
-	public:
+		//! Get the uid of the parameter. It will be exposed in the final compiled data.
+		const FString& GetUid() const;
+		void SetUid( const FString&);
 
-		// Node Interface
-		virtual const FNodeType* GetType() const override { return &StaticType; }
-		static const FNodeType* GetStaticType() { return &StaticType; }
+		//! Get the index of the default value of the parameter.
+		int GetDefaultValueIndex() const;
+		void SetDefaultValueIndex( int v );
+
+		//! Set the number of possible values.
+		void SetValueCount( int i );
+
+		//! Get the number of possible values.
+		int GetValueCount() const;
+
+		//! Set the data of one of the possible values of the parameter.
+		void SetValue( int i, float value, const FString& strName );
+
+        //! Set the number of ranges (dimensions) for this parameter.
+        //! By default a parameter has 0 ranges, meaning it only has one value.
+        void SetRangeCount( int i );
+        void SetRange( int i, NodeRangePtr pRange );
+
+		//-----------------------------------------------------------------------------------------
+		// Interface pattern
+		//-----------------------------------------------------------------------------------------
+		class Private;
+		Private* GetPrivate() const;
 
 	protected:
 
-		/** Forbidden. Manage with the Ptr<> template. */
-		~NodeScalarEnumParameter() {}
+		//! Forbidden. Manage with the Ptr<> template.
+		~NodeScalarEnumParameter();
 
 	private:
 
-		static FNodeType StaticType;
+		Private* m_pD;
 
 	};
 
+
 }
diff --git a/Engine/Plugins/Mutable/Source/MutableTools/Public/MuT/NodeScalarParameter.h b/Engine/Plugins/Mutable/Source/MutableTools/Public/MuT/NodeScalarParameter.h
--- a/Engine/Plugins/Mutable/Source/MutableTools/Public/MuT/NodeScalarParameter.h
+++ b/Engine/Plugins/Mutable/Source/MutableTools/Public/MuT/NodeScalarParameter.h
@@ -7,36 +7,75 @@
 #include "MuR/RefCounted.h"
 #include "MuT/Node.h"
 #include "MuT/NodeScalar.h"
-#include "MuT/NodeRange.h"
+
 
 namespace mu
 {
 
+	// Forward definitions
+    class NodeImage;
+    using NodeImagePtr = Ptr<NodeImage>;
+    using NodeImagePtrConst = Ptr<const NodeImage>;
+
+    class NodeRange;
+    using NodeRangePtr = Ptr<NodeRange>;
+    using NodeRangePtrConst = Ptr<const NodeRange>;
+
+    class NodeScalarParameter;
+    using NodeScalarParameterPtr = Ptr<NodeScalarParameter>;
+    using NodeScalarParameterPtrConst = Ptr<const NodeScalarParameter>;
+
+
 	//! Node that defines a scalar model parameter.
 	class MUTABLETOOLS_API NodeScalarParameter : public NodeScalar
 	{
 	public:
 
-		float DefaultValue = 0.0f;
-		FString Name;
-		FString UID;
-
-		TArray<Ptr<NodeRange>> Ranges;
+		//-----------------------------------------------------------------------------------------
+		// Life cycle
+		//-----------------------------------------------------------------------------------------
 
-	public:
+		NodeScalarParameter();
 
+		//-----------------------------------------------------------------------------------------
 		// Node Interface
-		virtual const FNodeType* GetType() const override { return &StaticType; }
-		static const FNodeType* GetStaticType() { return &StaticType; }
+		//-----------------------------------------------------------------------------------------
+
+        const FNodeType* GetType() const override;
+		static const FNodeType* GetStaticType();
+
+		//-----------------------------------------------------------------------------------------
+		// Own Interface
+		//-----------------------------------------------------------------------------------------
+
+		//! Get the name of the parameter. It will be exposed in the final compiled data.
+		void SetName( const FString& );
+
+		//! Get the uid of the parameter. It will be exposed in the final compiled data.
+		void SetUid( const FString&);
+
+		//! Get the default value of the parameter.
+		void SetDefaultValue( float v );
+
+        //! Set the number of ranges (dimensions) for this parameter.
+        //! By default a parameter has 0 ranges, meaning it only has one value.
+        void SetRangeCount( int i );
+        void SetRange( int i, NodeRangePtr pRange );
+
+		//-----------------------------------------------------------------------------------------
+		// Interface pattern
+		//-----------------------------------------------------------------------------------------
+		class Private;
+		Private* GetPrivate() const;
 
 	protected:
 
-		/** Forbidden. Manage with the Ptr<> template. */
-		~NodeScalarParameter() {}
+		//! Forbidden. Manage with the Ptr<> template.
+		~NodeScalarParameter();
 
 	private:
 
-		static FNodeType StaticType;
+		Private* m_pD;
 
 	};
 
diff --git a/Engine/Plugins/Mutable/Source/MutableTools/Public/MuT/NodeScalarTable.h b/Engine/Plugins/Mutable/Source/MutableTools/Public/MuT/NodeScalarTable.h
--- a/Engine/Plugins/Mutable/Source/MutableTools/Public/MuT/NodeScalarTable.h
+++ b/Engine/Plugins/Mutable/Source/MutableTools/Public/MuT/NodeScalarTable.h
@@ -8,10 +8,17 @@
 #include "MuT/NodeScalar.h"
 #include "MuT/Table.h"
 
+
 namespace mu
 {
 
-	/** This node provides the meshes stored in the column of a table. */
+	// Forward definitions
+	class NodeScalarTable;
+	typedef Ptr<NodeScalarTable> NodeScalarTablePtr;
+	typedef Ptr<const NodeScalarTable> NodeScalarTablePtrConst;
+
+
+	//! This node provides the meshes stored in the column of a table.
 	class MUTABLETOOLS_API NodeScalarTable : public NodeScalar
 	{
 	public:
@@ -28,9 +35,26 @@ namespace mu
 		virtual const FNodeType* GetType() const override { return GetStaticType(); }
 		static const FNodeType* GetStaticType() { return &StaticType; }
 
+		//-----------------------------------------------------------------------------------------
+		// Own Interface
+		//-----------------------------------------------------------------------------------------
+
+		//! Set the name of the implicit table parameter.
+		void SetParameterName( const FString& strName );
+
+		//!
+		void SetColumn( const FString& strName );
+
+		//! Adds the "None" option to the parameter that represents this table column
+		void SetNoneOption(bool bAddOption);
+
+		//! Set the row name to be used as default value
+		void SetDefaultRowName(const FString& RowName);
+
+
 	protected:
 
-		/** Forbidden. Manage with the Ptr<> template. */
+		//! Forbidden. Manage with the Ptr<> template.
 		~NodeScalarTable() {}
 
 	private:
@@ -39,4 +63,5 @@ namespace mu
 
 	};
 
+
 }
diff --git a/Engine/Plugins/Mutable/Source/MutableTools/Public/MuT/NodeScalarVariation.h b/Engine/Plugins/Mutable/Source/MutableTools/Public/MuT/NodeScalarVariation.h
--- a/Engine/Plugins/Mutable/Source/MutableTools/Public/MuT/NodeScalarVariation.h
+++ b/Engine/Plugins/Mutable/Source/MutableTools/Public/MuT/NodeScalarVariation.h
@@ -7,39 +7,61 @@
 #include "MuT/Node.h"
 #include "MuT/NodeScalar.h"
 
+
 namespace mu
 {
 
-	/** */
-    class MUTABLETOOLS_API NodeScalarVariation : public NodeScalar
+    // Forward definitions
+    class NodeScalarVariation;
+    typedef Ptr<NodeScalarVariation> NodeScalarVariationPtr;
+    typedef Ptr<const NodeScalarVariation> NodeScalarVariationPtrConst;
+
+
+     class MUTABLETOOLS_API NodeScalarVariation : public NodeScalar
     {
-	public:
+    public:
+ 
+        NodeScalarVariation();
 
-		Ptr<NodeScalar> DefaultScalar;
+        //-----------------------------------------------------------------------------------------
+        // Node Interface
+        //-----------------------------------------------------------------------------------------
 
-		struct FVariation
-		{
-			Ptr<NodeScalar> Scalar;
-			FString Tag;
-		};
+        const FNodeType* GetType() const override;
+        static const FNodeType* GetStaticType();
 
-		TArray<FVariation> Variations;
+        //-----------------------------------------------------------------------------------------
+        // Own Interface
+        //-----------------------------------------------------------------------------------------
 
-    public:
- 
-		// Node Interface
-		virtual const FNodeType* GetType() const override { return &StaticType; }
-		static const FNodeType* GetStaticType() { return &StaticType; }
+        //!
+        void SetDefaultScalar( NodeScalar* Scalar );
 
-    protected:
+        //! Set the number of tags to consider in this variation
+        void SetVariationCount( int32 count );
 
-		/** Forbidden. Manage with the Ptr<> template. */
-		~NodeScalarVariation() {}
+        //! Set the tag or state name that will enable a specific vartiation
+        void SetVariationTag( int32 index, const FString& strTag );
 
-    private:
+        //!
+        void SetVariationScalar( int32 index, NodeScalar* Scalar );
+
+        //!}
+
+
+        //-----------------------------------------------------------------------------------------
+        // Interface pattern
+        //-----------------------------------------------------------------------------------------
+        class Private;
+        Private* GetPrivate() const;
 
-		static FNodeType StaticType;
+    protected:
+        //! Forbidden. Manage with the Ptr<> template.
+        ~NodeScalarVariation();
+
+    private:
+        Private* m_pD;
+    };
 
-	};
 
-}
+} // namespace mu
diff --git a/Engine/Plugins/Mutable/Source/MutableTools/Public/MuT/NodeString.h b/Engine/Plugins/Mutable/Source/MutableTools/Public/MuT/NodeString.h
--- a/Engine/Plugins/Mutable/Source/MutableTools/Public/MuT/NodeString.h
+++ b/Engine/Plugins/Mutable/Source/MutableTools/Public/MuT/NodeString.h
@@ -11,23 +11,34 @@
 namespace mu
 {
 
-    /** Base class of any node that outputs a string value. */
+	// Forward definitions
+	class NodeString;
+	typedef Ptr<NodeString> NodeStringPtr;
+	typedef Ptr<const NodeString> NodeStringPtrConst;
+
+
+    //! %Base class of any node that outputs a string value.
 	class MUTABLETOOLS_API NodeString : public Node
 	{
 	public:
 
+		//-----------------------------------------------------------------------------------------
 		// Node Interface
-		virtual const FNodeType* GetType() const override { return &StaticType; }
-		static const FNodeType* GetStaticType() { return &StaticType; }
+		//-----------------------------------------------------------------------------------------
 
-	protected:
+		const FNodeType* GetType() const override;
+		static const FNodeType* GetStaticType();
 
-		/** Forbidden. Manage with the Ptr<> template. */
-		~NodeString() {}
 
-	private:
+		//-----------------------------------------------------------------------------------------
+		// Interface pattern
+		//-----------------------------------------------------------------------------------------
+		class Private;
+
+	protected:
 
-		static FNodeType StaticType;
+		//! Forbidden. Manage with the Ptr<> template.
+		inline ~NodeString() {}
 
 	};
 
diff --git a/Engine/Plugins/Mutable/Source/MutableTools/Public/MuT/NodeStringConstant.h b/Engine/Plugins/Mutable/Source/MutableTools/Public/MuT/NodeStringConstant.h
--- a/Engine/Plugins/Mutable/Source/MutableTools/Public/MuT/NodeStringConstant.h
+++ b/Engine/Plugins/Mutable/Source/MutableTools/Public/MuT/NodeStringConstant.h
@@ -11,27 +11,48 @@
 namespace mu
 {
 
-	/** Node returning a string constant value. */
+	// Forward definitions
+	class NodeStringConstant;
+	typedef Ptr<NodeStringConstant> NodeStringConstantPtr;
+	typedef Ptr<const NodeStringConstant> NodeStringConstantPtrConst;
+
+
+	//! Node returning a string constant value.
+	//! \ingroup model
 	class MUTABLETOOLS_API NodeStringConstant : public NodeString
 	{
 	public:
 
-		FString Value;
+		NodeStringConstant();
 
-	public:
 
+		//-----------------------------------------------------------------------------------------
 		// Node Interface
-		virtual const FNodeType* GetType() const override { return &StaticType; }
-		static const FNodeType* GetStaticType() { return &StaticType; }
+		//-----------------------------------------------------------------------------------------
+        const FNodeType* GetType() const override;
+		static const FNodeType* GetStaticType();
+
+		//-----------------------------------------------------------------------------------------
+		// Own Interface
+		//-----------------------------------------------------------------------------------------
+
+		//! Set the value to be returned by the node.
+		void SetValue( const FString& v );
+
+		//-----------------------------------------------------------------------------------------
+		// Interface pattern
+		//-----------------------------------------------------------------------------------------
+		class Private;
+		Private* GetPrivate() const;
 
 	protected:
 
-		/** Forbidden. Manage with the Ptr<> template. */
-		~NodeStringConstant() {}
+		//! Forbidden. Manage with the Ptr<> template.
+		~NodeStringConstant();
 
 	private:
 
-		static FNodeType StaticType;
+		Private* m_pD;
 
 	};
 
diff --git a/Engine/Plugins/Mutable/Source/MutableTools/Public/MuT/NodeStringParameter.h b/Engine/Plugins/Mutable/Source/MutableTools/Public/MuT/NodeStringParameter.h
--- a/Engine/Plugins/Mutable/Source/MutableTools/Public/MuT/NodeStringParameter.h
+++ b/Engine/Plugins/Mutable/Source/MutableTools/Public/MuT/NodeStringParameter.h
@@ -7,36 +7,71 @@
 #include "MuR/RefCounted.h"
 #include "MuT/Node.h"
 #include "MuT/NodeString.h"
-#include "MuT/NodeRange.h"
+
 
 namespace mu
 {
 
-	/** Node that defines a string model parameter. */
+	// Forward definitions
+    class NodeImage;
+    using NodeImagePtr = Ptr<NodeImage>;
+    using NodeImagePtrConst = Ptr<const NodeImage>;
+
+    class NodeRange;
+    using NodeRangePtr = Ptr<NodeRange>;
+    using NodeRangePtrConst = Ptr<const NodeRange>;
+
+    class NodeStringParameter;
+    using NodeStringParameterPtr = Ptr<NodeStringParameter>;
+    using NodeStringParameterPtrConst = Ptr<const NodeStringParameter>;
+
+
+	//! Node that defines a string model parameter.
 	class MUTABLETOOLS_API NodeStringParameter : public NodeString
 	{
 	public:
 
-		FString DefaultValue;
-		FString Name;
-		FString UID;
+		NodeStringParameter();
 
-		TArray<Ptr<NodeRange>> Ranges;
+		//-----------------------------------------------------------------------------------------
+		// Node Interface
+		//-----------------------------------------------------------------------------------------
 
-	public:
+        const FNodeType* GetType() const override;
+		static const FNodeType* GetStaticType();
 
-		// Node Interface
-		virtual const FNodeType* GetType() const override { return &StaticType; }
-		static const FNodeType* GetStaticType() { return &StaticType; }
+		//-----------------------------------------------------------------------------------------
+		// Own Interface
+		//-----------------------------------------------------------------------------------------
+
+		//! Set the name of the parameter. It will be exposed in the final compiled data.
+		void SetName( const FString& );
+
+		//! Set the uid of the parameter. It will be exposed in the final compiled data.
+		void SetUid( const FString& );
+
+		//! Set the default value of the parameter.
+		void SetDefaultValue( const FString& );
+
+        //! Set the number of ranges (dimensions) for this parameter.
+        //! By default a parameter has 0 ranges, meaning it only has one value.
+        void SetRangeCount( int i );
+        void SetRange( int i, NodeRangePtr pRange );
+
+		//-----------------------------------------------------------------------------------------
+		// Interface pattern
+		//-----------------------------------------------------------------------------------------
+		class Private;
+		Private* GetPrivate() const;
 
 	protected:
 
-		/** Forbidden. Manage with the Ptr<> template. */
-		~NodeStringParameter() {}
+		//! Forbidden. Manage with the Ptr<> template.
+		~NodeStringParameter();
 
 	private:
 
-		static FNodeType StaticType;
+		Private* m_pD;
 
 	};
 
