[mutable] Reduced verbosity of log that was affecting the validation of assets and preventing the upload of changes.
- An indirect update triggered by the compilation of a CO due to it being validated was raising this warning.

diff --git a/Engine/Plugins/Mutable/Source/CustomizableObject/Private/MuCO/CustomizableObjectSystem.cpp b/Engine/Plugins/Mutable/Source/CustomizableObject/Private/MuCO/CustomizableObjectSystem.cpp
--- a/Engine/Plugins/Mutable/Source/CustomizableObject/Private/MuCO/CustomizableObjectSystem.cpp
+++ b/Engine/Plugins/Mutable/Source/CustomizableObject/Private/MuCO/CustomizableObjectSystem.cpp
@@ -1430,7 +1430,7 @@ void UpdateSkeletalMesh(const TSharedRef<FUpdateContextPrivate>& Context)
 
 		if (Parent->HasOverrideMaterials())
 		{
-			UE_LOG(LogMutable, Warning, TEXT("Attaching Customizable Skeletal Component to Skeletal Mesh Component with overriden materials! Deleting overrides."));
+			UE_LOG(LogMutable, Log, TEXT("Attaching Customizable Skeletal Component to Skeletal Mesh Component with overriden materials! Deleting overrides."));
 			Parent->EmptyOverrideMaterials();
 		}
 
