[Mutable] Fix crash when accessing ClothMappingDataLODs of meshes that failed to restore clothing data.

#rnx

diff --git a/Engine/Plugins/Mutable/Source/CustomizableObjectEditor/Private/MuCOE/GenerateMutableSource/GenerateMutableSourceMesh.cpp b/Engine/Plugins/Mutable/Source/CustomizableObjectEditor/Private/MuCOE/GenerateMutableSource/GenerateMutableSourceMesh.cpp
--- a/Engine/Plugins/Mutable/Source/CustomizableObjectEditor/Private/MuCOE/GenerateMutableSource/GenerateMutableSourceMesh.cpp
+++ b/Engine/Plugins/Mutable/Source/CustomizableObjectEditor/Private/MuCOE/GenerateMutableSource/GenerateMutableSourceMesh.cpp
@@ -1854,7 +1854,12 @@ TSharedPtr<mu::FMesh> ConvertSkeletalMeshToMutable(
 				}
 			}, ClothingAssetIndex, PhysicsAssetIndex);
 
-		if (ClothingAssetIndex != INDEX_NONE)
+		if (ClothingAssetIndex != INDEX_NONE && MeshSection.ClothMappingDataLODs.IsEmpty())
+		{
+			UE_LOG(LogMutable, Display, TEXT("SkeletalMesh [%s] is missing ClothMappingDataLODs. "), *GetNameSafe(SkeletalMesh));
+		}
+
+		if (ClothingAssetIndex != INDEX_NONE && !MeshSection.ClothMappingDataLODs.IsEmpty())
 		{
 			MutableMesh->GetVertexBuffers().SetBufferCount(NextBufferIndex + 2);
 			{
