[Mutable] Fix possible assert in GetResourceKey when a matrix parameter is part of the relevant params list. Fix possible bug with multi value projectors.

#rnx

diff --git a/Engine/Plugins/Mutable/Source/MutableRuntime/Private/MuR/System.cpp b/Engine/Plugins/Mutable/Source/MutableRuntime/Private/MuR/System.cpp
--- a/Engine/Plugins/Mutable/Source/MutableRuntime/Private/MuR/System.cpp
+++ b/Engine/Plugins/Mutable/Source/MutableRuntime/Private/MuR/System.cpp
@@ -1849,7 +1849,7 @@ namespace mu
 			case EParameterType::Projector:
 			{
 				FPackedNormal TempVec;
-				FProjector& Value = Params->GetPrivate()->Values[ParamIndex].Get<FParamProjectorType>();
+				const FProjector& Value = Params->GetPrivate()->Values[ParamIndex].Get<FParamProjectorType>();
 				Blob.Serialize((void*)&Value.position, sizeof(FVector3f));
 				TempVec = FPackedNormal(Value.direction);
 				Blob.Serialize(&TempVec, sizeof(FPackedNormal));
@@ -1865,14 +1865,14 @@ namespace mu
 					AddMultiValueKeys(Blob, Multi);
 					for (const TPair< TArray<int32>, FParameterValue >& v : Multi)
 					{
-						Value = v.Value.Get<FParamProjectorType>();
-						Blob.Serialize((void*)&Value.position, sizeof(FVector3f));
-						TempVec = FPackedNormal(Value.direction);
+						const FProjector& MultiValue = v.Value.Get<FParamProjectorType>();
+						Blob.Serialize((void*)&MultiValue.position, sizeof(FVector3f));
+						TempVec = FPackedNormal(MultiValue.direction);
 						Blob.Serialize(&TempVec, sizeof(FPackedNormal));
-						TempVec = FPackedNormal(Value.up);
+						TempVec = FPackedNormal(MultiValue.up);
 						Blob.Serialize(&TempVec, sizeof(FPackedNormal));
-						Blob.Serialize((void*)&Value.scale, sizeof(FVector3f));
-						Blob.Serialize((void*)&Value.projectionAngle, sizeof(float));
+						Blob.Serialize((void*)&MultiValue.scale, sizeof(FVector3f));
+						Blob.Serialize((void*)&MultiValue.projectionAngle, sizeof(float));
 					}
 				}
 				break;
@@ -1896,6 +1896,27 @@ namespace mu
 				break;
 			}
 
+			case EParameterType::Matrix:
+			{
+				DataSize = sizeof(FMatrix44f);
+
+				const FMatrix44f& Value = Params->GetPrivate()->Values[ParamIndex].Get<FParamMatrixType>();
+				Blob.Serialize((void*)&Value, DataSize);
+
+				// Multi-values
+				if (Params->GetPrivate()->HasMultipleValues(ParamIndex))
+				{
+					const TMap<TArray<int32>, FParameterValue>& Multi = Params->GetPrivate()->MultiValues[ParamIndex];
+					AddMultiValueKeys(Blob, Multi);
+					for (const TPair<TArray<int32>, FParameterValue>& MultiValuePair : Multi)
+					{
+						const FMatrix44f& MultiValue = MultiValuePair.Value.Get<FParamMatrixType>();
+						Blob.Serialize((void*)&MultiValue, DataSize);
+					}
+				}
+				break;
+			}
+
 			default:
 				// unsupported parameter type
 				check(false);
