CustomizableObject: Use FAssetPackageData::GetSavedPackageHash instead of the deprecated PackageGuid; PackageGuid will be deleted soon.
#rnx

diff --git a/Engine/Plugins/Mutable/Source/CustomizableObjectEditor/Private/MuCOE/CustomizableObjectCookPackageSplitter.cpp b/Engine/Plugins/Mutable/Source/CustomizableObjectEditor/Private/MuCOE/CustomizableObjectCookPackageSplitter.cpp
--- a/Engine/Plugins/Mutable/Source/CustomizableObjectEditor/Private/MuCOE/CustomizableObjectCookPackageSplitter.cpp
+++ b/Engine/Plugins/Mutable/Source/CustomizableObjectEditor/Private/MuCOE/CustomizableObjectCookPackageSplitter.cpp
@@ -137,29 +137,17 @@ void GenerateNewPackage(const FCustomizableObjectStreamedResourceData& StreamedD
 	// To support iterative cooking, GenerationHash should only change when OwnerPackage
 	// changes.
 	//
-	// The simplest and fastest way to do this is to set it to OwnerPackage's GUID.
+	// The simplest and fastest way to do this is to set it to OwnerPackage's PackageSavedHash.
 	{
 		// Zero the hash, as we won't be writing all bytes of it below
 		Package.GenerationHash.Reset();
 
-		// Ensure there's no padding within FGuid, as the padding could be uninitialized
-		static_assert(sizeof(FGuid) == sizeof(uint32) * 4);
-		// Ensure an FGuid will fit within the hash buffer
-		static_assert(sizeof(FGuid) <= sizeof(FBlake3Hash::ByteArray));
+		FIoHash OwnerSavedHash = OwnerPackage->GetSavedHash();
+		static_assert(sizeof(Package.GenerationHash.GetBytes()) >= sizeof(OwnerSavedHash.GetBytes()));
+		static_assert(sizeof(Package.GenerationHash.GetBytes()) > 8); // It should be a byte array, not a pointer
+		static_assert(sizeof(OwnerSavedHash.GetBytes()) > 8); // It should be a byte array, not a pointer
+		FMemory::Memcpy(Package.GenerationHash.GetBytes(), OwnerSavedHash.GetBytes(), sizeof(OwnerSavedHash.GetBytes()));
 
-		// Create a copy that we can get the address of
-		//
-		// Note that UPackage::GetGuid is deprecated but there's no replacement yet. The engine
-		// team has advised me to use it anyway, so it's necessary to suppress the deprecation
-		// warning here.
-		PRAGMA_DISABLE_DEPRECATION_WARNINGS
-		const FGuid OwnerGuid = OwnerPackage->GetGuid();
-		PRAGMA_ENABLE_DEPRECATION_WARNINGS
-
-		FMemory::Memcpy(Package.GenerationHash.GetBytes(), &OwnerGuid, sizeof(FGuid));
-
-		// If OwnerGuid is non-zero, the hash should now be non-zero
-		check(!OwnerGuid.IsValid() || !Package.GenerationHash.IsZero());
 	}
 }
 
diff --git a/Engine/Plugins/Mutable/Source/CustomizableObjectEditor/Private/MuCOE/CustomizableObjectEditorModule.cpp b/Engine/Plugins/Mutable/Source/CustomizableObjectEditor/Private/MuCOE/CustomizableObjectEditorModule.cpp
--- a/Engine/Plugins/Mutable/Source/CustomizableObjectEditor/Private/MuCOE/CustomizableObjectEditorModule.cpp
+++ b/Engine/Plugins/Mutable/Source/CustomizableObjectEditor/Private/MuCOE/CustomizableObjectEditorModule.cpp
@@ -124,6 +124,17 @@ constexpr int32 MaxAccessTimespan = 30;
 
 IMPLEMENT_MODULE( FCustomizableObjectEditorModule, CustomizableObjectEditor );
 
+// TODO UE-226453: UPackage and FAssetPackageData originally had an FGuid PackageGuid, but this was changed to
+// FIoHash PackageSavedHash for SavePackage determinism. We wrote our ParticipatingObject change detection in terms of
+// PackageGuid, and we need to instead save and compare FIoHashes. Change our TMap<FName,FGuid> into
+// TMap<FName,FIoHash>. In the meantime, we change the format back to an FGuid by truncating the FIoHash.
+static FGuid TruncatePackageSavedHash(const FIoHash& PackageSavedHash)
+{
+	FGuid Result;
+	static_assert(sizeof(Result) < sizeof(PackageSavedHash.GetBytes()), "We copy sizeof(FGuid) bytes from FIoHash.GetBytes");
+	FMemory::Memcpy(&Result, PackageSavedHash.GetBytes(), sizeof(Result));
+	return Result;
+}
 
 void DeleteUnusedMutableStreamedDataEditorFiles()
 {
@@ -570,11 +581,7 @@ void IsCompilationOutOfDate_Async(const TSharedRef<FCompilationOutOfDateContext>
 		TSoftObjectPtr SoftObjectPtr(FSoftObjectPath(ParticipatingObject.Key.ToString()));
 		if (SoftObjectPtr) // If loaded
 		{
-			// TODO UE-226453 
-			PRAGMA_DISABLE_DEPRECATION_WARNINGS
-			const FGuid PackageGuid = SoftObjectPtr->GetPackage()->GetGuid();
-			PRAGMA_ENABLE_DEPRECATION_WARNINGS
-			
+			const FGuid PackageGuid = TruncatePackageSavedHash(SoftObjectPtr->GetPackage()->GetSavedHash());
 			if (PackageGuid != ParticipatingObject.Value)
 			{
 				Context->OutOfDatePackages.Add(ParticipatingObject.Key);
@@ -590,10 +597,7 @@ void IsCompilationOutOfDate_Async(const TSharedRef<FCompilationOutOfDateContext>
 				Context->RemovedPackages.Add(ParticipatingObject.Key);
 			}
 
-			PRAGMA_DISABLE_DEPRECATION_WARNINGS
-			const FGuid PackageGuid = AssetPackageData.PackageGuid;
-			PRAGMA_ENABLE_DEPRECATION_WARNINGS
-			
+			const FGuid PackageGuid = TruncatePackageSavedHash(AssetPackageData.GetPackageSavedHash());
 			if (PackageGuid != ParticipatingObject.Value)
 			{
 				Context->OutOfDatePackages.Add(ParticipatingObject.Key);
@@ -824,10 +828,7 @@ TMap<FName, FGuid> FCustomizableObjectEditorModule::GetParticipatingObjects(cons
 		if (Candidate)
 		{
 			// TODO UE-226453
-			PRAGMA_DISABLE_DEPRECATION_WARNINGS
-			const FGuid PackageGuid = Candidate->GetPackage()->GetGuid();
-			PRAGMA_ENABLE_DEPRECATION_WARNINGS
-
+			const FGuid PackageGuid = TruncatePackageSavedHash(Candidate->GetPackage()->GetSavedHash());
 			FName PackageName = Candidate->GetPackage()->GetFName();
 			Result.Emplace(PackageName, PackageGuid);
 		}
@@ -842,10 +843,7 @@ TMap<FName, FGuid> FCustomizableObjectEditorModule::GetParticipatingObjects(cons
 			return;
 		}
 
-		PRAGMA_DISABLE_DEPRECATION_WARNINGS
-		const FGuid PackageGuid = AssetPackageData.PackageGuid;
-		PRAGMA_ENABLE_DEPRECATION_WARNINGS
-
+		const FGuid PackageGuid = TruncatePackageSavedHash(AssetPackageData.GetPackageSavedHash());
 		Result.Emplace(SoftPath.GetLongPackageFName(), PackageGuid);
 	};
 	
