CustomizableObject: Suppress static analysis warning that sizeof(OwnerSavedHash.GetBytes()) looks like a common typo when using sizeof.
#rnx

diff --git a/Engine/Plugins/Mutable/Source/CustomizableObjectEditor/Private/MuCOE/CustomizableObjectCookPackageSplitter.cpp b/Engine/Plugins/Mutable/Source/CustomizableObjectEditor/Private/MuCOE/CustomizableObjectCookPackageSplitter.cpp
--- a/Engine/Plugins/Mutable/Source/CustomizableObjectEditor/Private/MuCOE/CustomizableObjectCookPackageSplitter.cpp
+++ b/Engine/Plugins/Mutable/Source/CustomizableObjectEditor/Private/MuCOE/CustomizableObjectCookPackageSplitter.cpp
@@ -143,10 +143,10 @@ void GenerateNewPackage(const FCustomizableObjectStreamedResourceData& StreamedD
 		Package.GenerationHash.Reset();
 
 		FIoHash OwnerSavedHash = OwnerPackage->GetSavedHash();
-		static_assert(sizeof(Package.GenerationHash.GetBytes()) >= sizeof(OwnerSavedHash.GetBytes()));
-		static_assert(sizeof(Package.GenerationHash.GetBytes()) > 8); // It should be a byte array, not a pointer
-		static_assert(sizeof(OwnerSavedHash.GetBytes()) > 8); // It should be a byte array, not a pointer
-		FMemory::Memcpy(Package.GenerationHash.GetBytes(), OwnerSavedHash.GetBytes(), sizeof(OwnerSavedHash.GetBytes()));
+		static_assert(sizeof(Package.GenerationHash.GetBytes()) >= sizeof(OwnerSavedHash.GetBytes()));  // -V568
+		static_assert(sizeof(Package.GenerationHash.GetBytes()) > 8); // It should be a byte array, not a pointer // -V568
+		static_assert(sizeof(OwnerSavedHash.GetBytes()) > 8); // It should be a byte array, not a pointer // -V568
+		FMemory::Memcpy(Package.GenerationHash.GetBytes(), OwnerSavedHash.GetBytes(), sizeof(OwnerSavedHash.GetBytes())); // -V568
 
 	}
 }
