[Mutable] First version of the Customizable Object Macro Compiler.
+ Adapted all mutable nodes to support being instantiated in Macro Graphs.
+ Adapted all GenerateMutableSource functions to support compiltaion of Mutable Macros.
+ Improved Macro Editor's widgets.
+ Now, COEditor inherits froma COGraphEditorToolkit.

#rnx

diff --git a/Engine/Plugins/Mutable/Source/CustomizableObject/Internal/MuCO/CustomizableObjectCustomVersion.h b/Engine/Plugins/Mutable/Source/CustomizableObject/Internal/MuCO/CustomizableObjectCustomVersion.h
--- a/Engine/Plugins/Mutable/Source/CustomizableObject/Internal/MuCO/CustomizableObjectCustomVersion.h
+++ b/Engine/Plugins/Mutable/Source/CustomizableObject/Internal/MuCO/CustomizableObjectCustomVersion.h
@@ -176,6 +176,8 @@ struct CUSTOMIZABLEOBJECT_API FCustomizableObjectCustomVersion
 
 		UpdatedNodesPinName3,
 
+		TextureVariationsToVariations,
+
 		// -----<new versions can be added above this line>-------------------------------------------------
 		VersionPlusOne,
 		LatestVersion = VersionPlusOne - 1
diff --git a/Engine/Plugins/Mutable/Source/CustomizableObjectEditor/Private/MuCOE/CustomizableInstanceDetails.cpp b/Engine/Plugins/Mutable/Source/CustomizableObjectEditor/Private/MuCOE/CustomizableInstanceDetails.cpp
--- a/Engine/Plugins/Mutable/Source/CustomizableObjectEditor/Private/MuCOE/CustomizableInstanceDetails.cpp
+++ b/Engine/Plugins/Mutable/Source/CustomizableObjectEditor/Private/MuCOE/CustomizableInstanceDetails.cpp
@@ -20,6 +20,7 @@
 #include "MuCO/CustomizableObject.h"
 #include "MuCO/CustomizableObjectInstance.h"
 #include "MuCO/CustomizableObjectSystem.h"
+#include "MuCOE/CustomizableObjectEditor.h"
 #include "MuCOE/UnrealEditorPortabilityHelpers.h"
 
 #include "AssetThumbnail.h"
@@ -70,14 +71,29 @@ void FCustomizableInstanceDetails::CustomizeDetails(const TSharedPtr<IDetailLayo
 	CustomInstance = Cast<UCustomizableObjectInstance>(DetailsView->GetSelectedObjects()[0].Get());
 	check(CustomInstance.IsValid());
 
-	TSharedPtr<IToolkit> FoundAssetEditor = FToolkitManager::Get().FindEditorForAsset(CustomInstance.Get()); // Tab spawned in a COEInstanceEditor
-	if (!FoundAssetEditor)
+	TSharedPtr<IToolkit> FoundAssetEditor;
+
+	// Tab spawned in a COEInstanceEditor
+	if (FoundAssetEditor = FToolkitManager::Get().FindEditorForAsset(CustomInstance.Get())) 
 	{
-		FoundAssetEditor = FToolkitManager::Get().FindEditorForAsset(CustomInstance->GetCustomizableObject()); // Tab spawned in a COEditor
+		if (TSharedPtr<FCustomizableObjectInstanceEditor> InstanceEditor = StaticCastSharedPtr<FCustomizableObjectInstanceEditor>(FoundAssetEditor))
+		{
+			WeakEditor = StaticCastSharedPtr<ICustomizableObjectInstanceEditor>(InstanceEditor).ToWeakPtr();
+		}
 	}
-	check(FoundAssetEditor);
 
-	WeakEditor = StaticCastSharedPtr<ICustomizableObjectInstanceEditor>(FoundAssetEditor).ToWeakPtr();
+	if (!FoundAssetEditor)
+	{
+		// Tab spawned in a COEditor
+		if (FoundAssetEditor = FToolkitManager::Get().FindEditorForAsset(CustomInstance->GetCustomizableObject()))
+		{
+			if (TSharedPtr<FCustomizableObjectEditor> ObjectEditor = StaticCastSharedPtr<FCustomizableObjectEditor>(FoundAssetEditor))
+			{
+				WeakEditor = StaticCastSharedPtr<ICustomizableObjectInstanceEditor>(ObjectEditor).ToWeakPtr();
+			}
+		}
+	}
+	check(WeakEditor.IsValid());
 
 	if (CustomInstance->GetPrivate()->IsSelectedParameterProfileDirty())
 	{
diff --git a/Engine/Plugins/Mutable/Source/CustomizableObjectEditor/Private/MuCOE/CustomizableObjectCompiler.cpp b/Engine/Plugins/Mutable/Source/CustomizableObjectEditor/Private/MuCOE/CustomizableObjectCompiler.cpp
--- a/Engine/Plugins/Mutable/Source/CustomizableObjectEditor/Private/MuCOE/CustomizableObjectCompiler.cpp
+++ b/Engine/Plugins/Mutable/Source/CustomizableObjectEditor/Private/MuCOE/CustomizableObjectCompiler.cpp
@@ -654,17 +654,21 @@ mu::Ptr<mu::NodeObject> GenerateMutableRoot(
 		}
 	}
 
+	TSet<UCustomizableObjectNode*> VisitedNodes;
+
 	// First pass. Only used to recollect info required for the primary pass.
 	// Notice that the traversal is different form the primary pass. Here we follow all pins indiscriminately,
 	// while the primary pass follows the Mutable Source structure (which may cut branches).
-	GraphTraversal::VisitNodes(*RootNodeObject, [&GenerationContext](UCustomizableObjectNode& Node)
+	GraphTraversal::RecursiveVisitNodes(*RootNodeObject, [&GenerationContext](UCustomizableObjectNode& Node)
 	{
 		if (UCustomizableObjectNodeComponentMesh* NodeComponentMesh = Cast<UCustomizableObjectNodeComponentMesh>(&Node))
 		{
 			FirstPass(*NodeComponentMesh, GenerationContext);
 		}
-	},
-	&GenerationContext.GroupIdToExternalNodeMap);
+	}, 
+	GenerationContext.MacroNodesStack, VisitedNodes, &GenerationContext.GroupIdToExternalNodeMap);
+
+	check(GenerationContext.MacroNodesStack.Num() == 0); // Ensure Macro context is empty after visiting all the nodes
 	
     GenerationContext.CompilationContext.RealTimeMorphTargetsOverrides = RootNodeObject->RealTimeMorphSelectionOverrides;
 
diff --git a/Engine/Plugins/Mutable/Source/CustomizableObjectEditor/Private/MuCOE/CustomizableObjectCustomSettings.cpp b/Engine/Plugins/Mutable/Source/CustomizableObjectEditor/Private/MuCOE/CustomizableObjectCustomSettings.cpp
--- a/Engine/Plugins/Mutable/Source/CustomizableObjectEditor/Private/MuCOE/CustomizableObjectCustomSettings.cpp
+++ b/Engine/Plugins/Mutable/Source/CustomizableObjectEditor/Private/MuCOE/CustomizableObjectCustomSettings.cpp
@@ -15,7 +15,6 @@
 #include "MuCO/CustomizableObject.h"
 #include "MuCO/CustomizableObjectInstance.h"
 #include "MuCO/LoadUtils.h"
-#include "MuCOE/ICustomizableObjectEditor.h"
 #include "MuCOE/ICustomizableObjectInstanceEditor.h"
 #include "Subsystems/ImportSubsystem.h"
 #include "Toolkits/ToolkitManager.h"
diff --git a/Engine/Plugins/Mutable/Source/CustomizableObjectEditor/Private/MuCOE/CustomizableObjectEditor.cpp b/Engine/Plugins/Mutable/Source/CustomizableObjectEditor/Private/MuCOE/CustomizableObjectEditor.cpp
--- a/Engine/Plugins/Mutable/Source/CustomizableObjectEditor/Private/MuCOE/CustomizableObjectEditor.cpp
+++ b/Engine/Plugins/Mutable/Source/CustomizableObjectEditor/Private/MuCOE/CustomizableObjectEditor.cpp
@@ -155,7 +155,6 @@ FCustomizableObjectEditor::~FCustomizableObjectEditor()
 	}
 	
 	ObjectDetailsView.Reset();
-	GEditor->UnregisterForUndo(this);
 
 	FCoreUObjectDelegates::OnObjectModified.RemoveAll(this);
 
@@ -179,11 +178,6 @@ void FCustomizableObjectEditor::InitCustomizableObjectEditor(const EToolkitMode:
 	
 	EditorProperties = NewObject<UCustomizableObjectEditorProperties>();
 
-	// Support undo/redo
-	CustomizableObject->SetFlags(RF_Transactional);
-
-	GEditor->RegisterForUndo(this);
-
 	// Register our commands. This will only register them if not previously registered
 	FGraphEditorCommands::Register();
 	FCustomizableObjectEditorCommands::Register();
@@ -191,6 +185,7 @@ void FCustomizableObjectEditor::InitCustomizableObjectEditor(const EToolkitMode:
 	FCustomizableObjectEditorNodeContextCommands::Register();
 
 	BindCommands();
+	BindGraphCommands();
 
 	FPropertyEditorModule& PropPlugin = FModuleManager::LoadModuleChecked<FPropertyEditorModule>("PropertyEditor");
 
@@ -317,12 +312,6 @@ FText FCustomizableObjectEditor::GetBaseToolkitName() const
 }
 
 
-void FCustomizableObjectEditor::SelectNode(const UEdGraphNode* Node)
-{
-	GraphEditor->JumpToNode(Node);
-}
-
-
 void FCustomizableObjectEditor::AddReferencedObjects( FReferenceCollector& Collector )
 {
 	Collector.AddReferencedObject( CustomizableObject );
@@ -406,114 +395,14 @@ TSharedRef<SDockTab> FCustomizableObjectEditor::SpawnTab_InstanceProperties( con
 }
 
 
-/** Create new tab for the supplied graph - don't call this directly, call SExplorer->FindTabForGraph.*/
-void FCustomizableObjectEditor::CreateGraphEditorWidget(UEdGraph* InGraph)
-{
-	UCustomizableObjectGraph* CustomizableObjectGraph = Cast<UCustomizableObjectGraph>(InGraph);
-	check(CustomizableObjectGraph != NULL);
-
-	// Add essential graph nodes if they do not exist
-	CustomizableObjectGraph->AddEssentialGraphNodes();
-
-	GraphEditorCommands = MakeShareable(new FUICommandList);
-
-	TSharedRef<SWidget> TitleBarWidget =
-		SNew(SHorizontalBox)
-		+ SHorizontalBox::Slot()
-		.HAlign(HAlign_Left)
-		.FillWidth(10.f)
-		.Padding(5.f)
-		[
-			SNew(SSearchBox)
-			.HintText(LOCTEXT("Search", "Search..."))
-			.ToolTipText(LOCTEXT("Search Nodes, Properties or Values that contain the inserted words", "Search Nodes, Properties or Values that contain the inserted words"))
-			.OnTextCommitted(FOnTextCommitted::CreateSP(this, &FCustomizableObjectEditor::OnEnterText))
-			.SelectAllTextWhenFocused(true)
-		];
-
-	// Create the appearance info
-	FGraphAppearanceInfo AppearanceInfo;	
-	AppearanceInfo.CornerText = LOCTEXT("ApperanceCornerText", "MUTABLE");
-
-	// TODO(Max) Each editor has to bind its own events
-	SGraphEditor::FGraphEditorEvents InEvents;
-	InEvents.OnSelectionChanged = SGraphEditor::FOnSelectionChanged::CreateSP( this, &FCustomizableObjectEditor::OnSelectedGraphNodesChanged );
-	InEvents.OnTextCommitted = FOnNodeTextCommitted::CreateSP(this, &FCustomizableObjectEditor::OnNodeTitleCommitted);
-	InEvents.OnNodeDoubleClicked = FSingleNodeEvent::CreateSP(this, &FCustomizableObjectEditor::OnNodeDoubleClicked);// TODO(Max): Remove this once this editor starts using the new toolkit
-
-	// Make full graph editor
-	GraphEditor = SNew(SGraphEditor)
-		.AdditionalCommands(GraphEditorCommands)
-		.Appearance(AppearanceInfo)
-		.GraphToEdit(InGraph)
-		.GraphEvents(InEvents)
-		.TitleBar(TitleBarWidget)
-		.ShowGraphStateOverlay(false); // Removes graph state overlays (border and text) such as "SIMULATING" and "READ-ONLY"
-
-	// Editing commands
-	GraphEditorCommands->MapAction(FGenericCommands::Get().Delete,
-		FExecuteAction::CreateSP(this, &FCustomizableObjectEditor::DeleteSelectedNodes),
-		FCanExecuteAction::CreateSP(this, &FCustomizableObjectEditor::CanDeleteNodes));
-
-	GraphEditorCommands->MapAction( FGenericCommands::Get().Copy,
-		FExecuteAction::CreateSP(this, &FCustomizableObjectEditor::CopySelectedNodes),
-		FCanExecuteAction::CreateSP(this, &FCustomizableObjectEditor::CanCopyNodes));
-
-	GraphEditorCommands->MapAction(FGenericCommands::Get().Paste,
-		FExecuteAction::CreateSP(this, &FCustomizableObjectEditor::PasteNodes),
-		FCanExecuteAction::CreateSP(this, &FCustomizableObjectEditor::CanPasteNodes));
-
-	GraphEditorCommands->MapAction(FGenericCommands::Get().Cut,
-		FExecuteAction::CreateSP(this, &FCustomizableObjectEditor::CutSelectedNodes),
-		FCanExecuteAction::CreateSP(this, &FCustomizableObjectEditor::CanCutNodes));
-
-	GraphEditorCommands->MapAction(FGenericCommands::Get().Duplicate,
-		FExecuteAction::CreateSP(this, &FCustomizableObjectEditor::DuplicateSelectedNodes),
-		FCanExecuteAction::CreateSP(this, &FCustomizableObjectEditor::CanDuplicateSelectedNodes));
-
-	GraphEditorCommands->MapAction(FGenericCommands::Get().Rename,
-		FExecuteAction::CreateSP(this, &FCustomizableObjectEditor::OnRenameNode),
-		FCanExecuteAction::CreateSP(this, &FCustomizableObjectEditor::CanRenameNodes));
-
-	GraphEditorCommands->MapAction(FCustomizableObjectEditorNodeContextCommands::Get().CreateComment,
-		FExecuteAction::CreateSP(this, &FCustomizableObjectEditor::CreateCommentBoxFromKey));
-
-	// Alignment Commands
-	GraphEditorCommands->MapAction(FGraphEditorCommands::Get().AlignNodesTop,
-		FExecuteAction::CreateSP(GraphEditor.Get(), &SGraphEditor::OnAlignTop));
-
-	GraphEditorCommands->MapAction(FGraphEditorCommands::Get().AlignNodesMiddle,
-		FExecuteAction::CreateSP(GraphEditor.Get(), &SGraphEditor::OnAlignMiddle));
-
-	GraphEditorCommands->MapAction(FGraphEditorCommands::Get().AlignNodesBottom,
-		FExecuteAction::CreateSP(GraphEditor.Get(), &SGraphEditor::OnAlignBottom));
-
-	GraphEditorCommands->MapAction(FGraphEditorCommands::Get().AlignNodesLeft,
-		FExecuteAction::CreateSP(GraphEditor.Get(), &SGraphEditor::OnAlignLeft));
-
-	GraphEditorCommands->MapAction(FGraphEditorCommands::Get().AlignNodesCenter,
-		FExecuteAction::CreateSP(GraphEditor.Get(), &SGraphEditor::OnAlignCenter));
-
-	GraphEditorCommands->MapAction(FGraphEditorCommands::Get().AlignNodesRight,
-		FExecuteAction::CreateSP(GraphEditor.Get(), &SGraphEditor::OnAlignRight));
-
-	GraphEditorCommands->MapAction(FGraphEditorCommands::Get().StraightenConnections,
-		FExecuteAction::CreateSP(GraphEditor.Get(), &SGraphEditor::OnStraightenConnections));
-
-	// Distribution Commands
-	GraphEditorCommands->MapAction(FGraphEditorCommands::Get().DistributeNodesHorizontally,
-		FExecuteAction::CreateSP(GraphEditor.Get(), &SGraphEditor::OnDistributeNodesH));
-	
-	GraphEditorCommands->MapAction(FGraphEditorCommands::Get().DistributeNodesVertically,
-		FExecuteAction::CreateSP(GraphEditor.Get(), &SGraphEditor::OnDistributeNodesV));
-}
-
-
 TSharedRef<SDockTab> FCustomizableObjectEditor::SpawnTab_Graph( const FSpawnTabArgs& Args )
 {
 	check( Args.GetTabId().TabType == GraphTabId );
 
-	CreateGraphEditorWidget(CustomizableObject->GetPrivate()->GetSource());
+	SGraphEditor::FGraphEditorEvents InEvents;
+	InEvents.OnSelectionChanged = SGraphEditor::FOnSelectionChanged::CreateSP(this, &FCustomizableObjectEditor::OnSelectedGraphNodesChanged);
+
+	CreateGraphEditorWidget(CustomizableObject->GetPrivate()->GetSource(), InEvents);
 
 	TSharedRef<SDockTab> DockTab = SNew(SDockTab)
 		.Label( FText::FromString( GetTabPrefix() + LOCTEXT( "SourceGraph", "Source Graph" ).ToString() ) )
@@ -613,30 +502,6 @@ void FCustomizableObjectEditor::BindCommands()
 		FExecuteAction::CreateSP(this, &FCustomizableObjectEditor::OpenPerformanceAnalyzerTab),
 		FCanExecuteAction(),
 		FIsActionChecked());
-	
-	// Undo-Redo
-	ToolkitCommands->MapAction(
-		FGenericCommands::Get().Undo,
-		FExecuteAction::CreateSP(this, &FCustomizableObjectEditor::UndoGraphAction));
-
-	ToolkitCommands->MapAction(
-		FGenericCommands::Get().Redo,
-		FExecuteAction::CreateSP(this, &FCustomizableObjectEditor::RedoGraphAction));
-}
-
-
-void FCustomizableObjectEditor::UndoGraphAction()
-{
-	GEditor->UndoTransaction();
-}
-
-
-void FCustomizableObjectEditor::RedoGraphAction()
-{
-	// Clear selection, to avoid holding refs to nodes that go away
-	GraphEditor->ClearSelectionSet();
-
-	GEditor->RedoTransaction();
 }
 
 
@@ -717,17 +582,6 @@ UCustomSettings* FCustomizableObjectEditor::GetCustomSettings()
 }
 
 
-void FCustomizableObjectEditor::SelectSingleNode(UCustomizableObjectNode& Node)
-{
-	FGraphPanelSelectionSet SelectedNodes = GraphEditor->GetSelectedNodes();
-	if (SelectedNodes.Num() != 1 || Cast<UCustomizableObjectNode>(*SelectedNodes.CreateIterator()) != &Node)
-	{
-		GraphEditor->ClearSelectionSet();
-		GraphEditor->SetNodeSelection(&Node, true);
-	}
-}
-
-
 void FCustomizableObjectEditor::HideGizmo()
 {
 	HideGizmoProjectorNodeProjectorConstant();
@@ -1103,6 +957,8 @@ UCustomizableObjectEditorProperties* FCustomizableObjectEditor::GetEditorPropert
 
 void FCustomizableObjectEditor::PostUndo(bool bSuccess)
 {
+	FCustomizableObjectGraphEditorToolkit::PostUndo(bSuccess);
+
 	if (bSuccess)
 	{
 		if (ObjectDetailsView.IsValid())
@@ -1115,7 +971,6 @@ void FCustomizableObjectEditor::PostUndo(bool bSuccess)
 			CustomizableInstanceDetailsView->RemoveInvalidObjects();
 		}
 
-		GraphEditor->NotifyGraphChanged();
 		CustomizableObject->MarkPackageDirty();
 
 		FSlateApplication::Get().DismissAllMenus();
@@ -1542,73 +1397,6 @@ void FCustomizableObjectEditor::SaveAsset_Execute()
 }
 
 
-void FCustomizableObjectEditor::DeleteSelectedNodes()
-{
-	if (!GraphEditor.IsValid())
-	{
-		return;
-	}
-
-	const FScopedTransaction Transaction(LOCTEXT("UEdGraphSchema_CustomizableObject", "Delete Nodes"));
-
-	const FGraphPanelSelectionSet SelectedNodes = GraphEditor->GetSelectedNodes();
-	GraphEditor->ClearSelectionSet();
-
-	for (FGraphPanelSelectionSet::TConstIterator NodeIt( SelectedNodes ); NodeIt; ++NodeIt)
-	{
-		if (UEdGraphNode* Node = Cast<UEdGraphNode>(*NodeIt))
-		{
-			if (Node->CanUserDeleteNode())
-			{
-				if (const UEdGraph* GraphObj = Node->GetGraph())
-				{
-					if (const UEdGraphSchema* Schema = GraphObj->GetSchema())
-					{
-						Schema->BreakNodeLinks(*Node);  // Required to notify to all connected nodes (UEdGraphNode::PinConnectionListChanged() and UEdGraphNode::PinConnectionListChanged(...))
-					}
-				}
-
-				Node->DestroyNode();
-			}
-		}
-	}
-}
-
-
-bool FCustomizableObjectEditor::CanDeleteNodes() const
-{
-	if (GraphEditor.IsValid() && GraphEditor->GetSelectedNodes().Num() > 0)
-	{
-		for (auto Itr = GraphEditor->GetSelectedNodes().CreateConstIterator(); Itr; ++Itr)
-		{
-			UCustomizableObjectNode* Node = Cast<UCustomizableObjectNode>(*Itr);
-
-			if (Node && !Node->CanUserDeleteNode())
-			{
-				return false;
-			}
-		}
-
-		return true;
-	}
-
-	return false;
-}
-
-
-void FCustomizableObjectEditor::DuplicateSelectedNodes()
-{
-	CopySelectedNodes();
-	PasteNodes();
-}
-
-
-bool FCustomizableObjectEditor::CanDuplicateSelectedNodes() const
-{
-	return CanCopyNodes();
-}
-
-
 void FCustomizableObjectEditor::OnSelectedGraphNodesChanged(const FGraphPanelSelectionSet& NewSelection)
 {
 	TArray<UObject*> Objects;
@@ -1686,15 +1474,6 @@ void FCustomizableObjectEditor::OnSelectedGraphNodesChanged(const FGraphPanelSel
 }
 
 
-void FCustomizableObjectEditor::OnNodeDoubleClicked(UEdGraphNode* Node) //TODO(Max): Remove this once this editor uses the new toolkit
-{
-	if (Node->CanJumpToDefinition())
-	{
-		Node->JumpToDefinition();
-	}
-}
-
-
 void FCustomizableObjectEditor::NotifyPostChange( const FPropertyChangedEvent& PropertyChangedEvent, FProperty* PropertyThatChanged  )
 {
 	// Is it a source graph node?
@@ -1713,235 +1492,6 @@ void FCustomizableObjectEditor::NotifyPostChange( const FPropertyChangedEvent& P
 	}
 }
 
-void FCustomizableObjectEditor::CopySelectedNodes()
-{
-	// Export the selected nodes and place the text on the clipboard
-	const FGraphPanelSelectionSet SelectedNodes = GraphEditor->GetSelectedNodes();
-
-	FString ExportedText;
-
-	for (FGraphPanelSelectionSet::TConstIterator SelectedIter(SelectedNodes); SelectedIter; ++SelectedIter)
-	{
-		if(UEdGraphNode* Node = Cast<UEdGraphNode>(*SelectedIter))
-		{
-			Node->PrepareForCopying();
-		}
-	}
-
-	FEdGraphUtilities::ExportNodesToText(SelectedNodes, /*out*/ ExportedText);
-
-	FPlatformApplicationMisc::ClipboardCopy(*ExportedText);
-
-	// Make sure Material remains the owner of the copied nodes
-	for (FGraphPanelSelectionSet::TConstIterator SelectedIter(SelectedNodes); SelectedIter; ++SelectedIter)
-	{
-		if (UMaterialGraphNode* Node = Cast<UMaterialGraphNode>(*SelectedIter))
-		{
-			Node->PostCopyNode();
-		}
-		else if (UMaterialGraphNode_Comment* Comment = Cast<UMaterialGraphNode_Comment>(*SelectedIter))
-		{
-			Comment->PostCopyNode();
-		}
-	}
-}
-
-bool FCustomizableObjectEditor::CanCopyNodes() const
-{
-	// If any of the nodes can be duplicated then we should allow copying
-	const FGraphPanelSelectionSet SelectedNodes = GraphEditor->GetSelectedNodes();
-	for (FGraphPanelSelectionSet::TConstIterator SelectedIter(SelectedNodes); SelectedIter; ++SelectedIter)
-	{
-		UEdGraphNode* Node = Cast<UEdGraphNode>(*SelectedIter);
-		if ((Node != NULL) && Node->CanDuplicateNode())
-		{
-			return true;
-		}
-	}
-
-	return false;
-}
-
-
-void FCustomizableObjectEditor::PasteNodes()
-{
-	PasteNodesHere(GraphEditor->GetPasteLocation());
-}
-
-
-void FCustomizableObjectEditor::PasteNodesHere(const FVector2D& Location)
-{
-	// Undo/Redo support
-	const FScopedTransaction Transaction( LOCTEXT("CustomizableObjectEditorPaste", "Customizable Object Editor Editor: Paste") );
-	CustomizableObject->GetPrivate()->GetSource()->Modify();
-	CustomizableObject->Modify();
-
-	// Clear the selection set (newly pasted stuff will be selected)
-	GraphEditor->ClearSelectionSet();
-
-	// Grab the text to paste from the clipboard.
-	FString TextToImport;
-	FPlatformApplicationMisc::ClipboardPaste(TextToImport);
-
-	// Import the nodes
-	TSet<UEdGraphNode*> PastedNodes;
-	FEdGraphUtilities::ImportNodesFromText(CustomizableObject->GetPrivate()->GetSource(), TextToImport, /*out*/ PastedNodes);
-	
-	//Average position of nodes so we can move them while still maintaining relative distances to each other
-	FVector2D AvgNodePosition(0.0f,0.0f);
-
-	for (TSet<UEdGraphNode*>::TIterator It(PastedNodes); It; ++It)
-	{
-		UEdGraphNode* Node = *It;
-		AvgNodePosition.X += Node->NodePosX;
-		AvgNodePosition.Y += Node->NodePosY;
-	}
-
-	if ( PastedNodes.Num() > 0 )
-	{
-		float InvNumNodes = 1.0f/float(PastedNodes.Num());
-		AvgNodePosition.X *= InvNumNodes;
-		AvgNodePosition.Y *= InvNumNodes;
-	}
-
-	for (TSet<UEdGraphNode*>::TIterator It(PastedNodes); It; ++It)
-	{
-		UEdGraphNode* Node = *It;
-		if (UCustomizableObjectNode* GraphNode = Cast<UCustomizableObjectNode>(Node))
-		{
-			// There can be only one default mesh paint texture.
-			//UMaterialExpressionTextureBase* TextureSample = Cast<UMaterialExpressionTextureBase>( NewExpression );
-			//if( TextureSample )
-			//{
-			//	TextureSample->IsDefaultMeshpaintTexture = false;
-			//}
-
-			//NewExpression->UpdateParameterGuid(true, true);
-
-			//UMaterialExpressionFunctionInput* FunctionInput = Cast<UMaterialExpressionFunctionInput>( NewExpression );
-			//if( FunctionInput )
-			//{
-			//	FunctionInput->ConditionallyGenerateId(true);
-			//	FunctionInput->ValidateName();
-			//}
-
-			//UMaterialExpressionFunctionOutput* FunctionOutput = Cast<UMaterialExpressionFunctionOutput>( NewExpression );
-			//if( FunctionOutput )
-			//{
-			//	FunctionOutput->ConditionallyGenerateId(true);
-			//	FunctionOutput->ValidateName();
-			//}
-		}
-		//else if (UMaterialGraphNode_Comment* CommentNode = Cast<UMaterialGraphNode_Comment>(Node))
-		//{
-		//	CommentNode->MaterialDirtyDelegate = Material->MaterialGraph->MaterialDirtyDelegate;
-		//	CommentNode->MaterialExpressionComment->Material = Material;
-		//	Material->EditorComments.Add(CommentNode->MaterialExpressionComment);
-		//}
-
-		// Select the newly pasted stuff
-		GraphEditor->SetNodeSelection(Node, true);
-
-		Node->NodePosX = (Node->NodePosX - AvgNodePosition.X) + Location.X ;
-		Node->NodePosY = (Node->NodePosY - AvgNodePosition.Y) + Location.Y ;
-
-		Node->SnapToGrid(SNodePanel::GetSnapGridSize());
-
-		// Give new node a different Guid from the old one
-		Node->CreateNewGuid();
-	}
-
-	// Force new pasted Material Expressions to have same connections as graph nodes
-	//Material->MaterialGraph->LinkMaterialExpressionsFromGraph();
-	
-	for (UEdGraphNode* PastedNode : PastedNodes)
-	{
-		if (UCustomizableObjectNode* TypedNode = Cast<UCustomizableObjectNode>(PastedNode))
-		{
-			TypedNode->PostBackwardsCompatibleFixup();
-		}
-	}
-	
-	// Update UI
-	GraphEditor->NotifyGraphChanged();
-
-	CustomizableObject->PostEditChange();
-	CustomizableObject->MarkPackageDirty();
-}
-
-
-bool FCustomizableObjectEditor::CanPasteNodes() const
-{
-	FString ClipboardContent;
-
-	FPlatformApplicationMisc::ClipboardPaste(ClipboardContent);
-
-	return FEdGraphUtilities::CanImportNodesFromText(CustomizableObject->GetPrivate()->GetSource(), ClipboardContent);
-}
-
-
-void FCustomizableObjectEditor::CutSelectedNodes()
-{
-	CopySelectedNodes();
-	// Cut should only delete nodes that can be duplicated
-	DeleteSelectedNodes();
-}
-
-
-bool FCustomizableObjectEditor::CanCutNodes() const
-{
-	return CanCopyNodes() && CanDeleteNodes();
-}
-
-
-void FCustomizableObjectEditor::CreateCommentBoxFromKey()
-{
-	CreateCommentBox(GraphEditor->GetPasteLocation());
-}
-
-UEdGraphNode* FCustomizableObjectEditor::CreateCommentBox(const FVector2D& InTargetPosition)
-{
-	//const FScopedTransaction Transaction(LOCTEXT("UEdGraphSchema_CustomizableObject", "Add Comment Box"));
-
-	UEdGraphNode_Comment* CommentTemplate = NewObject<UEdGraphNode_Comment>();
-
-	UEdGraphNode_Comment* NewComment = nullptr;
-	{
-		CustomizableObject->Modify();
-
-		// const FGraphPanelSelectionSet& SelectionSet = GraphEditor->GetSelectedNodes();
-		FSlateRect Bounds;
-		FVector2D Location;
-		FVector2D Size;
-
-		if (GraphEditor->GetBoundsForSelectedNodes(Bounds, 50.f))
-		{
-			Location.X = Bounds.Left;
-			Location.Y = Bounds.Top;
-			Size = Bounds.GetSize();
-		}
-		else
-		{
-			Location.X = InTargetPosition.X;
-			Location.Y = InTargetPosition.Y;
-			Size.X = 400;
-			Size.Y = 100;
-		}
-
-		NewComment = FEdGraphSchemaAction_NewNode::SpawnNodeFromTemplate<UEdGraphNode_Comment>(GraphEditor->GetCurrentGraph(), CommentTemplate, InTargetPosition, true);
-		NewComment->NodePosX = Location.X;
-		NewComment->NodePosY = Location.Y;
-		NewComment->NodeWidth = Size.X;
-		NewComment->NodeHeight = Size.Y;
-		NewComment->NodeComment = FString(TEXT("Comment"));
-	}
-
-	CustomizableObject->MarkPackageDirty();
-	GraphEditor->NotifyGraphChanged();
-
-	return NewComment;
-}
-
 
 TSharedPtr<SCustomizableObjectEditorAdvancedPreviewSettings> FCustomizableObjectEditor::GetAdvancedPreviewSettings()
 {
@@ -1961,6 +1511,14 @@ bool FCustomizableObjectEditor::ShowProfileManagementOptions()
 }
 
 
+const UObject* FCustomizableObjectEditor::GetObjectBeingEdited()
+{
+	ensure(GetObjectsCurrentlyBeingEdited()->Num());
+
+	return (*GetObjectsCurrentlyBeingEdited())[0];
+}
+
+
 void FCustomizableObjectEditor::FindProperty(const FProperty* Property, const void* InContainer, const FString& FindString, const UObject& Context, bool& bFound)
 {
 	if (!Property || !InContainer)
@@ -2065,72 +1623,6 @@ void FCustomizableObjectEditor::OnPostCompile()
 }
 
 
-void FCustomizableObjectEditor::OnEnterText(const FText& NewText, ETextCommit::Type TextType)
-{
-	if (TextType != ETextCommit::OnEnter)
-	{
-		return;
-	}
-	
-	if (!GraphEditor)
-	{
-		return;
-	}
-
-	const UEdGraph* Graph = GraphEditor->GetCurrentGraph();
-	if (!Graph)
-	{
-		return;
-	}
-
-	bool bFound = false;
-
-	const FString FindString = NewText.ToString();
-
-	for (TObjectPtr<UEdGraphNode> Node : Graph->Nodes)
-	{
-		if (!Node)
-		{
-			continue;
-		}
-
-		// Node names are not in the reflection system
-		const FString NodeName = Node->GetNodeTitle(ENodeTitleType::FullTitle).ToString().Replace(TEXT("\n"), TEXT(" "));
-		if (NodeName.Contains(NewText.ToString(), ESearchCase::IgnoreCase))
-		{
-			LogSearchResult(*Node, "Node", bFound, NodeName);
-			bFound = true;
-		}
-
-		// Pins are not in the reflection system
-		for (const UEdGraphPin* Pin : Node->GetAllPins())
-		{
-			const FString PinFriendlyName = Pin->PinFriendlyName.ToString();
-			if (PinFriendlyName.Contains(FindString))
-			{
-				LogSearchResult(*Node, "Pin", bFound, PinFriendlyName);
-				bFound = true;
-			}
-		}
-
-		// Find anything marked as a UPROPERTY
-		for (TFieldIterator<FProperty> It(Node->GetClass()); It; ++It)
-		{
-			FindProperty(*It, Node, FindString, *Node, bFound);
-		}
-	}
-
-	const FText Text = bFound ?
-		LOCTEXT("SearchCompleted", "Search completed") :
-		FText::FromString("No Results for: " + NewText.ToString());
-
-	FCustomizableObjectEditorLogger::CreateLog(Text)
-		.Category(ELoggerCategory::GraphSearch)
-		.CustomNotification()
-		.Log();
-}
-
-
 void FCustomizableObjectEditor::LogSearchResult(const UObject& Context, const FString& Type, bool bIsFirst, const FString& Result) const
 {
 	if (!bIsFirst)
@@ -2157,27 +1649,6 @@ void FCustomizableObjectEditor::OnUpdatePreviewInstance(UCustomizableObjectInsta
 }
 
 
-void FCustomizableObjectEditor::OnNodeTitleCommitted(const FText& NewText, ETextCommit::Type CommitInfo, UEdGraphNode* NodeBeingChanged)
-{
-	if (NodeBeingChanged)
-	{
-		const FScopedTransaction Transaction(LOCTEXT("RenameNode", "Rename Node"));
-		NodeBeingChanged->Modify();
-		NodeBeingChanged->OnRenameNode(NewText.ToString());
-	}
-}
-
-
-void FCustomizableObjectEditor::UpdateGraphNodeProperties()
-{
-	// Cache a copy of the selected nodes so we can later restore them
-	const FGraphPanelSelectionSet PreClearingSelectedNodes = GraphEditor->GetSelectedNodes();
-	
-	OnSelectedGraphNodesChanged(FGraphPanelSelectionSet());
-	OnSelectedGraphNodesChanged(PreClearingSelectedNodes);
-}
-
-
 void FCustomizableObjectEditor::OpenTextureAnalyzerTab()
 {
 	TabManager->TryInvokeTab(TextureAnalyzerTabId);
@@ -2298,45 +1769,4 @@ void FCustomizableObjectEditor::OnCustomizableObjectStatusChanged(FCustomizableO
 	}
 }
 
-
-void FCustomizableObjectEditor::OnRenameNode()
-{
-	if (GraphEditor.IsValid())
-	{
-		const FGraphPanelSelectionSet& SelectedNodes = GraphEditor->GetSelectedNodes();
-
-		for (FGraphPanelSelectionSet::TConstIterator NodeIt(SelectedNodes); NodeIt; ++NodeIt)
-		{
-			// Rename only the first valid selected node
-			UEdGraphNode* SelectedNode = Cast<UEdGraphNode>(*NodeIt);
-			if (SelectedNode && SelectedNode->GetCanRenameNode())
-			{
-				GraphEditor->IsNodeTitleVisible(SelectedNode, true);
-				break;
-			}
-		}
-	}
-}
-
-
-bool FCustomizableObjectEditor::CanRenameNodes() const
-{
-	if (GraphEditor.IsValid())
-	{
-		const FGraphPanelSelectionSet& SelectedNodes = GraphEditor->GetSelectedNodes();
-
-		for (FGraphPanelSelectionSet::TConstIterator SelectedIter(SelectedNodes); SelectedIter; ++SelectedIter)
-		{
-			if (UEdGraphNode* Node = Cast<UEdGraphNode>(*SelectedIter))
-			{
-				return Node->GetCanRenameNode();
-			}
-		}
-	}
-
-	return false;
-}
-
-
-
 #undef LOCTEXT_NAMESPACE
diff --git a/Engine/Plugins/Mutable/Source/CustomizableObjectEditor/Private/MuCOE/CustomizableObjectEditor.h b/Engine/Plugins/Mutable/Source/CustomizableObjectEditor/Private/MuCOE/CustomizableObjectEditor.h
--- a/Engine/Plugins/Mutable/Source/CustomizableObjectEditor/Private/MuCOE/CustomizableObjectEditor.h
+++ b/Engine/Plugins/Mutable/Source/CustomizableObjectEditor/Private/MuCOE/CustomizableObjectEditor.h
@@ -3,10 +3,9 @@
 #pragma once
 
 #include "CustomizableObjectInstanceEditor.h"
-#include "EditorUndoClient.h"
+#include "MuCOE/CustomizableObjectMacroLibrary/CustomizableObjectGraphEditorToolkit.h"
 #include "GraphEditor.h"
 #include "Misc/NotifyHook.h"
-#include "MuCOE/ICustomizableObjectEditor.h"
 #include "MuCO/CustomizableObjectPrivate.h"
 #include "Widgets/Input/SNumericDropDown.h"
 
@@ -80,11 +79,10 @@ public:
 /**
  * CustomizableObject Editor class
  */
-class FCustomizableObjectEditor : 
-	public ICustomizableObjectEditor, 
-	public FGCObject,
-	public FNotifyHook,
-	public FEditorUndoClient
+class FCustomizableObjectEditor :
+	public FCustomizableObjectGraphEditorToolkit,
+	public ICustomizableObjectInstanceEditor,
+	public FGCObject
 {
 public:
 	explicit FCustomizableObjectEditor(UCustomizableObject& ObjectToEdit);
@@ -117,15 +115,11 @@ public:
 	virtual FString GetWorldCentricTabPrefix() const override;
 	virtual FLinearColor GetWorldCentricTabColorScale() const override;
 
-	// ICustomizableObjectEditor interface
-	virtual UCustomizableObject* GetCustomizableObject() override;
+	// ICustomizableObjectInstanceEditor interface
+	UCustomizableObject* GetCustomizableObject();
 	virtual void RefreshTool() override;
 	virtual TSharedPtr<SCustomizableObjectEditorViewportTabBody> GetViewport() override;
 	virtual UCustomizableObjectInstance* GetPreviewInstance() override;
-	virtual bool CanPasteNodes() const override;
-	virtual void PasteNodesHere(const FVector2D& Location) override;
-	virtual void SelectNode(const UEdGraphNode* Node) override;
-	virtual void ReconstructAllChildNodes(UCustomizableObjectNode& StartNode, const UClass& NodeType) override;
 	virtual UProjectorParameter* GetProjectorParameter() override;
 	virtual UCustomSettings* GetCustomSettings() override;
 	virtual void HideGizmo() override;
@@ -145,15 +139,13 @@ public:
 	virtual TSharedPtr< SCustomizableObjectEditorAdvancedPreviewSettings> GetAdvancedPreviewSettings() override;
 	virtual bool ShowLightingSettings() override;
 	virtual bool ShowProfileManagementOptions() override;
+	virtual const UObject* GetObjectBeingEdited() override;
 
-	/** Select only this node only. Do nothing if already was only selected. */
-	void SelectSingleNode(UCustomizableObjectNode& Node);
-	
-	/** Called to undo the last action */
-	void UndoGraphAction();
 
-	/** Called to redo the last undone action */
-	void RedoGraphAction();
+	/** Reconstructs all child the nodes that match the given type.
+	 * @param StartNode Root node to start the graph traversal. This one also will be reconstructed.
+	 * @param NodeType Node types to reconstruct. */
+	void ReconstructAllChildNodes(UCustomizableObjectNode& StartNode, const UClass& NodeType);
 
 	/** Utility method: Test whether the CO Node Object given as parameter is linked to any of the CO Node Object Group nodes
 	* in the Test CO given as parameter */
@@ -161,7 +153,6 @@ public:
 
 	// FEditorUndoClient Interface
 	virtual void PostUndo(bool bSuccess) override;
-	virtual void PostRedo(bool bSuccess) override { PostUndo(bSuccess); }
 
 	// FNotifyHook interface
 	virtual void NotifyPostChange( const FPropertyChangedEvent& PropertyChangedEvent, FProperty* PropertyThatChanged ) override;
@@ -172,18 +163,6 @@ public:
 	void DuplicateSelectedNodes();
 	bool CanDuplicateSelectedNodes() const;
 	void OnSelectedGraphNodesChanged(const FGraphPanelSelectionSet& NewSelection);
-	void OnNodeDoubleClicked(UEdGraphNode* Node); //TODO(Max): Remove this once this editor uses the new toolkit
-
-	/**
-	 * Called when a node's title is committed for a rename
-	 *
-	 * @param	NewText				New title text
-	 * @param	CommitInfo			How text was committed
-	 * @param	NodeBeingChanged	The node being changed
-	 */
-	void OnNodeTitleCommitted(const FText& NewText, ETextCommit::Type CommitInfo, UEdGraphNode* NodeBeingChanged);
-
-	virtual void UpdateGraphNodeProperties() override;
 
 	/** Callback to notify the editor when the PreviewInstance has been updated */
 	void OnUpdatePreviewInstance(UCustomizableObjectInstance* Instance);
@@ -243,32 +222,6 @@ private:
 
 	/** Callback for the object modified event */
 	void OnObjectModified(UObject* Object);
-		
-	void CreateGraphEditorWidget(UEdGraph* InGraph);
-
-	/** Copy the currently selected nodes */
-	void CopySelectedNodes();
-	/** Whether we are able to copy the currently selected nodes */
-	bool CanCopyNodes() const;
-
-	/** Paste the contents of the clipboard */
-	void PasteNodes();
-
-	/** Cut the currently selected nodes */
-	void CutSelectedNodes();
-	/** Whether we are able to cut the currently selected nodes */
-	bool CanCutNodes() const;
-
-	/** On starting to rename node */
-	void OnRenameNode();
-
-	/** Check if node can be renamed */
-	bool CanRenameNodes() const;
-
-	virtual UEdGraphNode* CreateCommentBox(const FVector2D& NodePos) override;
-	
-	/** Searches a node that contains the inserted word */
-	void OnEnterText(const FText& NewText, ETextCommit::Type TextType);
 
 	/** Logs the search results of the search
 	 * @param Context The UObject we have found to be related with the searched string.
@@ -331,9 +284,6 @@ private:
 	/** Property View for Customizable Object and Nodes. */
 	TSharedPtr<class IDetailsView> ObjectDetailsView;
 
-	/** */
-	TSharedPtr<SGraphEditor> GraphEditor;
-
 	/** Widget to select which node pins are visible. */
 	TSharedPtr<class SCustomizableObjectNodePinViewer> NodePinViewer;
 	
diff --git a/Engine/Plugins/Mutable/Source/CustomizableObjectEditor/Private/MuCOE/CustomizableObjectEditorLogger.cpp b/Engine/Plugins/Mutable/Source/CustomizableObjectEditor/Private/MuCOE/CustomizableObjectEditorLogger.cpp
--- a/Engine/Plugins/Mutable/Source/CustomizableObjectEditor/Private/MuCOE/CustomizableObjectEditorLogger.cpp
+++ b/Engine/Plugins/Mutable/Source/CustomizableObjectEditor/Private/MuCOE/CustomizableObjectEditorLogger.cpp
@@ -8,7 +8,9 @@
 #include "Misc/UObjectToken.h"
 #include "MuCO/ICustomizableObjectEditorModule.h"
 #include "MuCOE/CustomizableObjectGraph.h"
-#include "MuCOE/ICustomizableObjectEditor.h"
+#include "MuCOE/CustomizableObjectMacroLibrary/CustomizableObjectGraphEditorToolkit.h"
+#include "MuCOE/CustomizableObjectMacroLibrary/CustomizableObjectMacroLibraryEditor.h"
+#include "MuCOE/CustomizableObjectMacroLibrary/CustomizableObjectMacroLibrary.h"
 #include "MuCOE/Nodes/CustomizableObjectNode.h"
 #include "Toolkits/ToolkitManager.h"
 #include "Widgets/Notifications/SNotificationList.h"
@@ -62,10 +64,30 @@ void OnMessageLogLinkActivated(const TSharedRef<IMessageToken>& Token)
 	if (Token->GetType() == EMessageToken::Object)
 	{
 		const TSharedRef<FUObjectToken> UObjectToken = StaticCastSharedRef<FUObjectToken>(Token);
+		
 		if (UObjectToken->GetObject().IsValid())
 		{
-			UObject* Object = UObjectToken->GetObject().Get();
-			GEditor->GetEditorSubsystem<UAssetEditorSubsystem>()->OpenEditorForAsset(Object);
+			if (UCustomizableObject* CustomizableObject = Cast<UCustomizableObject>(UObjectToken->GetObject().Get()))
+			{
+				GEditor->GetEditorSubsystem<UAssetEditorSubsystem>()->OpenEditorForAsset(CustomizableObject);
+			}
+			else if (UCustomizableObjectMacro* CustomizableObjectMacro = Cast<UCustomizableObjectMacro>(UObjectToken->GetObject().Get()))
+			{
+				if (UCustomizableObjectMacroLibrary* ParentMacroLibrary = Cast<UCustomizableObjectMacroLibrary>(CustomizableObjectMacro->GetOuter()))
+				{
+					// Open the editor for the Macro library
+					GEditor->GetEditorSubsystem<UAssetEditorSubsystem>()->OpenEditorForAsset(ParentMacroLibrary);
+
+					// Find the editor we just opened
+					IAssetEditorInstance* AssetEditor = GEditor->GetEditorSubsystem<UAssetEditorSubsystem>()->FindEditorForAsset(ParentMacroLibrary, false);
+
+					// Set the Parent Macro to edit
+					if (FCustomizableObjectMacroLibraryEditor* Editor = static_cast<FCustomizableObjectMacroLibraryEditor*>(AssetEditor))
+					{
+						Editor->SetSelectedMacro(CustomizableObjectMacro, true);
+					}
+				}
+			}
 		}
 	}
 
@@ -73,20 +95,40 @@ void OnMessageLogLinkActivated(const TSharedRef<IMessageToken>& Token)
 	else if (Token->GetType() == EMessageToken::Text)
 	{
 		const TSharedRef<FCustomizableObjectToken> UObjectToken = StaticCastSharedRef<FCustomizableObjectToken>(Token);
+
 		if (UObjectToken->GetNode())
 		{
 			const UEdGraphNode* Node = UObjectToken->GetNode();
 			UObject* Object = Node->GetOuter()->GetOuter();
 
-			// Make sure the editor exists for this asset
-			GEditor->GetEditorSubsystem<UAssetEditorSubsystem>()->OpenEditorForAsset(Object);
-
 			// Find it
 			if (const UCustomizableObject* CustomizableObject = Cast<UCustomizableObject>(Object))
 			{
+				// Make sure the editor exists for this asset
+				GEditor->GetEditorSubsystem<UAssetEditorSubsystem>()->OpenEditorForAsset(CustomizableObject);
+
 				if (const TSharedPtr<IToolkit> FoundAssetEditor = FToolkitManager::Get().FindEditorForAsset(CustomizableObject))
 				{
-					StaticCastSharedPtr<ICustomizableObjectEditor>(FoundAssetEditor)->SelectNode(Node);
+					StaticCastSharedPtr<FCustomizableObjectGraphEditorToolkit>(FoundAssetEditor)->SelectNode(Node);
+				}
+			}
+
+			else if (UCustomizableObjectMacro* CustomizableObjectMacro = Cast<UCustomizableObjectMacro>(Object))
+			{
+				if (UCustomizableObjectMacroLibrary* ParentMacroLibrary = Cast<UCustomizableObjectMacroLibrary>(CustomizableObjectMacro->GetOuter()))
+				{
+					// Open the editor for the Macro library
+					GEditor->GetEditorSubsystem<UAssetEditorSubsystem>()->OpenEditorForAsset(ParentMacroLibrary);
+
+					// Find the editor we just opened
+					IAssetEditorInstance* AssetEditor = GEditor->GetEditorSubsystem<UAssetEditorSubsystem>()->FindEditorForAsset(ParentMacroLibrary, false);
+
+					// Set the Parent Macro to edit
+					if (FCustomizableObjectMacroLibraryEditor* Editor = static_cast<FCustomizableObjectMacroLibraryEditor*>(AssetEditor))
+					{
+						Editor->SetSelectedMacro(CustomizableObjectMacro, true);
+						Editor->SelectNode(Node);
+					}
 				}
 			}
 		}
diff --git a/Engine/Plugins/Mutable/Source/CustomizableObjectEditor/Private/MuCOE/CustomizableObjectEditor_Deprecated.h b/Engine/Plugins/Mutable/Source/CustomizableObjectEditor/Private/MuCOE/CustomizableObjectEditor_Deprecated.h
--- a/Engine/Plugins/Mutable/Source/CustomizableObjectEditor/Private/MuCOE/CustomizableObjectEditor_Deprecated.h
+++ b/Engine/Plugins/Mutable/Source/CustomizableObjectEditor/Private/MuCOE/CustomizableObjectEditor_Deprecated.h
@@ -191,3 +191,15 @@ struct FComponentSettings
 	UPROPERTY()
 	TArray<FLODReductionSettings> LODReductionSettings;
 };
+
+
+// UCustomizableObjectNodeTextureVariation
+// Deprecated, do not use!
+USTRUCT()
+struct CUSTOMIZABLEOBJECTEDITOR_API FCustomizableObjectTextureVariation
+{
+	GENERATED_USTRUCT_BODY()
+
+	UPROPERTY(EditAnywhere, Category = CustomizableObject)
+	FString Tag;
+};
\ No newline at end of file
diff --git a/Engine/Plugins/Mutable/Source/CustomizableObjectEditor/Private/MuCOE/CustomizableObjectGraphEditorToolkit.cpp b/Engine/Plugins/Mutable/Source/CustomizableObjectEditor/Private/MuCOE/CustomizableObjectGraphEditorToolkit.cpp
deleted file mode 100644
diff --git a/Engine/Plugins/Mutable/Source/CustomizableObjectEditor/Private/MuCOE/CustomizableObjectInstanceEditor.cpp b/Engine/Plugins/Mutable/Source/CustomizableObjectEditor/Private/MuCOE/CustomizableObjectInstanceEditor.cpp
--- a/Engine/Plugins/Mutable/Source/CustomizableObjectEditor/Private/MuCOE/CustomizableObjectInstanceEditor.cpp
+++ b/Engine/Plugins/Mutable/Source/CustomizableObjectEditor/Private/MuCOE/CustomizableObjectInstanceEditor.cpp
@@ -607,6 +607,14 @@ bool FCustomizableObjectInstanceEditor::ShowProfileManagementOptions()
 }
 
 
+const UObject* FCustomizableObjectInstanceEditor::GetObjectBeingEdited()
+{
+	ensure(GetObjectsCurrentlyBeingEdited()->Num());
+
+	return (*GetObjectsCurrentlyBeingEdited())[0];
+}
+
+
 void FCustomizableObjectInstanceEditor::OnInstancePropertySelectionChanged(FProperty* InProperty)
 {
 	Viewport->GetViewportClient()->Invalidate();
diff --git a/Engine/Plugins/Mutable/Source/CustomizableObjectEditor/Private/MuCOE/CustomizableObjectInstanceEditor.h b/Engine/Plugins/Mutable/Source/CustomizableObjectEditor/Private/MuCOE/CustomizableObjectInstanceEditor.h
--- a/Engine/Plugins/Mutable/Source/CustomizableObjectEditor/Private/MuCOE/CustomizableObjectInstanceEditor.h
+++ b/Engine/Plugins/Mutable/Source/CustomizableObjectEditor/Private/MuCOE/CustomizableObjectInstanceEditor.h
@@ -172,6 +172,7 @@ public:
  */
 class FCustomizableObjectInstanceEditor : 
 	public ICustomizableObjectInstanceEditor,
+	public FAssetEditorToolkit,
 	public FGCObject,
 	public FTickableEditorObject
 {
@@ -219,6 +220,7 @@ public:
 	virtual TSharedPtr<SCustomizableObjectEditorAdvancedPreviewSettings> GetAdvancedPreviewSettings() override;
 	virtual bool ShowLightingSettings() override;
 	virtual bool ShowProfileManagementOptions() override;
+	virtual const UObject* GetObjectBeingEdited() override;
 
 	/** Callback to notify the editor when the PreviewInstance has been updated */
 	void OnUpdatePreviewInstance(UCustomizableObjectInstance* Instance);
diff --git a/Engine/Plugins/Mutable/Source/CustomizableObjectEditor/Private/MuCOE/CustomizableObjectLayout.cpp b/Engine/Plugins/Mutable/Source/CustomizableObjectEditor/Private/MuCOE/CustomizableObjectLayout.cpp
--- a/Engine/Plugins/Mutable/Source/CustomizableObjectEditor/Private/MuCOE/CustomizableObjectLayout.cpp
+++ b/Engine/Plugins/Mutable/Source/CustomizableObjectEditor/Private/MuCOE/CustomizableObjectLayout.cpp
@@ -9,7 +9,7 @@
 #include "MuCOE/GenerateMutableSource/GenerateMutableSourceLayout.h"
 #include "MuCOE/GraphTraversal.h"
 #include "MuCOE/MutableUtils.h"
-#include "MuCOE/ICustomizableObjectEditor.h"
+#include "MuCOE/CustomizableObjectEditor.h"
 #include "MuCOE/Nodes/CustomizableObjectNodeSkeletalMesh.h"
 #include "MuCOE/Nodes/CustomizableObjectNodeMesh.h"
 #include "MuCOE/Nodes/CustomizableObjectNodeTable.h"
@@ -85,8 +85,9 @@ void UCustomizableObjectLayout::GenerateAutomaticBlocksFromUVs()
 {
 	UCustomizableObjectNode* Node = Cast<UCustomizableObjectNode>(GetOuter());
 	TSoftObjectPtr<const UObject> Mesh = GetMesh();
+	TSharedPtr<FCustomizableObjectEditor> Editor = StaticCastSharedPtr<FCustomizableObjectEditor>(Node->GetGraphEditor());
 
-	if (!Node || !Mesh)
+	if (!Node || !Mesh || !Editor)
 	{
 		return;
 	}
@@ -98,7 +99,9 @@ void UCustomizableObjectLayout::GenerateAutomaticBlocksFromUVs()
 
 	// Create a GenerationContext
 	TSharedRef<FCustomizableObjectCompiler> Compiler = MakeShared<FCustomizableObjectCompiler>();
-	UCustomizableObject* Object = Node->GetGraphEditor()->GetCustomizableObject();
+	UCustomizableObject* Object = Editor->GetCustomizableObject();
+	check(Object);
+
 	FCompilationOptions Options = Object->GetPrivate()->GetCompileOptions();
 	
 	FMutableCompilationContext CompilationContext(Object, Compiler, Options);
diff --git a/Engine/Plugins/Mutable/Source/CustomizableObjectEditor/Private/MuCOE/CustomizableObjectMacroLibrary/CustomizableObjectGraphEditorToolkit.cpp b/Engine/Plugins/Mutable/Source/CustomizableObjectEditor/Private/MuCOE/CustomizableObjectMacroLibrary/CustomizableObjectGraphEditorToolkit.cpp
new file mode 100644
--- /dev/null
+++ b/Engine/Plugins/Mutable/Source/CustomizableObjectEditor/Private/MuCOE/CustomizableObjectMacroLibrary/CustomizableObjectGraphEditorToolkit.cpp
@@ -0,0 +1,706 @@
+// Copyright Epic Games, Inc. All Rights Reserved.
+
+#include "MuCOE/CustomizableObjectMacroLibrary/CustomizableObjectGraphEditorToolkit.h"
+
+#include "Framework/Commands/GenericCommands.h"
+#include "Framework/Commands/UIAction.h"
+#include "HAL/PlatformApplicationMisc.h"
+#include "MaterialGraph/MaterialGraphNode.h"
+#include "MaterialGraph/MaterialGraphNode_Comment.h"
+#include "MuCOE/CustomizableObjectEditorNodeContextCommands.h"
+#include "MuCOE/CustomizableObjectEditorLogger.h"
+#include "MuCOE/CustomizableObjectGraph.h"
+#include "MuCOE/Nodes/CustomizableObjectNode.h"
+#include "SAssetSearchBox.h"
+#include "SNodePanel.h"
+#include "Widgets/Input/SSearchBox.h"
+
+#define LOCTEXT_NAMESPACE "CustomizableObjectGraphEditorToolkit"
+
+bool FCustomizableObjectGraphEditorToolkit::CanPasteNodes() const 
+{
+	if (!GraphEditor.IsValid() || !GraphEditor->GetCurrentGraph())
+	{
+		return false;
+	}
+
+	FString ClipboardContent;
+	FPlatformApplicationMisc::ClipboardPaste(ClipboardContent);
+
+	return FEdGraphUtilities::CanImportNodesFromText(GraphEditor->GetCurrentGraph(), ClipboardContent);
+}
+
+
+void FCustomizableObjectGraphEditorToolkit::PasteNodes()
+{
+	if (!GraphEditor.IsValid())
+	{
+		return;
+	}
+
+	PasteNodesHere(GraphEditor->GetPasteLocation());
+}
+
+
+void FCustomizableObjectGraphEditorToolkit::PasteNodesHere(const FVector2D& Location) 
+{
+	if (!GraphEditor.IsValid() || !GraphEditor->GetCurrentGraph())
+	{
+		return;
+	}
+
+	const FScopedTransaction Transaction(LOCTEXT("CustomizableObjectEditorPaste", "Customizable Object Editor Editor: Paste"));
+	GraphEditor->GetCurrentGraph()->Modify();
+
+	// Clear the selection set (newly pasted stuff will be selected)
+	GraphEditor->ClearSelectionSet();
+
+	// Grab the text to paste from the clipboard.
+	FString TextToImport;
+	FPlatformApplicationMisc::ClipboardPaste(TextToImport);
+
+	// Import the nodes
+	TSet<UEdGraphNode*> PastedNodes;
+	FEdGraphUtilities::ImportNodesFromText(GraphEditor->GetCurrentGraph(), TextToImport, PastedNodes);
+
+	//Average position of nodes so we can move them while still maintaining relative distances to each other
+	FVector2D AvgNodePosition(0.0f, 0.0f);
+
+	for (TSet<UEdGraphNode*>::TIterator It(PastedNodes); It; ++It)
+	{
+		UEdGraphNode* Node = *It;
+		AvgNodePosition.X += Node->NodePosX;
+		AvgNodePosition.Y += Node->NodePosY;
+	}
+
+	if (PastedNodes.Num() > 0)
+	{
+		float InvNumNodes = 1.0f / float(PastedNodes.Num());
+		AvgNodePosition.X *= InvNumNodes;
+		AvgNodePosition.Y *= InvNumNodes;
+	}
+
+	for (TSet<UEdGraphNode*>::TIterator It(PastedNodes); It; ++It)
+	{
+		UEdGraphNode* Node = *It;
+
+		// Select the newly pasted stuff
+		GraphEditor->SetNodeSelection(Node, true);
+
+		Node->NodePosX = (Node->NodePosX - AvgNodePosition.X) + Location.X;
+		Node->NodePosY = (Node->NodePosY - AvgNodePosition.Y) + Location.Y;
+
+		Node->SnapToGrid(SNodePanel::GetSnapGridSize());
+
+		// Give new node a different Guid from the old one
+		Node->CreateNewGuid();
+	}
+
+	for (UEdGraphNode* PastedNode : PastedNodes)
+	{
+		if (UCustomizableObjectNode* TypedNode = Cast<UCustomizableObjectNode>(PastedNode))
+		{
+			TypedNode->PostBackwardsCompatibleFixup();
+		}
+	}
+
+	// Update UI
+	GraphEditor->NotifyGraphChanged();
+
+	//TODO(Max): Review why do we need to call the line below. Does it make sence?
+	//CustomizableObject->PostEditChange();
+	GraphEditor->GetCurrentGraph()->MarkPackageDirty();
+}
+
+
+void FCustomizableObjectGraphEditorToolkit::SelectNode(const UEdGraphNode* Node) 
+{
+	if (!GraphEditor.IsValid())
+	{
+		return;
+	}
+
+	GraphEditor->JumpToNode(Node);
+}
+
+
+void FCustomizableObjectGraphEditorToolkit::SelectSingleNode(UCustomizableObjectNode& Node)
+{
+	FGraphPanelSelectionSet SelectedNodes = GraphEditor->GetSelectedNodes();
+	if (SelectedNodes.Num() != 1 || Cast<UCustomizableObjectNode>(*SelectedNodes.CreateIterator()) != &Node)
+	{
+		GraphEditor->ClearSelectionSet();
+		GraphEditor->SetNodeSelection(&Node, true);
+	}
+}
+
+
+void FCustomizableObjectGraphEditorToolkit::DeleteSelectedNodes()
+{
+	if (!GraphEditor.IsValid())
+	{
+		return;
+	}
+
+	const FScopedTransaction Transaction(LOCTEXT("UEdGraphSchema_CustomizableObject", "Delete Nodes"));
+
+	const FGraphPanelSelectionSet SelectedNodes = GraphEditor->GetSelectedNodes();
+	GraphEditor->ClearSelectionSet();
+
+	for (FGraphPanelSelectionSet::TConstIterator NodeIt(SelectedNodes); NodeIt; ++NodeIt)
+	{
+		if (UEdGraphNode* Node = Cast<UEdGraphNode>(*NodeIt))
+		{
+			if (Node->CanUserDeleteNode())
+			{
+				if (const UEdGraph* GraphObj = Node->GetGraph())
+				{
+					if (const UEdGraphSchema* Schema = GraphObj->GetSchema())
+					{
+						Schema->BreakNodeLinks(*Node);  // Required to notify to all connected nodes (UEdGraphNode::PinConnectionListChanged() and UEdGraphNode::PinConnectionListChanged(...))
+					}
+				}
+
+				Node->DestroyNode();
+			}
+		}
+	}
+}
+
+
+bool FCustomizableObjectGraphEditorToolkit::CanDeleteNodes() const
+{
+	if (GraphEditor.IsValid() && GraphEditor->GetSelectedNodes().Num() > 0)
+	{
+		for (auto Itr = GraphEditor->GetSelectedNodes().CreateConstIterator(); Itr; ++Itr)
+		{
+			UCustomizableObjectNode* Node = Cast<UCustomizableObjectNode>(*Itr);
+
+			if (Node && !Node->CanUserDeleteNode())
+			{
+				return false;
+			}
+		}
+
+		return true;
+	}
+
+	return false;
+}
+
+
+void FCustomizableObjectGraphEditorToolkit::DuplicateSelectedNodes()
+{
+	CopySelectedNodes();
+	PasteNodes();
+}
+
+
+bool FCustomizableObjectGraphEditorToolkit::CanDuplicateSelectedNodes() const
+{
+	return CanCopyNodes();
+}
+
+
+void FCustomizableObjectGraphEditorToolkit::OnNodeTitleCommitted(const FText& NewText, ETextCommit::Type CommitInfo, UEdGraphNode* NodeBeingChanged)
+{
+	if (NodeBeingChanged)
+	{
+		const FScopedTransaction Transaction(LOCTEXT("RenameNode", "Rename Node"));
+		NodeBeingChanged->Modify();
+		NodeBeingChanged->OnRenameNode(NewText.ToString());
+	}
+}
+
+
+void FCustomizableObjectGraphEditorToolkit::OnNodeDoubleClicked(UEdGraphNode* Node)
+{
+	if (Node->CanJumpToDefinition())
+	{
+		Node->JumpToDefinition();
+	}
+}
+
+
+void FCustomizableObjectGraphEditorToolkit::CopySelectedNodes()
+{
+	// Export the selected nodes and place the text on the clipboard
+	const FGraphPanelSelectionSet SelectedNodes = GraphEditor->GetSelectedNodes();
+
+	FString ExportedText;
+
+	for (FGraphPanelSelectionSet::TConstIterator SelectedIter(SelectedNodes); SelectedIter; ++SelectedIter)
+	{
+		if (UEdGraphNode* Node = Cast<UEdGraphNode>(*SelectedIter))
+		{
+			Node->PrepareForCopying();
+		}
+	}
+
+	FEdGraphUtilities::ExportNodesToText(SelectedNodes, /*out*/ ExportedText);
+	FPlatformApplicationMisc::ClipboardCopy(*ExportedText);
+
+	//TODO(Max): This is weird... review it
+	// Make sure Material remains the owner of the copied nodes
+	for (FGraphPanelSelectionSet::TConstIterator SelectedIter(SelectedNodes); SelectedIter; ++SelectedIter)
+	{
+		if (UMaterialGraphNode* Node = Cast<UMaterialGraphNode>(*SelectedIter))
+		{
+			Node->PostCopyNode();
+		}
+		else if (UMaterialGraphNode_Comment* Comment = Cast<UMaterialGraphNode_Comment>(*SelectedIter))
+		{
+			Comment->PostCopyNode();
+		}
+	}
+}
+
+
+bool FCustomizableObjectGraphEditorToolkit::CanCopyNodes() const
+{
+	// If any of the nodes can be duplicated then we should allow copying
+	const FGraphPanelSelectionSet SelectedNodes = GraphEditor->GetSelectedNodes();
+	for (FGraphPanelSelectionSet::TConstIterator SelectedIter(SelectedNodes); SelectedIter; ++SelectedIter)
+	{
+		UEdGraphNode* Node = Cast<UEdGraphNode>(*SelectedIter);
+		if ((Node != NULL) && Node->CanDuplicateNode())
+		{
+			return true;
+		}
+	}
+
+	return false;
+}
+
+
+void FCustomizableObjectGraphEditorToolkit::CutSelectedNodes()
+{
+	CopySelectedNodes();
+	// Cut should only delete nodes that can be duplicated
+	DeleteSelectedNodes();
+}
+
+
+bool FCustomizableObjectGraphEditorToolkit::CanCutNodes() const
+{
+	return CanCopyNodes() && CanDeleteNodes();
+}
+
+
+void FCustomizableObjectGraphEditorToolkit::OnRenameNode()
+{
+	if (GraphEditor.IsValid())
+	{
+		const FGraphPanelSelectionSet& SelectedNodes = GraphEditor->GetSelectedNodes();
+
+		for (FGraphPanelSelectionSet::TConstIterator NodeIt(SelectedNodes); NodeIt; ++NodeIt)
+		{
+			// Rename only the first valid selected node
+			UEdGraphNode* SelectedNode = Cast<UEdGraphNode>(*NodeIt);
+			if (SelectedNode && SelectedNode->GetCanRenameNode())
+			{
+				GraphEditor->IsNodeTitleVisible(SelectedNode, true);
+				break;
+			}
+		}
+	}
+}
+
+
+bool FCustomizableObjectGraphEditorToolkit::CanRenameNodes() const
+{
+	if (GraphEditor.IsValid())
+	{
+		const FGraphPanelSelectionSet& SelectedNodes = GraphEditor->GetSelectedNodes();
+
+		for (FGraphPanelSelectionSet::TConstIterator SelectedIter(SelectedNodes); SelectedIter; ++SelectedIter)
+		{
+			if (UEdGraphNode* Node = Cast<UEdGraphNode>(*SelectedIter))
+			{
+				return Node->GetCanRenameNode();
+			}
+		}
+	}
+
+	return false;
+}
+
+
+void FCustomizableObjectGraphEditorToolkit::CreateCommentBoxFromKey()
+{
+	CreateCommentBox(GraphEditor->GetPasteLocation());
+}
+
+
+UEdGraphNode* FCustomizableObjectGraphEditorToolkit::CreateCommentBox(const FVector2D& NodePos)
+{
+	if (!GraphEditor.IsValid() || !GraphEditor->GetCurrentGraph())
+	{
+		return nullptr;
+	}
+
+	//const FScopedTransaction Transaction(LOCTEXT("UEdGraphSchema_CustomizableObject", "Add Comment Box"));
+	UEdGraphNode_Comment* CommentTemplate = NewObject<UEdGraphNode_Comment>();
+	UEdGraphNode_Comment* NewComment = nullptr;
+
+	{
+		// const FGraphPanelSelectionSet& SelectionSet = GraphEditor->GetSelectedNodes();
+		FSlateRect Bounds;
+		FVector2D Location;
+		FVector2D Size;
+
+		if (GraphEditor->GetBoundsForSelectedNodes(Bounds, 50.f))
+		{
+			Location.X = Bounds.Left;
+			Location.Y = Bounds.Top;
+			Size = Bounds.GetSize();
+		}
+		else
+		{
+			Location.X = NodePos.X;
+			Location.Y = NodePos.Y;
+			Size.X = 400;
+			Size.Y = 100;
+		}
+
+		NewComment = FEdGraphSchemaAction_NewNode::SpawnNodeFromTemplate<UEdGraphNode_Comment>(GraphEditor->GetCurrentGraph(), CommentTemplate, NodePos, true);
+		NewComment->NodePosX = Location.X;
+		NewComment->NodePosY = Location.Y;
+		NewComment->NodeWidth = Size.X;
+		NewComment->NodeHeight = Size.Y;
+		NewComment->NodeComment = FString(TEXT("Comment"));
+	}
+
+	GraphEditor->GetCurrentGraph()->MarkPackageDirty();
+	GraphEditor->NotifyGraphChanged();
+
+	return NewComment;
+}
+
+
+void FCustomizableObjectGraphEditorToolkit::OnEnterText(const FText& NewText, ETextCommit::Type TextType)
+{
+	if (TextType != ETextCommit::OnEnter)
+	{
+		return;
+	}
+
+	if (!GraphEditor)
+	{
+		return;
+	}
+
+	const UEdGraph* Graph = GraphEditor->GetCurrentGraph();
+	if (!Graph)
+	{
+		return;
+	}
+
+	bool bFound = false;
+
+	const FString FindString = NewText.ToString();
+
+	for (TObjectPtr<UEdGraphNode> Node : Graph->Nodes)
+	{
+		if (!Node)
+		{
+			continue;
+		}
+
+		// Node names are not in the reflection system
+		const FString NodeName = Node->GetNodeTitle(ENodeTitleType::FullTitle).ToString().Replace(TEXT("\n"), TEXT(" "));
+		if (NodeName.Contains(NewText.ToString(), ESearchCase::IgnoreCase))
+		{
+			LogSearchResult(*Node, "Node", bFound, NodeName);
+			bFound = true;
+		}
+
+		// Pins are not in the reflection system
+		for (const UEdGraphPin* Pin : Node->GetAllPins())
+		{
+			const FString PinFriendlyName = Pin->PinFriendlyName.ToString();
+			if (PinFriendlyName.Contains(FindString))
+			{
+				LogSearchResult(*Node, "Pin", bFound, PinFriendlyName);
+				bFound = true;
+			}
+		}
+
+		// Find anything marked as a UPROPERTY
+		for (TFieldIterator<FProperty> It(Node->GetClass()); It; ++It)
+		{
+			FindProperty(*It, Node, FindString, *Node, bFound);
+		}
+	}
+
+	const FText Text = bFound ?
+		LOCTEXT("SearchCompleted", "Search completed") :
+		FText::FromString("No Results for: " + NewText.ToString());
+
+	FCustomizableObjectEditorLogger::CreateLog(Text)
+		.Category(ELoggerCategory::GraphSearch)
+		.CustomNotification()
+		.Log();
+}
+
+
+void FCustomizableObjectGraphEditorToolkit::FindProperty(const FProperty* Property, const void* InContainer, const FString& FindString, const UObject& Context, bool& bFound)
+{
+	if (!Property || !InContainer)
+	{
+		return;
+	}
+
+	const FString PropertyName = Property->GetDisplayNameText().ToString();
+	if (PropertyName.Contains(FindString))
+	{
+		LogSearchResult(Context, "Property Name", bFound, *PropertyName);
+		bFound = true;
+	}
+
+	for (int32 Index = 0; Index < Property->ArrayDim; ++Index)
+	{
+		const uint8* ValuePtr = Property->ContainerPtrToValuePtr<uint8>(InContainer, Index);
+
+		if (const FStrProperty* StringProperty = CastField<FStrProperty>(Property))
+		{
+			const FString* StringResult = StringProperty->GetPropertyValuePtr(ValuePtr);
+			if (StringResult->Contains(FindString))
+			{
+				LogSearchResult(Context, "Property Value", bFound, *StringResult);
+				bFound = true;
+			}
+		}
+		else if (const FEnumProperty* EnumProperty = CastField<FEnumProperty>(Property))
+		{
+			const UEnum* EnumResult = EnumProperty->GetEnum();
+
+			const FString StringResult = EnumResult->GetDisplayNameTextByIndex(*ValuePtr).ToString();
+			if (StringResult.Contains(FindString))
+			{
+				LogSearchResult(Context, "Property Value", bFound, StringResult);
+				bFound = true;
+			}
+		}
+		else if (const FSoftObjectProperty* SoftObjectProperty = CastField<FSoftObjectProperty>(Property))
+		{
+			const FString ObjectPath = SoftObjectProperty->GetPropertyValuePtr(ValuePtr)->ToString();
+			if (ObjectPath.Contains(FindString))
+			{
+				LogSearchResult(Context, "Property Value", bFound, ObjectPath);
+				bFound = true;
+			}
+		}
+		else if (const FObjectPropertyBase* ObjectProperty = CastField<FObjectPropertyBase>(Property))
+		{
+			if (const UObject* ObjectValue = ObjectProperty->GetObjectPropertyValue(ValuePtr))
+			{
+				const FString Name = ObjectValue->GetName();
+
+				if (ObjectValue->GetName().Contains(FindString))
+				{
+					LogSearchResult(Context, "Property Value", bFound, Name);
+					bFound = true;
+				}
+			}
+		}
+		else if (const FStructProperty* StructProperty = CastField<FStructProperty>(Property))
+		{
+			for (TFieldIterator<FProperty> It(StructProperty->Struct); It; ++It)
+			{
+				FindProperty(*It, ValuePtr, FindString, Context, bFound);
+			}
+		}
+		else if (const FArrayProperty* ArrayProperty = CastField<FArrayProperty>(Property))
+		{
+			FScriptArrayHelper ArrayHelper(ArrayProperty, ValuePtr);
+			for (int32 ValueIdx = 0; ValueIdx < ArrayHelper.Num(); ++ValueIdx)
+			{
+				FindProperty(ArrayProperty->Inner, ArrayHelper.GetRawPtr(ValueIdx), FindString, Context, bFound);
+			}
+		}
+		else if (const FSetProperty* SetProperty = CastField<FSetProperty>(Property))
+		{
+			FScriptSetHelper SetHelper(SetProperty, ValuePtr);
+			for (FScriptSetHelper::FIterator SetIt = SetHelper.CreateIterator(); SetIt; ++SetIt)
+			{
+				FindProperty(SetProperty->ElementProp, SetHelper.GetElementPtr(SetIt), FindString, Context, bFound);
+			}
+		}
+		else if (const FMapProperty* MapProperty = CastField<FMapProperty>(Property))
+		{
+			FScriptMapHelper MapHelper(MapProperty, ValuePtr);
+			for (FScriptMapHelper::FIterator MapIt = MapHelper.CreateIterator(); MapIt; ++MapIt)
+			{
+				const uint8* MapValuePtr = MapHelper.GetPairPtr(MapIt);
+				FindProperty(MapProperty->KeyProp, MapValuePtr, FindString, Context, bFound);
+				FindProperty(MapProperty->ValueProp, MapValuePtr, FindString, Context, bFound);
+			}
+		}
+	}
+}
+
+
+void FCustomizableObjectGraphEditorToolkit::LogSearchResult(const UObject& Context, const FString& Type, bool bIsFirst, const FString& Result) const
+{
+	if (!bIsFirst)
+	{
+		FCustomizableObjectEditorLogger::CreateLog(LOCTEXT("SearchResults", "Search Results:"))
+			.Notification(false)
+			.Log();
+	}
+
+	FCustomizableObjectEditorLogger::CreateLog(FText::FromString(Type + ": " + Result))
+		.Context(Context)
+		.BaseObject()
+		.Notification(false)
+		.Log();
+}
+
+
+/** Create new tab for the supplied graph - don't call this directly, call SExplorer->FindTabForGraph.*/
+void FCustomizableObjectGraphEditorToolkit::CreateGraphEditorWidget(UEdGraph* InGraph, SGraphEditor::FGraphEditorEvents& InEvents)
+{
+	UCustomizableObjectGraph* CustomizableObjectGraph = Cast<UCustomizableObjectGraph>(InGraph);
+	check(CustomizableObjectGraph);
+
+	GraphEditorCommands = MakeShareable(new FUICommandList);
+
+	TSharedRef<SWidget> TitleBarWidget =
+		SNew(SHorizontalBox)
+		+ SHorizontalBox::Slot()
+		.HAlign(HAlign_Left)
+		.FillWidth(10.f)
+		.Padding(5.f)
+		[
+			SNew(SSearchBox)
+				.HintText(LOCTEXT("Search", "Search..."))
+				.ToolTipText(LOCTEXT("Search Nodes, Properties or Values that contain the inserted words", "Search Nodes, Properties or Values that contain the inserted words"))
+				.OnTextCommitted(FOnTextCommitted::CreateSP(this, &FCustomizableObjectGraphEditorToolkit::OnEnterText))
+				.SelectAllTextWhenFocused(true)
+		];
+
+	// Create the appearance info
+	FGraphAppearanceInfo AppearanceInfo;
+	AppearanceInfo.CornerText = CustomizableObjectGraph->IsMacro() ? LOCTEXT("ApperanceCornerMacroText", "MUTABLE MACRO") : LOCTEXT("ApperanceCornerText", "MUTABLE");
+
+	// Add toolkit common events
+	InEvents.OnTextCommitted = FOnNodeTextCommitted::CreateSP(this, &FCustomizableObjectGraphEditorToolkit::OnNodeTitleCommitted);
+	InEvents.OnNodeDoubleClicked = FSingleNodeEvent::CreateSP(this, &FCustomizableObjectGraphEditorToolkit::OnNodeDoubleClicked);
+
+	// Make full graph editor
+	GraphEditor = SNew(SGraphEditor)
+		.AdditionalCommands(GraphEditorCommands)
+		.Appearance(AppearanceInfo)
+		.GraphToEdit(InGraph)
+		.GraphEvents(InEvents)
+		.TitleBar(TitleBarWidget)
+		.ShowGraphStateOverlay(false); // Removes graph state overlays (border and text) such as "SIMULATING" and "READ-ONLY"
+
+	// Ensure commands are registered
+	FGraphEditorCommands::Register();
+	FCustomizableObjectEditorNodeContextCommands::Register();
+
+	// Editing commands
+	GraphEditorCommands->MapAction(FGenericCommands::Get().Delete,
+		FExecuteAction::CreateSP(this, &FCustomizableObjectGraphEditorToolkit::DeleteSelectedNodes),
+		FCanExecuteAction::CreateSP(this, &FCustomizableObjectGraphEditorToolkit::CanDeleteNodes));
+
+	GraphEditorCommands->MapAction(FGenericCommands::Get().Copy,
+		FExecuteAction::CreateSP(this, &FCustomizableObjectGraphEditorToolkit::CopySelectedNodes),
+		FCanExecuteAction::CreateSP(this, &FCustomizableObjectGraphEditorToolkit::CanCopyNodes));
+
+	GraphEditorCommands->MapAction(FGenericCommands::Get().Paste,
+		FExecuteAction::CreateSP(this, &FCustomizableObjectGraphEditorToolkit::PasteNodes),
+		FCanExecuteAction::CreateSP(this, &FCustomizableObjectGraphEditorToolkit::CanPasteNodes));
+
+	GraphEditorCommands->MapAction(FGenericCommands::Get().Cut,
+		FExecuteAction::CreateSP(this, &FCustomizableObjectGraphEditorToolkit::CutSelectedNodes),
+		FCanExecuteAction::CreateSP(this, &FCustomizableObjectGraphEditorToolkit::CanCutNodes));
+
+	GraphEditorCommands->MapAction(FGenericCommands::Get().Duplicate,
+		FExecuteAction::CreateSP(this, &FCustomizableObjectGraphEditorToolkit::DuplicateSelectedNodes),
+		FCanExecuteAction::CreateSP(this, &FCustomizableObjectGraphEditorToolkit::CanDuplicateSelectedNodes));
+
+	GraphEditorCommands->MapAction(FGenericCommands::Get().Rename,
+		FExecuteAction::CreateSP(this, &FCustomizableObjectGraphEditorToolkit::OnRenameNode),
+		FCanExecuteAction::CreateSP(this, &FCustomizableObjectGraphEditorToolkit::CanRenameNodes));
+
+	GraphEditorCommands->MapAction(FCustomizableObjectEditorNodeContextCommands::Get().CreateComment,
+		FExecuteAction::CreateSP(this, &FCustomizableObjectGraphEditorToolkit::CreateCommentBoxFromKey));
+
+	// Alignment Commands
+	GraphEditorCommands->MapAction(FGraphEditorCommands::Get().AlignNodesTop,
+		FExecuteAction::CreateSP(GraphEditor.Get(), &SGraphEditor::OnAlignTop));
+
+	GraphEditorCommands->MapAction(FGraphEditorCommands::Get().AlignNodesMiddle,
+		FExecuteAction::CreateSP(GraphEditor.Get(), &SGraphEditor::OnAlignMiddle));
+
+	GraphEditorCommands->MapAction(FGraphEditorCommands::Get().AlignNodesBottom,
+		FExecuteAction::CreateSP(GraphEditor.Get(), &SGraphEditor::OnAlignBottom));
+
+	GraphEditorCommands->MapAction(FGraphEditorCommands::Get().AlignNodesLeft,
+		FExecuteAction::CreateSP(GraphEditor.Get(), &SGraphEditor::OnAlignLeft));
+
+	GraphEditorCommands->MapAction(FGraphEditorCommands::Get().AlignNodesCenter,
+		FExecuteAction::CreateSP(GraphEditor.Get(), &SGraphEditor::OnAlignCenter));
+
+	GraphEditorCommands->MapAction(FGraphEditorCommands::Get().AlignNodesRight,
+		FExecuteAction::CreateSP(GraphEditor.Get(), &SGraphEditor::OnAlignRight));
+
+	GraphEditorCommands->MapAction(FGraphEditorCommands::Get().StraightenConnections,
+		FExecuteAction::CreateSP(GraphEditor.Get(), &SGraphEditor::OnStraightenConnections));
+
+	// Distribution Commands
+	GraphEditorCommands->MapAction(FGraphEditorCommands::Get().DistributeNodesHorizontally,
+		FExecuteAction::CreateSP(GraphEditor.Get(), &SGraphEditor::OnDistributeNodesH));
+
+	GraphEditorCommands->MapAction(FGraphEditorCommands::Get().DistributeNodesVertically,
+		FExecuteAction::CreateSP(GraphEditor.Get(), &SGraphEditor::OnDistributeNodesV));
+}
+
+
+void FCustomizableObjectGraphEditorToolkit::UpdateGraphNodeProperties()
+{
+	// Cache a copy of the selected nodes so we can later restore them
+	const FGraphPanelSelectionSet PreClearingSelectedNodes = GraphEditor->GetSelectedNodes();
+
+	OnSelectedGraphNodesChanged(FGraphPanelSelectionSet());
+	OnSelectedGraphNodesChanged(PreClearingSelectedNodes);
+}
+
+
+void FCustomizableObjectGraphEditorToolkit::BindGraphCommands()
+{
+	// Undo-Redo
+	ToolkitCommands->MapAction(
+		FGenericCommands::Get().Undo,
+		FExecuteAction::CreateSP(this, &FCustomizableObjectGraphEditorToolkit::UndoGraphAction));
+
+	ToolkitCommands->MapAction(
+		FGenericCommands::Get().Redo,
+		FExecuteAction::CreateSP(this, &FCustomizableObjectGraphEditorToolkit::RedoGraphAction));
+}
+
+
+void FCustomizableObjectGraphEditorToolkit::PostUndo(bool bSuccess)
+{
+	GraphEditor->NotifyGraphChanged();
+}
+
+
+void FCustomizableObjectGraphEditorToolkit::UndoGraphAction()
+{
+	GEditor->UndoTransaction();
+}
+
+
+void FCustomizableObjectGraphEditorToolkit::RedoGraphAction()
+{
+	// Clear selection, to avoid holding refs to nodes that go away
+	GraphEditor->ClearSelectionSet();
+
+	GEditor->RedoTransaction();
+}
+
+#undef LOCTEXT_NAMESPACE
\ No newline at end of file
diff --git a/Engine/Plugins/Mutable/Source/CustomizableObjectEditor/Private/MuCOE/CustomizableObjectMacroLibrary/CustomizableObjectGraphEditorToolkit.h b/Engine/Plugins/Mutable/Source/CustomizableObjectEditor/Private/MuCOE/CustomizableObjectMacroLibrary/CustomizableObjectGraphEditorToolkit.h
new file mode 100644
--- /dev/null
+++ b/Engine/Plugins/Mutable/Source/CustomizableObjectEditor/Private/MuCOE/CustomizableObjectMacroLibrary/CustomizableObjectGraphEditorToolkit.h
@@ -0,0 +1,101 @@
+// Copyright Epic Games, Inc. All Rights Reserved.
+
+#pragma once
+
+#include "BlueprintUtilities.h"
+#include "EdGraph/EdGraph.h"
+#include "EdGraph/EdGraphNode.h"
+#include "EdGraphUtilities.h"
+#include "EditorUndoClient.h"
+#include "GraphEditor.h"
+#include "GraphEditorActions.h"
+#include "Misc/NotifyHook.h"
+#include "ScopedTransaction.h"
+#include "Toolkits/AssetEditorToolkit.h"
+
+class UClass;
+class UCustomizableObjectNode;
+
+// Public interface to Customizable Object Editor
+class FCustomizableObjectGraphEditorToolkit : public FAssetEditorToolkit, public FNotifyHook, public FSelfRegisteringEditorUndoClient
+{
+public:
+	
+	/** Searches a node that contains the inserted word */
+	void OnEnterText(const FText& NewText, ETextCommit::Type TextType);
+	void FindProperty(const FProperty* Property, const void* InContainer, const FString& FindString, const UObject& Context, bool& bFound);
+	void LogSearchResult(const UObject& Context, const FString& Type, bool bIsFirst, const FString& Result) const;
+
+	/** Bind common graph commands */
+	void BindGraphCommands();
+
+	// FEditorUndoClient Interface
+	virtual void PostUndo(bool bSuccess) override;
+	virtual void PostRedo(bool bSuccess) override { PostUndo(bSuccess); }
+
+	// Graph Node Actions --------------------------
+	void SelectNode(const UEdGraphNode* Node);
+
+	/** Select only and this node only. Do nothing if already was only selected. */
+	void SelectSingleNode(UCustomizableObjectNode& Node);
+		
+	void DeleteSelectedNodes();
+	bool CanDeleteNodes() const;
+	
+	void DuplicateSelectedNodes();
+	bool CanDuplicateSelectedNodes() const;
+
+	/**
+	 * Called when a node's title is committed for a rename
+	 *
+	 * @param	NewText				New title text
+	 * @param	CommitInfo			How text was committed
+	 * @param	NodeBeingChanged	The node being changed
+	 */
+	void OnNodeTitleCommitted(const FText& NewText, ETextCommit::Type CommitInfo, UEdGraphNode* NodeBeingChanged);
+	
+	void CopySelectedNodes();
+	bool CanCopyNodes() const;
+	
+	void PasteNodes();
+	void PasteNodesHere(const FVector2D& Location);
+	bool CanPasteNodes() const;
+	
+	void CutSelectedNodes();
+	bool CanCutNodes() const;
+
+	void OnRenameNode();
+	bool CanRenameNodes() const;
+
+	void CreateCommentBoxFromKey();
+	UEdGraphNode* CreateCommentBox(const FVector2D& NodePos);
+
+	void OnNodeDoubleClicked(UEdGraphNode* Node);
+
+	//Undo/Redo graph Actions
+	void UndoGraphAction();
+	void RedoGraphAction();
+
+	// ------------------------------------------
+
+	/** Creates a new graph editor widget 
+	* @param InGraph is the graph that will be represented in the graph editor.
+	* @param InEvents are custom events occurring in/on the graph.
+	*/
+	void CreateGraphEditorWidget(UEdGraph* InGraph, SGraphEditor::FGraphEditorEvents& InEvents);
+
+	/** Graph Editor callbacks */
+	virtual void OnSelectedGraphNodesChanged(const FGraphPanelSelectionSet& NewSelection) = 0;
+
+	virtual void UpdateGraphNodeProperties();
+
+public:
+
+	/** Pointer to the graph editor widget */
+	TSharedPtr<SGraphEditor> GraphEditor;
+
+private:
+
+	/** Command list for the graph editor */
+	TSharedPtr<FUICommandList> GraphEditorCommands;
+};
diff --git a/Engine/Plugins/Mutable/Source/CustomizableObjectEditor/Private/MuCOE/CustomizableObjectMacroLibrary/CustomizableObjectMacroDetails.cpp b/Engine/Plugins/Mutable/Source/CustomizableObjectEditor/Private/MuCOE/CustomizableObjectMacroLibrary/CustomizableObjectMacroDetails.cpp
--- a/Engine/Plugins/Mutable/Source/CustomizableObjectEditor/Private/MuCOE/CustomizableObjectMacroLibrary/CustomizableObjectMacroDetails.cpp
+++ b/Engine/Plugins/Mutable/Source/CustomizableObjectEditor/Private/MuCOE/CustomizableObjectMacroLibrary/CustomizableObjectMacroDetails.cpp
@@ -11,7 +11,7 @@
 #include "PropertyCustomizationHelpers.h"
 #include "SListViewSelectorDropdownMenu.h"
 #include "Widgets/Input/SButton.h"
-#include "Widgets/Input/SComboButton.h"
+#include "Widgets/Input/SComboBox.h"
 #include "Widgets/SBoxPanel.h"
 #include "Widgets/Views/SListView.h"
 #include "Widgets/Views/STableRow.h"
@@ -230,60 +230,44 @@ void SCOMacroPinTypeSelector::Construct(const FArguments& InArgs)
 	const UEdGraphSchema_CustomizableObject* Schema = Cast<UEdGraphSchema_CustomizableObject>(Macro->Graph->GetSchema());
 	check(Schema);
 
+	TSharedPtr<FPinNameRowData> InitialSelection;
+
 	for (const FName& PinType : Schema->SupportedMacroPinTypes)
 	{
 		TSharedPtr<FPinNameRowData> NewSourceData = MakeShareable(new FPinNameRowData());
 		NewSourceData->PinCategory = PinType;
 		NewSourceData->PinFriendlyName = Schema->GetPinCategoryName(PinType);
 
-		ListViewSource.Add(NewSourceData);
-	}
-	
-	TSharedPtr<SListView<TSharedPtr<FPinNameRowData>>> PinTypesListView = SNew(SListView<TSharedPtr<FPinNameRowData>>)
-	.ListItemsSource(&ListViewSource)
-	.SelectionMode(ESelectionMode::Single)
-	.OnMouseButtonClick_Lambda([this](TSharedPtr<FPinNameRowData> Item)
-	{
-		if (Item)
+		ComboBoxSource.Add(NewSourceData);
+
+		if (Variable->PinCategoryType == PinType)
 		{
-			Variable->PinCategoryType = Item->PinCategory;
-			IONode->ReconstructNode();
+			InitialSelection = ComboBoxSource.Last();
 		}
-	})
-	.ListViewStyle(&FAppStyle::Get().GetWidgetStyle<FTableViewStyle>("SimpleListView"))
-	.OnGenerateRow(this, &SCOMacroPinTypeSelector::OnGenerateRow);
+	}
 
 	ChildSlot
 	[
 		SNew(SHorizontalBox)
 		+ SHorizontalBox::Slot()
-		.AutoWidth()
+		.FillWidth(0.7f)
 		[
 			// Combo button that summons the dropdown menu
-			SNew(SComboButton)
-			.ButtonContent()
-			[
-				SNew(STextBlock)
-				.Font(FAppStyle::GetFontStyle(TEXT("PropertyWindow.NormalFont")))
-				.Text(this, &SCOMacroPinTypeSelector::GetVariablePinTypeText)
-			]
-			.MenuContent()
+			SNew(SComboBox<TSharedPtr<FPinNameRowData>>)
+			.OptionsSource(&ComboBoxSource)
+			.OnGenerateWidget(this, &SCOMacroPinTypeSelector::OnGenerateRow)
+			.OnSelectionChanged_Lambda([this](TSharedPtr<FPinNameRowData> Item, ESelectInfo::Type)
+			{
+				if (Item.IsValid())
+				{
+					Variable->PinCategoryType = Item->PinCategory;
+					IONode->ReconstructNode();
+				}
+			})
+			.Content()
 			[
-				SNew(SListViewSelectorDropdownMenu<TSharedPtr<FPinNameRowData>>, nullptr, PinTypesListView)
-				[
-					SNew(SBorder)
-					.BorderImage(FAppStyle::GetBrush("Menu.Background"))
-					.Padding(2)
-					[
-						SNew(SBox)
-						[
-							PinTypesListView.ToSharedRef()
-						]
-					]
-				]
+				GenerateSelectedContent()
 			]
-			.IsFocusable(true)
-			.ContentPadding(2.0f)
 		]
 
 		+ SHorizontalBox::Slot()
@@ -296,32 +280,61 @@ void SCOMacroPinTypeSelector::Construct(const FArguments& InArgs)
 }
 
 
-TSharedRef<ITableRow> SCOMacroPinTypeSelector::OnGenerateRow(TSharedPtr<FPinNameRowData> RowData, const TSharedRef<STableViewBase>& OwnerTable)
+TSharedRef<SWidget> SCOMacroPinTypeSelector::OnGenerateRow(TSharedPtr<FPinNameRowData> Option)
 {
-	TSharedPtr<STableRow<TSharedPtr<FString>>> ListTableRow;
+	//return SNew(STextBlock).Text(FText::FromName(Item->PinFriendlyName));
+	const FSlateBrush* IconBrush = FAppStyle::GetBrush(TEXT("Kismet.VariableList.TypeIcon"));
 
-	SAssignNew(ListTableRow, STableRow<TSharedPtr<FString>>, OwnerTable)
+	TSharedPtr<SHorizontalBox> RowWidget = SNew(SHorizontalBox)
+	+ SHorizontalBox::Slot()
+	.AutoWidth()
 	[
-		SNew(SVerticalBox)
-			+ SVerticalBox::Slot()
-			.AutoHeight()
-			.Padding(20.0f, 1.0f, 20.0f, 1.0f)
-			[
-				SNew(STextBlock)
-					.Text(FText::FromName(RowData.IsValid() ? RowData->PinFriendlyName : FName("Invalid")))
-			]
+		SNew(SImage)
+		.Image(IconBrush)
+		.ColorAndOpacity(UEdGraphSchema_CustomizableObject::GetPinTypeColor(Option->PinCategory))
+	]
+	+ SHorizontalBox::Slot()
+	.AutoWidth()
+	.Padding(7.5f, 0.0f, 0.0f, 0.0f)
+	[
+		SNew(STextBlock)
+		.Text(FText::FromName(Option.IsValid() ? Option->PinFriendlyName : FName("Invalid")))
+		.Font(IDetailLayoutBuilder::GetDetailFont())
 	];
 
-	return ListTableRow.ToSharedRef();
+	return RowWidget.ToSharedRef();
 }
 
 
-FText SCOMacroPinTypeSelector::GetVariablePinTypeText() const
+TSharedRef<SWidget> SCOMacroPinTypeSelector::GenerateSelectedContent() const
 {
-	const UEdGraphSchema_CustomizableObject* Schema = Cast<UEdGraphSchema_CustomizableObject>(Macro->Graph->GetSchema());
-	check(Schema);
+	//return SNew(STextBlock).Text(FText::FromName(Item->PinFriendlyName));
+	const FSlateBrush* IconBrush = FAppStyle::GetBrush(TEXT("Kismet.VariableList.TypeIcon"));
+
+	TSharedPtr<SHorizontalBox> RowWidget = SNew(SHorizontalBox)
+	+ SHorizontalBox::Slot()
+	.AutoWidth()
+	[
+		SNew(SImage)
+		.Image(IconBrush)
+		.ColorAndOpacity_Lambda([this]() 
+		{
+			return UEdGraphSchema_CustomizableObject::GetPinTypeColor(Variable->PinCategoryType); 
+		})
+	]
+	+ SHorizontalBox::Slot()
+	.AutoWidth()
+	.Padding(7.5f, 0.0f, 0.0f, 0.0f)
+	[
+		SNew(STextBlock)
+		.Font(IDetailLayoutBuilder::GetDetailFont())
+		.Text_Lambda([this]() 
+		{
+			return FText::FromName(Variable ? UEdGraphSchema_CustomizableObject::GetPinCategoryName(Variable->PinCategoryType) : FName("Invalid"));
+		})
+	];
 
-	return Variable ? FText::FromName(Schema->GetPinCategoryName(Variable->PinCategoryType)) : FText();
+	return RowWidget.ToSharedRef();
 }
 
 
diff --git a/Engine/Plugins/Mutable/Source/CustomizableObjectEditor/Private/MuCOE/CustomizableObjectMacroLibrary/CustomizableObjectMacroDetails.h b/Engine/Plugins/Mutable/Source/CustomizableObjectEditor/Private/MuCOE/CustomizableObjectMacroLibrary/CustomizableObjectMacroDetails.h
--- a/Engine/Plugins/Mutable/Source/CustomizableObjectEditor/Private/MuCOE/CustomizableObjectMacroLibrary/CustomizableObjectMacroDetails.h
+++ b/Engine/Plugins/Mutable/Source/CustomizableObjectEditor/Private/MuCOE/CustomizableObjectMacroLibrary/CustomizableObjectMacroDetails.h
@@ -37,10 +37,10 @@ class SCOMacroPinTypeSelector : public SCompoundWidget
 	void Construct(const FArguments& InArgs);
 
 	/** Generates the widget of a List view row. */
-	TSharedRef<ITableRow> OnGenerateRow(TSharedPtr<FPinNameRowData> RowData, const TSharedRef<STableViewBase>& OwnerTable);
+	TSharedRef<SWidget> OnGenerateRow(TSharedPtr<FPinNameRowData> Option);
 
 	/** Returns the name of the pin type of the variable.*/
-	FText GetVariablePinTypeText() const;
+	TSharedRef<SWidget> GenerateSelectedContent() const;
 
 	/** Callback of the remove variable button. */
 	void OnRemoveClicked();
@@ -57,7 +57,7 @@ private:
 	TObjectPtr<UCustomizableObjectNodeTunnel> IONode;
 
 	/** Array that contains all the pin types showed in the list view widget.*/
-	TArray<TSharedPtr<FPinNameRowData>> ListViewSource;
+	TArray<TSharedPtr<FPinNameRowData>> ComboBoxSource;
 
 	/** Callback to indicate that a variable has been removed */
 	FOnVariableRemoved OnVariableRemoved;
diff --git a/Engine/Plugins/Mutable/Source/CustomizableObjectEditor/Private/MuCOE/CustomizableObjectMacroLibrary/CustomizableObjectMacroLibrary.cpp b/Engine/Plugins/Mutable/Source/CustomizableObjectEditor/Private/MuCOE/CustomizableObjectMacroLibrary/CustomizableObjectMacroLibrary.cpp
--- a/Engine/Plugins/Mutable/Source/CustomizableObjectEditor/Private/MuCOE/CustomizableObjectMacroLibrary/CustomizableObjectMacroLibrary.cpp
+++ b/Engine/Plugins/Mutable/Source/CustomizableObjectEditor/Private/MuCOE/CustomizableObjectMacroLibrary/CustomizableObjectMacroLibrary.cpp
@@ -66,17 +66,17 @@ void UCustomizableObjectMacro::RemoveVariable(UCustomizableObjectMacroInputOutpu
 	}
 }
 
-const UCustomizableObjectNodeTunnel* UCustomizableObjectMacro::GetIONode(ECOMacroIOType Type) const
+UCustomizableObjectNodeTunnel* UCustomizableObjectMacro::GetIONode(ECOMacroIOType Type) const
 {
-	const UCustomizableObjectNodeTunnel* OutNode = nullptr;
+	UCustomizableObjectNodeTunnel* OutNode = nullptr;
 
 	if (Graph)
 	{
-		TArray<const UCustomizableObjectNodeTunnel*> IONodes;
-		Graph->GetNodesOfClass<const UCustomizableObjectNodeTunnel>(IONodes);
+		TArray<UCustomizableObjectNodeTunnel*> IONodes;
+		Graph->GetNodesOfClass<UCustomizableObjectNodeTunnel>(IONodes);
 		check(IONodes.Num() == 2);
 
-		for (const UCustomizableObjectNodeTunnel* IONode : IONodes)
+		for (UCustomizableObjectNodeTunnel* IONode : IONodes)
 		{
 			if (IONode->bIsInputNode && Type == ECOMacroIOType::COMVT_Input)
 			{
diff --git a/Engine/Plugins/Mutable/Source/CustomizableObjectEditor/Private/MuCOE/CustomizableObjectMacroLibrary/CustomizableObjectMacroLibrary.h b/Engine/Plugins/Mutable/Source/CustomizableObjectEditor/Private/MuCOE/CustomizableObjectMacroLibrary/CustomizableObjectMacroLibrary.h
--- a/Engine/Plugins/Mutable/Source/CustomizableObjectEditor/Private/MuCOE/CustomizableObjectMacroLibrary/CustomizableObjectMacroLibrary.h
+++ b/Engine/Plugins/Mutable/Source/CustomizableObjectEditor/Private/MuCOE/CustomizableObjectMacroLibrary/CustomizableObjectMacroLibrary.h
@@ -56,7 +56,7 @@ public:
 	void RemoveVariable(UCustomizableObjectMacroInputOutput* Variable);
 
 	/** Returns the input or output node */
-	const UCustomizableObjectNodeTunnel* GetIONode(ECOMacroIOType Type) const;
+	UCustomizableObjectNodeTunnel* GetIONode(ECOMacroIOType Type) const;
 
 public:
 
diff --git a/Engine/Plugins/Mutable/Source/CustomizableObjectEditor/Private/MuCOE/CustomizableObjectMacroLibrary/CustomizableObjectMacroLibraryEditor.cpp b/Engine/Plugins/Mutable/Source/CustomizableObjectEditor/Private/MuCOE/CustomizableObjectMacroLibrary/CustomizableObjectMacroLibraryEditor.cpp
--- a/Engine/Plugins/Mutable/Source/CustomizableObjectEditor/Private/MuCOE/CustomizableObjectMacroLibrary/CustomizableObjectMacroLibraryEditor.cpp
+++ b/Engine/Plugins/Mutable/Source/CustomizableObjectEditor/Private/MuCOE/CustomizableObjectMacroLibrary/CustomizableObjectMacroLibraryEditor.cpp
@@ -26,10 +26,8 @@ void FCustomizableObjectMacroLibraryEditor::InitEditor(const EToolkitMode::Type
 {
 	MacroLibrary = Cast<UCustomizableObjectMacroLibrary>(ObjectToEdit);
 
-	if (MacroLibrary->Macros.Num())
-	{
-		SelectedMacro = MacroLibrary->Macros[0];
-	}
+	// Bind Commands
+	BindGraphCommands();
 
 	// Tab Generation
 	const TSharedRef<FTabManager::FLayout> StandaloneDefaultLayout = FTabManager::NewLayout("Standalone_CustomizableObjectMacroLibraryEditor_Layout_v0.1")
@@ -138,7 +136,7 @@ TSharedRef<SDockTab> FCustomizableObjectMacroLibraryEditor::SpawnTab_MacroSelect
 		.MacroLibrary(MacroLibrary)
 		.OnAddMacroButtonClicked(FExecuteAction::CreateSP(this, &FCustomizableObjectMacroLibraryEditor::OnAddMacroButtonClicked))
 		.OnSelectMacro(SCustomizableObjectMacroLibraryList::FOnSelectMacroDelegate::CreateSP(this, &FCustomizableObjectMacroLibraryEditor::OnMacroSelectionChanged))
-		.OnRemoveMacro(SCustomizableObjectMacroLibraryList::FOnSelectMacroDelegate::CreateSP(this, &FCustomizableObjectMacroLibraryEditor::OnRemoveMacroButtonClicked));
+		.OnRemoveMacro(SCustomizableObjectMacroLibraryList::FOnRemoveMacroDelegate::CreateSP(this, &FCustomizableObjectMacroLibraryEditor::OnRemoveMacroButtonClicked));
 	}
 
 	return SNew(SDockTab)
@@ -238,20 +236,10 @@ FLinearColor FCustomizableObjectMacroLibraryEditor::GetWorldCentricTabColorScale
 }
 
 
-void FCustomizableObjectMacroLibraryEditor::NotifyPostChange(const FPropertyChangedEvent& PropertyChangedEvent, FProperty* PropertyThatChanged)
-{
-	check(PropertyThatChanged);
-
-	// Refresh macros listview if a macro has been renamed
-	if (PropertyThatChanged->NamePrivate == "Name")
-	{
-		OnMacroNameChanged();
-	}
-}
-
-
 void FCustomizableObjectMacroLibraryEditor::PostUndo(bool bSuccess) 
 {
+	FCustomizableObjectGraphEditorToolkit::PostUndo(bSuccess);
+
 	if (bSuccess)
 	{
 		//TODO(Max): Select the correct macro after an undo.
@@ -265,6 +253,8 @@ void FCustomizableObjectMacroLibraryEditor::PostUndo(bool bSuccess)
 
 void FCustomizableObjectMacroLibraryEditor::PostRedo(bool bSuccess)
 {
+	FCustomizableObjectGraphEditorToolkit::PostRedo(bSuccess);
+
 	if (bSuccess)
 	{
 		//TODO(Max): Select the correct macro after a redo.
@@ -284,8 +274,6 @@ void FCustomizableObjectMacroLibraryEditor::OnAddMacroButtonClicked()
 		MacroLibrary->Modify();
 
 		UCustomizableObjectMacro* NewMacro = MacroLibrary->AddMacro();
-		
-		SetSelectedMacro(NewMacro);
 	}
 }
 
@@ -304,22 +292,13 @@ void FCustomizableObjectMacroLibraryEditor::OnRemoveMacroButtonClicked(UCustomiz
 }
 
 
-void FCustomizableObjectMacroLibraryEditor::OnMacroNameChanged()
-{
-	if (MacroSelector.IsValid())
-	{
-		MacroSelector->GenerateRowView();
-	}
-}
-
-
 void FCustomizableObjectMacroLibraryEditor::OnMacroSelectionChanged(UCustomizableObjectMacro* NewSelection)
 {
 	SetSelectedMacro(NewSelection);
 }
 
 
-void FCustomizableObjectMacroLibraryEditor::SetSelectedMacro(UCustomizableObjectMacro* NewSelection)
+void FCustomizableObjectMacroLibraryEditor::SetSelectedMacro(UCustomizableObjectMacro* NewSelection, bool bRefreshMacroSelection)
 {
 	if (!NewSelection)
 	{
@@ -339,15 +318,22 @@ void FCustomizableObjectMacroLibraryEditor::SetSelectedMacro(UCustomizableObject
 
 			SGraphEditor::FGraphEditorEvents InEvents;
 			InEvents.OnSelectionChanged = SGraphEditor::FOnSelectionChanged::CreateSP(this, &FCustomizableObjectMacroLibraryEditor::OnSelectedGraphNodesChanged);
-			InEvents.OnNodeDoubleClicked = FSingleNodeEvent::CreateSP(this, &FCustomizableObjectMacroLibraryEditor::OnNodeDoubleClicked);
 
 			CreateGraphEditorWidget(NewSelection->Graph, InEvents);
 
 			GraphTab->SetContent(GraphEditor.ToSharedRef());
+
+			if (bRefreshMacroSelection)
+			{
+				MacroSelector->SetSelectedMacro(NewSelection);
+			}
 		}
 	}
 
-	DetailsView->SetObject(SelectedMacro, true);
+	if (DetailsView.IsValid())
+	{
+		DetailsView->SetObject(SelectedMacro, true);
+	}
 }
 
 #undef LOCTEXT_NAMESPACE
diff --git a/Engine/Plugins/Mutable/Source/CustomizableObjectEditor/Private/MuCOE/CustomizableObjectMacroLibrary/CustomizableObjectMacroLibraryEditor.h b/Engine/Plugins/Mutable/Source/CustomizableObjectEditor/Private/MuCOE/CustomizableObjectMacroLibrary/CustomizableObjectMacroLibraryEditor.h
--- a/Engine/Plugins/Mutable/Source/CustomizableObjectEditor/Private/MuCOE/CustomizableObjectMacroLibrary/CustomizableObjectMacroLibraryEditor.h
+++ b/Engine/Plugins/Mutable/Source/CustomizableObjectEditor/Private/MuCOE/CustomizableObjectMacroLibrary/CustomizableObjectMacroLibraryEditor.h
@@ -3,14 +3,14 @@
 #pragma once
 
 #include "EditorUndoClient.h"
-#include "MuCOE/CustomizableObjectGraphEditorToolkit.h"
+#include "MuCOE/CustomizableObjectMacroLibrary/CustomizableObjectGraphEditorToolkit.h"
 #include "MuCOE/CustomizableObjectMacroLibrary/CustomizableObjectMacroLibrary.h"
 
 class SCustomizableObjectMacroLibraryList;
 class SDockTab;
 
 /** Macro Library Editor. */
-class FCustomizableObjectMacroLibraryEditor : public FCustomizableObjectGraphEditorToolkit, public FSelfRegisteringEditorUndoClient
+class FCustomizableObjectMacroLibraryEditor : public FCustomizableObjectGraphEditorToolkit
 {
 public:
 
@@ -29,14 +29,14 @@ public:
 	virtual FLinearColor GetWorldCentricTabColorScale() const override;
 
 	/** NotifyHook Interface */
-	virtual void NotifyPostChange(const FPropertyChangedEvent& PropertyChangedEvent, FProperty* PropertyThatChanged) override;
+	//virtual void NotifyPostChange(const FPropertyChangedEvent& PropertyChangedEvent, FProperty* PropertyThatChanged) override;
 
 	/** FEditorUndoClient Interface */
 	virtual void PostUndo(bool bSuccess) override;
 	virtual void PostRedo(bool bSuccess) override;
 
 	/** Sets which macro is edited in the editor. */
-	void SetSelectedMacro(UCustomizableObjectMacro* NewSelection);
+	void SetSelectedMacro(UCustomizableObjectMacro* NewSelection, bool bRefreshMacroSelection = false);
 
 private:
 
@@ -45,16 +45,14 @@ private:
 	TSharedRef<SDockTab> SpawnTab_MacroSelector(const FSpawnTabArgs& Args);
 	TSharedRef<SDockTab> SpawnTab_Graph(const FSpawnTabArgs& Args);
 
+	// CO Graph Editor Toolkit Interface
 	/** Graph Editor callbacks */
-	void OnSelectedGraphNodesChanged(const FGraphPanelSelectionSet& NewSelection);
+	virtual void OnSelectedGraphNodesChanged(const FGraphPanelSelectionSet& NewSelection) override;
 
 	/** Macro List View callbacks */
 	void OnAddMacroButtonClicked();
 	void OnRemoveMacroButtonClicked(UCustomizableObjectMacro* NewSelection);
 	void OnMacroSelectionChanged(UCustomizableObjectMacro* NewSelection);
-	
-	/** Editor Details View callbacks */
-	void OnMacroNameChanged();
 
 public:
 
diff --git a/Engine/Plugins/Mutable/Source/CustomizableObjectEditor/Private/MuCOE/EdGraphSchema_CustomizableObject.cpp b/Engine/Plugins/Mutable/Source/CustomizableObjectEditor/Private/MuCOE/EdGraphSchema_CustomizableObject.cpp
--- a/Engine/Plugins/Mutable/Source/CustomizableObjectEditor/Private/MuCOE/EdGraphSchema_CustomizableObject.cpp
+++ b/Engine/Plugins/Mutable/Source/CustomizableObjectEditor/Private/MuCOE/EdGraphSchema_CustomizableObject.cpp
@@ -14,7 +14,8 @@
 #include "MuCO/CustomizableObjectExtension.h"
 #include "MuCO/ICustomizableObjectModule.h"
 #include "MuCOE/CustomizableObjectGraph.h"
-#include "MuCOE/ICustomizableObjectEditor.h"
+#include "MuCOE/CustomizableObjectMacroLibrary/CustomizableObjectGraphEditorToolkit.h"
+#include "MuCOE/CustomizableObjectEditor.h"
 #include "MuCOE/Nodes/CustomizableObjectNodeAnimationPose.h"
 #include "MuCOE/Nodes/CustomizableObjectNodeColorArithmeticOp.h"
 #include "MuCOE/Nodes/CustomizableObjectNodeColorConstant.h"
@@ -110,18 +111,17 @@ namespace MutablePrivate
 }
 
 
-TSharedPtr<ICustomizableObjectEditor> GetCustomizableObjectEditor(const UEdGraph* ParentGraph)
+TSharedPtr<FCustomizableObjectGraphEditorToolkit> GetCustomizableObjectEditor(const UEdGraph* ParentGraph)
 {
 	// Find the associated Editor
-	UCustomizableObject* CustomizableObject = Cast<UCustomizableObject>(ParentGraph->GetOuter());
+	const UCustomizableObjectGraph* CustomizableObjectGraph = Cast<UCustomizableObjectGraph>(ParentGraph);
 
-	TSharedPtr<ICustomizableObjectEditor> CustomizableObjectEditor;
-	if (CustomizableObject)
+	if (CustomizableObjectGraph)
 	{
-		TSharedPtr< IToolkit > FoundAssetEditor = FToolkitManager::Get().FindEditorForAsset(CustomizableObject);
+		TSharedPtr< IToolkit > FoundAssetEditor = FToolkitManager::Get().FindEditorForAsset(CustomizableObjectGraph);
 		if (FoundAssetEditor.IsValid())
 		{
-			return StaticCastSharedPtr<ICustomizableObjectEditor>(FoundAssetEditor);
+			return StaticCastSharedPtr<FCustomizableObjectGraphEditorToolkit>(FoundAssetEditor);
 		}
 	}
 
@@ -234,7 +234,7 @@ void FCustomizableObjectSchemaAction_NewNode::AddReferencedObjects( FReferenceCo
 
 UEdGraphNode* FCustomizableObjectSchemaAction_Paste::PerformAction(class UEdGraph* ParentGraph, UEdGraphPin* FromPin, const FVector2D Location, bool bSelectNewNode/* = true*/)
 {
-	TSharedPtr<ICustomizableObjectEditor> CustomizableObjectEditor = GetCustomizableObjectEditor(ParentGraph);
+	TSharedPtr<FCustomizableObjectGraphEditorToolkit> CustomizableObjectEditor = GetCustomizableObjectEditor(ParentGraph);
 
 	if (CustomizableObjectEditor.IsValid() && CustomizableObjectEditor->CanPasteNodes())
 	{
@@ -270,7 +270,7 @@ const FName UEdGraphSchema_CustomizableObject::PC_Transform("transform");
 const FName UEdGraphSchema_CustomizableObject::PC_String("string");
 
 // Add more pin types to this array if needed
-const TArray<FName> UEdGraphSchema_CustomizableObject::SupportedMacroPinTypes = { PC_Component, PC_Material, PC_Modifier, PC_Mesh, PC_Image,
+const TArray<FName> UEdGraphSchema_CustomizableObject::SupportedMacroPinTypes = { PC_Object, PC_Component, PC_Material, PC_Modifier, PC_Mesh, PC_Image,
 PC_PassThroughImage, PC_Projector, PC_GroupProjector, PC_Color, PC_Float, PC_Enum, PC_Stack, PC_MaterialAsset, PC_PoseAsset, PC_Transform, PC_String };
 
 // Node categories
@@ -913,8 +913,10 @@ void GetContextMenuActionsReconstructAllChildNodes(UToolMenu* Menu, TWeakObjectP
 				return;
 			}
 
-			const TSharedRef<ICustomizableObjectEditor> Editor = GetCustomizableObjectEditor(Context->Graph).ToSharedRef();
-			Editor->ReconstructAllChildNodes(*Node, *NodeType);			
+			if (TSharedPtr<FCustomizableObjectEditor> Editor = StaticCastSharedPtr<FCustomizableObjectEditor>(GetCustomizableObjectEditor(Context->Graph)))
+			{
+				Editor->ReconstructAllChildNodes(*Node, *NodeType);
+			};
 		};
 
 		SubSection.AddMenuEntry(
@@ -980,7 +982,7 @@ void UEdGraphSchema_CustomizableObject::GetContextMenuActions(UToolMenu* Menu, U
 			{
 				if (Schema && Graph)
 				{
-					if (TSharedPtr<ICustomizableObjectEditor> CustomizableObjectEditor = GetCustomizableObjectEditor(Graph))
+					if (TSharedPtr<FCustomizableObjectGraphEditorToolkit> CustomizableObjectEditor = GetCustomizableObjectEditor(Graph))
 					{
 						CustomizableObjectEditor->CreateCommentBox(FVector2D::ZeroVector);
 					}
diff --git a/Engine/Plugins/Mutable/Source/CustomizableObjectEditor/Private/MuCOE/GenerateMutableSource/GenerateMutableSource.cpp b/Engine/Plugins/Mutable/Source/CustomizableObjectEditor/Private/MuCOE/GenerateMutableSource/GenerateMutableSource.cpp
--- a/Engine/Plugins/Mutable/Source/CustomizableObjectEditor/Private/MuCOE/GenerateMutableSource/GenerateMutableSource.cpp
+++ b/Engine/Plugins/Mutable/Source/CustomizableObjectEditor/Private/MuCOE/GenerateMutableSource/GenerateMutableSource.cpp
@@ -17,10 +17,12 @@
 #include "MuCOE/GenerateMutableSource/GenerateMutableSourceModifier.h"
 #include "MuCOE/GraphTraversal.h"
 #include "MuCOE/ICustomizableObjectExtensionNode.h"
+#include "MuCOE/Nodes/CustomizableObjectNodeMacroInstance.h"
 #include "MuCOE/Nodes/CustomizableObjectNodeMaterial.h"
 #include "MuCOE/Nodes/CustomizableObjectNodeMesh.h"
 #include "MuCOE/Nodes/CustomizableObjectNodeModifierClipWithMesh.h"
 #include "MuCOE/Nodes/CustomizableObjectNodeObjectGroup.h"
+#include "MuCOE/Nodes/CustomizableObjectNodeStringConstant.h"
 #include "MuCOE/Nodes/CustomizableObjectNodeTable.h"
 #include "MuCOE/Nodes/CustomizableObjectNodeTexture.h"
 #include "MuCOE/Nodes/CustomizableObjectNodePassThroughTexture.h"
@@ -38,6 +40,8 @@
 #include "MuCOE/Nodes/CustomizableObjectNodeTextureTransform.h"
 #include "MuCOE/Nodes/CustomizableObjectNodeTextureSaturate.h"
 #include "MuCOE/Nodes/CustomizableObjectNodeTextureVariation.h"
+#include "MuCOE/Nodes/CustomizableObjectNodeTunnel.h"
+#include "MuCOE/Nodes/CustomizableObjectNodeVariation.h"
 #include "MuCOE/UnrealEditorPortabilityHelpers.h"
 #include "MuT/NodeComponentEdit.h"
 #include "MuT/NodeMeshConstant.h"
@@ -57,7 +61,7 @@
 
 uint32 GetTypeHash(const FGraphCycleKey& Key)
 {
-	return HashCombine(GetTypeHash(&Key.Pin), GetTypeHash(Key.Id));
+	return HashCombine(HashCombine(GetTypeHash(&Key.Pin), GetTypeHash(Key.Id)), GetTypeHash(Key.MacroContext));
 }
 
 
@@ -404,6 +408,7 @@ FGeneratedKey::FGeneratedKey(void* InFunctionAddress, const UEdGraphPin& InPin,
 	FunctionAddress = InFunctionAddress;
 	Pin = &InPin;
 	LOD = Node.IsAffectedByLOD() ? GenerationContext.CurrentLOD : 0;
+	MacroContext = GenerationContext.MacroNodesStack;
 
 	if (UseMesh)
 	{
@@ -425,21 +430,23 @@ uint32 GetTypeHash(const FGeneratedKey& Key)
 	//Hash = HashCombine(Hash, GetTypeHash(Key.MeshMorphStack)); // Does not support array
 	Hash = HashCombine(Hash, GetTypeHash(Key.bOnlyConnectedLOD));
 	Hash = HashCombine(Hash, GetTypeHash(Key.CurrentMeshComponent));
+	//Hash = HashCombine(Hash, GetTypeHash(Key.MacroContext)); // Does not support array
 	
 	return Hash;
 }
 
 
-FGraphCycleKey::FGraphCycleKey(const UEdGraphPin& Pin, const FString& Id) :
+FGraphCycleKey::FGraphCycleKey(const UEdGraphPin& Pin, const FString& Id, const UCustomizableObjectNodeMacroInstance* MacroContext) :
 	Pin(Pin),
-	Id(Id)
+	Id(Id),
+	MacroContext(MacroContext)
 {
 }
 
 
 bool FGraphCycleKey::operator==(const FGraphCycleKey& Other) const
 {
-	return &Pin == &Other.Pin && Id == Other.Id;
+	return &Pin == &Other.Pin && Id == Other.Id && MacroContext == Other.MacroContext;
 }
 
 
@@ -764,6 +771,10 @@ mu::Ptr<mu::NodeObject> GenerateMutableSource(const UEdGraphPin * Pin, FMutableG
 		return static_cast<mu::NodeObject*>(Generated->Node.get());
 	}
 
+	// Bool that determines if a node can be added to the cache of nodes.
+	// Most nodes need to be added to the cache but there are some that don't. For exampel, MacroInstanceNodes
+	bool bCacheNode = true;
+
 	mu::Ptr<mu::NodeObject> Result;
 	
 	if (const UCustomizableObjectNodeObject* TypedNodeObj = Cast<UCustomizableObjectNodeObject>(Node))
@@ -855,9 +866,26 @@ mu::Ptr<mu::NodeObject> GenerateMutableSource(const UEdGraphPin * Pin, FMutableG
 		if (const UEdGraphPin* ModifierPin = TypedNodeObj->ModifiersPin())
 		{
 			TArray<UEdGraphPin*> ConnectedModifierPins = FollowInputPinArray(*ModifierPin);
-			for (UEdGraphPin* const ChildNodePin : ConnectedModifierPins)
+			for (UEdGraphPin* const ConnectedModifier : ConnectedModifierPins)
 			{
-				if (!Cast<UCustomizableObjectNodeModifierBase>(ChildNodePin->GetOwningNode()))
+				bool bMacroModifier = false;
+				UEdGraphPin* ChildNodePin = ConnectedModifier;
+
+				if (const UCustomizableObjectNodeMacroInstance* TypedNodeMacro = Cast<UCustomizableObjectNodeMacroInstance>(ChildNodePin->GetOwningNode()))
+				{
+					if (const UEdGraphPin* OutputPin = TypedNodeMacro->GetMacroIOPin(ECOMacroIOType::COMVT_Output, Pin->PinName))
+					{
+						// Get the modifier pin inside the macro
+						if (UEdGraphPin* FollowPin = FollowInputPin(*OutputPin))
+						{
+							bMacroModifier = true;
+							GenerationContext.MacroNodesStack.Push(TypedNodeMacro);
+							ChildNodePin = FollowPin;
+						}
+					}
+				}
+				
+				else if (!Cast<UCustomizableObjectNodeModifierBase>(ChildNodePin->GetOwningNode()))
 				{
 					continue;
 				}
@@ -873,6 +901,11 @@ mu::Ptr<mu::NodeObject> GenerateMutableSource(const UEdGraphPin * Pin, FMutableG
 					GenerationContext.CurrentMeshComponent = {};
 					GenerationContext.CurrentAutoLODStrategy = {};
 				}
+
+				if (bMacroModifier)
+				{
+					GenerationContext.MacroNodesStack.Pop();
+				}
 			}
 		}
 		
@@ -1190,13 +1223,79 @@ mu::Ptr<mu::NodeObject> GenerateMutableSource(const UEdGraphPin * Pin, FMutableG
 		ensure(GenerationContext.SocketPriorityStack.Num() > 0);
 		GenerationContext.SocketPriorityStack.Pop();
 	}
+
+	else if (const UCustomizableObjectNodeMacroInstance* TypedNodeMacro = Cast<UCustomizableObjectNodeMacroInstance>(Node))
+	{
+		bCacheNode = false;
+
+		if (const UEdGraphPin* OutputPin = TypedNodeMacro->GetMacroIOPin(ECOMacroIOType::COMVT_Output, Pin->PinName))
+		{
+			if (const UEdGraphPin* FollowPin = FollowInputPin(*OutputPin))
+			{
+				GenerationContext.MacroNodesStack.Push(TypedNodeMacro);
+				Result = GenerateMutableSource(FollowPin, GenerationContext);
+				GenerationContext.MacroNodesStack.Pop();
+			}
+			else
+			{
+				FString Msg = FString::Printf(TEXT("Macro Pin [%s] not linked."), *Pin->PinName.ToString());
+				GenerationContext.Log(FText::FromString(Msg), Node);
+			}
+		}
+		else
+		{
+			FString Msg = FString::Printf(TEXT("Macro Output node does not contain a pin with name [%s]."), *Pin->PinName.ToString());
+			GenerationContext.Log(FText::FromString(Msg), Node);
+		}
+	}
+
+	else if (const UCustomizableObjectNodeTunnel* TypedNodeTunnel = Cast<UCustomizableObjectNodeTunnel>(Node))
+	{
+		check(TypedNodeTunnel->bIsInputNode);
+		check(GenerationContext.MacroNodesStack.Num());
+
+		bCacheNode = false;
+
+		const UCustomizableObjectNodeMacroInstance* MacroInstanceNode = GenerationContext.MacroNodesStack.Pop();
+		check(MacroInstanceNode);
+
+		if (const UEdGraphPin* InputPin = MacroInstanceNode->FindPin(Pin->PinName, EEdGraphPinDirection::EGPD_Input))
+		{
+			if (const UEdGraphPin* FollowPin = FollowInputPin(*InputPin))
+			{
+				Result = GenerateMutableSource(FollowPin, GenerationContext);
+			}
+			else
+			{
+				FString Msg = FString::Printf(TEXT("Macro Instance Pin [%s] is not linked."), *Pin->PinName.ToString());
+				GenerationContext.Log(FText::FromString(Msg), Node);
+			}
+		}
+		else
+		{
+			FString Msg = FString::Printf(TEXT("Macro Instance Node does not contain a pin with name [%s]."), *Pin->PinName.ToString());
+			GenerationContext.Log(FText::FromString(Msg), Node);
+		}
+
+		if (!Result)
+		{
+			GenerationContext.Log(LOCTEXT("ScalarMacroInstanceEmptyAssets", "Something went wrong"), Node);
+		}
+
+		// Push the Macro again even if the result is null
+		GenerationContext.MacroNodesStack.Push(MacroInstanceNode);
+	}
+
 	else
 	{
 		GenerationContext.Log(LOCTEXT("UnimplementedNode", "Node type not implemented yet."), Node);
 	}
 
-	GenerationContext.Generated.Add(Key, FGeneratedData(Node, Result));
-	GenerationContext.GeneratedNodes.Add(Node);
+	if (bCacheNode)
+	{
+		GenerationContext.Generated.Add(Key, FGeneratedData(Node, Result));
+		GenerationContext.GeneratedNodes.Add(Node);
+	}
 
 	return Result;
 }
@@ -1733,6 +1832,38 @@ void FMutableComponentInfo::AccumulateBonesToRemovePerLOD(const TArray<FLODReduc
 	}
 }
 
+FString GetVariationTag(const UCustomizableObjectNodeVariation& VariationNode, int32 VariationIndex, FMutableGraphGenerationContext& GenerationContext)
+{
+	if (VariationNode.IsInMacro())
+	{
+		if (const UEdGraphPin* VariationPin = VariationNode.VariationPin(VariationIndex))
+		{
+			const FString PinName = VariationPin->PinName.ToString() + UCustomizableObjectNodeVariation::VariationTagSufix;
+			const UEdGraphPin* TagPin = VariationNode.FindPin(PinName, EEdGraphPinDirection::EGPD_Input);
+
+			if (TagPin && TagPin->LinkedTo.Num())
+			{
+				const UEdGraphPin* ConnectedStringPin = FollowInputPin(*TagPin);
+				const UEdGraphPin* SourceStringPin = ConnectedStringPin ? GraphTraversal::FindStringSource(*ConnectedStringPin, &GenerationContext.MacroNodesStack) : nullptr;
+
+				if (SourceStringPin)
+				{
+					if (const UCustomizableObjectNodeStringConstant* StringNode = Cast<UCustomizableObjectNodeStringConstant>(SourceStringPin->GetOwningNode()))
+					{
+						return StringNode->Value;
+					}
+				}
+				else
+				{
+					GenerationContext.Log(LOCTEXT("MorphStringNodeFailed", "Could not found a linked String node."), &VariationNode);
+				}
+			}
+		}
+	}
+
+	return VariationNode.GetVariation(VariationIndex).Tag;
+}
+
 
 #undef LOCTEXT_NAMESPACE
 
diff --git a/Engine/Plugins/Mutable/Source/CustomizableObjectEditor/Private/MuCOE/GenerateMutableSource/GenerateMutableSource.h b/Engine/Plugins/Mutable/Source/CustomizableObjectEditor/Private/MuCOE/GenerateMutableSource/GenerateMutableSource.h
--- a/Engine/Plugins/Mutable/Source/CustomizableObjectEditor/Private/MuCOE/GenerateMutableSource/GenerateMutableSource.h
+++ b/Engine/Plugins/Mutable/Source/CustomizableObjectEditor/Private/MuCOE/GenerateMutableSource/GenerateMutableSource.h
@@ -41,9 +41,11 @@ class FCustomizableObjectCompiler;
 class UTextureLODSettings;
 class UAnimInstance;
 class UCompositeDataTable;
+class UCustomizableObjectNodeMacroInstance;
 class UCustomizableObjectNodeMaterial;
 class UCustomizableObjectNodeMeshMorph;
 class UCustomizableObjectNodeObjectGroup;
+class UCustomizableObjectNodeVariation;
 class UEdGraphNode;
 class UMaterialInterface;
 class UObject;
@@ -113,7 +115,7 @@ struct FMorphNodeData
 	FString MorphTargetName;
 
 	// Pin to the node that generates the factor of the morph
-	const UEdGraphPin* FactorPin;
+	mu::Ptr<mu::NodeScalar> FactorNode;
 
 	// Pin of the mesh where the morphs will ble apllied
 	const UEdGraphPin* MeshPin;
@@ -121,7 +123,7 @@ struct FMorphNodeData
 	bool operator==(const FMorphNodeData& Other) const
 	{
 		return OwningNode == Other.OwningNode && MorphTargetName == Other.MorphTargetName
-			&& FactorPin == Other.FactorPin && MeshPin == Other.MeshPin;
+			&& FactorNode == Other.FactorNode && MeshPin == Other.MeshPin;
 	}
 };
 
@@ -156,6 +158,9 @@ public:
 	
 	/** When caching a generated mesh, true if we force to generate the connected LOD when using Automatic LODs From Mesh. */
 	bool bOnlyConnectedLOD = false;
+
+	/** Pointer to control if this is a node inside a Mutable Macro. */
+	TArray<const UCustomizableObjectNodeMacroInstance*> MacroContext;
 };
 
 
@@ -396,6 +401,9 @@ struct FMutableComponentInfo
 	TArray<TMap<FName, bool>> BonesToRemovePerLOD;
 	
 	UCustomizableObjectNodeComponentMesh* Node = nullptr;
+
+	// Keeps track of the macro context where this component node is instantiated
+	TArray<const UCustomizableObjectNodeMacroInstance*> MacroContext;
 };
 
 
@@ -408,7 +416,7 @@ struct FGraphCycleKey
 {
 	friend uint32 GetTypeHash(const FGraphCycleKey& Key);
 
-	FGraphCycleKey(const UEdGraphPin& Pin, const FString& Id);
+	FGraphCycleKey(const UEdGraphPin& Pin, const FString& Id, const UCustomizableObjectNodeMacroInstance* MacroContext);
 
 	bool operator==(const FGraphCycleKey& Other) const;
 	
@@ -417,6 +425,8 @@ struct FGraphCycleKey
 
 	/** Unique id. */
 	FString Id;
+
+	const UCustomizableObjectNodeMacroInstance* MacroContext;
 };
 
 /** Graph Cycle scope.
@@ -442,7 +452,7 @@ private:
 
 /** Return the default value if there is a cycle. */
 #define RETURN_ON_CYCLE(Pin, GenerationContext) \
-	FGraphCycle GraphCycle(FGraphCycleKey(Pin, TEXT(__FILE__ PREPROCESSOR_TO_STRING(__LINE__))), GenerationContext); \
+	FGraphCycle GraphCycle(FGraphCycleKey(Pin, TEXT(__FILE__ PREPROCESSOR_TO_STRING(__LINE__)),GenerationContext.MacroNodesStack.Num() ? GenerationContext.MacroNodesStack.Top() : nullptr), GenerationContext); \
 	if (GraphCycle.FoundCycle()) \
 	{ \
 		return {}; \
@@ -871,6 +881,10 @@ struct FMutableGraphGenerationContext
 
 	/** Index of the Referenced Material being generated */
 	int32 CurrentReferencedMaterialIndex = -1;
+	
+	// This should be automatized (create a struct that each time that a Macro is stacked, it adds 1 to the count.
+	// it also substracts 1 each time we remove a Macro from the stack)
+	TArray<const UCustomizableObjectNodeMacroInstance*> MacroNodesStack;
 };
 
 
@@ -985,3 +999,5 @@ uint32 AddUniqueSurfaceMetadata(const FSkeletalMaterial* Material, const FSkelMe
 /** Convert a mesh to Mutable format.
 */
 TSharedPtr<mu::FMesh> ConvertSkeletalMeshToMutable(const FMutableSourceMeshData& Source, FMutableCompilationContext& Context, const FString& MorphName);
+
+FString GetVariationTag(const UCustomizableObjectNodeVariation& VariationNode, int32 VariationIndex, FMutableGraphGenerationContext& GenerationContext);
diff --git a/Engine/Plugins/Mutable/Source/CustomizableObjectEditor/Private/MuCOE/GenerateMutableSource/GenerateMutableSourceColor.cpp b/Engine/Plugins/Mutable/Source/CustomizableObjectEditor/Private/MuCOE/GenerateMutableSource/GenerateMutableSourceColor.cpp
--- a/Engine/Plugins/Mutable/Source/CustomizableObjectEditor/Private/MuCOE/GenerateMutableSource/GenerateMutableSourceColor.cpp
+++ b/Engine/Plugins/Mutable/Source/CustomizableObjectEditor/Private/MuCOE/GenerateMutableSource/GenerateMutableSourceColor.cpp
@@ -3,6 +3,7 @@
 #include "MuCOE/GenerateMutableSource/GenerateMutableSourceColor.h"
 
 #include "MuCOE/CustomizableObjectCompiler.h"
+#include "MuCOE/CustomizableObjectGraph.h"
 #include "MuCOE/EdGraphSchema_CustomizableObject.h"
 #include "MuCOE/GenerateMutableSource/GenerateMutableSourceFloat.h"
 #include "MuCOE/GenerateMutableSource/GenerateMutableSourceImage.h"
@@ -14,8 +15,10 @@
 #include "MuCOE/Nodes/CustomizableObjectNodeColorParameter.h"
 #include "MuCOE/Nodes/CustomizableObjectNodeColorSwitch.h"
 #include "MuCOE/Nodes/CustomizableObjectNodeColorVariation.h"
+#include "MuCOE/Nodes/CustomizableObjectNodeMacroInstance.h"
 #include "MuCOE/Nodes/CustomizableObjectNodeTable.h"
 #include "MuCOE/Nodes/CustomizableObjectNodeTextureSample.h"
+#include "MuCOE/Nodes/CustomizableObjectNodeTunnel.h"
 #include "MuT/NodeColourArithmeticOperation.h"
 #include "MuT/NodeColourConstant.h"
 #include "MuT/NodeColourFromScalars.h"
@@ -52,6 +55,10 @@ mu::Ptr<mu::NodeColour> GenerateMutableSourceColor(const UEdGraphPin* Pin, FMuta
 		Node->SetRefreshNodeWarning();
 	}
 
+	// Bool that determines if a node can be added to the cache of nodes.
+	// Most nodes need to be added to the cache but there are some that don't. For exampel, MacroInstanceNodes
+	bool bCacheNode = true;
+
 	mu::Ptr<mu::NodeColour> Result;
 
 	if (const UCustomizableObjectNodeColorConstant* TypedNodeColorConst = Cast<UCustomizableObjectNodeColorConstant>(Node))
@@ -255,7 +262,8 @@ mu::Ptr<mu::NodeColour> GenerateMutableSourceColor(const UEdGraphPin* Pin, FMuta
 			const UEdGraphPin* VariationPin = TypedNodeColorVar->VariationPin(VariationIndex);
 			if (!VariationPin) continue;
 
-			ColorNode->Variations[VariationIndex].Tag = TypedNodeColorVar->GetVariation(VariationIndex).Tag;
+			ColorNode->Variations[VariationIndex].Tag = GetVariationTag(*TypedNodeColorVar,VariationIndex, GenerationContext);
+
 			if (const UEdGraphPin* ConnectedPin = FollowInputPin(*VariationPin))
 			{
 				mu::Ptr<mu::NodeColour> ChildNode = GenerateMutableSourceColor(ConnectedPin, GenerationContext);
@@ -336,14 +344,79 @@ mu::Ptr<mu::NodeColour> GenerateMutableSourceColor(const UEdGraphPin* Pin, FMuta
 		}
 	}
 
+	else if (const UCustomizableObjectNodeMacroInstance* TypedNodeMacro = Cast<UCustomizableObjectNodeMacroInstance>(Node))
+	{
+		bCacheNode = false;
+
+		if (const UEdGraphPin* OutputPin = TypedNodeMacro->GetMacroIOPin(ECOMacroIOType::COMVT_Output, Pin->PinName))
+		{
+			if (const UEdGraphPin* FollowPin = FollowInputPin(*OutputPin))
+			{
+				GenerationContext.MacroNodesStack.Push(TypedNodeMacro);
+				Result = GenerateMutableSourceColor(FollowPin, GenerationContext);
+				GenerationContext.MacroNodesStack.Pop();
+			}
+			else
+			{
+				FString Msg = FString::Printf(TEXT("Macro Pin %s not linked."), *Pin->PinName.ToString());
+				GenerationContext.Log(FText::FromString(Msg), Node);
+			}
+		}
+		else
+		{
+			FString Msg = FString::Printf(TEXT("Macro Output node does not contain a pin with name %s."), *Pin->PinName.ToString());
+			GenerationContext.Log(FText::FromString(Msg), Node);
+		}
+	}
+
+	else if (const UCustomizableObjectNodeTunnel* TypedNodeTunnel = Cast<UCustomizableObjectNodeTunnel>(Node))
+	{
+		check(TypedNodeTunnel->bIsInputNode);
+		check(GenerationContext.MacroNodesStack.Num());
+
+		bCacheNode = false;
+
+		const UCustomizableObjectNodeMacroInstance* MacroInstanceNode = GenerationContext.MacroNodesStack.Pop();
+		check(MacroInstanceNode);
+
+		if (const UEdGraphPin* InputPin = MacroInstanceNode->FindPin(Pin->PinName, EEdGraphPinDirection::EGPD_Input))
+		{
+			if (const UEdGraphPin* FollowPin = FollowInputPin(*InputPin))
+			{
+				Result = GenerateMutableSourceColor(FollowPin, GenerationContext);
+			}
+			else
+			{
+				FString Msg = FString::Printf(TEXT("Macro Instance Pin %s is not linked."), *Pin->PinName.ToString());
+				GenerationContext.Log(FText::FromString(Msg), Node);
+			}
+		}
+		else
+		{
+			FString Msg = FString::Printf(TEXT("Macro Instance Node does not contain a pin with name %s."), *Pin->PinName.ToString());
+			GenerationContext.Log(FText::FromString(Msg), Node);
+		}
+
+		if (!Result)
+		{
+			GenerationContext.Log(LOCTEXT("ScalarMacroInstanceEmptyAssets", "Something went wrong"), Node);
+		}
+
+		// Push the Macro again even if the result is null
+		GenerationContext.MacroNodesStack.Push(MacroInstanceNode);
+	}
+
 	else
 	{
 		GenerationContext.Log(LOCTEXT("UnimplementedNode", "Node type not implemented yet."), Node);
 	}
 
-	FGeneratedData CacheData = FGeneratedData(Node, Result);
-	GenerationContext.Generated.Add(Key, CacheData);
-	GenerationContext.GeneratedNodes.Add(Node);
+	if (bCacheNode)
+	{
+		FGeneratedData CacheData = FGeneratedData(Node, Result);
+		GenerationContext.Generated.Add(Key, CacheData);
+		GenerationContext.GeneratedNodes.Add(Node);
+	}
 	
 	if (Result)
 	{
diff --git a/Engine/Plugins/Mutable/Source/CustomizableObjectEditor/Private/MuCOE/GenerateMutableSource/GenerateMutableSourceComponent.cpp b/Engine/Plugins/Mutable/Source/CustomizableObjectEditor/Private/MuCOE/GenerateMutableSource/GenerateMutableSourceComponent.cpp
--- a/Engine/Plugins/Mutable/Source/CustomizableObjectEditor/Private/MuCOE/GenerateMutableSource/GenerateMutableSourceComponent.cpp
+++ b/Engine/Plugins/Mutable/Source/CustomizableObjectEditor/Private/MuCOE/GenerateMutableSource/GenerateMutableSourceComponent.cpp
@@ -18,6 +18,9 @@
 #include "MuCOE/Nodes/CustomizableObjectNodeComponentSwitch.h"
 #include "MuCOE/Nodes/CustomizableObjectNodeComponentVariation.h"
 #include "MuCOE/Nodes/CustomizableObjectNodeModifierBase.h"
+#include "MuCOE/Nodes/CustomizableObjectNodeMacroInstance.h"
+#include "MuCOE/Nodes/CustomizableObjectNodeStringConstant.h"
+#include "MuCOE/Nodes/CustomizableObjectNodeTunnel.h"
 #include "MuT/NodeComponent.h"
 #include "MuT/NodeComponentEdit.h"
 #include "MuT/NodeComponentNew.h"
@@ -38,6 +41,41 @@
 
 #define LOCTEXT_NAMESPACE "CustomizableObjectEditor"
 
+FName GetComponentName(const UCustomizableObjectNode& Node, FMutableGraphGenerationContext& GenerationContext)
+{
+	// find the name of the linked string node
+	if (Node.IsInMacro())
+	{
+		if (const UEdGraphPin* ComponentNamePin = Node.FindPin(TEXT("Component Name"), EEdGraphPinDirection::EGPD_Input))
+		{
+			const UEdGraphPin* StringPin = GraphTraversal::FindStringSource(*ComponentNamePin, &GenerationContext.MacroNodesStack);
+
+			if (const UCustomizableObjectNodeStringConstant* StringNode = StringPin ? Cast<UCustomizableObjectNodeStringConstant>(StringPin->GetOwningNode()) : nullptr)
+			{
+				return FName(StringNode->Value);
+			}
+		}
+	}
+
+	// return the name set in the node
+	if (const UCustomizableObjectNodeComponentMesh* TypedComponentMesh = Cast<UCustomizableObjectNodeComponentMesh>(&Node))
+	{
+		return TypedComponentMesh->ComponentName;
+	}
+
+	else if(const UCustomizableObjectNodeComponentMeshAddTo* TypedComponentMeshExtend = Cast<UCustomizableObjectNodeComponentMeshAddTo>(&Node))
+	{
+		return TypedComponentMeshExtend->ParentComponentName;
+	}
+
+	else if(const UCustomizableObjectNodeComponentPassthroughMesh* TypedComponentPassthroughMesh = Cast<UCustomizableObjectNodeComponentPassthroughMesh>(&Node))
+	{
+		return TypedComponentPassthroughMesh->ComponentName;
+	}
+
+	return FName();
+}
+
 
 /** Generate LOD pins of the given NodeComponentBase (NodeComponent, NodeComponentExtend...).
  * @param TypedComponentMesh Given component node.
@@ -127,6 +165,10 @@ mu::Ptr<mu::NodeComponent> GenerateMutableSourceComponent(const UEdGraphPin* Pin
 	{
 		return static_cast<mu::NodeComponent*>(Generated->Node.get());
 	}
+
+	// Bool that determines if a node can be added to the cache of nodes.
+	// Most nodes need to be added to the cache but there are some that don't. For exampel, MacroInstanceNodes
+	bool bCacheNode = true;
 	
 	mu::Ptr<mu::NodeComponent> Result;
 	
@@ -184,8 +226,15 @@ mu::Ptr<mu::NodeComponent> GenerateMutableSourceComponent(const UEdGraphPin* Pin
 				GenerationContext.Log(Msg, TypedComponentMeshExtend, EMessageSeverity::Warning);
 			}
 
+			// Swap the macro contex since this component can be in another macro
+			TArray<const UCustomizableObjectNodeMacroInstance*> MacroContextCopy = GenerationContext.MacroNodesStack;
+			GenerationContext.MacroNodesStack = FindResult->MacroContext;
+
 			mu::Ptr<mu::NodeComponent> ParentNodeComponent = GenerateMutableSourceComponent(TypedParentComponentMesh->OutputPin.Get(), GenerationContext);
 
+			// Restore GenerationContext MacroContex
+			GenerationContext.MacroNodesStack = MacroContextCopy;
+
 			mu::Ptr<mu::NodeComponentEdit> NodeComponentEdit = new mu::NodeComponentEdit();
 			NodeComponentEdit->Parent = ParentNodeComponent.get();
 			NodeComponentEdit->SetMessageContext(TypedComponentMeshExtend);
@@ -396,7 +445,7 @@ mu::Ptr<mu::NodeComponent> GenerateMutableSourceComponent(const UEdGraphPin* Pin
 
 			if (UEdGraphPin* VariationPin = TypedNodeVar->VariationPin(VariationIndex))
 			{
-				SurfNode->Variations[VariationIndex].Tag = TypedNodeVar->GetVariation(VariationIndex).Tag;
+				SurfNode->Variations[VariationIndex].Tag = GetVariationTag(*TypedNodeVar, VariationIndex, GenerationContext);
 				for (const UEdGraphPin* ConnectedPin : FollowInputPinArray(*VariationPin))
 				{
 					// Is it a modifier?
@@ -414,15 +463,80 @@ mu::Ptr<mu::NodeComponent> GenerateMutableSourceComponent(const UEdGraphPin* Pin
 		}
 	}
 
+	else if (const UCustomizableObjectNodeMacroInstance* TypedNodeMacro = Cast<UCustomizableObjectNodeMacroInstance>(Node))
+	{
+		bCacheNode = false;
+
+		if (const UEdGraphPin* OutputPin = TypedNodeMacro->GetMacroIOPin(ECOMacroIOType::COMVT_Output, Pin->PinName))
+		{
+			if (const UEdGraphPin* FollowPin = FollowInputPin(*OutputPin))
+			{
+				GenerationContext.MacroNodesStack.Push(TypedNodeMacro);
+				Result = GenerateMutableSourceComponent(FollowPin, GenerationContext);
+				GenerationContext.MacroNodesStack.Pop();
+			}
+			else
+			{
+				FString Msg = FString::Printf(TEXT("Macro Pin [%s] not linked."), *Pin->PinName.ToString());
+				GenerationContext.Log(FText::FromString(Msg), Node);
+			}
+		}
+		else
+		{
+			FString Msg = FString::Printf(TEXT("Macro Output node does not contain a pin with name [%s]."), *Pin->PinName.ToString());
+			GenerationContext.Log(FText::FromString(Msg), Node);
+		}
+	}
+
+	else if (const UCustomizableObjectNodeTunnel* TypedNodeTunnel = Cast<UCustomizableObjectNodeTunnel>(Node))
+	{
+		check(TypedNodeTunnel->bIsInputNode);
+		check(GenerationContext.MacroNodesStack.Num());
+
+		bCacheNode = false;
+
+		const UCustomizableObjectNodeMacroInstance* MacroInstanceNode = GenerationContext.MacroNodesStack.Pop();
+		check(MacroInstanceNode);
+
+		if (const UEdGraphPin* InputPin = MacroInstanceNode->FindPin(Pin->PinName, EEdGraphPinDirection::EGPD_Input))
+		{
+			if (const UEdGraphPin* FollowPin = FollowInputPin(*InputPin))
+			{
+				Result = GenerateMutableSourceComponent(FollowPin, GenerationContext);
+			}
+			else
+			{
+				FString Msg = FString::Printf(TEXT("Macro Instance Pin [%s] is not linked."), *Pin->PinName.ToString());
+				GenerationContext.Log(FText::FromString(Msg), Node);
+			}
+		}
+		else
+		{
+			FString Msg = FString::Printf(TEXT("Macro Instance Node does not contain a pin with name [%s]."), *Pin->PinName.ToString());
+			GenerationContext.Log(FText::FromString(Msg), Node);
+		}
+
+		if (!Result)
+		{
+			GenerationContext.Log(LOCTEXT("ScalarMacroInstanceEmptyAssets", "Something went wrong"), Node);
+		}
+
+		// Push the Macro again even if the result is null
+		GenerationContext.MacroNodesStack.Push(MacroInstanceNode);
+	}
+
 	else
 	{
 		GenerationContext.Log(LOCTEXT("UnimplementedNode", "Node type not implemented yet."), Node);
 		ensure(false);
 	}
 
-	GenerationContext.Generated.Add(Key, FGeneratedData(Node, Result));
-	GenerationContext.GeneratedNodes.Add(Node);
-	
+	if (bCacheNode)
+	{
+		GenerationContext.Generated.Add(Key, FGeneratedData(Node, Result));
+		GenerationContext.GeneratedNodes.Add(Node);
+	}
+
 	return Result;
 }
 
@@ -550,6 +664,7 @@ void FirstPass(UCustomizableObjectNodeComponentMesh& Node, FMutableGraphGenerati
 	FMutableComponentInfo ComponentInfo(Node.ComponentName, RefSkeletalMesh);
 	ComponentInfo.Node = &Node;
 	ComponentInfo.AccumulateBonesToRemovePerLOD(Node.LODReductionSettings, Node.NumLODs);
+	ComponentInfo.MacroContext = GenerationContext.MacroNodesStack;
 
 	GenerationContext.CompilationContext.ComponentInfos.Add(ComponentInfo);
 
diff --git a/Engine/Plugins/Mutable/Source/CustomizableObjectEditor/Private/MuCOE/GenerateMutableSource/GenerateMutableSourceFloat.cpp b/Engine/Plugins/Mutable/Source/CustomizableObjectEditor/Private/MuCOE/GenerateMutableSource/GenerateMutableSourceFloat.cpp
--- a/Engine/Plugins/Mutable/Source/CustomizableObjectEditor/Private/MuCOE/GenerateMutableSource/GenerateMutableSourceFloat.cpp
+++ b/Engine/Plugins/Mutable/Source/CustomizableObjectEditor/Private/MuCOE/GenerateMutableSource/GenerateMutableSourceFloat.cpp
@@ -6,6 +6,7 @@
 #include "Curves/CurveLinearColor.h"
 #include "Curves/CurveVector.h"
 #include "MuCOE/CustomizableObjectCompiler.h"
+#include "MuCOE/CustomizableObjectGraph.h"
 #include "MuCOE/EdGraphSchema_CustomizableObject.h"
 #include "MuCOE/GenerateMutableSource/GenerateMutableSourceImage.h"
 #include "MuCOE/GenerateMutableSource/GenerateMutableSourceTable.h"
@@ -17,7 +18,9 @@
 #include "MuCOE/Nodes/CustomizableObjectNodeFloatSwitch.h"
 #include "MuCOE/Nodes/CustomizableObjectNodeFloatVariation.h"
 #include "MuCOE/Nodes/CustomizableObjectNodeFloatArithmeticOp.h"
+#include "MuCOE/Nodes/CustomizableObjectNodeMacroInstance.h"
 #include "MuCOE/Nodes/CustomizableObjectNodeTable.h"
+#include "MuCOE/Nodes/CustomizableObjectNodeTunnel.h"
 #include "MuT/NodeScalarConstant.h"
 #include "MuT/NodeScalarCurve.h"
 #include "MuT/NodeScalarSwitch.h"
@@ -28,7 +31,6 @@
 #define LOCTEXT_NAMESPACE "CustomizableObjectEditor"
 
 
-
 mu::Ptr<mu::NodeScalar> GenerateMutableSourceFloat(const UEdGraphPin* Pin, FMutableGraphGenerationContext& GenerationContext)
 {
 	check(Pin)
@@ -51,7 +53,9 @@ mu::Ptr<mu::NodeScalar> GenerateMutableSourceFloat(const UEdGraphPin* Pin, FMuta
 		Node->SetRefreshNodeWarning();
 	}
 
-	bool bDoNotAddToGeneratedCache = false;
+	// Bool that determines if a node can be added to the cache of nodes.
+	// Most nodes need to be added to the cache but there are some that don't. For exampel, MacroInstanceNodes
+	bool bCacheNode = true;
 
 	mu::Ptr<mu::NodeScalar> Result;
 	
@@ -251,7 +255,8 @@ mu::Ptr<mu::NodeScalar> GenerateMutableSourceFloat(const UEdGraphPin* Pin, FMuta
 			UEdGraphPin* VariationPin = TypedNodeFloatVar->VariationPin(VariationIndex);
 			if (!VariationPin) continue;
 
-			FloatNode->Variations[VariationIndex].Tag = TypedNodeFloatVar->GetVariation(VariationIndex).Tag;
+			FloatNode->Variations[VariationIndex].Tag = GetVariationTag(*TypedNodeFloatVar, VariationIndex, GenerationContext);
+
 			if (const UEdGraphPin* ConnectedPin = FollowInputPin(*VariationPin))
 			{
 				mu::Ptr<mu::NodeScalar> ChildNode = GenerateMutableSourceFloat(ConnectedPin, GenerationContext);
@@ -311,7 +316,7 @@ mu::Ptr<mu::NodeScalar> GenerateMutableSourceFloat(const UEdGraphPin* Pin, FMuta
 		if (Pin->PinType.PinCategory == Schema->PC_MaterialAsset)
 		{
 			// Material pins have to skip the cache of nodes or they will return always the same column node
-			bDoNotAddToGeneratedCache = true;
+			bCacheNode = false;
 		}
 
 		bool bSuccess = true;
@@ -384,12 +389,74 @@ mu::Ptr<mu::NodeScalar> GenerateMutableSourceFloat(const UEdGraphPin* Pin, FMuta
 		}
 	}
 
+	else if (const UCustomizableObjectNodeMacroInstance* TypedNodeMacro = Cast<UCustomizableObjectNodeMacroInstance>(Node))
+	{
+		bCacheNode = false;
+
+		if (const UEdGraphPin* OutputPin = TypedNodeMacro->GetMacroIOPin(ECOMacroIOType::COMVT_Output, Pin->PinName))
+		{
+			if (const UEdGraphPin* FollowPin = FollowInputPin(*OutputPin))
+			{
+				GenerationContext.MacroNodesStack.Push(TypedNodeMacro);
+				Result = GenerateMutableSourceFloat(FollowPin, GenerationContext);
+				GenerationContext.MacroNodesStack.Pop();
+			}
+			else
+			{
+				FString Msg = FString::Printf(TEXT("Macro Pin [%s] not linked."), *Pin->PinName.ToString());
+				GenerationContext.Log(FText::FromString(Msg), Node);
+			}
+		}
+		else
+		{
+			FString Msg = FString::Printf(TEXT("Macro Output node does not contain a pin with name [%s]."), *Pin->PinName.ToString());
+			GenerationContext.Log(FText::FromString(Msg), Node);
+		}
+	}
+
+	else if (const UCustomizableObjectNodeTunnel* TypedNodeTunnel = Cast<UCustomizableObjectNodeTunnel>(Node))
+	{
+		check(TypedNodeTunnel->bIsInputNode);
+		check(GenerationContext.MacroNodesStack.Num());
+
+		bCacheNode = false;
+
+		const UCustomizableObjectNodeMacroInstance* MacroInstanceNode = GenerationContext.MacroNodesStack.Pop();
+		check(MacroInstanceNode);
+
+		if (const UEdGraphPin* InputPin = MacroInstanceNode->FindPin(Pin->PinName, EEdGraphPinDirection::EGPD_Input))
+		{
+			if (const UEdGraphPin* FollowPin = FollowInputPin(*InputPin))
+			{
+				Result = GenerateMutableSourceFloat(FollowPin, GenerationContext);
+			}
+			else
+			{
+				FString Msg = FString::Printf(TEXT("Macro Instance Pin [%s] is not linked."), *Pin->PinName.ToString());
+				GenerationContext.Log(FText::FromString(Msg), Node);
+			}
+		}
+		else
+		{
+			FString Msg = FString::Printf(TEXT("Macro Instance Node does not contain a pin with name [%s]."), *Pin->PinName.ToString());
+			GenerationContext.Log(FText::FromString(Msg), Node);
+		}
+
+		if (!Result)
+		{
+			GenerationContext.Log(LOCTEXT("ScalarMacroInstanceEmptyAssets", "Something went wrong"), Node);
+		}
+
+		// Push the Macro again even if the result is null
+		GenerationContext.MacroNodesStack.Push(MacroInstanceNode);
+	}
+
 	else
 	{
 		GenerationContext.Log(LOCTEXT("UnimplementedNode", "Node type not implemented yet."), Node);
 	}
 
-	if (!bDoNotAddToGeneratedCache)
+	if (bCacheNode)
 	{
 		GenerationContext.Generated.Add(Key, FGeneratedData(Node, Result));
 		GenerationContext.GeneratedNodes.Add(Node);
diff --git a/Engine/Plugins/Mutable/Source/CustomizableObjectEditor/Private/MuCOE/GenerateMutableSource/GenerateMutableSourceImage.cpp b/Engine/Plugins/Mutable/Source/CustomizableObjectEditor/Private/MuCOE/GenerateMutableSource/GenerateMutableSourceImage.cpp
--- a/Engine/Plugins/Mutable/Source/CustomizableObjectEditor/Private/MuCOE/GenerateMutableSource/GenerateMutableSourceImage.cpp
+++ b/Engine/Plugins/Mutable/Source/CustomizableObjectEditor/Private/MuCOE/GenerateMutableSource/GenerateMutableSourceImage.cpp
@@ -5,6 +5,7 @@
 #include "Engine/TextureLODSettings.h"
 
 #include "MuCOE/CustomizableObjectCompiler.h"
+#include "MuCOE/CustomizableObjectGraph.h"
 #include "MuCOE/EdGraphSchema_CustomizableObject.h"
 #include "MuCOE/GenerateMutableSource/GenerateMutableSourceColor.h"
 #include "MuCOE/GenerateMutableSource/GenerateMutableSourceFloat.h"
@@ -12,6 +13,7 @@
 #include "MuCOE/GenerateMutableSource/GenerateMutableSourceProjector.h"
 #include "MuCOE/GenerateMutableSource/GenerateMutableSourceTable.h"
 #include "MuCOE/GraphTraversal.h"
+#include "MuCOE/Nodes/CustomizableObjectNodeMacroInstance.h"
 #include "MuCOE/Nodes/CustomizableObjectNodeMesh.h"
 #include "MuCOE/Nodes/CustomizableObjectNodeTable.h"
 #include "MuCOE/Nodes/CustomizableObjectNodeMaterial.h"
@@ -33,6 +35,7 @@
 #include "MuCOE/Nodes/CustomizableObjectNodeTextureTransform.h"
 #include "MuCOE/Nodes/CustomizableObjectNodeTextureSaturate.h"
 #include "MuCOE/Nodes/CustomizableObjectNodeTextureVariation.h"
+#include "MuCOE/Nodes/CustomizableObjectNodeTunnel.h"
 #include "MuCOE/UnrealEditorPortabilityHelpers.h"
 #include "MuCO/UnrealToMutableTextureConversionUtils.h"
 #include "MuT/NodeImageBinarise.h"
@@ -98,7 +101,9 @@ mu::Ptr<mu::NodeImage> GenerateMutableSourceImage(const UEdGraphPin* Pin, FMutab
 		Node->SetRefreshNodeWarning();
 	}
 
-	bool bDoNotAddToGeneratedCache = false;
+	// Bool that determines if a node can be added to the cache of nodes.
+	// Most nodes need to be added to the cache but there are some that don't. For exampel, MacroInstanceNodes
+	bool bCacheNode = true;
 
 	mu::Ptr<mu::NodeImage> Result;
 
@@ -427,8 +432,10 @@ mu::Ptr<mu::NodeImage> GenerateMutableSourceImage(const UEdGraphPin* Pin, FMutab
 			}
 		}
 
-		TextureNode->Variations.SetNum(TypedNodeImageVar->Variations.Num());
-		for (int32 VariationIndex = 0; VariationIndex < TypedNodeImageVar->Variations.Num(); ++VariationIndex)
+		const int32 NumVariations = TypedNodeImageVar->GetNumVariations();
+		TextureNode->Variations.SetNum(NumVariations);
+		
+		for (int32 VariationIndex = 0; VariationIndex < NumVariations; ++VariationIndex)
 		{
 			const UEdGraphPin* VariationPin = TypedNodeImageVar->VariationPin(VariationIndex);
 			if (!VariationPin)
@@ -436,7 +443,8 @@ mu::Ptr<mu::NodeImage> GenerateMutableSourceImage(const UEdGraphPin* Pin, FMutab
 				continue;
 			}
 
-			TextureNode->Variations[VariationIndex].Tag = TypedNodeImageVar->Variations[VariationIndex].Tag;
+			TextureNode->Variations[VariationIndex].Tag = GetVariationTag(*TypedNodeImageVar, VariationIndex, GenerationContext);
+
 			if (const UEdGraphPin* ConnectedPin = FollowInputPin(*VariationPin))
 			{
 				mu::Ptr<mu::NodeImage> ChildNode = GenerateMutableSourceImage(ConnectedPin, GenerationContext, ReferenceTextureSize);
@@ -1016,7 +1024,7 @@ mu::Ptr<mu::NodeImage> GenerateMutableSourceImage(const UEdGraphPin* Pin, FMutab
 			if (Pin->PinType.PinCategory == Schema->PC_MaterialAsset)
 			{
 				// Material pins have to skip the cache of nodes or they will return always the same column node
-				bDoNotAddToGeneratedCache = true;
+				bCacheNode = false;
 			}
 
 			UDataTable* DataTable = GetDataTable(TypedNodeTable, GenerationContext);
@@ -1152,12 +1160,74 @@ mu::Ptr<mu::NodeImage> GenerateMutableSourceImage(const UEdGraphPin* Pin, FMutab
 		}
 	}
 
+	else if (const UCustomizableObjectNodeMacroInstance* TypedNodeMacro = Cast<UCustomizableObjectNodeMacroInstance>(Node))
+	{
+		bCacheNode = false;
+
+		if (const UEdGraphPin* OutputPin = TypedNodeMacro->GetMacroIOPin(ECOMacroIOType::COMVT_Output, Pin->PinName))
+		{
+			if (const UEdGraphPin* FollowPin = FollowInputPin(*OutputPin))
+			{
+				GenerationContext.MacroNodesStack.Push(TypedNodeMacro);
+				Result = GenerateMutableSourceImage(FollowPin, GenerationContext, ReferenceTextureSize);
+				GenerationContext.MacroNodesStack.Pop();
+			}
+			else
+			{
+				FString Msg = FString::Printf(TEXT("Macro Pin [%s] not linked."), *Pin->PinName.ToString());
+				GenerationContext.Log(FText::FromString(Msg), Node);
+			}
+		}
+		else
+		{
+			FString Msg = FString::Printf(TEXT("Macro Output node does not contain a pin with name [%s]."), *Pin->PinName.ToString());
+			GenerationContext.Log(FText::FromString(Msg), Node);
+		}
+	}
+
+	else if (const UCustomizableObjectNodeTunnel* TypedNodeTunnel = Cast<UCustomizableObjectNodeTunnel>(Node))
+	{
+		check(TypedNodeTunnel->bIsInputNode);
+		check(GenerationContext.MacroNodesStack.Num());
+
+		bCacheNode = false;
+
+		const UCustomizableObjectNodeMacroInstance* MacroInstanceNode = GenerationContext.MacroNodesStack.Pop();
+		check(MacroInstanceNode);
+
+		if (const UEdGraphPin* InputPin = MacroInstanceNode->FindPin(Pin->PinName, EEdGraphPinDirection::EGPD_Input))
+		{
+			if (const UEdGraphPin* FollowPin = FollowInputPin(*InputPin))
+			{
+				Result = GenerateMutableSourceImage(FollowPin, GenerationContext, ReferenceTextureSize);
+			}
+			else
+			{
+				FString Msg = FString::Printf(TEXT("Macro Instance Pin [%s] is not linked."), *Pin->PinName.ToString());
+				GenerationContext.Log(FText::FromString(Msg), Node);
+			}
+		}
+		else
+		{
+			FString Msg = FString::Printf(TEXT("Macro Instance Node does not contain a pin with name [%s]."), *Pin->PinName.ToString());
+			GenerationContext.Log(FText::FromString(Msg), Node);
+		}
+
+		if (!Result)
+		{
+			GenerationContext.Log(LOCTEXT("ScalarMacroInstanceEmptyAssets", "Something went wrong"), Node);
+		}
+
+		// Push the Macro again even if the result is null
+		GenerationContext.MacroNodesStack.Push(MacroInstanceNode);
+	}
+
 	else
 	{
 		GenerationContext.Log(LOCTEXT("UnimplementedNode", "Node type not implemented yet."), Node);
 	}
 
-	if (!bDoNotAddToGeneratedCache)
+	if (bCacheNode)
 	{
 		GenerationContext.Generated.Add(Key, FGeneratedData(Node, Result));
 		GenerationContext.GeneratedNodes.Add(Node);
diff --git a/Engine/Plugins/Mutable/Source/CustomizableObjectEditor/Private/MuCOE/GenerateMutableSource/GenerateMutableSourceMesh.cpp b/Engine/Plugins/Mutable/Source/CustomizableObjectEditor/Private/MuCOE/GenerateMutableSource/GenerateMutableSourceMesh.cpp
--- a/Engine/Plugins/Mutable/Source/CustomizableObjectEditor/Private/MuCOE/GenerateMutableSource/GenerateMutableSourceMesh.cpp
+++ b/Engine/Plugins/Mutable/Source/CustomizableObjectEditor/Private/MuCOE/GenerateMutableSource/GenerateMutableSourceMesh.cpp
@@ -26,6 +26,7 @@
 #include "MuCOE/Nodes/CustomizableObjectNodeAnimationPose.h"
 #include "MuCOE/Nodes/CustomizableObjectNodeFloatConstant.h"
 #include "MuCOE/Nodes/CustomizableObjectNodeFloatParameter.h"
+#include "MuCOE/Nodes/CustomizableObjectNodeMacroInstance.h"
 #include "MuCOE/Nodes/CustomizableObjectNodeMaterial.h"
 #include "MuCOE/Nodes/CustomizableObjectNodeMeshMorph.h"
 #include "MuCOE/Nodes/CustomizableObjectNodeMeshMorphStackApplication.h"
@@ -37,7 +38,9 @@
 #include "MuCOE/Nodes/CustomizableObjectNodeMeshParameter.h"
 #include "MuCOE/Nodes/CustomizableObjectNodeSkeletalMesh.h"
 #include "MuCOE/Nodes/CustomizableObjectNodeStaticMesh.h"
+#include "MuCOE/Nodes/CustomizableObjectNodeStringConstant.h"
 #include "MuCOE/Nodes/CustomizableObjectNodeTable.h"
+#include "MuCOE/Nodes/CustomizableObjectNodeTunnel.h"
 #include "MuCOE/UnrealEditorPortabilityHelpers.h"
 #include "MuT/NodeMeshConstant.h"
 #include "MuT/NodeMeshMakeMorph.h"
@@ -2963,8 +2966,10 @@ TSharedPtr<mu::FMesh> BuildMorphedMutableMesh(const UEdGraphPin* BaseSourcePin,
 }
 
 
-void GenerateMorphFactor(const UCustomizableObjectNode* Node, const UEdGraphPin& FactorPin, FMutableGraphGenerationContext& GenerationContext, mu::Ptr<mu::NodeMeshMorph> MeshNode)
+mu::Ptr<mu::NodeScalar> GenerateMorphFactor(const UCustomizableObjectNode* Node, const UEdGraphPin& FactorPin, FMutableGraphGenerationContext& GenerationContext)
 {
+	mu::Ptr<mu::NodeScalar> FactorNode;
+
 	if (const UEdGraphPin* ConnectedPin = FollowInputPin(FactorPin))
 	{
 		UEdGraphNode* floatNode = ConnectedPin->GetOwningNode();
@@ -3004,10 +3009,11 @@ void GenerateMorphFactor(const UCustomizableObjectNode* Node, const UEdGraphPin&
 
 		if (validStaticFactor)
 		{
-			mu::Ptr<mu::NodeScalar> FactorNode = GenerateMutableSourceFloat(ConnectedPin, GenerationContext);
-			MeshNode->Factor = FactorNode;
+			FactorNode = GenerateMutableSourceFloat(ConnectedPin, GenerationContext);
 		}
 	}
+
+	return FactorNode;
 }
 
 TArray<TTuple<USkeletalMesh*, TSoftClassPtr<UAnimInstance>>> GetSkeletalMeshesInfoForReshapeSelection(
@@ -3471,8 +3477,9 @@ mu::NodeMeshPtr GenerateMorphMesh(const UEdGraphPin* Pin,
 	mu::Ptr<mu::NodeMeshMorph> Result = new mu::NodeMeshMorph();
 	
 	// Factor
-	GenerateMorphFactor(MorphNode, *TypedNodeMorphs[MorphIndex].FactorPin, GenerationContext, Result);
-	
+	//GenerateMorphFactor(MorphNode, , GenerationContext, Result);
+	Result->Factor = TypedNodeMorphs[MorphIndex].FactorNode;
+
 	// Base
 	if (MorphIndex == TypedNodeMorphs.Num() - 1)
 	{
@@ -3573,8 +3580,8 @@ mu::NodeMeshPtr GenerateMorphMesh(const UEdGraphPin* Pin,
 				Result->bReshapeSkeleton = TypedMorphNode->bReshapeSkeleton;
 				Result->bReshapePhysicsVolumes = TypedMorphNode->bReshapePhysicsVolumes;
 				{
-					const UEdGraphPin* ConnectedPin = FollowInputPin(*TypedMorphNode->MeshPin());
-					const UEdGraphPin* SourceMeshPin = ConnectedPin ? FindMeshBaseSource(*ConnectedPin, false) : nullptr;
+					//const UEdGraphPin* ConnectedPin = FollowInputPin(*TypedMorphNode->MeshPin());
+					const UEdGraphPin* SourceMeshPin = Pin ? FindMeshBaseSource(*Pin, false) : nullptr;
 					const UEdGraphNode* SkeletalMeshNode = SourceMeshPin ? SourceMeshPin->GetOwningNode() : nullptr;
 
 					TArray<TTuple<USkeletalMesh*, TSoftClassPtr<UAnimInstance>>> SkeletalMeshesToDeform = 
@@ -3626,6 +3633,7 @@ mu::NodeMeshPtr GenerateMorphMesh(const UEdGraphPin* Pin,
 	return Result;
 }
 
+
 mu::Ptr<mu::NodeMesh> GenerateMutableSourceMesh(const UEdGraphPin* Pin,
 	FMutableGraphGenerationContext& GenerationContext,
 	FMutableGraphMeshGenerationData& MeshData,
@@ -3654,6 +3662,10 @@ mu::Ptr<mu::NodeMesh> GenerateMutableSourceMesh(const UEdGraphPin* Pin,
 		Node->SetRefreshNodeWarning();
 	}
 
+	// Bool that determines if a node can be added to the cache of nodes.
+	// Most nodes need to be added to the cache but there are some that don't. For exampel, MacroInstanceNodes
+	bool bCacheNode = true;
+
 	//SkeletalMesh Result
 	mu::Ptr<mu::NodeMesh> Result;
 
@@ -3959,8 +3971,31 @@ mu::Ptr<mu::NodeMesh> GenerateMutableSourceMesh(const UEdGraphPin* Pin,
 	{
 		if (const UEdGraphPin* ConnectedPin = FollowInputPin(*TypedNodeMorph->MeshPin()))
 		{
+			FString MorphTargetName = TypedNodeMorph->MorphTargetName;
+
+			if (UEdGraphPin* MorphTargetNamePin = Node->FindPin(FString("Morph Target Name"), EEdGraphPinDirection::EGPD_Input))
+			{
+				const UEdGraphPin* ConnectedStringPin = FollowInputPin(*MorphTargetNamePin);
+				const UEdGraphPin* SourceStringPin = ConnectedStringPin ? GraphTraversal::FindStringSource(*ConnectedStringPin, &GenerationContext.MacroNodesStack) : nullptr;
+
+				if (SourceStringPin)
+				{
+					if (const UCustomizableObjectNodeStringConstant* StringNode = Cast<UCustomizableObjectNodeStringConstant>(SourceStringPin->GetOwningNode()))
+					{
+						MorphTargetName = StringNode->Value;
+					}
+				}
+				else
+				{
+					GenerationContext.Log(LOCTEXT("MorphStringNodeFailed", "Could not found a linked String node."), Node);
+				}
+			}
+
+			// Factor
+			mu::Ptr<mu::NodeScalar> FactorNode = GenerateMorphFactor(Node, *TypedNodeMorph->FactorPin(), GenerationContext);
+
 			// Mesh Morph Stack Management
-			FMorphNodeData NewMorphData = { TypedNodeMorph, TypedNodeMorph->MorphTargetName ,TypedNodeMorph->FactorPin(), TypedNodeMorph->MeshPin() };
+			FMorphNodeData NewMorphData = { TypedNodeMorph, MorphTargetName, FactorNode, TypedNodeMorph->MeshPin() };
 			GenerationContext.MeshMorphStack.Push(NewMorphData);
 			Result = GenerateMutableSourceMesh(ConnectedPin, GenerationContext, MeshData, SurfaceMetadataId, false, bOnlyConnectedLOD);
 			GenerationContext.MeshMorphStack.Pop(EAllowShrinking::Yes);
@@ -4015,8 +4050,11 @@ mu::Ptr<mu::NodeMesh> GenerateMutableSourceMesh(const UEdGraphPin* Pin,
 						continue;
 					}
 
+					//Factor
+					mu::Ptr<mu::NodeScalar> FactorNode = GenerateMorphFactor(Node, *MorphPin, GenerationContext);
+
 					// Mesh Morph Stack Management. TODO(Max): should we add the stack application node here instead of the def? Or both?
-					FMorphNodeData NewMorphData = { TypedNodeMeshMorphStackDef, MorphName, MorphPin, TypedNodeMeshMorphStackApp->GetMeshPin() };
+					FMorphNodeData NewMorphData = { TypedNodeMeshMorphStackDef, MorphName, FactorNode, TypedNodeMeshMorphStackApp->GetMeshPin() };
 					GenerationContext.MeshMorphStack.Push(NewMorphData);
 
 					AddedMorphs++;
@@ -4134,13 +4172,17 @@ mu::Ptr<mu::NodeMesh> GenerateMutableSourceMesh(const UEdGraphPin* Pin,
 		}
 
 		const int32 NumVariations = TypedNodeMeshVar->GetNumVariations();
+
 		MeshNode->Variations.SetNum(NumVariations);
+
 		for (int VariationIndex = 0; VariationIndex < NumVariations; ++VariationIndex)
 		{
 			const UEdGraphPin* VariationPin = TypedNodeMeshVar->VariationPin(VariationIndex);
 			if (!VariationPin) continue;
 
-			MeshNode->Variations[VariationIndex].Tag = StringCast<ANSICHAR>(*TypedNodeMeshVar->GetVariation(VariationIndex).Tag).Get();
+			FString VariationTag = GetVariationTag(*TypedNodeMeshVar, VariationIndex, GenerationContext);
+			MeshNode->Variations[VariationIndex].Tag = StringCast<ANSICHAR>(*VariationTag).Get();
+
 			if (const UEdGraphPin* ConnectedPin = FollowInputPin(*VariationPin))
 			{
 				FMutableGraphMeshGenerationData VariationMeshData;
@@ -4607,13 +4649,78 @@ mu::Ptr<mu::NodeMesh> GenerateMutableSourceMesh(const UEdGraphPin* Pin,
 		Result = ParameterNode;
 	}
 
+	else if (const UCustomizableObjectNodeMacroInstance* TypedNodeMacro = Cast<UCustomizableObjectNodeMacroInstance>(Node))
+	{
+		bCacheNode = false;
+
+		if (const UEdGraphPin* OutputPin = TypedNodeMacro->GetMacroIOPin(ECOMacroIOType::COMVT_Output, Pin->PinName))
+		{
+			if (const UEdGraphPin* FollowPin = FollowInputPin(*OutputPin))
+			{
+				GenerationContext.MacroNodesStack.Push(TypedNodeMacro);
+				Result = GenerateMutableSourceMesh(FollowPin, GenerationContext, MeshData, SurfaceMetadataId, bLinkedToExtendMaterial, bOnlyConnectedLOD);
+				GenerationContext.MacroNodesStack.Pop();
+			}
+			else
+			{
+				FString Msg = FString::Printf(TEXT("Macro Pin [%s] not linked."), *Pin->PinName.ToString());
+				GenerationContext.Log(FText::FromString(Msg), Node);
+			}
+		}
+		else
+		{
+			FString Msg = FString::Printf(TEXT("Macro Output node does not contain a pin with name [%s]."), *Pin->PinName.ToString());
+			GenerationContext.Log(FText::FromString(Msg), Node);
+		}
+	}
+
+	else if (const UCustomizableObjectNodeTunnel* TypedNodeTunnel = Cast<UCustomizableObjectNodeTunnel>(Node))
+	{
+		check(TypedNodeTunnel->bIsInputNode);
+		check(GenerationContext.MacroNodesStack.Num());
+
+		bCacheNode = false;
+
+		const UCustomizableObjectNodeMacroInstance* MacroInstanceNode = GenerationContext.MacroNodesStack.Pop();
+		check(MacroInstanceNode);
+
+		if (const UEdGraphPin* InputPin = MacroInstanceNode->FindPin(Pin->PinName, EEdGraphPinDirection::EGPD_Input))
+		{
+			if (const UEdGraphPin* FollowPin = FollowInputPin(*InputPin))
+			{
+				Result = GenerateMutableSourceMesh(FollowPin, GenerationContext, MeshData, SurfaceMetadataId, bLinkedToExtendMaterial, bOnlyConnectedLOD);
+			}
+			else
+			{
+				FString Msg = FString::Printf(TEXT("Macro Instance Pin [%s] is not linked."), *Pin->PinName.ToString());
+				GenerationContext.Log(FText::FromString(Msg), Node);
+			}
+		}
+		else
+		{
+			FString Msg = FString::Printf(TEXT("Macro Instance Node does not contain a pin with name [%s]."), *Pin->PinName.ToString());
+			GenerationContext.Log(FText::FromString(Msg), Node);
+		}
+
+		if (!Result)
+		{
+			GenerationContext.Log(LOCTEXT("ScalarMacroInstanceEmptyAssets", "Something went wrong"), Node);
+		}
+
+		// Push the Macro again even if the result is null
+		GenerationContext.MacroNodesStack.Push(MacroInstanceNode);
+	}
+
 	else
 	{
 		GenerationContext.Log(LOCTEXT("UnimplementedMeshNode", "Mesh node type not implemented yet."), Node);
 	}
 	
-	GenerationContext.Generated.Add(Key, FGeneratedData(Node, Result, &MeshData));
-	GenerationContext.GeneratedNodes.Add(Node);
+	if (bCacheNode)
+	{
+		GenerationContext.Generated.Add(Key, FGeneratedData(Node, Result, &MeshData));
+		GenerationContext.GeneratedNodes.Add(Node);
+	}
 
 	// We return the mesh modified by morphs if there is any
 	if (MorphResult)
diff --git a/Engine/Plugins/Mutable/Source/CustomizableObjectEditor/Private/MuCOE/GenerateMutableSource/GenerateMutableSourceModifier.cpp b/Engine/Plugins/Mutable/Source/CustomizableObjectEditor/Private/MuCOE/GenerateMutableSource/GenerateMutableSourceModifier.cpp
--- a/Engine/Plugins/Mutable/Source/CustomizableObjectEditor/Private/MuCOE/GenerateMutableSource/GenerateMutableSourceModifier.cpp
+++ b/Engine/Plugins/Mutable/Source/CustomizableObjectEditor/Private/MuCOE/GenerateMutableSource/GenerateMutableSourceModifier.cpp
@@ -14,6 +14,7 @@
 #include "MuCOE/GenerateMutableSource/GenerateMutableSourceTransform.h"
 #include "MuCOE/GraphTraversal.h"
 #include "MuCOE/MutableUtils.h"
+#include "MuCOE/Nodes/CustomizableObjectNodeMacroInstance.h"
 #include "MuCOE/Nodes/CustomizableObjectNodeModifierClipDeform.h"
 #include "MuCOE/Nodes/CustomizableObjectNodeModifierClipMorph.h"
 #include "MuCOE/Nodes/CustomizableObjectNodeModifierClipWithMesh.h"
@@ -29,6 +30,8 @@
 #include "MuCOE/Nodes/CustomizableObjectNodeFloatParameter.h"
 #include "MuCOE/Nodes/CustomizableObjectNodeFloatConstant.h"
 #include "MuCOE/Nodes/CustomizableObjectNodeModifierTransformInMesh.h"
+#include "MuCOE/Nodes/CustomizableObjectNodeStringConstant.h"
+#include "MuCOE/Nodes/CustomizableObjectNodeTunnel.h"
 #include "MuR/Mesh.h"
 #include "MuT/NodeMeshTransform.h"
 #include "MuT/NodeModifierMeshClipDeform.h"
@@ -43,6 +46,36 @@
 
 #define LOCTEXT_NAMESPACE "CustomizableObjectEditor"
 
+TArray<FString> GetNodeRequiredTags(const UCustomizableObjectNodeModifierBase& Node, FMutableGraphGenerationContext& GenerationContext)
+{
+	TArray<FString> OutTags = Node.RequiredTags;
+
+	if (Node.IsInMacro())
+	{
+		UEdGraphPin* TagsPin = Node.RequiredTagsPin();
+		check(TagsPin);
+
+		for (const UEdGraphPin* StringPin : FollowInputPinArray(*TagsPin))
+		{
+			const UEdGraphPin* SourceStringPin = GraphTraversal::FindStringSource(*StringPin, &GenerationContext.MacroNodesStack);
+
+			if (SourceStringPin)
+			{
+				if (const UCustomizableObjectNodeStringConstant* StringNode = Cast<UCustomizableObjectNodeStringConstant>(SourceStringPin->GetOwningNode()))
+				{
+					OutTags.Add(StringNode->Value);
+				}
+			}
+			else
+			{
+				GenerationContext.Log(LOCTEXT("MorphStringNodeFailed", "Could not found a linked String node."), &Node);
+			}
+		}
+	}
+
+	return OutTags;
+}
+
 
 mu::Ptr<mu::NodeModifier> GenerateMutableSourceModifier(const UEdGraphPin * Pin, FMutableGraphGenerationContext & GenerationContext)
 {
@@ -63,9 +96,10 @@ mu::Ptr<mu::NodeModifier> GenerateMutableSourceModifier(const UEdGraphPin * Pin,
 	
 	mu::Ptr<mu::NodeModifier> Result;
 
-	bool bDoNotAddToGeneratedCache = false; // TODO Remove on MTBL-829 
+	// Bool that determines if a node can be added to the cache of nodes.
+	// Most nodes need to be added to the cache but there are some that don't. For exampel, MacroInstanceNodes
+	bool bCacheNode = true; 
 
-	
 	if (const UCustomizableObjectNodeModifierClipMorph* TypedNodeClip = Cast<UCustomizableObjectNodeModifierClipMorph>(Node))
 	{
 		const EMutableMeshConversionFlags ModifiersMeshFlags =
@@ -75,7 +109,7 @@ mu::Ptr<mu::NodeModifier> GenerateMutableSourceModifier(const UEdGraphPin * Pin,
 
 		// This modifier can be connected to multiple nodes at the same time and, when that happens and if the cache is being used, only the first node to be processed does work. 
 		// By not caching the mutable node we avoid this from even happening
-		bDoNotAddToGeneratedCache = true;
+		bCacheNode = false;
 		
 		mu::Ptr<mu::NodeModifierMeshClipMorphPlane> ClipNode = new mu::NodeModifierMeshClipMorphPlane();
 		Result = ClipNode;
@@ -90,7 +124,7 @@ mu::Ptr<mu::NodeModifier> GenerateMutableSourceModifier(const UEdGraphPin * Pin,
 		ClipNode->SetVertexSelectionBone(GenerationContext.CompilationContext.GetBoneUnique(TypedNodeClip->BoneName), TypedNodeClip->MaxEffectRadius);
 
 		ClipNode->MultipleTagsPolicy = TypedNodeClip->MultipleTagPolicy;
-		ClipNode->RequiredTags = TypedNodeClip->RequiredTags;
+		ClipNode->RequiredTags = GetNodeRequiredTags(*TypedNodeClip, GenerationContext);
 
 		ClipNode->Parameters.FaceCullStrategy = TypedNodeClip->FaceCullStrategy;
 
@@ -143,7 +177,7 @@ mu::Ptr<mu::NodeModifier> GenerateMutableSourceModifier(const UEdGraphPin * Pin,
 		}
 	
 		ClipNode->MultipleTagsPolicy = TypedNodeClipDeform->MultipleTagPolicy;
-		ClipNode->RequiredTags = TypedNodeClipDeform->RequiredTags;
+		ClipNode->RequiredTags = GetNodeRequiredTags(*TypedNodeClipDeform, GenerationContext);
 
 		GenerationContext.MeshGenerationFlags.Pop();
 	}
@@ -157,7 +191,7 @@ mu::Ptr<mu::NodeModifier> GenerateMutableSourceModifier(const UEdGraphPin * Pin,
 
 		// MeshClipWithMesh can be connected to multiple objects, so the compiled NodeModifierMeshClipWithMesh
 		// needs to be different for each object. If it were added to the Generated cache, all the objects would get the same.
-		bDoNotAddToGeneratedCache = true;
+		bCacheNode = false;
 
 		mu::Ptr<mu::NodeModifierMeshClipWithMesh> ClipNode = new mu::NodeModifierMeshClipWithMesh();
 		Result = ClipNode;
@@ -189,7 +223,7 @@ mu::Ptr<mu::NodeModifier> GenerateMutableSourceModifier(const UEdGraphPin * Pin,
 		}
 
 		ClipNode->MultipleTagsPolicy = TypedNodeClipMesh->MultipleTagPolicy;
-		ClipNode->RequiredTags = TypedNodeClipMesh->RequiredTags;
+		ClipNode->RequiredTags = GetNodeRequiredTags(*TypedNodeClipMesh, GenerationContext);
 
 		GenerationContext.MeshGenerationFlags.Pop();
 	}
@@ -203,7 +237,7 @@ mu::Ptr<mu::NodeModifier> GenerateMutableSourceModifier(const UEdGraphPin * Pin,
 
 		// This modifier can be connected to multiple objects, so the compiled node
 		// needs to be different for each object. If it were added to the Generated cache, all the objects would get the same.
-		bDoNotAddToGeneratedCache = true;
+		bCacheNode = false;
 
 		mu::Ptr<mu::NodeModifierMeshClipWithUVMask> ClipNode = new mu::NodeModifierMeshClipWithUVMask();
 		Result = ClipNode;
@@ -228,7 +262,7 @@ mu::Ptr<mu::NodeModifier> GenerateMutableSourceModifier(const UEdGraphPin * Pin,
 		ClipNode->LayoutIndex = TypedNodeClipUVMask->UVChannelForMask;
 
 		ClipNode->MultipleTagsPolicy = TypedNodeClipUVMask->MultipleTagPolicy;
-		ClipNode->RequiredTags = TypedNodeClipUVMask->RequiredTags;
+		ClipNode->RequiredTags = GetNodeRequiredTags(*TypedNodeClipUVMask, GenerationContext);
 
 		GenerationContext.MeshGenerationFlags.Pop();
 	}
@@ -245,7 +279,7 @@ mu::Ptr<mu::NodeModifier> GenerateMutableSourceModifier(const UEdGraphPin * Pin,
 		const int32 LOD = Node->IsAffectedByLOD() ? GenerationContext.CurrentLOD : 0;
 
 		SurfNode->MultipleTagsPolicy = TypedNodeExt->MultipleTagPolicy;
-		SurfNode->RequiredTags = TypedNodeExt->RequiredTags;
+		SurfNode->RequiredTags = GetNodeRequiredTags(*TypedNodeExt, GenerationContext);
 
 		// Is this enough? Should we try to narrow down with potential mesh sections modified by this?
 		int32 LODCount = GenerationContext.NumLODs[GenerationContext.CurrentMeshComponent];
@@ -271,7 +305,7 @@ mu::Ptr<mu::NodeModifier> GenerateMutableSourceModifier(const UEdGraphPin * Pin,
 				if (ConnectedPin)
 				{
 					//NOTE: This is the same is done in GenerateMutableSourceSurface. 
-					if (const UEdGraphPin* SkeletalMeshPin = FindMeshBaseSource(*ConnectedPin, false))
+					if (const UEdGraphPin* SkeletalMeshPin = FindMeshBaseSource(*ConnectedPin, false, &GenerationContext.MacroNodesStack))
 					{
 						FSkeletalMaterial* SkeletalMaterial = nullptr;
 						const FSkelMeshSection* ReferenceSkelMeshSection = nullptr;
@@ -373,6 +407,31 @@ mu::Ptr<mu::NodeModifier> GenerateMutableSourceModifier(const UEdGraphPin * Pin,
 		}
 
 		SurfNode->EnableTags = TypedNodeExt->Tags;
+		
+		// Add Macro linked Tags
+		if (TypedNodeExt->IsInMacro())
+		{
+			UEdGraphPin* TagsPin = TypedNodeExt->FindPin(FString("Enable Tags"), EEdGraphPinDirection::EGPD_Input);
+			check(TagsPin);
+
+			for (const UEdGraphPin* StringPin : FollowInputPinArray(*TagsPin))
+			{
+				const UEdGraphPin* SourceStringPin = GraphTraversal::FindStringSource(*StringPin, &GenerationContext.MacroNodesStack);
+
+				if (SourceStringPin)
+				{
+					if (const UCustomizableObjectNodeStringConstant* StringNode = Cast<UCustomizableObjectNodeStringConstant>(SourceStringPin->GetOwningNode()))
+					{
+						SurfNode->EnableTags.Add(StringNode->Value);
+					}
+				}
+				else
+				{
+					GenerationContext.Log(LOCTEXT("MorphStringNodeFailed", "Could not found a linked String node."), Node);
+				}
+			}
+		}
+
 		SurfNode->EnableTags.AddUnique(TypedNodeExt->GetInternalTag());
 
 		GenerationContext.MeshGenerationFlags.Pop();
@@ -392,7 +451,7 @@ mu::Ptr<mu::NodeModifier> GenerateMutableSourceModifier(const UEdGraphPin * Pin,
 		Result = SurfNode;
 
 		SurfNode->MultipleTagsPolicy = TypedNodeRem->MultipleTagPolicy;
-		SurfNode->RequiredTags = TypedNodeRem->RequiredTags;
+		SurfNode->RequiredTags = GetNodeRequiredTags(*TypedNodeRem, GenerationContext);
 
 		if (const UEdGraphPin* ConnectedPin = FollowInputPin(*TypedNodeRem->RemoveMeshPin()))
 		{
@@ -431,7 +490,7 @@ mu::Ptr<mu::NodeModifier> GenerateMutableSourceModifier(const UEdGraphPin * Pin,
 		ClipNode->FaceCullStrategy = TypedNodeRemBlocks->FaceCullStrategy;
 
 		ClipNode->MultipleTagsPolicy = TypedNodeRemBlocks->MultipleTagPolicy;
-		ClipNode->RequiredTags = TypedNodeRemBlocks->RequiredTags;
+		ClipNode->RequiredTags = GetNodeRequiredTags(*TypedNodeRemBlocks, GenerationContext);
 
 		bool bWasEmpty = false;
 		mu::Ptr<mu::NodeLayout> SourceLayout = CreateMutableLayoutNode(GenerationContext, TypedNodeRemBlocks->Layout, true, bWasEmpty);
@@ -452,7 +511,7 @@ mu::Ptr<mu::NodeModifier> GenerateMutableSourceModifier(const UEdGraphPin * Pin,
 		Result = SurfNode;
 
 		SurfNode->MultipleTagsPolicy = TypedNodeEdit->MultipleTagPolicy;
-		SurfNode->RequiredTags = TypedNodeEdit->RequiredTags;
+		SurfNode->RequiredTags = GetNodeRequiredTags(*TypedNodeEdit, GenerationContext);
 
 		// Is this enough? Should we try to narrow down with potential mesh sections modified by this?
 		int32 LODCount = GenerationContext.NumLODs[GenerationContext.CurrentMeshComponent];
@@ -532,9 +591,29 @@ mu::Ptr<mu::NodeModifier> GenerateMutableSourceModifier(const UEdGraphPin * Pin,
 		SurfNode->bApplyBeforeNormalOperations = true;
 
 		SurfNode->MultipleTagsPolicy = TypedNodeMorph->MultipleTagPolicy;
-		SurfNode->RequiredTags = TypedNodeMorph->RequiredTags;
+		SurfNode->RequiredTags = GetNodeRequiredTags(*TypedNodeMorph, GenerationContext);
 
-		SurfNode->MeshMorph = TypedNodeMorph->MorphTargetName;
+		FString MorphTargetName = TypedNodeMorph->MorphTargetName;
+
+		if (UEdGraphPin* MorphTargetNamePin = Node->FindPin(FString("Morph Target Name"), EEdGraphPinDirection::EGPD_Input))
+		{
+			const UEdGraphPin* ConnectedStringPin = FollowInputPin(*MorphTargetNamePin);
+			const UEdGraphPin* SourceStringPin = ConnectedStringPin ? GraphTraversal::FindStringSource(*ConnectedStringPin, &GenerationContext.MacroNodesStack) : nullptr;
+
+			if (SourceStringPin)
+			{
+				if (const UCustomizableObjectNodeStringConstant* StringNode = Cast<UCustomizableObjectNodeStringConstant>(SourceStringPin->GetOwningNode()))
+				{
+					MorphTargetName = StringNode->Value;
+				}
+			}
+			else
+			{
+				GenerationContext.Log(LOCTEXT("MorphStringNodeFailed", "Could not found a linked String node."), Node);
+			}
+		}
+
+		SurfNode->MeshMorph = MorphTargetName;
 
 		if (const UEdGraphPin* ConnectedPin = FollowInputPin(*TypedNodeMorph->FactorPin()))
 		{
@@ -590,7 +669,7 @@ mu::Ptr<mu::NodeModifier> GenerateMutableSourceModifier(const UEdGraphPin * Pin,
 
 		// MeshTransformInMesh can be connected to multiple objects, so the compiled NodeModifierMeshTransformInMesh
 		// needs to be different for each object. If it were added to the Generated cache, all the objects would get the same.
-		bDoNotAddToGeneratedCache = true;
+		bCacheNode = false;
 
 		mu::Ptr<mu::NodeModifierMeshTransformInMesh> TransformNode = new mu::NodeModifierMeshTransformInMesh();
 		Result = TransformNode;
@@ -620,11 +699,73 @@ mu::Ptr<mu::NodeModifier> GenerateMutableSourceModifier(const UEdGraphPin * Pin,
 		}
 
 		TransformNode->MultipleTagsPolicy = TypedNodeTransformMesh->MultipleTagPolicy;
-		TransformNode->RequiredTags = TypedNodeTransformMesh->RequiredTags;
+		TransformNode->RequiredTags = GetNodeRequiredTags(*TypedNodeTransformMesh, GenerationContext);
 
 		GenerationContext.MeshGenerationFlags.Pop();
 	}
 
+	else if (const UCustomizableObjectNodeMacroInstance* TypedNodeMacro = Cast<UCustomizableObjectNodeMacroInstance>(Node))
+	{
+		bCacheNode = false;
+
+		if (const UEdGraphPin* OutputPin = TypedNodeMacro->GetMacroIOPin(ECOMacroIOType::COMVT_Output, Pin->PinName))
+		{
+			if (const UEdGraphPin* FollowPin = FollowInputPin(*OutputPin))
+			{
+				GenerationContext.MacroNodesStack.Push(TypedNodeMacro);
+				Result = GenerateMutableSourceModifier(FollowPin, GenerationContext);
+				GenerationContext.MacroNodesStack.Pop();
+			}
+			else
+			{
+				FString Msg = FString::Printf(TEXT("Macro Pin [%s] not linked."), *Pin->PinName.ToString());
+				GenerationContext.Log(FText::FromString(Msg), Node);
+			}
+		}
+		else
+		{
+			FString Msg = FString::Printf(TEXT("Macro Output node does not contain a pin with name [%s]."), *Pin->PinName.ToString());
+			GenerationContext.Log(FText::FromString(Msg), Node);
+		}
+	}
+
+	else if (const UCustomizableObjectNodeTunnel* TypedNodeTunnel = Cast<UCustomizableObjectNodeTunnel>(Node))
+	{
+		check(TypedNodeTunnel->bIsInputNode);
+		check(GenerationContext.MacroNodesStack.Num());
+
+		bCacheNode = false;
+
+		const UCustomizableObjectNodeMacroInstance* MacroInstanceNode = GenerationContext.MacroNodesStack.Pop();
+		check(MacroInstanceNode);
+
+		if (const UEdGraphPin* InputPin = MacroInstanceNode->FindPin(Pin->PinName, EEdGraphPinDirection::EGPD_Input))
+		{
+			if (const UEdGraphPin* FollowPin = FollowInputPin(*InputPin))
+			{
+				Result = GenerateMutableSourceModifier(FollowPin, GenerationContext);
+			}
+			else
+			{
+				FString Msg = FString::Printf(TEXT("Macro Instance Pin [%s] is not linked."), *Pin->PinName.ToString());
+				GenerationContext.Log(FText::FromString(Msg), Node);
+			}
+		}
+		else
+		{
+			FString Msg = FString::Printf(TEXT("Macro Instance Node does not contain a pin with name [%s]."), *Pin->PinName.ToString());
+			GenerationContext.Log(FText::FromString(Msg), Node);
+		}
+
+		if (!Result)
+		{
+			GenerationContext.Log(LOCTEXT("ScalarMacroInstanceEmptyAssets", "Something went wrong"), Node);
+		}
+
+		// Push the Macro again even if the result is null
+		GenerationContext.MacroNodesStack.Push(MacroInstanceNode);
+	}
+
 	else
 	{
 		GenerationContext.Log(LOCTEXT("UnimplementedNode", "Node type not implemented yet."), Node);
@@ -639,7 +780,7 @@ mu::Ptr<mu::NodeModifier> GenerateMutableSourceModifier(const UEdGraphPin * Pin,
 		Result->RequiredComponentId = ComponentId;
 	}
 
-	if (!bDoNotAddToGeneratedCache)
+	if (bCacheNode)
 	{
 		GenerationContext.Generated.Add(Key, FGeneratedData(Node, Result));
 	}
diff --git a/Engine/Plugins/Mutable/Source/CustomizableObjectEditor/Private/MuCOE/GenerateMutableSource/GenerateMutableSourceSurface.cpp b/Engine/Plugins/Mutable/Source/CustomizableObjectEditor/Private/MuCOE/GenerateMutableSource/GenerateMutableSourceSurface.cpp
--- a/Engine/Plugins/Mutable/Source/CustomizableObjectEditor/Private/MuCOE/GenerateMutableSource/GenerateMutableSourceSurface.cpp
+++ b/Engine/Plugins/Mutable/Source/CustomizableObjectEditor/Private/MuCOE/GenerateMutableSource/GenerateMutableSourceSurface.cpp
@@ -32,6 +32,7 @@
 #include "MuCOE/Nodes/CustomizableObjectNodeCopyMaterial.h"
 #include "MuCOE/Nodes/CustomizableObjectNodeFloatConstant.h"
 #include "MuCOE/Nodes/CustomizableObjectNodeFloatParameter.h"
+#include "MuCOE/Nodes/CustomizableObjectNodeMacroInstance.h"
 #include "MuCOE/Nodes/CustomizableObjectNodeMaterialVariation.h"
 #include "MuCOE/Nodes/CustomizableObjectNodeMaterialSwitch.h"
 #include "MuCOE/Nodes/CustomizableObjectNodeModifierEditMeshSection.h"
@@ -40,7 +41,9 @@
 #include "MuCOE/Nodes/CustomizableObjectNodeModifierRemoveMesh.h"
 #include "MuCOE/Nodes/CustomizableObjectNodeModifierRemoveMeshBlocks.h"
 #include "MuCOE/Nodes/CustomizableObjectNodeSkeletalMesh.h"
+#include "MuCOE/Nodes/CustomizableObjectNodeStringConstant.h"
 #include "MuCOE/Nodes/CustomizableObjectNodeTable.h"
+#include "MuCOE/Nodes/CustomizableObjectNodeTunnel.h"
 #include "MuT/NodeImageFormat.h"
 #include "MuT/NodeImageMipmap.h"
 #include "MuT/NodeImageNormalComposite.h"
@@ -69,6 +72,10 @@ mu::Ptr<mu::NodeSurface> GenerateMutableSourceSurface(const UEdGraphPin * Pin, F
 
 	UCustomizableObjectNode* Node = CastChecked<UCustomizableObjectNode>(Pin->GetOwningNode());
 
+	// Bool that determines if a node can be added to the cache of nodes.
+	// Most nodes need to be added to the cache but there are some that don't. For exampel, MacroInstanceNodes
+	bool bCacheNode = true;
+
 	const FGeneratedKey Key(reinterpret_cast<void*>(&GenerateMutableSourceSurface), *Pin, *Node, GenerationContext, true);
 	if (const FGeneratedData* Generated = GenerationContext.Generated.Find(Key))
 	{
@@ -96,7 +103,7 @@ mu::Ptr<mu::NodeSurface> GenerateMutableSourceSurface(const UEdGraphPin * Pin, F
 		// TODO: delete the if clause and the warning when static meshes are operational again
 		if (ConnectedMaterialPin)
 		{
-			if (const UEdGraphPin* StaticMeshPin = FindMeshBaseSource(*ConnectedMaterialPin, true))
+			if (const UEdGraphPin* StaticMeshPin = FindMeshBaseSource(*ConnectedMaterialPin, true, &GenerationContext.MacroNodesStack))
 			{
 				const UCustomizableObjectNode* StaticMeshNode = CastChecked<UCustomizableObjectNode>(StaticMeshPin->GetOwningNode());
 				GenerationContext.Log(LOCTEXT("UnsupportedStaticMeshes", "Static meshes are currently not supported as material meshes"), StaticMeshNode);
@@ -127,7 +134,7 @@ mu::Ptr<mu::NodeSurface> GenerateMutableSourceSurface(const UEdGraphPin * Pin, F
 
 			if (ConnectedMaterialPin)
 			{
-				if (const UEdGraphPin* SkeletalMeshPin = FindMeshBaseSource(*ConnectedMaterialPin, false))
+				if (const UEdGraphPin* SkeletalMeshPin = FindMeshBaseSource(*ConnectedMaterialPin, false, &GenerationContext.MacroNodesStack))
 				{
 					FSkeletalMaterial* SkeletalMaterial = nullptr;
 					const FSkelMeshSection* ReferenceSkelMeshSection = nullptr;
@@ -816,6 +823,32 @@ mu::Ptr<mu::NodeSurface> GenerateMutableSourceSurface(const UEdGraphPin * Pin, F
 			SurfNode->Tags.AddUnique( TypedNodeMat->GetInternalTag() );
 		}
 
+		// Add Macro linked Tags
+		if (TypedNodeMat->IsInMacro())
+		{
+			UEdGraphPin* TagsPin = TypedNodeMat->GetEnableTagsPin();
+			
+			if (TagsPin)
+			{
+				for (const UEdGraphPin* StringPin : FollowInputPinArray(*TagsPin))
+				{
+					const UEdGraphPin* SourceStringPin = GraphTraversal::FindStringSource(*StringPin, &GenerationContext.MacroNodesStack);
+
+					if (SourceStringPin)
+					{
+						if (const UCustomizableObjectNodeStringConstant* StringNode = Cast<UCustomizableObjectNodeStringConstant>(SourceStringPin->GetOwningNode()))
+						{
+							SurfNode->Tags.AddUnique(StringNode->Value);
+						}
+					}
+					else
+					{
+						GenerationContext.Log(LOCTEXT("MorphStringNodeFailed", "Could not found a linked String node."), Node);
+					}
+				}
+			}
+		}
+
 		// If an alternate resolution for a particular state is present, clone the surface node, add the image resizing and inject the surface variation node
 		if (TextureNameToProjectionResFactor.Num() > 0 && !AlternateResStateName.IsEmpty())
 		{
@@ -914,7 +947,7 @@ mu::Ptr<mu::NodeSurface> GenerateMutableSourceSurface(const UEdGraphPin * Pin, F
 
 			if (UEdGraphPin* VariationPin = TypedNodeVar->VariationPin(VariationIndex))
 			{
-				SurfNode->Variations[VariationIndex].Tag = TypedNodeVar->GetVariation(VariationIndex).Tag;
+				SurfNode->Variations[VariationIndex].Tag = GetVariationTag(*TypedNodeVar, VariationIndex, GenerationContext);
 				for (const UEdGraphPin* ConnectedPin : FollowInputPinArray(*VariationPin))
 				{
 					// Is it a modifier?
@@ -1006,6 +1039,68 @@ mu::Ptr<mu::NodeSurface> GenerateMutableSourceSurface(const UEdGraphPin * Pin, F
 		}(); // invoke lambda.
 	}
 
+	else if (const UCustomizableObjectNodeMacroInstance* TypedNodeMacro = Cast<UCustomizableObjectNodeMacroInstance>(Node))
+	{
+		bCacheNode = false;
+
+		if (const UEdGraphPin* OutputPin = TypedNodeMacro->GetMacroIOPin(ECOMacroIOType::COMVT_Output, Pin->PinName))
+		{
+			if (const UEdGraphPin* FollowPin = FollowInputPin(*OutputPin))
+			{
+				GenerationContext.MacroNodesStack.Push(TypedNodeMacro);
+				Result = GenerateMutableSourceSurface(FollowPin, GenerationContext);
+				GenerationContext.MacroNodesStack.Pop();
+			}
+			else
+			{
+				FString Msg = FString::Printf(TEXT("Macro Pin [%s] not linked."), *Pin->PinName.ToString());
+				GenerationContext.Log(FText::FromString(Msg), Node);
+			}
+		}
+		else
+		{
+			FString Msg = FString::Printf(TEXT("Macro Output node does not contain a pin with name [%s]."), *Pin->PinName.ToString());
+			GenerationContext.Log(FText::FromString(Msg), Node);
+		}
+	}
+
+	else if (const UCustomizableObjectNodeTunnel* TypedNodeTunnel = Cast<UCustomizableObjectNodeTunnel>(Node))
+	{
+		check(TypedNodeTunnel->bIsInputNode);
+		check(GenerationContext.MacroNodesStack.Num());
+
+		bCacheNode = false;
+
+		const UCustomizableObjectNodeMacroInstance* MacroInstanceNode = GenerationContext.MacroNodesStack.Pop();
+		check(MacroInstanceNode);
+
+		if (const UEdGraphPin* InputPin = MacroInstanceNode->FindPin(Pin->PinName, EEdGraphPinDirection::EGPD_Input))
+		{
+			if (const UEdGraphPin* FollowPin = FollowInputPin(*InputPin))
+			{
+				Result = GenerateMutableSourceSurface(FollowPin, GenerationContext);
+			}
+			else
+			{
+				FString Msg = FString::Printf(TEXT("Macro Instance Pin [%s] is not linked."), *Pin->PinName.ToString());
+				GenerationContext.Log(FText::FromString(Msg), Node);
+			}
+		}
+		else
+		{
+			FString Msg = FString::Printf(TEXT("Macro Instance Node does not contain a pin with name [%s]."), *Pin->PinName.ToString());
+			GenerationContext.Log(FText::FromString(Msg), Node);
+		}
+
+		if (!Result)
+		{
+			GenerationContext.Log(LOCTEXT("ScalarMacroInstanceEmptyAssets", "Something went wrong"), Node);
+		}
+
+		// Push the Macro again even if the result is null
+		GenerationContext.MacroNodesStack.Push(MacroInstanceNode);
+	}
+
 	else
 	{
 		GenerationContext.Log(LOCTEXT("UnimplementedNode", "Node type not implemented yet."), Node);
@@ -1017,8 +1112,11 @@ mu::Ptr<mu::NodeSurface> GenerateMutableSourceSurface(const UEdGraphPin * Pin, F
 		Result->SetMessageContext(Node);
 	}
 
-	GenerationContext.Generated.Add(Key, FGeneratedData(Node, Result));
-	GenerationContext.GeneratedNodes.Add(Node);
+	if (bCacheNode)
+	{
+		GenerationContext.Generated.Add(Key, FGeneratedData(Node, Result));
+		GenerationContext.GeneratedNodes.Add(Node);
+	}
 
 	return Result;
 }
diff --git a/Engine/Plugins/Mutable/Source/CustomizableObjectEditor/Private/MuCOE/GraphTraversal.cpp b/Engine/Plugins/Mutable/Source/CustomizableObjectEditor/Private/MuCOE/GraphTraversal.cpp
--- a/Engine/Plugins/Mutable/Source/CustomizableObjectEditor/Private/MuCOE/GraphTraversal.cpp
+++ b/Engine/Plugins/Mutable/Source/CustomizableObjectEditor/Private/MuCOE/GraphTraversal.cpp
@@ -10,6 +10,7 @@
 #include "MuCOE/Nodes/CustomizableObjectNodeExposePin.h"
 #include "MuCOE/Nodes/CustomizableObjectNodeModifierExtendMeshSection.h"
 #include "MuCOE/Nodes/CustomizableObjectNodeExternalPin.h"
+#include "MuCOE/Nodes/CustomizableObjectNodeMacroInstance.h"
 #include "MuCOE/Nodes/CustomizableObjectNodeMaterial.h"
 #include "MuCOE/Nodes/CustomizableObjectNodeMaterialVariation.h"
 #include "MuCOE/Nodes/CustomizableObjectNodeMeshMorph.h"
@@ -22,7 +23,9 @@
 #include "MuCOE/Nodes/CustomizableObjectNodeObject.h"
 #include "MuCOE/Nodes/CustomizableObjectNodeSkeletalMesh.h"
 #include "MuCOE/Nodes/CustomizableObjectNodeStaticMesh.h"
+#include "MuCOE/Nodes/CustomizableObjectNodeStringConstant.h"
 #include "MuCOE/Nodes/CustomizableObjectNodeTable.h"
+#include "MuCOE/Nodes/CustomizableObjectNodeTunnel.h"
 #include "MuCOE/Nodes/CustomizableObjectNodeAnimationPose.h"
 #include "MuCOE/Nodes/CustomizableObjectNodeObjectGroup.h"
 #include "MuCOE/Nodes/CustomizableObjectNodeReroute.h"
@@ -368,6 +371,11 @@ UCustomizableObject* GetFullGraphRootObject(const UCustomizableObjectNodeObject*
 
 UCustomizableObject* GraphTraversal::GetObject(const UCustomizableObjectNode& Node)
 {
+	if (Node.IsInMacro())
+	{
+		return nullptr;
+	}
+
 	return CastChecked<UCustomizableObject>(Node.GetGraph()->GetOuter());
 }
 
@@ -460,6 +468,129 @@ void GraphTraversal::VisitNodes(UCustomizableObjectNode& StartNode, const TFunct
 }
 
 
+void GraphTraversal::RecursiveVisitNodes(UCustomizableObjectNode& CurrentNode, const TFunction<void(UCustomizableObjectNode&)>& VisitFunction,
+	TArray<const UCustomizableObjectNodeMacroInstance*>& MacroContext, TSet<UCustomizableObjectNode*>& VisitedNodes, const TMultiMap<FGuid, UCustomizableObjectNodeObject*>* ObjectGroupMap)
+{
+	// Check if we already visited this node
+	if (VisitedNodes.Contains(&CurrentNode))
+	{
+		return;
+	}
+
+	TMultiMap<FGuid, UCustomizableObjectNodeObject*> LocalObjectGroupMap;
+	if (!ObjectGroupMap)
+	{
+		LocalObjectGroupMap = GetNodeGroupObjectNodeMapping(GetObject(CurrentNode));
+		ObjectGroupMap = &LocalObjectGroupMap;
+	}
+
+	// Add it to the cache of already visited nodes
+	VisitedNodes.Add(&CurrentNode);
+
+	// Visit the current Node
+	VisitFunction(CurrentNode);
+
+	// Iterate through all nodes linked to the current node
+	for (UEdGraphPin* Pin : CurrentNode.GetAllNonOrphanPins())
+	{
+		if (Pin->Direction != EGPD_Input)
+		{
+			continue;
+		}
+
+		for (UEdGraphPin* ConnectedPin : FollowInputPinArray(*Pin))
+		{
+			UEdGraphNode* ConnectedNode = ConnectedPin->GetOwningNode();
+
+			if (UCustomizableObjectNodeObjectGroup* ObjectGroupNode = Cast<UCustomizableObjectNodeObjectGroup>(ConnectedNode))
+			{
+				TArray<UCustomizableObjectNodeObject*> ChildObjectNodes;
+				ObjectGroupMap->MultiFind(ObjectGroupNode->NodeGuid, ChildObjectNodes);
+
+				for (UCustomizableObjectNode* ChildObjectNode : ChildObjectNodes)
+				{
+					GraphTraversal::RecursiveVisitNodes(*ChildObjectNode, VisitFunction, MacroContext, VisitedNodes, ObjectGroupMap);
+				}
+			}
+
+			else if (UCustomizableObjectNodeMacroInstance* MacroInstanceNode = Cast<UCustomizableObjectNodeMacroInstance>(ConnectedNode))
+			{
+				// We visit this node here because we need to pass through it multiple times
+				if (!VisitedNodes.Contains(MacroInstanceNode))
+				{
+					VisitFunction(*MacroInstanceNode);
+					VisitedNodes.Add(MacroInstanceNode);
+				}
+
+				if (const UEdGraphPin* OutputPin = MacroInstanceNode->GetMacroIOPin(ECOMacroIOType::COMVT_Output, ConnectedPin->PinName))
+				{
+					if (UCustomizableObjectNode* OutputNode = Cast<UCustomizableObjectNode>(OutputPin->GetOwningNode()))
+					{
+						// We visit this node here because we need to pass through it multiple times
+						if (!VisitedNodes.Contains(OutputNode))
+						{
+							VisitFunction(*OutputNode);
+							VisitedNodes.Add(OutputNode);
+						}
+					}
+
+					// Going to the node linked to the OutputPin
+					const UEdGraphPin* OutputLinkedPin = FollowInputPin(*OutputPin);
+
+					// We go directly to the node linked to the output node to ensure that is linked and we need to visit it.
+					if (UCustomizableObjectNode* OutputLinkedNode = Cast<UCustomizableObjectNode>(OutputLinkedPin->GetOwningNode()))
+					{
+						MacroContext.Push(MacroInstanceNode);
+						GraphTraversal::RecursiveVisitNodes(*OutputLinkedNode, VisitFunction, MacroContext, VisitedNodes, ObjectGroupMap);
+						MacroContext.Pop();
+					}
+				}
+			}
+
+			else if (UCustomizableObjectNodeTunnel* IOMacroNode = Cast<UCustomizableObjectNodeTunnel>(ConnectedNode))
+			{
+				// Output Nodes are already visited
+				if (IOMacroNode->bIsInputNode)
+				{
+					check(MacroContext.Num());
+
+					// We visit this node here because we need to pass through it multiple times
+					if (!VisitedNodes.Contains(IOMacroNode))
+					{
+						VisitFunction(*IOMacroNode);
+						VisitedNodes.Add(IOMacroNode);
+					}
+
+					// Going to the parent's graph where the current macro was instantiated
+					const UCustomizableObjectNodeMacroInstance* ParentMacroInstanceNode = MacroContext.Pop();
+
+					if (ParentMacroInstanceNode)
+					{
+						if (const UEdGraphPin* InputPin = ParentMacroInstanceNode->FindPin(ConnectedPin->PinName, EEdGraphPinDirection::EGPD_Input))
+						{
+							// Going to the node linked to the Macro Instance
+							const UEdGraphPin* InputLinkedPin = FollowInputPin(*InputPin);
+
+							if (UCustomizableObjectNode* InputLinkedNode = Cast<UCustomizableObjectNode>(InputLinkedPin->GetOwningNode()))
+							{
+								GraphTraversal::RecursiveVisitNodes(*InputLinkedNode, VisitFunction, MacroContext, VisitedNodes, ObjectGroupMap);
+							}
+						}
+					}
+
+					MacroContext.Push(ParentMacroInstanceNode);
+				}
+			}
+
+			else if (UCustomizableObjectNode* Node = Cast<UCustomizableObjectNode>(ConnectedNode))
+			{
+				GraphTraversal::RecursiveVisitNodes(*Node, VisitFunction, MacroContext, VisitedNodes, ObjectGroupMap);
+			}
+		}
+	}
+}
+
+
 UCustomizableObjectNodeObject* GraphTraversal::GetFullGraphRootNode(const UCustomizableObject* Object, TArray<const UCustomizableObject*>& VisitedObjects)
 {
 	if (Object != nullptr)
@@ -490,7 +621,7 @@ UCustomizableObjectNodeObject* GraphTraversal::GetFullGraphRootNode(const UCusto
 }
 
 
-const UEdGraphPin* FindMeshBaseSource(const UEdGraphPin& Pin, const bool bOnlyLookForStaticMesh)
+const UEdGraphPin* FindMeshBaseSource(const UEdGraphPin& Pin, const bool bOnlyLookForStaticMesh, TArray<const UCustomizableObjectNodeMacroInstance*>* MacroContext)
 {
 	check(Pin.Direction == EGPD_Output);
 	check(Pin.PinType.PinCategory == UEdGraphSchema_CustomizableObject::PC_Mesh
@@ -522,7 +653,7 @@ const UEdGraphPin* FindMeshBaseSource(const UEdGraphPin& Pin, const bool bOnlyLo
 	{
 		if (const UEdGraphPin* ConnectedPin = FollowInputPin(*TypedNodeReshape->BaseMeshPin()))
 		{
-			return FindMeshBaseSource(*ConnectedPin, bOnlyLookForStaticMesh);
+			return FindMeshBaseSource(*ConnectedPin, bOnlyLookForStaticMesh, MacroContext);
 		}
 	}
 
@@ -530,7 +661,7 @@ const UEdGraphPin* FindMeshBaseSource(const UEdGraphPin& Pin, const bool bOnlyLo
 	{
 		if (const UEdGraphPin* ConnectedPin = FollowInputPin(*TypedNodeMorph->MeshPin()))
 		{
-			return FindMeshBaseSource(*ConnectedPin, bOnlyLookForStaticMesh);
+			return FindMeshBaseSource(*ConnectedPin, bOnlyLookForStaticMesh, MacroContext);
 		}
 	}
 
@@ -544,7 +675,7 @@ const UEdGraphPin* FindMeshBaseSource(const UEdGraphPin& Pin, const bool bOnlyLo
 				{
 					if (const UEdGraphPin* ConnectedPin = FollowInputPin(*DefaultPin))
 					{
-						return FindMeshBaseSource(*ConnectedPin, bOnlyLookForStaticMesh);
+						return FindMeshBaseSource(*ConnectedPin, bOnlyLookForStaticMesh, MacroContext);
 					}
 				}
 			}
@@ -555,14 +686,14 @@ const UEdGraphPin* FindMeshBaseSource(const UEdGraphPin& Pin, const bool bOnlyLo
 	{
 		if (const UEdGraphPin* ConnectedPin = FollowInputPin(*TypedNodeMeshVar->DefaultPin()))
 		{
-			return FindMeshBaseSource(*ConnectedPin, bOnlyLookForStaticMesh);
+			return FindMeshBaseSource(*ConnectedPin, bOnlyLookForStaticMesh, MacroContext);
 		}
 
 		for (int32 i = 0; i < TypedNodeMeshVar->GetNumVariations(); ++i)
 		{
 			if (const UEdGraphPin* ConnectedPin = FollowInputPin(*TypedNodeMeshVar->VariationPin(i)))
 			{
-				return FindMeshBaseSource(*ConnectedPin, bOnlyLookForStaticMesh);
+				return FindMeshBaseSource(*ConnectedPin, bOnlyLookForStaticMesh, MacroContext);
 			}
 		}
 	}
@@ -571,7 +702,7 @@ const UEdGraphPin* FindMeshBaseSource(const UEdGraphPin& Pin, const bool bOnlyLo
 	{
 		if (const UEdGraphPin* ConnectedPin = FollowInputPin(*TypedNodeMat->GetMeshPin()))
 		{
-			return FindMeshBaseSource(*ConnectedPin, bOnlyLookForStaticMesh);
+			return FindMeshBaseSource(*ConnectedPin, bOnlyLookForStaticMesh, MacroContext);
 		}
 	}
 
@@ -579,7 +710,7 @@ const UEdGraphPin* FindMeshBaseSource(const UEdGraphPin& Pin, const bool bOnlyLo
 	{
 		if (const UEdGraphPin* ConnectedPin = FollowInputPin(*TypedNodeMatVar->DefaultPin()))
 		{
-			return FindMeshBaseSource(*ConnectedPin, bOnlyLookForStaticMesh);
+			return FindMeshBaseSource(*ConnectedPin, bOnlyLookForStaticMesh, MacroContext);
 		}
 	}
 
@@ -587,7 +718,7 @@ const UEdGraphPin* FindMeshBaseSource(const UEdGraphPin& Pin, const bool bOnlyLo
 	{
 		if (const UEdGraphPin* ConnectedPin = FollowInputPin(*TypedNodeExtend->AddMeshPin()))
 		{
-			return FindMeshBaseSource(*ConnectedPin, bOnlyLookForStaticMesh);
+			return FindMeshBaseSource(*ConnectedPin, bOnlyLookForStaticMesh, MacroContext);
 		}
 	}
 
@@ -595,7 +726,7 @@ const UEdGraphPin* FindMeshBaseSource(const UEdGraphPin& Pin, const bool bOnlyLo
 	{
 		if (const UEdGraphPin* ConnectedPin = FollowInputPin(*TypedNodeMorphStackDef->GetMeshPin()))
 		{
-			return FindMeshBaseSource(*ConnectedPin, bOnlyLookForStaticMesh);
+			return FindMeshBaseSource(*ConnectedPin, bOnlyLookForStaticMesh, MacroContext);
 		}
 	}
 
@@ -603,7 +734,7 @@ const UEdGraphPin* FindMeshBaseSource(const UEdGraphPin& Pin, const bool bOnlyLo
 	{
 		if (const UEdGraphPin* ConnectedPin = FollowInputPin(*TypedNodeMorphStackApp->GetMeshPin()))
 		{
-			return FindMeshBaseSource(*ConnectedPin, bOnlyLookForStaticMesh);
+			return FindMeshBaseSource(*ConnectedPin, bOnlyLookForStaticMesh, MacroContext);
 		}
 	}
 
@@ -611,7 +742,7 @@ const UEdGraphPin* FindMeshBaseSource(const UEdGraphPin& Pin, const bool bOnlyLo
 	{
 		if (const UEdGraphPin* ConnectedPin = FollowInputPin(*NodeMeshReshape->BaseMeshPin()))
 		{
-			return FindMeshBaseSource(*ConnectedPin, bOnlyLookForStaticMesh);
+			return FindMeshBaseSource(*ConnectedPin, bOnlyLookForStaticMesh, MacroContext);
 		}
 	}
 
@@ -627,7 +758,7 @@ const UEdGraphPin* FindMeshBaseSource(const UEdGraphPin& Pin, const bool bOnlyLo
 	{
 		if (const UEdGraphPin* ConnectedPin = FollowInputPin(*NodeMeshPose->GetInputMeshPin()))
 		{
-			return FindMeshBaseSource(*ConnectedPin, bOnlyLookForStaticMesh);
+			return FindMeshBaseSource(*ConnectedPin, bOnlyLookForStaticMesh, MacroContext);
 		}
 	}
 
@@ -639,6 +770,55 @@ const UEdGraphPin* FindMeshBaseSource(const UEdGraphPin& Pin, const bool bOnlyLo
 		}
 	}
 
+	else if (const UCustomizableObjectNodeMacroInstance* NodeMacro = Cast<UCustomizableObjectNodeMacroInstance>(Node))
+	{
+		if (const UEdGraphPin* OutputPin = NodeMacro->GetMacroIOPin(ECOMacroIOType::COMVT_Output, Pin.PinName))
+		{
+			if (const UEdGraphPin* ConnectedPin = FollowInputPin(*OutputPin))
+			{
+				TArray<const UCustomizableObjectNodeMacroInstance*> NewMacroContext;
+
+				if (!MacroContext)
+				{
+					MacroContext = &NewMacroContext;
+				}
+
+				MacroContext->Push(NodeMacro);
+				const UEdGraphPin* ReturnPin = FindMeshBaseSource(*ConnectedPin, bOnlyLookForStaticMesh, MacroContext);
+				MacroContext->Pop();
+
+				return ReturnPin;
+			}
+		}
+	}
+
+	else if (const UCustomizableObjectNodeTunnel* NodeTunnel = Cast<UCustomizableObjectNodeTunnel>(Node))
+	{
+		check(NodeTunnel->bIsInputNode);
+
+		const UEdGraphPin* ReturnPin = nullptr;
+
+		if (MacroContext && MacroContext->Num())
+		{
+			const UCustomizableObjectNodeMacroInstance* MacroInstanceNode = MacroContext->Pop();
+
+			if (MacroInstanceNode)
+			{
+				if (const UEdGraphPin* InputPin = MacroInstanceNode->FindPin(Pin.PinName, EEdGraphPinDirection::EGPD_Input))
+				{
+					if (const UEdGraphPin* ConnectedPin = FollowInputPin(*InputPin))
+					{
+						ReturnPin = FindMeshBaseSource(*ConnectedPin, bOnlyLookForStaticMesh, MacroContext);
+					}
+				}
+			}
+
+			MacroContext->Push(MacroInstanceNode);
+		}
+
+		return ReturnPin;
+	}
+
 	else
 	{
 		unimplemented(); // Case missing.
@@ -777,3 +957,68 @@ void NodePinConnectionListChanged(const TArray<UEdGraphPin*>& Pins)
 		Pair.Key->NodeConnectionListChanged();
 	}
 }
+
+
+const UEdGraphPin* GraphTraversal::FindStringSource(const UEdGraphPin& Pin, TArray<const UCustomizableObjectNodeMacroInstance*>* MacroContext)
+{
+	const UEdGraphPin* ReturnPin = nullptr;
+	const UEdGraphNode* Node = Pin.GetOwningNode();
+	check(Node);
+
+	if (Cast<UCustomizableObjectNodeStringConstant>(Node))
+	{
+		return &Pin;
+	}
+
+	else if (const UCustomizableObjectNodeMacroInstance* NodeMacro = Cast<UCustomizableObjectNodeMacroInstance>(Node))
+	{
+		if (const UEdGraphPin* OutputPin = NodeMacro->GetMacroIOPin(ECOMacroIOType::COMVT_Output, Pin.PinName))
+		{
+			if (const UEdGraphPin* ConnectedPin = FollowInputPin(*OutputPin))
+			{
+				TArray<const UCustomizableObjectNodeMacroInstance*> NewMacroContext;
+
+				if (!MacroContext)
+				{
+					MacroContext = &NewMacroContext;
+				}
+
+				MacroContext->Push(NodeMacro);
+				ReturnPin = GraphTraversal::FindStringSource(*ConnectedPin, MacroContext);
+				MacroContext->Pop();
+
+				return ReturnPin;
+			}
+		}
+	}
+
+	else if (const UCustomizableObjectNodeTunnel* NodeTunnel = Cast<UCustomizableObjectNodeTunnel>(Node))
+	{
+		check(NodeTunnel->bIsInputNode);
+		check(MacroContext->Num());
+
+		const UCustomizableObjectNodeMacroInstance* MacroInstanceNode = MacroContext->Pop();
+
+		if (MacroInstanceNode)
+		{
+			if (const UEdGraphPin* InputPin = MacroInstanceNode->FindPin(Pin.PinName, EEdGraphPinDirection::EGPD_Input))
+			{
+				if (const UEdGraphPin* ConnectedPin = FollowInputPin(*InputPin))
+				{
+					ReturnPin = GraphTraversal::FindStringSource(*ConnectedPin, MacroContext);
+				}
+			}
+		}
+
+		MacroContext->Push(MacroInstanceNode);
+
+		return ReturnPin;
+	}
+
+	else
+	{
+		unimplemented(); // Case missing.
+	}
+
+	return nullptr;
+}
diff --git a/Engine/Plugins/Mutable/Source/CustomizableObjectEditor/Private/MuCOE/GraphTraversal.h b/Engine/Plugins/Mutable/Source/CustomizableObjectEditor/Private/MuCOE/GraphTraversal.h
--- a/Engine/Plugins/Mutable/Source/CustomizableObjectEditor/Private/MuCOE/GraphTraversal.h
+++ b/Engine/Plugins/Mutable/Source/CustomizableObjectEditor/Private/MuCOE/GraphTraversal.h
@@ -9,6 +9,7 @@ class UCustomizableObjectNodeObject;
 class UEdGraph;
 class UEdGraphPin;
 class UCustomizableObjectNode;
+class UCustomizableObjectNodeMacroInstance;
 
 
 /** Follow the given pin returning its connected pin.
@@ -67,7 +68,9 @@ UCustomizableObject* GetFullGraphRootObject(const UCustomizableObjectNodeObject*
 
 namespace GraphTraversal
 {
-	/** Return Customizable Object of the given node. */
+	/** Return Customizable Object of the given node.
+	* Returns null if the Node belongs to a Macro graph
+	*/
 	UCustomizableObject* GetObject(const UCustomizableObjectNode& Node);
 
 	/** Provided a CO object it provides the root CO it is connected. In other words : it returns the root of the entire
@@ -87,10 +90,19 @@ namespace GraphTraversal
 	
 	/** Return the full graph Customizable Object Node root of the node given as parameter */
 	UCustomizableObjectNodeObject* GetFullGraphRootNode(const UCustomizableObject* Object, TArray<const UCustomizableObject*>& VisitedObjects);
+
+	void RecursiveVisitNodes(UCustomizableObjectNode& CurrentNode, const TFunction<void(UCustomizableObjectNode&)>& VisitFunction,
+		TArray<const UCustomizableObjectNodeMacroInstance*>& MacroContext, TSet<UCustomizableObjectNode*>& VisitedNodes, const TMultiMap<FGuid, UCustomizableObjectNodeObject*>* ObjectGroupMap = nullptr);
+
+	/** Given an output pin, return the output pin where the string node is located.
+	* @param Pin: Input pin where the search starts.
+	* @param MacroContext: Copy of the current macro context.
+	*/
+	const UEdGraphPin* FindStringSource(const UEdGraphPin& Pin, TArray<const UCustomizableObjectNodeMacroInstance*>* MacroContext);
 }
 
 /** Given an output pin, return the output pin where the mesh is located. */
-const UEdGraphPin* FindMeshBaseSource(const UEdGraphPin& Pin, const bool bOnlyLookForStaticMesh);
+const UEdGraphPin* FindMeshBaseSource(const UEdGraphPin& Pin, const bool bOnlyLookForStaticMesh, TArray<const UCustomizableObjectNodeMacroInstance*>* MacroContext = nullptr);
 
 /** Return the mapping of Group Object Nodes to Child Object Nodes of the given hierarchy.
  * @param Object Child or root Object to start the search from. */
@@ -107,4 +119,4 @@ namespace GraphTraversal
 
 /** For each given pin, call PinConnectionListChanged and NodeConnectionListChanged in the correct order.
  *  Order: For each node, first call all PinConnectionListChanged, then NodeConnectionListChanged. */
-void NodePinConnectionListChanged(const TArray<UEdGraphPin*>& Pins);
\ No newline at end of file
+void NodePinConnectionListChanged(const TArray<UEdGraphPin*>& Pins);
diff --git a/Engine/Plugins/Mutable/Source/CustomizableObjectEditor/Private/MuCOE/ICustomizableObjectEditor.h b/Engine/Plugins/Mutable/Source/CustomizableObjectEditor/Private/MuCOE/ICustomizableObjectEditor.h
deleted file mode 100644
diff --git a/Engine/Plugins/Mutable/Source/CustomizableObjectEditor/Private/MuCOE/Nodes/CustomizableObjectNode.cpp b/Engine/Plugins/Mutable/Source/CustomizableObjectEditor/Private/MuCOE/Nodes/CustomizableObjectNode.cpp
--- a/Engine/Plugins/Mutable/Source/CustomizableObjectEditor/Private/MuCOE/Nodes/CustomizableObjectNode.cpp
+++ b/Engine/Plugins/Mutable/Source/CustomizableObjectEditor/Private/MuCOE/Nodes/CustomizableObjectNode.cpp
@@ -103,21 +103,21 @@ void UCustomizableObjectNode::PostEditChangeProperty(FPropertyChangedEvent& Prop
 }
 
 
-TSharedPtr<ICustomizableObjectEditor> UCustomizableObjectNode::GetGraphEditor() const
+TSharedPtr<FCustomizableObjectGraphEditorToolkit> UCustomizableObjectNode::GetGraphEditor() const
 {
 	const UCustomizableObject* CustomizableObject = Cast<UCustomizableObject>(GetCustomizableObjectGraph()->GetOuter());
 
-	TSharedPtr<ICustomizableObjectEditor> CustomizableObjectEditor;
+	TSharedPtr<FCustomizableObjectGraphEditorToolkit> CustomizableObjectEditor;
 	if (CustomizableObject)
 	{
 		TSharedPtr< IToolkit > FoundAssetEditor = FToolkitManager::Get().FindEditorForAsset(CustomizableObject);
 		if (FoundAssetEditor.IsValid())
 		{
-			return StaticCastSharedPtr<ICustomizableObjectEditor>(FoundAssetEditor);
+			return StaticCastSharedPtr<FCustomizableObjectGraphEditorToolkit>(FoundAssetEditor);
 		}
 	}
 
-	return TSharedPtr<ICustomizableObjectEditor>();
+	return TSharedPtr<FCustomizableObjectGraphEditorToolkit>();
 }
 
 
@@ -778,5 +778,15 @@ UCustomizableObjectNodePinData* UCustomizableObjectNode::GetPinData(const UEdGra
 	}
 }
 
+bool UCustomizableObjectNode::IsInMacro() const
+{
+	if (UCustomizableObjectGraph* Graph = Cast<UCustomizableObjectGraph>(GetGraph()))
+	{
+		return Graph->IsMacro();
+	}
+
+	return false;
+}
+
 
 #undef LOCTEXT_NAMESPACE
diff --git a/Engine/Plugins/Mutable/Source/CustomizableObjectEditor/Private/MuCOE/Nodes/CustomizableObjectNodeComponentMeshBase.cpp b/Engine/Plugins/Mutable/Source/CustomizableObjectEditor/Private/MuCOE/Nodes/CustomizableObjectNodeComponentMeshBase.cpp
--- a/Engine/Plugins/Mutable/Source/CustomizableObjectEditor/Private/MuCOE/Nodes/CustomizableObjectNodeComponentMeshBase.cpp
+++ b/Engine/Plugins/Mutable/Source/CustomizableObjectEditor/Private/MuCOE/Nodes/CustomizableObjectNodeComponentMeshBase.cpp
@@ -1,4 +1,4 @@
-﻿// Copyright Epic Games, Inc. All Rights Reserved.
+// Copyright Epic Games, Inc. All Rights Reserved.
 
 #include "MuCOE/Nodes/CustomizableObjectNodeComponentMeshBase.h"
 
@@ -19,7 +19,12 @@ void UCustomizableObjectNodeComponentMeshBase::AllocateDefaultPins(UCustomizable
 		UEdGraphPin* Pin = CustomCreatePin(EGPD_Input, Schema->PC_Material, FName(*LODName), true);
 		LODPins.Add(Pin);
 	}
-	
+
+	if (IsInMacro())
+	{
+		UEdGraphPin* Pin = CustomCreatePin(EGPD_Input, Schema->PC_String, FName("Component Name"));
+	}
+
 	OutputPin = CustomCreatePin(EGPD_Output, Schema->PC_Component, FName(TEXT("Component")));
 }
 
diff --git a/Engine/Plugins/Mutable/Source/CustomizableObjectEditor/Private/MuCOE/Nodes/CustomizableObjectNodeCopyMaterial.cpp b/Engine/Plugins/Mutable/Source/CustomizableObjectEditor/Private/MuCOE/Nodes/CustomizableObjectNodeCopyMaterial.cpp
--- a/Engine/Plugins/Mutable/Source/CustomizableObjectEditor/Private/MuCOE/Nodes/CustomizableObjectNodeCopyMaterial.cpp
+++ b/Engine/Plugins/Mutable/Source/CustomizableObjectEditor/Private/MuCOE/Nodes/CustomizableObjectNodeCopyMaterial.cpp
@@ -135,6 +135,19 @@ bool UCustomizableObjectNodeCopyMaterial::RealMaterialDataHasChanged() const
 }
 
 
+UEdGraphPin* UCustomizableObjectNodeCopyMaterial::GetEnableTagsPin() const
+{
+	if (UCustomizableObjectNodeMaterial* NodeMaterial = GetMaterialNode())
+	{
+		return NodeMaterial->GetEnableTagsPin();
+	}
+	else
+	{
+		return nullptr;
+	}
+}
+
+
 UEdGraphPin* UCustomizableObjectNodeCopyMaterial::GetMeshSectionPin() const
 {
 	UEdGraphPin* Pin = FindPin(CustomizableObjectNodeCopyMaterial::MeshSectionPinName, EEdGraphPinDirection::EGPD_Input);
diff --git a/Engine/Plugins/Mutable/Source/CustomizableObjectEditor/Private/MuCOE/Nodes/CustomizableObjectNodeCopyMaterial.h b/Engine/Plugins/Mutable/Source/CustomizableObjectEditor/Private/MuCOE/Nodes/CustomizableObjectNodeCopyMaterial.h
--- a/Engine/Plugins/Mutable/Source/CustomizableObjectEditor/Private/MuCOE/Nodes/CustomizableObjectNodeCopyMaterial.h
+++ b/Engine/Plugins/Mutable/Source/CustomizableObjectEditor/Private/MuCOE/Nodes/CustomizableObjectNodeCopyMaterial.h
@@ -56,6 +56,7 @@ public:
 	virtual TArray<UCustomizableObjectLayout*> GetLayouts() const override;
 	virtual UEdGraphPin* OutputPin() const override;
 	virtual bool RealMaterialDataHasChanged() const override;
+	virtual UEdGraphPin* GetEnableTagsPin() const override;
 	
 	// Own interface
 	UCustomizableObjectNodeSkeletalMesh* GetMeshNode() const;
diff --git a/Engine/Plugins/Mutable/Source/CustomizableObjectEditor/Private/MuCOE/Nodes/CustomizableObjectNodeMacroInstance.cpp b/Engine/Plugins/Mutable/Source/CustomizableObjectEditor/Private/MuCOE/Nodes/CustomizableObjectNodeMacroInstance.cpp
--- a/Engine/Plugins/Mutable/Source/CustomizableObjectEditor/Private/MuCOE/Nodes/CustomizableObjectNodeMacroInstance.cpp
+++ b/Engine/Plugins/Mutable/Source/CustomizableObjectEditor/Private/MuCOE/Nodes/CustomizableObjectNodeMacroInstance.cpp
@@ -5,6 +5,7 @@
 #include "Editor.h"
 #include "MuCOE/CustomizableObjectMacroLibrary/CustomizableObjectMacroLibraryEditor.h"
 #include "MuCOE/EdGraphSchema_CustomizableObject.h"
+#include "MuCOE/Nodes/CustomizableObjectNodeTunnel.h"
 #include "Subsystems/AssetEditorSubsystem.h"
 
 #define LOCTEXT_NAMESPACE "CustomizableObjectNodeMacroInstance"
@@ -124,7 +125,7 @@ void UCustomizableObjectNodeMacroInstance::JumpToDefinition() const
 	// Set the Parent Macro to edit
 	if (FCustomizableObjectMacroLibraryEditor* Editor = static_cast<FCustomizableObjectMacroLibraryEditor*>(AssetEditor))
 	{
-		Editor->SetSelectedMacro(ParentMacro);
+		Editor->SetSelectedMacro(ParentMacro, true);
 	}
 }
 
@@ -159,4 +160,21 @@ void UCustomizableObjectNodeMacroInstance::PostEditChangeProperty(FPropertyChang
 	}
 }
 
+UEdGraphPin* UCustomizableObjectNodeMacroInstance::GetMacroIOPin(ECOMacroIOType IONodeType, const FName& PinName) const
+{
+	UEdGraphPin* MacroIOPin = nullptr;
+
+	if (ParentMacroLibrary && ParentMacro)
+	{
+		// Input Node contains output pins and output Node contains input pins!
+		EEdGraphPinDirection PinDirection = IONodeType == ECOMacroIOType::COMVT_Input ? EEdGraphPinDirection::EGPD_Output : EEdGraphPinDirection::EGPD_Input;
+		const UCustomizableObjectNodeTunnel* OutputNode = ParentMacro->GetIONode(IONodeType);
+		check(OutputNode);
+
+		MacroIOPin = OutputNode->FindPin(PinName, PinDirection);
+	}
+
+	return (MacroIOPin && !MacroIOPin->bOrphanedPin) ? MacroIOPin : nullptr;
+}
+
 #undef LOCTEXT_NAMESPACE
\ No newline at end of file
diff --git a/Engine/Plugins/Mutable/Source/CustomizableObjectEditor/Private/MuCOE/Nodes/CustomizableObjectNodeMacroInstance.h b/Engine/Plugins/Mutable/Source/CustomizableObjectEditor/Private/MuCOE/Nodes/CustomizableObjectNodeMacroInstance.h
--- a/Engine/Plugins/Mutable/Source/CustomizableObjectEditor/Private/MuCOE/Nodes/CustomizableObjectNodeMacroInstance.h
+++ b/Engine/Plugins/Mutable/Source/CustomizableObjectEditor/Private/MuCOE/Nodes/CustomizableObjectNodeMacroInstance.h
@@ -70,4 +70,8 @@ public:
 	virtual UCustomizableObjectNodeMacroInstanceRemapPins* CreateRemapPinsDefault() const override;
 	virtual void PostEditChangeProperty(FPropertyChangedEvent& PropertyChangedEvent) override;
 
+	// Node MacroInstance Interface
+	/** Returns the pin of the Macro's input/output node with the same name */
+	UEdGraphPin* GetMacroIOPin(ECOMacroIOType IONodeType, const FName& PinName) const;
+
 };
\ No newline at end of file
diff --git a/Engine/Plugins/Mutable/Source/CustomizableObjectEditor/Private/MuCOE/Nodes/CustomizableObjectNodeMaterial.cpp b/Engine/Plugins/Mutable/Source/CustomizableObjectEditor/Private/MuCOE/Nodes/CustomizableObjectNodeMaterial.cpp
--- a/Engine/Plugins/Mutable/Source/CustomizableObjectEditor/Private/MuCOE/Nodes/CustomizableObjectNodeMaterial.cpp
+++ b/Engine/Plugins/Mutable/Source/CustomizableObjectEditor/Private/MuCOE/Nodes/CustomizableObjectNodeMaterial.cpp
@@ -335,6 +335,15 @@ void UCustomizableObjectNodeMaterial::AllocateDefaultPins(UCustomizableObjectNod
 		TableMaterialPin->PinToolTip = "Pin for a Material from a Table Node";
 	}
 
+	if (IsInMacro())
+	{
+		const FString PinFriendlyName = "Enable Tags";
+		const FString PinName = PinFriendlyName + FString(TEXT("_Input_Pin"));
+		UEdGraphPin* TagsPin = CustomCreatePin(EGPD_Input, Schema->PC_String, FName(PinName), true);
+		TagsPin->PinFriendlyName = FText::FromString(PinFriendlyName);
+		TagsPin->bDefaultValueIsIgnored = true;
+	}
+
 	for (const EMaterialParameterType Type : ParameterTypes)
 	{
 		AllocateDefaultParameterPins(Type);
@@ -876,6 +885,29 @@ UEdGraphPin* UCustomizableObjectNodeMaterial::GetMaterialAssetPin() const
 }
 
 
+UEdGraphPin* UCustomizableObjectNodeMaterial::GetEnableTagsPin() const
+{
+	if(IsInMacro())
+	{
+		const FString PinFriendlyName = "Enable Tags";
+		const FString PinName = PinFriendlyName + FString(TEXT("_Input_Pin"));
+
+		UEdGraphPin* Pin = FindPin(PinName, EEdGraphPinDirection::EGPD_Input);
+
+		if (Pin)
+		{
+			return Pin;
+		}
+		else
+		{
+			return FindPin(PinFriendlyName, EEdGraphPinDirection::EGPD_Input);
+		}
+	}
+
+	return nullptr;
+}
+
+
 const UCustomizableObjectNodeMaterial* UCustomizableObjectNodeMaterial::GetMaterialNode() const
 {
 	return this;
diff --git a/Engine/Plugins/Mutable/Source/CustomizableObjectEditor/Private/MuCOE/Nodes/CustomizableObjectNodeMaterial.h b/Engine/Plugins/Mutable/Source/CustomizableObjectEditor/Private/MuCOE/Nodes/CustomizableObjectNodeMaterial.h
--- a/Engine/Plugins/Mutable/Source/CustomizableObjectEditor/Private/MuCOE/Nodes/CustomizableObjectNodeMaterial.h
+++ b/Engine/Plugins/Mutable/Source/CustomizableObjectEditor/Private/MuCOE/Nodes/CustomizableObjectNodeMaterial.h
@@ -152,6 +152,7 @@ public:
 	virtual int32 GetImageUVLayout(int32 ImageIndex) const override;
 	virtual UEdGraphPin* GetMeshPin() const override;
 	virtual UEdGraphPin* GetMaterialAssetPin() const override;
+	virtual UEdGraphPin* GetEnableTagsPin() const override;
 	virtual const UCustomizableObjectNodeMaterial* GetMaterialNode() const override;
 	virtual bool RealMaterialDataHasChanged() const override;
 	virtual FPostImagePinModeChangedDelegate* GetPostImagePinModeChangedDelegate() override;
diff --git a/Engine/Plugins/Mutable/Source/CustomizableObjectEditor/Private/MuCOE/Nodes/CustomizableObjectNodeMaterialBase.h b/Engine/Plugins/Mutable/Source/CustomizableObjectEditor/Private/MuCOE/Nodes/CustomizableObjectNodeMaterialBase.h
--- a/Engine/Plugins/Mutable/Source/CustomizableObjectEditor/Private/MuCOE/Nodes/CustomizableObjectNodeMaterialBase.h
+++ b/Engine/Plugins/Mutable/Source/CustomizableObjectEditor/Private/MuCOE/Nodes/CustomizableObjectNodeMaterialBase.h
@@ -71,6 +71,8 @@ public:
 	virtual UEdGraphPin* GetMeshPin() const PURE_VIRTUAL(UCustomizableObjectNodeMaterial::GetMeshPin, return {}; );
 
 	virtual UEdGraphPin* GetMaterialAssetPin() const PURE_VIRTUAL(UCustomizableObjectNodeMaterial::GetMaterialAssetPin, return {}; );
+
+	virtual UEdGraphPin* GetEnableTagsPin() const PURE_VIRTUAL(UCustomizableObjectNodeMaterial::GetMaterialAssetPin, return {}; );
 	
 	/** Returns the number of Material Parameters. */
 	virtual int32 GetNumParameters(EMaterialParameterType Type) const PURE_VIRTUAL(UCustomizableObjectNodeMaterial::GetNumParameters, return {}; );
diff --git a/Engine/Plugins/Mutable/Source/CustomizableObjectEditor/Private/MuCOE/Nodes/CustomizableObjectNodeMeshClipDeform.cpp b/Engine/Plugins/Mutable/Source/CustomizableObjectEditor/Private/MuCOE/Nodes/CustomizableObjectNodeMeshClipDeform.cpp
--- a/Engine/Plugins/Mutable/Source/CustomizableObjectEditor/Private/MuCOE/Nodes/CustomizableObjectNodeMeshClipDeform.cpp
+++ b/Engine/Plugins/Mutable/Source/CustomizableObjectEditor/Private/MuCOE/Nodes/CustomizableObjectNodeMeshClipDeform.cpp
@@ -18,7 +18,10 @@ void UCustomizableObjectNodeModifierClipDeform::AllocateDefaultPins(UCustomizabl
 	ClipMeshPin->bDefaultValueIsIgnored = true;
 	
 	UEdGraphPin* OutputPin = CustomCreatePin(EGPD_Output, Schema->PC_Modifier, FName("Modifier"));
-	ClipMeshPin->bDefaultValueIsIgnored = true;
+	OutputPin->bDefaultValueIsIgnored = true;
+
+	//Create Node Modifier Common Pins
+	Super::AllocateDefaultPins(RemapPins);
 }
 
 
diff --git a/Engine/Plugins/Mutable/Source/CustomizableObjectEditor/Private/MuCOE/Nodes/CustomizableObjectNodeMeshClipMorph.cpp b/Engine/Plugins/Mutable/Source/CustomizableObjectEditor/Private/MuCOE/Nodes/CustomizableObjectNodeMeshClipMorph.cpp
--- a/Engine/Plugins/Mutable/Source/CustomizableObjectEditor/Private/MuCOE/Nodes/CustomizableObjectNodeMeshClipMorph.cpp
+++ b/Engine/Plugins/Mutable/Source/CustomizableObjectEditor/Private/MuCOE/Nodes/CustomizableObjectNodeMeshClipMorph.cpp
@@ -4,7 +4,7 @@
 
 #include "MuCO/CustomizableObjectCustomVersion.h"
 #include "MuCOE/EdGraphSchema_CustomizableObject.h"
-#include "MuCOE/ICustomizableObjectEditor.h"
+#include "MuCOE/CustomizableObjectEditor.h"
 #include "MuCOE/UnrealEditorPortabilityHelpers.h"
 #include "MuCOE/Nodes/CustomizableObjectNodeObject.h"
 
@@ -102,7 +102,10 @@ void UCustomizableObjectNodeModifierClipMorph::PostEditChangeProperty(FPropertyC
 		ChangeStartOffsetTransform();
 	}
 
-	GetGraphEditor()->ShowGizmoClipMorph(*this);
+	if (TSharedPtr<FCustomizableObjectEditor> Editor = StaticCastSharedPtr<FCustomizableObjectEditor>(GetGraphEditor()))
+	{
+		Editor->ShowGizmoClipMorph(*this);
+	}
 
 	Super::PostEditChangeProperty(PropertyChangedEvent);
 }
@@ -136,8 +139,11 @@ void UCustomizableObjectNodeModifierClipMorph::Serialize(FArchive& Ar)
 void UCustomizableObjectNodeModifierClipMorph::AllocateDefaultPins(UCustomizableObjectNodeRemapPins* RemapPins)
 {
 	const UEdGraphSchema_CustomizableObject* Schema = GetDefault<UEdGraphSchema_CustomizableObject>();
-
+	
 	UEdGraphPin* OutputPin = CustomCreatePin(EGPD_Output, Schema->PC_Modifier, FName("Modifier"));
+
+	//Create Node Modifier Common Pins
+	Super::AllocateDefaultPins(RemapPins);
 }
 
 
@@ -153,7 +159,7 @@ void UCustomizableObjectNodeModifierClipMorph::PinConnectionListChanged(UEdGraph
 
 	if (Pin == OutputPin())
 	{
-		TSharedPtr<ICustomizableObjectEditor> Editor = GetGraphEditor();
+		TSharedPtr<FCustomizableObjectGraphEditorToolkit> Editor = GetGraphEditor();
 
 		if (Editor.IsValid())
 		{
diff --git a/Engine/Plugins/Mutable/Source/CustomizableObjectEditor/Private/MuCOE/Nodes/CustomizableObjectNodeMeshClipWithMesh.cpp b/Engine/Plugins/Mutable/Source/CustomizableObjectEditor/Private/MuCOE/Nodes/CustomizableObjectNodeMeshClipWithMesh.cpp
--- a/Engine/Plugins/Mutable/Source/CustomizableObjectEditor/Private/MuCOE/Nodes/CustomizableObjectNodeMeshClipWithMesh.cpp
+++ b/Engine/Plugins/Mutable/Source/CustomizableObjectEditor/Private/MuCOE/Nodes/CustomizableObjectNodeMeshClipWithMesh.cpp
@@ -4,7 +4,6 @@
 
 #include "MuCOE/CustomizableObjectGraph.h"
 #include "MuCOE/EdGraphSchema_CustomizableObject.h"
-#include "MuCOE/ICustomizableObjectEditor.h"
 #include "MuCO/CustomizableObjectCustomVersion.h"
 #include "MuCOE/GraphTraversal.h"
 
@@ -27,6 +26,9 @@ void UCustomizableObjectNodeModifierClipWithMesh::AllocateDefaultPins(UCustomiza
 	ClipMeshPin->bDefaultValueIsIgnored = true;
 
 	UEdGraphPin* OutputPin_p = CustomCreatePin(EGPD_Output, Schema->PC_Modifier, FName("Modifier"));
+
+	//Create Node Modifier Common Pins
+	Super::AllocateDefaultPins(RemapPins);
 }
 
 
diff --git a/Engine/Plugins/Mutable/Source/CustomizableObjectEditor/Private/MuCOE/Nodes/CustomizableObjectNodeMeshMorph.cpp b/Engine/Plugins/Mutable/Source/CustomizableObjectEditor/Private/MuCOE/Nodes/CustomizableObjectNodeMeshMorph.cpp
--- a/Engine/Plugins/Mutable/Source/CustomizableObjectEditor/Private/MuCOE/Nodes/CustomizableObjectNodeMeshMorph.cpp
+++ b/Engine/Plugins/Mutable/Source/CustomizableObjectEditor/Private/MuCOE/Nodes/CustomizableObjectNodeMeshMorph.cpp
@@ -4,9 +4,10 @@
 
 #include "Engine/SkeletalMesh.h"
 #include "MuCO/CustomizableObjectCustomVersion.h"
+#include "MuCOE/CustomizableObjectGraph.h"
 #include "MuCOE/EdGraphSchema_CustomizableObject.h"
 #include "MuCOE/GraphTraversal.h"
-#include "MuCOE/ICustomizableObjectEditor.h"
+#include "MuCOE/CustomizableObjectMacroLibrary/CustomizableObjectGraphEditorToolkit.h"
 #include "MuCOE/Nodes/CustomizableObjectNodeMeshReshapeCommon.h"
 #include "MuCOE/Nodes/CustomizableObjectNodeSkeletalMesh.h"
 
@@ -37,6 +38,14 @@ void UCustomizableObjectNodeMeshMorph::AllocateDefaultPins(UCustomizableObjectNo
 	PinName = TEXT("Factor");
 	UEdGraphPin* FactorPin = CustomCreatePin(EGPD_Input, Schema->PC_Float, FName(*PinName));
 	FactorPin->bDefaultValueIsIgnored = true;
+
+	// Create this pin only if we are in a Macro graph
+	if (IsInMacro())
+	{
+		PinName = TEXT("Morph Target Name");
+		UEdGraphPin* MorphNamePin = CustomCreatePin(EGPD_Input, Schema->PC_String, FName(*PinName));
+		MorphNamePin->bDefaultValueIsIgnored = true;
+	}
 }
 
 
@@ -119,7 +128,7 @@ void UCustomizableObjectNodeMeshMorph::PinConnectionListChanged(UEdGraphPin* Pin
 
 	if (Pin->Direction == EGPD_Input && Pin->PinName == "Mesh")
 	{
-		TSharedPtr<ICustomizableObjectEditor> Editor = GetGraphEditor();
+		TSharedPtr<FCustomizableObjectGraphEditorToolkit> Editor = GetGraphEditor();
 
 		if (Editor.IsValid())
 		{
diff --git a/Engine/Plugins/Mutable/Source/CustomizableObjectEditor/Private/MuCOE/Nodes/CustomizableObjectNodeMeshMorph.h b/Engine/Plugins/Mutable/Source/CustomizableObjectEditor/Private/MuCOE/Nodes/CustomizableObjectNodeMeshMorph.h
--- a/Engine/Plugins/Mutable/Source/CustomizableObjectEditor/Private/MuCOE/Nodes/CustomizableObjectNodeMeshMorph.h
+++ b/Engine/Plugins/Mutable/Source/CustomizableObjectEditor/Private/MuCOE/Nodes/CustomizableObjectNodeMeshMorph.h
@@ -48,7 +48,7 @@ public:
 
 	void Serialize(FArchive& Ar) override;
 
-	UPROPERTY(Category=CustomizableObject, EditAnywhere)
+	UPROPERTY(Category=MeshMorph, EditAnywhere)
 	FString MorphTargetName;
 
 	/** Experimental - Enable the deformation of the skeleton of the base mesh. */
diff --git a/Engine/Plugins/Mutable/Source/CustomizableObjectEditor/Private/MuCOE/Nodes/CustomizableObjectNodeMeshMorphDetails.cpp b/Engine/Plugins/Mutable/Source/CustomizableObjectEditor/Private/MuCOE/Nodes/CustomizableObjectNodeMeshMorphDetails.cpp
--- a/Engine/Plugins/Mutable/Source/CustomizableObjectEditor/Private/MuCOE/Nodes/CustomizableObjectNodeMeshMorphDetails.cpp
+++ b/Engine/Plugins/Mutable/Source/CustomizableObjectEditor/Private/MuCOE/Nodes/CustomizableObjectNodeMeshMorphDetails.cpp
@@ -33,80 +33,91 @@ void FCustomizableObjectNodeMeshMorphDetails::CustomizeDetails( IDetailLayoutBui
 	}
 
 	IDetailCategoryBuilder& BlocksCategory = DetailBuilder.EditCategory( "Morph Target" );
+	TSharedRef<IPropertyHandle> MorphTargetNameProperty = DetailBuilder.GetProperty("MorphTargetName");
+
 	//BlocksCategory.CategoryIcon( "ActorClassIcon.CustomizableObject" );
 	
 	MorphTargetComboOptions.Empty();
 
 	bool SourceMeshFound = false;
-	if (Node)
+
+	if (!Node)
 	{
-		if (const UEdGraphPin* ConnectedPin = FollowInputPin(*Node->MeshPin()))
+		return;
+	}
+
+	if (const UEdGraphPin* ConnectedPin = FollowInputPin(*Node->MeshPin()))
+	{
+		TSharedPtr<FString> ItemToSelect;
+		const UEdGraphPin* BaseSourcePin = FindMeshBaseSource(*ConnectedPin, false );
+		
+		if ( BaseSourcePin )
 		{
-			SourceMeshFound = true;
-			
-			TSharedPtr<FString> ItemToSelect;
+			USkeletalMesh* SkeletalMesh = nullptr;
+			if ( const UCustomizableObjectNodeSkeletalMesh* TypedNodeSkeletalMesh = Cast<UCustomizableObjectNodeSkeletalMesh>(BaseSourcePin->GetOwningNode()) )
+			{
+				SkeletalMesh = MutablePrivate::LoadObject(TypedNodeSkeletalMesh->SkeletalMesh);
+			}
+			else if (const UCustomizableObjectNodeTable* TypedNodeTable = Cast<UCustomizableObjectNodeTable>(BaseSourcePin->GetOwningNode()))
+			{
+				SkeletalMesh = TypedNodeTable->GetColumnDefaultAssetByType<USkeletalMesh>(BaseSourcePin);
+			}
 
-			const UEdGraphPin* BaseSourcePin = FindMeshBaseSource(*ConnectedPin, false );
-			if ( BaseSourcePin )
+			if (SkeletalMesh)
 			{
-				USkeletalMesh* SkeletalMesh = nullptr;
-				if ( const UCustomizableObjectNodeSkeletalMesh* TypedNodeSkeletalMesh = Cast<UCustomizableObjectNodeSkeletalMesh>(BaseSourcePin->GetOwningNode()) )
-				{
-					SkeletalMesh = MutablePrivate::LoadObject(TypedNodeSkeletalMesh->SkeletalMesh);
-				}
-				else if (const UCustomizableObjectNodeTable* TypedNodeTable = Cast<UCustomizableObjectNodeTable>(BaseSourcePin->GetOwningNode()))
-				{
-					SkeletalMesh = TypedNodeTable->GetColumnDefaultAssetByType<USkeletalMesh>(BaseSourcePin);
-				}
+				SourceMeshFound = true;
 
-				if (SkeletalMesh)
+				for (int m = 0; m < SkeletalMesh->GetMorphTargets().Num(); ++m)
 				{
-					for (int m = 0; m < SkeletalMesh->GetMorphTargets().Num(); ++m)
-					{
-						FString MorphName = *SkeletalMesh->GetMorphTargets()[m]->GetName();
-						MorphTargetComboOptions.Add(MakeShareable(new FString(MorphName)));
+					FString MorphName = *SkeletalMesh->GetMorphTargets()[m]->GetName();
+					MorphTargetComboOptions.Add(MakeShareable(new FString(MorphName)));
 
-						if (Node->MorphTargetName == MorphName)
-						{
-							ItemToSelect = MorphTargetComboOptions.Last();
-						}
+					if (Node->MorphTargetName == MorphName)
+					{
+						ItemToSelect = MorphTargetComboOptions.Last();
 					}
 				}
-			}
 
-			MorphTargetComboOptions.Sort(CompareNames);
+				MorphTargetComboOptions.Sort(CompareNames);
 
-			TSharedRef<IPropertyHandle> MorphTargetNameProperty = DetailBuilder.GetProperty("MorphTargetName");
-
-			BlocksCategory.AddCustomRow( LOCTEXT("MeshMorphDetails_Target", "Target") )
-			[
-				SNew(SProperty, MorphTargetNameProperty)
-				.ShouldDisplayName( false)
-				.CustomWidget()
+				BlocksCategory.AddCustomRow(LOCTEXT("MeshMorphDetails_Target", "Target"))
 				[
-					SNew(SBorder)
-					.BorderImage(UE_MUTABLE_GET_BRUSH("NoBorder"))
-					.Padding(FMargin(0.0f, 0.0f, 10.0f, 0.0f))
+					SNew(SProperty, MorphTargetNameProperty)
+					.ShouldDisplayName(false)
+					.CustomWidget()
 					[
-						SNew(STextComboBox)
-							.OptionsSource( &MorphTargetComboOptions )
-							.InitiallySelectedItem( ItemToSelect )
+						SNew(SBorder)
+						.BorderImage(UE_MUTABLE_GET_BRUSH("NoBorder"))
+						.Padding(FMargin(0.0f, 0.0f, 10.0f, 0.0f))
+						[
+							SNew(STextComboBox)
+							.OptionsSource(&MorphTargetComboOptions)
+							.InitiallySelectedItem(ItemToSelect)
 							.OnSelectionChanged(this, &FCustomizableObjectNodeMeshMorphDetails::OnMorphTargetComboBoxSelectionChanged, MorphTargetNameProperty)
+						]
 					]
-				]
-			];
+				];
+			}
 		}
 	}
-	
+
 	if (!SourceMeshFound)
 	{
-		BlocksCategory.AddCustomRow( LOCTEXT("MeshMorphDetails_Node", "Node") )
-		[
-			SNew( STextBlock )
-			.Text( LOCTEXT( "MeshMorphDetails_NoSource", "No source mesh found." ) )			
-		];
-	}
+		FText Message = Node->IsInMacro() ? LOCTEXT("MeshMorphDetails_PinMessage", "In Mutable Macros, Morph Target Names are deffined through String Nodes.") : LOCTEXT("MeshMorphDetails_NoSource", "No source mesh found.");
 
+		DetailBuilder.EditDefaultProperty(MorphTargetNameProperty)->
+			CustomWidget()
+			.NameContent()
+			[
+				SNew(STextBlock).Text(LOCTEXT("MeshMorphDetails_MorphTargetName", "Morph Target Name")).Font(IDetailLayoutBuilder::GetDetailFont())
+			]
+			.ValueContent()
+			[
+				SNew(STextBlock)
+					.Text(Message)
+					.Font(IDetailLayoutBuilder::GetDetailFont())
+			];
+	}
 }
 
 void FCustomizableObjectNodeMeshMorphDetails::OnMorphTargetComboBoxSelectionChanged(TSharedPtr<FString> Selection, ESelectInfo::Type SelectInfo, TSharedRef<IPropertyHandle> ParentProperty)
diff --git a/Engine/Plugins/Mutable/Source/CustomizableObjectEditor/Private/MuCOE/Nodes/CustomizableObjectNodeMeshMorphStackApplication.cpp b/Engine/Plugins/Mutable/Source/CustomizableObjectEditor/Private/MuCOE/Nodes/CustomizableObjectNodeMeshMorphStackApplication.cpp
--- a/Engine/Plugins/Mutable/Source/CustomizableObjectEditor/Private/MuCOE/Nodes/CustomizableObjectNodeMeshMorphStackApplication.cpp
+++ b/Engine/Plugins/Mutable/Source/CustomizableObjectEditor/Private/MuCOE/Nodes/CustomizableObjectNodeMeshMorphStackApplication.cpp
@@ -6,7 +6,6 @@
 #include "MuCO/CustomizableObjectCustomVersion.h"
 #include "MuCOE/EdGraphSchema_CustomizableObject.h"
 #include "MuCOE/GraphTraversal.h"
-#include "MuCOE/ICustomizableObjectEditor.h"
 #include "MuCOE/Nodes/CustomizableObjectNodeSkeletalMesh.h"
 #include "MuCOE/Nodes/CustomizableObjectNodeTable.h"
 
diff --git a/Engine/Plugins/Mutable/Source/CustomizableObjectEditor/Private/MuCOE/Nodes/CustomizableObjectNodeMeshMorphStackDefinition.cpp b/Engine/Plugins/Mutable/Source/CustomizableObjectEditor/Private/MuCOE/Nodes/CustomizableObjectNodeMeshMorphStackDefinition.cpp
--- a/Engine/Plugins/Mutable/Source/CustomizableObjectEditor/Private/MuCOE/Nodes/CustomizableObjectNodeMeshMorphStackDefinition.cpp
+++ b/Engine/Plugins/Mutable/Source/CustomizableObjectEditor/Private/MuCOE/Nodes/CustomizableObjectNodeMeshMorphStackDefinition.cpp
@@ -5,7 +5,6 @@
 #include "Engine/SkeletalMesh.h"
 #include "MuCOE/EdGraphSchema_CustomizableObject.h"
 #include "MuCOE/GraphTraversal.h"
-#include "MuCOE/ICustomizableObjectEditor.h"
 #include "MuCOE/Nodes/CustomizableObjectNodeSkeletalMesh.h"
 #include "MuCOE/Nodes/CustomizableObjectNodeTable.h"
 
diff --git a/Engine/Plugins/Mutable/Source/CustomizableObjectEditor/Private/MuCOE/Nodes/CustomizableObjectNodeModifierBase.cpp b/Engine/Plugins/Mutable/Source/CustomizableObjectEditor/Private/MuCOE/Nodes/CustomizableObjectNodeModifierBase.cpp
--- a/Engine/Plugins/Mutable/Source/CustomizableObjectEditor/Private/MuCOE/Nodes/CustomizableObjectNodeModifierBase.cpp
+++ b/Engine/Plugins/Mutable/Source/CustomizableObjectEditor/Private/MuCOE/Nodes/CustomizableObjectNodeModifierBase.cpp
@@ -9,6 +9,7 @@
 #include "MuCO/CustomizableObjectCustomVersion.h"
 #include "MuT/NodeModifier.h"
 
+const FName UCustomizableObjectNodeModifierBase::RequiredTagsPinName(TEXT("Required Tags"));
 
 FLinearColor UCustomizableObjectNodeModifierBase::GetNodeTitleColor() const
 {
@@ -23,6 +24,12 @@ UEdGraphPin* UCustomizableObjectNodeModifierBase::OutputPin() const
 }
 
 
+UEdGraphPin* UCustomizableObjectNodeModifierBase::RequiredTagsPin() const
+{
+	return FindPin(RequiredTagsPinName, EEdGraphPinDirection::EGPD_Input);
+}
+
+
 bool UCustomizableObjectNodeModifierBase::IsApplicableTo(UCustomizableObjectNode* Candidate)
 {
 	if (!Candidate)
@@ -291,3 +298,13 @@ FString UCustomizableObjectNodeModifierBase::MakeNodeAutoTag( UEdGraphNode* Node
 	FString NewLegacyTag = FString::Printf(TEXT("%s_%s"), *ShortName, *Node->NodeGuid.ToString());
 	return NewLegacyTag;
 }
+
+
+void UCustomizableObjectNodeModifierBase::AllocateDefaultPins(UCustomizableObjectNodeRemapPins* RemapPins)
+{
+	if (IsInMacro())
+	{
+		UEdGraphPin* RquiredTags = CustomCreatePin(EGPD_Input, UEdGraphSchema_CustomizableObject::PC_String, RequiredTagsPinName, true);
+		RquiredTags->bDefaultValueIsIgnored = true;
+	}
+}
diff --git a/Engine/Plugins/Mutable/Source/CustomizableObjectEditor/Private/MuCOE/Nodes/CustomizableObjectNodeModifierBase.h b/Engine/Plugins/Mutable/Source/CustomizableObjectEditor/Private/MuCOE/Nodes/CustomizableObjectNodeModifierBase.h
--- a/Engine/Plugins/Mutable/Source/CustomizableObjectEditor/Private/MuCOE/Nodes/CustomizableObjectNodeModifierBase.h
+++ b/Engine/Plugins/Mutable/Source/CustomizableObjectEditor/Private/MuCOE/Nodes/CustomizableObjectNodeModifierBase.h
@@ -52,9 +52,11 @@ public:
 	// UCustomizableObjectNode interface
 	virtual void BackwardsCompatibleFixup(int32 CustomizableObjectCustomVersion) override;
 	virtual void PostBackwardsCompatibleFixup() override;
+	virtual void AllocateDefaultPins(UCustomizableObjectNodeRemapPins* RemapPins) override;
 
 	// Own interface
 	virtual UEdGraphPin* OutputPin() const;
+	virtual UEdGraphPin* RequiredTagsPin() const;
 
 	/** Return true if the tags and policy in this node would make this node modify the given parameter node. */
 	bool IsApplicableTo(UCustomizableObjectNode* Candidate);
@@ -75,5 +77,10 @@ protected:
 	*/
 	UPROPERTY()
 	TArray<FLegacyTag> LegacyBackportsRequiredTags;
+
+private:
+	
+	static const FName RequiredTagsPinName;
+
 };
 
diff --git a/Engine/Plugins/Mutable/Source/CustomizableObjectEditor/Private/MuCOE/Nodes/CustomizableObjectNodeModifierBaseDetails.cpp b/Engine/Plugins/Mutable/Source/CustomizableObjectEditor/Private/MuCOE/Nodes/CustomizableObjectNodeModifierBaseDetails.cpp
--- a/Engine/Plugins/Mutable/Source/CustomizableObjectEditor/Private/MuCOE/Nodes/CustomizableObjectNodeModifierBaseDetails.cpp
+++ b/Engine/Plugins/Mutable/Source/CustomizableObjectEditor/Private/MuCOE/Nodes/CustomizableObjectNodeModifierBaseDetails.cpp
@@ -48,41 +48,39 @@ void FCustomizableObjectNodeModifierBaseDetails::CustomizeDetails( IDetailLayout
 	ModifierCategory.SetSortOrder(-10000);
 
 	// Add the required tags widget
-	{
-		RequiredTagsPropertyHandle = DetailBuilder.GetProperty(GET_MEMBER_NAME_CHECKED(UCustomizableObjectNodeModifierBase, RequiredTags), UCustomizableObjectNodeModifierBase::StaticClass());
-		DetailBuilder.HideProperty(RequiredTagsPropertyHandle);
-
-		RequiredTagsPropertyHandle->SetOnPropertyValueChanged(FSimpleDelegate::CreateSP(this, &FCustomizableObjectNodeModifierBaseDetails::OnRequiredTagsPropertyChanged));
-		RequiredTagsPropertyHandle->SetOnChildPropertyValueChanged(FSimpleDelegate::CreateSP(this, &FCustomizableObjectNodeModifierBaseDetails::OnRequiredTagsPropertyChanged));
-
-		TagsPolicyPropertyHandle = DetailBuilder.GetProperty(GET_MEMBER_NAME_CHECKED(UCustomizableObjectNodeModifierBase, MultipleTagPolicy), UCustomizableObjectNodeModifierBase::StaticClass());
-		TagsPolicyPropertyHandle->SetOnPropertyValueChanged(FSimpleDelegate::CreateSP(this, &FCustomizableObjectNodeModifierBaseDetails::OnRequiredTagsPropertyChanged));
-
-		ModifierCategory.AddCustomRow(FText::FromString(TEXT("Required Tags")))
-			.PropertyHandleList({ RequiredTagsPropertyHandle })
-			.NameContent()
-				.VAlign(VAlign_Fill)
-			[
-				SNew(SVerticalBox)
-					+ SVerticalBox::Slot()
-					.VAlign(VAlign_Top)
-					.Padding(FMargin(0,4.0f,0,4.0f))
-					[
-						SNew(STextBlock)
-							.Text(LOCTEXT("ModifierDetails_RequiredTags", "Required Tags"))
-							.Font(IDetailLayoutBuilder::GetDetailFont())
-					]
-			]
-			.ValueContent()
-				.HAlign(HAlign_Fill)
-			[
-				SAssignNew(this->TagListWidget, SMutableTagListWidget)
-					.Node(Node)
-					.TagArray(&Node->RequiredTags)
-					.EmptyListText(LOCTEXT("ModifierDetails_NoRequiredTagsWarning", "Warning: There are no required tags, so this modifier will not do anything."))
-					.OnTagListChanged( this, &FCustomizableObjectNodeModifierBaseDetails::OnRequiredTagsPropertyChanged )
-			];
-	}
+	RequiredTagsPropertyHandle = DetailBuilder.GetProperty(GET_MEMBER_NAME_CHECKED(UCustomizableObjectNodeModifierBase, RequiredTags), UCustomizableObjectNodeModifierBase::StaticClass());
+	DetailBuilder.HideProperty(RequiredTagsPropertyHandle);
+
+	RequiredTagsPropertyHandle->SetOnPropertyValueChanged(FSimpleDelegate::CreateSP(this, &FCustomizableObjectNodeModifierBaseDetails::OnRequiredTagsPropertyChanged));
+	RequiredTagsPropertyHandle->SetOnChildPropertyValueChanged(FSimpleDelegate::CreateSP(this, &FCustomizableObjectNodeModifierBaseDetails::OnRequiredTagsPropertyChanged));
+
+	TagsPolicyPropertyHandle = DetailBuilder.GetProperty(GET_MEMBER_NAME_CHECKED(UCustomizableObjectNodeModifierBase, MultipleTagPolicy), UCustomizableObjectNodeModifierBase::StaticClass());
+	TagsPolicyPropertyHandle->SetOnPropertyValueChanged(FSimpleDelegate::CreateSP(this, &FCustomizableObjectNodeModifierBaseDetails::OnRequiredTagsPropertyChanged));
+
+	ModifierCategory.AddCustomRow(FText::FromString(TEXT("Required Tags")))
+		.PropertyHandleList({ RequiredTagsPropertyHandle })
+		.NameContent()
+			.VAlign(VAlign_Fill)
+		[
+			SNew(SVerticalBox)
+				+ SVerticalBox::Slot()
+				.VAlign(VAlign_Top)
+				.Padding(FMargin(0,4.0f,0,4.0f))
+				[
+					SNew(STextBlock)
+						.Text(LOCTEXT("ModifierDetails_RequiredTags", "Required Tags"))
+						.Font(IDetailLayoutBuilder::GetDetailFont())
+				]
+		]
+		.ValueContent()
+			.HAlign(HAlign_Fill)
+		[
+			SAssignNew(this->TagListWidget, SMutableTagListWidget)
+				.Node(Node)
+				.TagArray(&Node->RequiredTags)
+				.EmptyListText(LOCTEXT("ModifierDetails_NoRequiredTagsWarning", "Warning: There are no required tags, so this modifier will not do anything."))
+				.OnTagListChanged( this, &FCustomizableObjectNodeModifierBaseDetails::OnRequiredTagsPropertyChanged )
+		];
 }
 
 
diff --git a/Engine/Plugins/Mutable/Source/CustomizableObjectEditor/Private/MuCOE/Nodes/CustomizableObjectNodeModifierClipWithUVMask.cpp b/Engine/Plugins/Mutable/Source/CustomizableObjectEditor/Private/MuCOE/Nodes/CustomizableObjectNodeModifierClipWithUVMask.cpp
--- a/Engine/Plugins/Mutable/Source/CustomizableObjectEditor/Private/MuCOE/Nodes/CustomizableObjectNodeModifierClipWithUVMask.cpp
+++ b/Engine/Plugins/Mutable/Source/CustomizableObjectEditor/Private/MuCOE/Nodes/CustomizableObjectNodeModifierClipWithUVMask.cpp
@@ -4,7 +4,7 @@
 
 #include "MuCOE/CustomizableObjectGraph.h"
 #include "MuCOE/EdGraphSchema_CustomizableObject.h"
-#include "MuCOE/ICustomizableObjectEditor.h"
+#include "MuCOE/CustomizableObjectMacroLibrary/CustomizableObjectGraphEditorToolkit.h"
 #include "MuCO/CustomizableObjectCustomVersion.h"
 
 class UCustomizableObjectNodeRemapPins;
@@ -20,6 +20,9 @@ void UCustomizableObjectNodeModifierClipWithUVMask::AllocateDefaultPins(UCustomi
 	ClipMaskPin->bDefaultValueIsIgnored = true;
 
 	UEdGraphPin* OutputPin_p = CustomCreatePin(EGPD_Output, Schema->PC_Modifier, FName("Modifier"));
+
+	//Create Node Modifier Common Pins
+	Super::AllocateDefaultPins(RemapPins);
 }
 
 
@@ -33,7 +36,7 @@ void UCustomizableObjectNodeModifierClipWithUVMask::PinConnectionListChanged(UEd
 {
 	Super::PinConnectionListChanged(Pin);
 
-	TSharedPtr<ICustomizableObjectEditor> Editor = GetGraphEditor();
+	TSharedPtr<FCustomizableObjectGraphEditorToolkit> Editor = GetGraphEditor();
 
 	if (Editor.IsValid())
 	{
diff --git a/Engine/Plugins/Mutable/Source/CustomizableObjectEditor/Private/MuCOE/Nodes/CustomizableObjectNodeModifierExtendMeshSection.cpp b/Engine/Plugins/Mutable/Source/CustomizableObjectEditor/Private/MuCOE/Nodes/CustomizableObjectNodeModifierExtendMeshSection.cpp
--- a/Engine/Plugins/Mutable/Source/CustomizableObjectEditor/Private/MuCOE/Nodes/CustomizableObjectNodeModifierExtendMeshSection.cpp
+++ b/Engine/Plugins/Mutable/Source/CustomizableObjectEditor/Private/MuCOE/Nodes/CustomizableObjectNodeModifierExtendMeshSection.cpp
@@ -115,6 +115,15 @@ void UCustomizableObjectNodeModifierExtendMeshSection::AllocateDefaultPins(UCust
 
 	CustomCreatePin(EGPD_Input, Schema->PC_Mesh, FName("Add Mesh") );
 
+	//Create Node Modifier Common Pins
+	Super::AllocateDefaultPins(RemapPins);
+
+	if (IsInMacro())
+	{
+		UEdGraphPin* EnableTags = CustomCreatePin(EGPD_Input, Schema->PC_String, FName("Enable Tags"), true);
+		EnableTags->bDefaultValueIsIgnored = true;
+	}
+
 	// Begin texture pins
 	{
 		const int32 NumImages = GetNumParameters(EMaterialParameterType::Texture);
diff --git a/Engine/Plugins/Mutable/Source/CustomizableObjectEditor/Private/MuCOE/Nodes/CustomizableObjectNodeModifierMorphMeshSection.cpp b/Engine/Plugins/Mutable/Source/CustomizableObjectEditor/Private/MuCOE/Nodes/CustomizableObjectNodeModifierMorphMeshSection.cpp
--- a/Engine/Plugins/Mutable/Source/CustomizableObjectEditor/Private/MuCOE/Nodes/CustomizableObjectNodeModifierMorphMeshSection.cpp
+++ b/Engine/Plugins/Mutable/Source/CustomizableObjectEditor/Private/MuCOE/Nodes/CustomizableObjectNodeModifierMorphMeshSection.cpp
@@ -6,7 +6,6 @@
 #include "MuCO/CustomizableObjectCustomVersion.h"
 #include "MuCOE/EdGraphSchema_CustomizableObject.h"
 #include "MuCOE/GraphTraversal.h"
-#include "MuCOE/ICustomizableObjectEditor.h"
 #include "MuCOE/Nodes/CustomizableObjectNodeMaterial.h"
 #include "MuCOE/Nodes/CustomizableObjectNodeSkeletalMesh.h"
 
@@ -27,6 +26,16 @@ void UCustomizableObjectNodeModifierMorphMeshSection::AllocateDefaultPins(UCusto
 	PinName = TEXT("Factor");
 	UEdGraphPin* FactorPin = CustomCreatePin(EGPD_Input, Schema->PC_Float, FName(*PinName));
 	FactorPin->bDefaultValueIsIgnored = true;
+
+	if (IsInMacro())
+	{
+		PinName = TEXT("Morph Target Name");
+		UEdGraphPin* StringPin = CustomCreatePin(EGPD_Input, Schema->PC_String, FName(*PinName));
+		StringPin->bDefaultValueIsIgnored = true;
+	}
+
+	//Create Node Modifier Common Pins
+	Super::AllocateDefaultPins(RemapPins);
 }
 
 
diff --git a/Engine/Plugins/Mutable/Source/CustomizableObjectEditor/Private/MuCOE/Nodes/CustomizableObjectNodeModifierMorphMeshSection.h b/Engine/Plugins/Mutable/Source/CustomizableObjectEditor/Private/MuCOE/Nodes/CustomizableObjectNodeModifierMorphMeshSection.h
--- a/Engine/Plugins/Mutable/Source/CustomizableObjectEditor/Private/MuCOE/Nodes/CustomizableObjectNodeModifierMorphMeshSection.h
+++ b/Engine/Plugins/Mutable/Source/CustomizableObjectEditor/Private/MuCOE/Nodes/CustomizableObjectNodeModifierMorphMeshSection.h
@@ -20,7 +20,7 @@ class CUSTOMIZABLEOBJECTEDITOR_API UCustomizableObjectNodeModifierMorphMeshSecti
 	
 public:
 
-	UPROPERTY(EditAnywhere, Category = CustomizableObject)
+	UPROPERTY(EditAnywhere, Category = MeshMorph)
 	FString MorphTargetName;
 
 public:
diff --git a/Engine/Plugins/Mutable/Source/CustomizableObjectEditor/Private/MuCOE/Nodes/CustomizableObjectNodeModifierMorphMeshSectionDetails.cpp b/Engine/Plugins/Mutable/Source/CustomizableObjectEditor/Private/MuCOE/Nodes/CustomizableObjectNodeModifierMorphMeshSectionDetails.cpp
--- a/Engine/Plugins/Mutable/Source/CustomizableObjectEditor/Private/MuCOE/Nodes/CustomizableObjectNodeModifierMorphMeshSectionDetails.cpp
+++ b/Engine/Plugins/Mutable/Source/CustomizableObjectEditor/Private/MuCOE/Nodes/CustomizableObjectNodeModifierMorphMeshSectionDetails.cpp
@@ -52,52 +52,70 @@ void FCustomizableObjectNodeModifierMorphMeshSectionDetails::CustomizeDetails( I
 	
 	// Add a morph selection widget.
 	IDetailCategoryBuilder& MorphCategory = DetailBuilder.EditCategory( "Morph" );
+	TSharedRef<IPropertyHandle> MorphTargetNameProperty = DetailBuilder.GetProperty("MorphTargetName");
 
-	// Scan for hint morph names
-	RefreshMorphOptions();
+	if (!Node->IsInMacro())
+	{
+		// Scan for hint morph names
+		RefreshMorphOptions();
 
-	TSharedRef<IPropertyHandle> MorphTargetNameProperty = DetailBuilder.GetProperty("MorphTargetName");
-	MorphCategory.AddCustomRow(LOCTEXT("MorphMaterialDetails_Target", "Target"))
+		MorphCategory.AddCustomRow(LOCTEXT("MorphMaterialDetails_Target", "Target"))
 		[
 			SNew(SProperty, MorphTargetNameProperty)
-				.ShouldDisplayName(false)
-				.CustomWidget()
+			.ShouldDisplayName(false)
+			.CustomWidget()
+			[
+				SNew(SBorder)
+				.BorderImage(UE_MUTABLE_GET_BRUSH("NoBorder"))
+				.Padding(FMargin(0.0f, 0.0f, 10.0f, 0.0f))
 				[
-					SNew(SBorder)
-						.BorderImage(UE_MUTABLE_GET_BRUSH("NoBorder"))
-						.Padding(FMargin(0.0f, 0.0f, 10.0f, 0.0f))
+					SNew(SHorizontalBox)
+
+					+ SHorizontalBox::Slot()
+					.FillWidth(10.0f)
+					.VAlign(VAlign_Center)
+					[
+						SNew(STextBlock)
+							.Text(LOCTEXT("MorphMaterialDetails_MorphTarget", "Morph Target"))
+							.Font(IDetailLayoutBuilder::GetDetailFont())
+					]
+
+					+ SHorizontalBox::Slot()
+					.FillWidth(10.0f)
+					.HAlign(HAlign_Fill)
+					.VAlign(VAlign_Center)
+					[
+						SAssignNew(this->MorphCombo , SMutableSearchComboBox)
+						.ButtonStyle(FAppStyle::Get(), "NoBorder")
+						.OptionsSource(&MorphOptionsSource)
+						.OnSelectionChanged(this, &FCustomizableObjectNodeModifierMorphMeshSectionDetails::OnMorphTargetComboBoxSelectionChanged)
+						.Content()
 						[
-							SNew(SHorizontalBox)
-
-							+ SHorizontalBox::Slot()
-							.FillWidth(10.0f)
-							.VAlign(VAlign_Center)
-							[
-								SNew(STextBlock)
-									.Text(LOCTEXT("MorphMaterialDetails_MorphTarget", "Morph Target"))
-									.Font(IDetailLayoutBuilder::GetDetailFont())
-							]
-
-							+ SHorizontalBox::Slot()
-							.FillWidth(10.0f)
-							.HAlign(HAlign_Fill)
-							.VAlign(VAlign_Center)
-							[
-								SAssignNew(this->MorphCombo , SMutableSearchComboBox)
-								.ButtonStyle(FAppStyle::Get(), "NoBorder")
-								.OptionsSource(&MorphOptionsSource)
-								.OnSelectionChanged(this, &FCustomizableObjectNodeModifierMorphMeshSectionDetails::OnMorphTargetComboBoxSelectionChanged)
-								.Content()
-								[
-									SNew(SEditableTextBox)
-									.Font(IDetailLayoutBuilder::GetDetailFont())
-									.Text_Lambda([&]() { return Node ? FText::FromString(Node->MorphTargetName) : FText();})
-									.OnTextChanged(this, &FCustomizableObjectNodeModifierMorphMeshSectionDetails::OnMorphTargetComboBoxSelectionChanged)
-								]
-							]
+							SNew(SEditableTextBox)
+							.Font(IDetailLayoutBuilder::GetDetailFont())
+							.Text_Lambda([&]() { return Node ? FText::FromString(Node->MorphTargetName) : FText();})
+							.OnTextChanged(this, &FCustomizableObjectNodeModifierMorphMeshSectionDetails::OnMorphTargetComboBoxSelectionChanged)
 						]
+					]
 				]
+			]
+		];
+	}
+	else
+	{
+		DetailBuilder.EditDefaultProperty(MorphTargetNameProperty)->
+		CustomWidget()
+		.NameContent()
+		[
+			SNew(STextBlock).Text(LOCTEXT("MorphMaterialDetails_MorphTargetName", "Morph Target Name")).Font(IDetailLayoutBuilder::GetDetailFont())
+		]
+		.ValueContent()
+		[
+			SNew(STextBlock)
+			.Text(LOCTEXT("MorphMaterialDetails_PinMessage", "In Mutable Macros, Morph Target Names are defined through String Nodes."))
+			.Font(IDetailLayoutBuilder::GetDetailFont())
 		];
+	}
 }
 
 
diff --git a/Engine/Plugins/Mutable/Source/CustomizableObjectEditor/Private/MuCOE/Nodes/CustomizableObjectNodeModifierRemoveMesh.cpp b/Engine/Plugins/Mutable/Source/CustomizableObjectEditor/Private/MuCOE/Nodes/CustomizableObjectNodeModifierRemoveMesh.cpp
--- a/Engine/Plugins/Mutable/Source/CustomizableObjectEditor/Private/MuCOE/Nodes/CustomizableObjectNodeModifierRemoveMesh.cpp
+++ b/Engine/Plugins/Mutable/Source/CustomizableObjectEditor/Private/MuCOE/Nodes/CustomizableObjectNodeModifierRemoveMesh.cpp
@@ -3,7 +3,7 @@
 #include "MuCOE/Nodes/CustomizableObjectNodeModifierRemoveMesh.h"
 
 #include "MuCOE/EdGraphSchema_CustomizableObject.h"
-#include "MuCOE/ICustomizableObjectEditor.h"
+#include "MuCOE/CustomizableObjectMacroLibrary/CustomizableObjectGraphEditorToolkit.h"
 #include "MuCO/CustomizableObjectCustomVersion.h"
 
 class UCustomizableObjectNodeRemapPins;
@@ -19,6 +19,9 @@ void UCustomizableObjectNodeModifierRemoveMesh::AllocateDefaultPins(UCustomizabl
 	RemoveMeshPin->bDefaultValueIsIgnored = true;
 
 	UEdGraphPin* OutputPin = CustomCreatePin(EGPD_Output, Schema->PC_Modifier, FName("Modifier"));
+
+	//Create Node Modifier Common Pins
+	Super::AllocateDefaultPins(RemapPins);
 }
 
 
@@ -34,7 +37,7 @@ void UCustomizableObjectNodeModifierRemoveMesh::PinConnectionListChanged(UEdGrap
 
 	if (Pin == OutputPin())
 	{
-		TSharedPtr<ICustomizableObjectEditor> Editor = GetGraphEditor();
+		TSharedPtr<FCustomizableObjectGraphEditorToolkit> Editor = GetGraphEditor();
 
 		if (Editor.IsValid())
 		{
diff --git a/Engine/Plugins/Mutable/Source/CustomizableObjectEditor/Private/MuCOE/Nodes/CustomizableObjectNodeModifierRemoveMeshBlocks.cpp b/Engine/Plugins/Mutable/Source/CustomizableObjectEditor/Private/MuCOE/Nodes/CustomizableObjectNodeModifierRemoveMeshBlocks.cpp
--- a/Engine/Plugins/Mutable/Source/CustomizableObjectEditor/Private/MuCOE/Nodes/CustomizableObjectNodeModifierRemoveMeshBlocks.cpp
+++ b/Engine/Plugins/Mutable/Source/CustomizableObjectEditor/Private/MuCOE/Nodes/CustomizableObjectNodeModifierRemoveMeshBlocks.cpp
@@ -6,7 +6,6 @@
 #include "MuCO/CustomizableObjectCustomVersion.h"
 #include "MuCOE/CustomizableObjectLayout.h"
 #include "MuCOE/EdGraphSchema_CustomizableObject.h"
-#include "MuCOE/ICustomizableObjectEditor.h"
 #include "MuCOE/Nodes/CustomizableObjectNodeMaterial.h"
 #include "MuCOE/Nodes/CustomizableObjectNodeSkeletalMesh.h"
 #include "MuCOE/Nodes/CustomizableObjectNodeTable.h"
diff --git a/Engine/Plugins/Mutable/Source/CustomizableObjectEditor/Private/MuCOE/Nodes/CustomizableObjectNodeModifierTransformInMesh.cpp b/Engine/Plugins/Mutable/Source/CustomizableObjectEditor/Private/MuCOE/Nodes/CustomizableObjectNodeModifierTransformInMesh.cpp
--- a/Engine/Plugins/Mutable/Source/CustomizableObjectEditor/Private/MuCOE/Nodes/CustomizableObjectNodeModifierTransformInMesh.cpp
+++ b/Engine/Plugins/Mutable/Source/CustomizableObjectEditor/Private/MuCOE/Nodes/CustomizableObjectNodeModifierTransformInMesh.cpp
@@ -33,6 +33,9 @@ void UCustomizableObjectNodeModifierTransformInMesh::AllocateDefaultPins(UCustom
 	TransformPin->bDefaultValueIsIgnored = true;
 	
 	(void)CustomCreatePin(EGPD_Output, Schema->PC_Modifier, FName(OutputPinName));
+
+	//Create Node Modifier Common Pins
+	Super::AllocateDefaultPins(RemapPins);
 }
 
 UEdGraphPin* UCustomizableObjectNodeModifierTransformInMesh::OutputPin() const
diff --git a/Engine/Plugins/Mutable/Source/CustomizableObjectEditor/Private/MuCOE/Nodes/CustomizableObjectNodeObject.cpp b/Engine/Plugins/Mutable/Source/CustomizableObjectEditor/Private/MuCOE/Nodes/CustomizableObjectNodeObject.cpp
--- a/Engine/Plugins/Mutable/Source/CustomizableObjectEditor/Private/MuCOE/Nodes/CustomizableObjectNodeObject.cpp
+++ b/Engine/Plugins/Mutable/Source/CustomizableObjectEditor/Private/MuCOE/Nodes/CustomizableObjectNodeObject.cpp
@@ -9,7 +9,7 @@
 #include "MuCOE/CustomizableObjectGraph.h"
 #include "MuCOE/EdGraphSchema_CustomizableObject.h"
 #include "MuCOE/GraphTraversal.h"
-#include "MuCOE/ICustomizableObjectEditor.h"
+#include "MuCOE/CustomizableObjectEditor.h"
 #include "MuCOE/Nodes/CustomizableObjectNodeMaterial.h"
 #include "MuCOE/Nodes/CustomizableObjectNodeComponentMesh.h"
 #include "MuCOE/Nodes/CustomizableObjectNodeMaterialVariation.h"
@@ -32,6 +32,7 @@ const FName UCustomizableObjectNodeObject::ChildrenPinName(TEXT("Children"));
 const FName UCustomizableObjectNodeObject::ComponentsPinName(TEXT("Components"));
 const FName UCustomizableObjectNodeObject::ModifiersPinName(TEXT("Modifiers"));
 const FName UCustomizableObjectNodeObject::OutputPinName(TEXT("Object"));
+const FName UCustomizableObjectNodeObject::NamePinName(TEXT("Name"));
 
 
 UCustomizableObjectNodeObject::UCustomizableObjectNodeObject()
@@ -596,6 +597,19 @@ void UCustomizableObjectNodeObject::AllocateDefaultPins(UCustomizableObjectNodeR
 	{
 		OutputPin->bHidden = true;
 	}
+	else
+	{
+		if (IsInMacro())
+		{
+			UEdGraphPin* NamePin = CustomCreatePin(EGPD_Input, Schema->PC_String, NamePinName);
+		}
+	}
+}
+
+
+bool UCustomizableObjectNodeObject::IsNodeSupportedInMacros()
+{
+	return false;
 }
 
 
@@ -686,15 +700,18 @@ void UCustomizableObjectNodeObject::PostDuplicate(bool bDuplicateForPIE)
 
 void UCustomizableObjectNodeObject::SetParentObject(UCustomizableObject* CustomizableParentObject)
 {
-	if (CustomizableParentObject != GetGraphEditor()->GetCustomizableObject())
+	if (TSharedPtr<FCustomizableObjectEditor> Editor = StaticCastSharedPtr<FCustomizableObjectEditor>(GetGraphEditor()))
 	{
-		ParentObject = CustomizableParentObject;
-
-		// Update the cached flag in the main object
-		UCustomizableObject* CustomizableObject = Cast<UCustomizableObject>(GetCustomizableObjectGraph()->GetOuter());
-		if (CustomizableObject)
+		if (CustomizableParentObject != Editor->GetCustomizableObject())
 		{
-			CustomizableObject->GetPrivate()->SetIsChildObject(ParentObject != nullptr);
+			ParentObject = CustomizableParentObject;
+
+			// Update the cached flag in the main object
+			UCustomizableObject* CustomizableObject = Cast<UCustomizableObject>(GetCustomizableObjectGraph()->GetOuter());
+			if (CustomizableObject)
+			{
+				CustomizableObject->GetPrivate()->SetIsChildObject(ParentObject != nullptr);
+			}
 		}
 	}
 }
diff --git a/Engine/Plugins/Mutable/Source/CustomizableObjectEditor/Private/MuCOE/Nodes/CustomizableObjectNodeObject.h b/Engine/Plugins/Mutable/Source/CustomizableObjectEditor/Private/MuCOE/Nodes/CustomizableObjectNodeObject.h
--- a/Engine/Plugins/Mutable/Source/CustomizableObjectEditor/Private/MuCOE/Nodes/CustomizableObjectNodeObject.h
+++ b/Engine/Plugins/Mutable/Source/CustomizableObjectEditor/Private/MuCOE/Nodes/CustomizableObjectNodeObject.h
@@ -161,6 +161,7 @@ public:
 	virtual void BackwardsCompatibleFixup(int32 CustomizableObjectCustomVersion) override;
 	virtual void PostBackwardsCompatibleFixup() override;
 	void AllocateDefaultPins(UCustomizableObjectNodeRemapPins* RemapPins) override;
+	virtual bool IsNodeSupportedInMacros() override;
 
 	// Own Interface
 	void SetParentObject(UCustomizableObject* CustomizableParentObject);
@@ -203,5 +204,6 @@ private:
 	static const FName ComponentsPinName;
 	static const FName ModifiersPinName;
 	static const FName OutputPinName;
+	static const FName NamePinName;
 };
 
diff --git a/Engine/Plugins/Mutable/Source/CustomizableObjectEditor/Private/MuCOE/Nodes/CustomizableObjectNodeObjectChild.cpp b/Engine/Plugins/Mutable/Source/CustomizableObjectEditor/Private/MuCOE/Nodes/CustomizableObjectNodeObjectChild.cpp
--- a/Engine/Plugins/Mutable/Source/CustomizableObjectEditor/Private/MuCOE/Nodes/CustomizableObjectNodeObjectChild.cpp
+++ b/Engine/Plugins/Mutable/Source/CustomizableObjectEditor/Private/MuCOE/Nodes/CustomizableObjectNodeObjectChild.cpp
@@ -50,6 +50,13 @@ bool UCustomizableObjectNodeObjectChild::CanDuplicateNode() const
 	return true;
 }
 
+
+bool UCustomizableObjectNodeObjectChild::IsNodeSupportedInMacros()
+{
+	return true;
+}
+
+
 FText UCustomizableObjectNodeObjectChild::GetTooltipText() const
 {
 	return LOCTEXT("Child_Object_Tooltip",
diff --git a/Engine/Plugins/Mutable/Source/CustomizableObjectEditor/Private/MuCOE/Nodes/CustomizableObjectNodeObjectChild.h b/Engine/Plugins/Mutable/Source/CustomizableObjectEditor/Private/MuCOE/Nodes/CustomizableObjectNodeObjectChild.h
--- a/Engine/Plugins/Mutable/Source/CustomizableObjectEditor/Private/MuCOE/Nodes/CustomizableObjectNodeObjectChild.h
+++ b/Engine/Plugins/Mutable/Source/CustomizableObjectEditor/Private/MuCOE/Nodes/CustomizableObjectNodeObjectChild.h
@@ -24,5 +24,8 @@ public:
 	FText GetTooltipText() const override;
 	bool CanUserDeleteNode() const override;
 	bool CanDuplicateNode() const override;
+
+	// UCustomizableObjectNode interface
+	virtual bool IsNodeSupportedInMacros() override;
 };
 
diff --git a/Engine/Plugins/Mutable/Source/CustomizableObjectEditor/Private/MuCOE/Nodes/CustomizableObjectNodePassThroughTextureVariation.cpp b/Engine/Plugins/Mutable/Source/CustomizableObjectEditor/Private/MuCOE/Nodes/CustomizableObjectNodePassThroughTextureVariation.cpp
--- a/Engine/Plugins/Mutable/Source/CustomizableObjectEditor/Private/MuCOE/Nodes/CustomizableObjectNodePassThroughTextureVariation.cpp
+++ b/Engine/Plugins/Mutable/Source/CustomizableObjectEditor/Private/MuCOE/Nodes/CustomizableObjectNodePassThroughTextureVariation.cpp
@@ -8,25 +8,10 @@
 #define LOCTEXT_NAMESPACE "CustomizableObjectEditor"
 
 
-void UCustomizableObjectNodePassThroughTextureVariation::AllocateDefaultPins(UCustomizableObjectNodeRemapPins* RemapPins)
+// UCustomizableObjectNodeVariation interface
+FName UCustomizableObjectNodePassThroughTextureVariation::GetCategory() const
 {
-	const UEdGraphSchema_CustomizableObject* Schema = GetDefault<UEdGraphSchema_CustomizableObject>();
-
-	UEdGraphPin* OutputPin = CustomCreatePinSimple(EGPD_Output, UEdGraphSchema_CustomizableObject::PC_PassThroughImage);
-	OutputPin->bDefaultValueIsIgnored = true;
-
-	for (int VariationIndex = Variations.Num() - 1; VariationIndex >= 0; --VariationIndex)
-	{
-		FString PinName = FString::Printf(TEXT("Variation %d"), VariationIndex);
-		UEdGraphPin* VariationPin = CustomCreatePin(EGPD_Input, Schema->PC_PassThroughImage, FName(*PinName), false);
-		VariationPin->bDefaultValueIsIgnored = true;
-
-		FString FriendlyName = FString::Printf(TEXT("Variation %d [%s]"), VariationIndex, *Variations[VariationIndex].Tag);
-		VariationPin->PinFriendlyName = FText::FromString(FriendlyName);
-	}
-
-	UEdGraphPin* DefaultVariation = CustomCreatePin(EGPD_Input, Schema->PC_PassThroughImage, FName(TEXT("Default")), false);
-	DefaultVariation->bDefaultValueIsIgnored = true;
+	return UEdGraphSchema_CustomizableObject::PC_PassThroughImage;
 }
 
 
diff --git a/Engine/Plugins/Mutable/Source/CustomizableObjectEditor/Private/MuCOE/Nodes/CustomizableObjectNodePassThroughTextureVariation.h b/Engine/Plugins/Mutable/Source/CustomizableObjectEditor/Private/MuCOE/Nodes/CustomizableObjectNodePassThroughTextureVariation.h
--- a/Engine/Plugins/Mutable/Source/CustomizableObjectEditor/Private/MuCOE/Nodes/CustomizableObjectNodePassThroughTextureVariation.h
+++ b/Engine/Plugins/Mutable/Source/CustomizableObjectEditor/Private/MuCOE/Nodes/CustomizableObjectNodePassThroughTextureVariation.h
@@ -17,6 +17,8 @@ public:
 	virtual FText GetNodeTitle(ENodeTitleType::Type TitleType) const override;
 
 	// UCustomizableObjectNode interface
-	virtual void AllocateDefaultPins(UCustomizableObjectNodeRemapPins* RemapPins) override;
 	virtual void BackwardsCompatibleFixup(int32 CustomizableObjectCustomVersion) override;
+	
+	// UCustomizableObjectNodeVariation interface
+	virtual FName GetCategory() const override;
 };
diff --git a/Engine/Plugins/Mutable/Source/CustomizableObjectEditor/Private/MuCOE/Nodes/CustomizableObjectNodeProjectorParameterDetails.cpp b/Engine/Plugins/Mutable/Source/CustomizableObjectEditor/Private/MuCOE/Nodes/CustomizableObjectNodeProjectorParameterDetails.cpp
--- a/Engine/Plugins/Mutable/Source/CustomizableObjectEditor/Private/MuCOE/Nodes/CustomizableObjectNodeProjectorParameterDetails.cpp
+++ b/Engine/Plugins/Mutable/Source/CustomizableObjectEditor/Private/MuCOE/Nodes/CustomizableObjectNodeProjectorParameterDetails.cpp
@@ -9,7 +9,7 @@
 #include "ScopedTransaction.h"
 #include "MuCOE/CustomizableObjectEditorUtilities.h"
 #include "MuCOE/CustomizableObjectGraph.h"
-#include "MuCOE/ICustomizableObjectEditor.h"
+#include "MuCOE/CustomizableObjectMacroLibrary/CustomizableObjectGraphEditorToolkit.h"
 #include "MuCOE/Nodes/CustomizableObjectNodeProjectorConstant.h"
 #include "MuCOE/Nodes/CustomizableObjectNodeProjectorParameter.h"
 #include "PropertyCustomizationHelpers.h"
@@ -155,7 +155,7 @@ FReply FCustomizableObjectNodeProjectorParameterDetails::OnProjectorPastePressed
 
 	FPlatformApplicationMisc::ClipboardPaste(ClipText);
 
-	TSharedPtr<ICustomizableObjectEditor> Editor = nullptr;
+	TSharedPtr<FCustomizableObjectGraphEditorToolkit> Editor = nullptr;
 
 	if (NodeParameter)
 	{
diff --git a/Engine/Plugins/Mutable/Source/CustomizableObjectEditor/Private/MuCOE/Nodes/CustomizableObjectNodeTextureVariation.cpp b/Engine/Plugins/Mutable/Source/CustomizableObjectEditor/Private/MuCOE/Nodes/CustomizableObjectNodeTextureVariation.cpp
--- a/Engine/Plugins/Mutable/Source/CustomizableObjectEditor/Private/MuCOE/Nodes/CustomizableObjectNodeTextureVariation.cpp
+++ b/Engine/Plugins/Mutable/Source/CustomizableObjectEditor/Private/MuCOE/Nodes/CustomizableObjectNodeTextureVariation.cpp
@@ -10,10 +10,9 @@ class UCustomizableObjectNodeRemapPins;
 #define LOCTEXT_NAMESPACE "CustomizableObjectEditor"
 
 
-UCustomizableObjectNodeTextureVariation::UCustomizableObjectNodeTextureVariation()
-	: Super()
+FName UCustomizableObjectNodeTextureVariation::GetCategory() const
 {
-
+	return UEdGraphSchema_CustomizableObject::PC_Image;
 }
 
 
@@ -29,25 +28,6 @@ void UCustomizableObjectNodeTextureVariation::PostEditChangeProperty(FPropertyCh
 }
 
 
-void UCustomizableObjectNodeTextureVariation::AllocateDefaultPins(UCustomizableObjectNodeRemapPins* RemapPins)
-{
-	UEdGraphPin* OutputPin = CustomCreatePinSimple(EGPD_Output, UEdGraphSchema_CustomizableObject::PC_Image);
-	OutputPin->bDefaultValueIsIgnored = true;
-
-	for ( int VariationIndex = Variations.Num() - 1; VariationIndex >= 0; --VariationIndex)
-	{
-		FString PinName = FString::Printf( TEXT("Variation %d"), VariationIndex);
-		UEdGraphPin* VariationPin = CustomCreatePin(EGPD_Input, UEdGraphSchema_CustomizableObject::PC_Image, FName(*PinName), false);
-		VariationPin->bDefaultValueIsIgnored = true;
-
-		FString TagName = GetTagDisplayName(Variations[VariationIndex].Tag);
-		VariationPin->PinFriendlyName = FText::Format(LOCTEXT("Variation_Pin_FriendlyName", "Variation {0} [{1}]"), VariationIndex, FText::FromString(TagName));
-	}
-
-	UEdGraphPin* DefaultVariation = CustomCreatePin(EGPD_Input, UEdGraphSchema_CustomizableObject::PC_Image, FName(TEXT("Default")), false);
-	DefaultVariation->bDefaultValueIsIgnored = true;
-}
-
 void UCustomizableObjectNodeTextureVariation::BackwardsCompatibleFixup(int32 CustomizableObjectCustomVersion)
 {
 	Super::BackwardsCompatibleFixup(CustomizableObjectCustomVersion);
@@ -59,6 +39,26 @@ void UCustomizableObjectNodeTextureVariation::BackwardsCompatibleFixup(int32 Cus
 			TexturePin->PinFriendlyName = LOCTEXT("Image_Pin_Category", "Texture");
 		}
 	}
+
+	// Converting TextureVariations to VariationNodes
+	if (CustomizableObjectCustomVersion == FCustomizableObjectCustomVersion::TextureVariationsToVariations)
+	{
+		VariationsPins.SetNum(Variations_DEPRECATED.Num());
+		VariationsData.SetNum(Variations_DEPRECATED.Num());
+
+		for (int32 VariationIndex = 0; VariationIndex < Variations_DEPRECATED.Num(); ++VariationIndex)
+		{
+			if (UEdGraphPin* VariationPin = FindPin(FString::Printf(TEXT("Variation %d"), VariationIndex)))
+			{
+				VariationsPins[VariationIndex] = VariationPin;
+				VariationsData[VariationIndex].Tag = Variations_DEPRECATED[VariationIndex].Tag;
+			}
+			else
+			{
+				check(false);
+			}
+		}
+	}
 }
 
 
diff --git a/Engine/Plugins/Mutable/Source/CustomizableObjectEditor/Private/MuCOE/Nodes/CustomizableObjectNodeTextureVariation.h b/Engine/Plugins/Mutable/Source/CustomizableObjectEditor/Private/MuCOE/Nodes/CustomizableObjectNodeTextureVariation.h
--- a/Engine/Plugins/Mutable/Source/CustomizableObjectEditor/Private/MuCOE/Nodes/CustomizableObjectNodeTextureVariation.h
+++ b/Engine/Plugins/Mutable/Source/CustomizableObjectEditor/Private/MuCOE/Nodes/CustomizableObjectNodeTextureVariation.h
@@ -2,7 +2,8 @@
 
 #pragma once
 
-#include "MuCOE/Nodes/CustomizableObjectNode.h"
+#include "MuCOE/Nodes/CustomizableObjectNodeVariation.h"
+#include "MuCOE/CustomizableObjectEditor_Deprecated.h"
 
 #include "CustomizableObjectNodeTextureVariation.generated.h"
 
@@ -14,27 +15,12 @@ class UObject;
 struct FPropertyChangedEvent;
 
 
-USTRUCT()
-struct CUSTOMIZABLEOBJECTEDITOR_API FCustomizableObjectTextureVariation
-{
-	GENERATED_USTRUCT_BODY()
-
-	UPROPERTY(EditAnywhere, Category = CustomizableObject)
-	FString Tag;
-};
-
-
 UCLASS()
-class CUSTOMIZABLEOBJECTEDITOR_API UCustomizableObjectNodeTextureVariation : public UCustomizableObjectNode
+class CUSTOMIZABLEOBJECTEDITOR_API UCustomizableObjectNodeTextureVariation : public UCustomizableObjectNodeVariation
 {
 public:
 	GENERATED_BODY()
 
-	UCustomizableObjectNodeTextureVariation();
-
-	UPROPERTY(EditAnywhere, Category = CustomizableObject)
-	TArray<FCustomizableObjectTextureVariation> Variations;
-
 	// UObject interface.
 	virtual void PostEditChangeProperty(FPropertyChangedEvent& PropertyChangedEvent) override;
 
@@ -44,22 +30,14 @@ public:
 	virtual FText GetTooltipText() const override;
 	
 	// UCustomizableObjectNode interface
-	virtual void AllocateDefaultPins(UCustomizableObjectNodeRemapPins* RemapPins) override;
 	virtual void BackwardsCompatibleFixup(int32 CustomizableObjectCustomVersion) override;
-	
-	UEdGraphPin* DefaultPin() const
-	{
-		return FindPin(TEXT("Default"));
-	}
 
-	int32 GetNumVariations() const
-	{
-		return Variations.Num();
-	}
+	// UCustomizableObjectNodeVariation interface
+	virtual FName GetCategory() const override;
+
+private:
 
-	UEdGraphPin* VariationPin(int Index) const
-	{
-		return FindPin(FString::Printf(TEXT("Variation %d"),Index));
-	}
+	UPROPERTY()
+	TArray<FCustomizableObjectTextureVariation> Variations_DEPRECATED;
 };
 
diff --git a/Engine/Plugins/Mutable/Source/CustomizableObjectEditor/Private/MuCOE/Nodes/CustomizableObjectNodeVariation.cpp b/Engine/Plugins/Mutable/Source/CustomizableObjectEditor/Private/MuCOE/Nodes/CustomizableObjectNodeVariation.cpp
--- a/Engine/Plugins/Mutable/Source/CustomizableObjectEditor/Private/MuCOE/Nodes/CustomizableObjectNodeVariation.cpp
+++ b/Engine/Plugins/Mutable/Source/CustomizableObjectEditor/Private/MuCOE/Nodes/CustomizableObjectNodeVariation.cpp
@@ -9,6 +9,7 @@ class UCustomizableObjectNodeRemapPins;
 
 #define LOCTEXT_NAMESPACE "CustomizableObjectEditor"
 
+const FString UCustomizableObjectNodeVariation::VariationTagSufix = " - Tag";
 
 void UCustomizableObjectNodeVariation::PostEditChangeProperty(FPropertyChangedEvent& PropertyChangedEvent)
 {
@@ -44,6 +45,13 @@ void UCustomizableObjectNodeVariation::AllocateDefaultPins(UCustomizableObjectNo
 		VariationPin->PinFriendlyName = FText::Format(LOCTEXT("Variation_Pin_FriendlyName", "Variation {0} [{1}]"), VariationIndex, FText::FromString(TagName));
 		
 		VariationsPins[VariationIndex] = VariationPin;
+
+		if (IsInMacro())
+		{
+			const FName TagPinName = FName(PinName.ToString() + VariationTagSufix);
+			UEdGraphPin* VariationTagPin = CustomCreatePin(EGPD_Input, Schema->PC_String, TagPinName);
+			VariationTagPin->bDefaultValueIsIgnored = true;
+		}
 	}
 
 	CustomCreatePin(EGPD_Input, Category, FName(TEXT("Default")), bIsInputPinArray);
diff --git a/Engine/Plugins/Mutable/Source/CustomizableObjectEditor/Private/MuCOE/Nodes/CustomizableObjectNodeVariation.h b/Engine/Plugins/Mutable/Source/CustomizableObjectEditor/Private/MuCOE/Nodes/CustomizableObjectNodeVariation.h
--- a/Engine/Plugins/Mutable/Source/CustomizableObjectEditor/Private/MuCOE/Nodes/CustomizableObjectNodeVariation.h
+++ b/Engine/Plugins/Mutable/Source/CustomizableObjectEditor/Private/MuCOE/Nodes/CustomizableObjectNodeVariation.h
@@ -33,14 +33,19 @@ class CUSTOMIZABLEOBJECTEDITOR_API UCustomizableObjectNodeVariation : public UCu
 	GENERATED_BODY()
 
 public:
+
 	UPROPERTY(EditAnywhere, Category = CustomizableObject)
 	TArray<FCustomizableObjectVariation> VariationsData; // The variable name can not be Variations due issues with the on to UObject Serialization system
 
+	static const FString VariationTagSufix;
+
 private:
+
 	UPROPERTY()
 	TArray<FEdGraphPinReference> VariationsPins;
 	
 public:
+
 	// UObject interface.
 	virtual void PostEditChangeProperty(FPropertyChangedEvent& PropertyChangedEvent) override;
 
@@ -71,5 +76,9 @@ public:
 
 	/** Get the Variation Input Pin. */
 	UEdGraphPin* VariationPin(int32 Index) const;
+
+	// Needed to convert the NodeTextureVariation into a NodeVariation
+	friend class UCustomizableObjectNodeTextureVariation;
+
 };
 
diff --git a/Engine/Plugins/Mutable/Source/CustomizableObjectEditor/Private/MuCOE/SCustomizableObjectEditorAdvancedPreviewSettings.h b/Engine/Plugins/Mutable/Source/CustomizableObjectEditor/Private/MuCOE/SCustomizableObjectEditorAdvancedPreviewSettings.h
--- a/Engine/Plugins/Mutable/Source/CustomizableObjectEditor/Private/MuCOE/SCustomizableObjectEditorAdvancedPreviewSettings.h
+++ b/Engine/Plugins/Mutable/Source/CustomizableObjectEditor/Private/MuCOE/SCustomizableObjectEditorAdvancedPreviewSettings.h
@@ -5,7 +5,6 @@
 #include "CustomizableObjectInstanceEditor.h"
 #include "AssetRegistry/AssetData.h"
 #include "SAdvancedPreviewDetailsTab.h"
-#include "MuCOE/ICustomizableObjectEditor.h"
 #include "Widgets/SWindow.h"
 
 class FAdvancedPreviewScene;
diff --git a/Engine/Plugins/Mutable/Source/CustomizableObjectEditor/Private/MuCOE/SCustomizableObjectEditorTagExplorer.cpp b/Engine/Plugins/Mutable/Source/CustomizableObjectEditor/Private/MuCOE/SCustomizableObjectEditorTagExplorer.cpp
--- a/Engine/Plugins/Mutable/Source/CustomizableObjectEditor/Private/MuCOE/SCustomizableObjectEditorTagExplorer.cpp
+++ b/Engine/Plugins/Mutable/Source/CustomizableObjectEditor/Private/MuCOE/SCustomizableObjectEditorTagExplorer.cpp
@@ -269,7 +269,7 @@ void SCustomizableObjectEditorTagExplorer::OnTagTableSelectionChanged(TWeakObjec
 		GEditor->GetEditorSubsystem<UAssetEditorSubsystem>()->OpenEditorForAsset(Object);
 
 		// Find it
-		TSharedPtr<ICustomizableObjectEditor> Editor = Entry->GetGraphEditor();
+		TSharedPtr<FCustomizableObjectGraphEditorToolkit> Editor = Entry->GetGraphEditor();
 		Editor->SelectNode(Entry.Get());
 
 		if (ListViewWidget.IsValid())
diff --git a/Engine/Plugins/Mutable/Source/CustomizableObjectEditor/Private/MuCOE/SCustomizableObjectEditorViewport.cpp b/Engine/Plugins/Mutable/Source/CustomizableObjectEditor/Private/MuCOE/SCustomizableObjectEditorViewport.cpp
--- a/Engine/Plugins/Mutable/Source/CustomizableObjectEditor/Private/MuCOE/SCustomizableObjectEditorViewport.cpp
+++ b/Engine/Plugins/Mutable/Source/CustomizableObjectEditor/Private/MuCOE/SCustomizableObjectEditorViewport.cpp
@@ -145,7 +145,7 @@ FText SCustomizableObjectEditorViewport::GetCompileErrorOverlayText() const
 		return {};
 	}
 
-	const UObject* EditingObject = (*Editor->GetObjectsCurrentlyBeingEdited())[0];
+	const UObject* EditingObject = Editor->GetObjectBeingEdited();
 
 	const UCustomizableObject* Object = nullptr;
 	const UCustomizableObjectInstance* Instance = nullptr;
diff --git a/Engine/Plugins/Mutable/Source/CustomizableObjectEditor/Private/MuCOE/SMutableTagListWidget.cpp b/Engine/Plugins/Mutable/Source/CustomizableObjectEditor/Private/MuCOE/SMutableTagListWidget.cpp
--- a/Engine/Plugins/Mutable/Source/CustomizableObjectEditor/Private/MuCOE/SMutableTagListWidget.cpp
+++ b/Engine/Plugins/Mutable/Source/CustomizableObjectEditor/Private/MuCOE/SMutableTagListWidget.cpp
@@ -281,34 +281,6 @@ void SMutableTagComboBox::RefreshOptions()
 					}
 				}
 			}
-
-			// Texture variation nodes
-			{
-				UCustomizableObjectNodeTextureVariation* VariationNode = Cast<UCustomizableObjectNodeTextureVariation>(CandidateNode);
-				if (VariationNode)
-				{
-					for (const FCustomizableObjectTextureVariation& Var : VariationNode->Variations)
-					{
-						TSharedPtr<SMutableSearchComboBox::FFilteredOption> NodeOption = AddNodeHierarchyOptions(VariationNode, AddedOptions);
-
-						if (Var.Tag.IsEmpty())
-						{
-							continue;
-						}
-
-						if (!bAllowInternalTags && UCustomizableObjectNode::IsInternalTag(Var.Tag))
-						{
-							continue;
-						}
-
-						TSharedRef Option = MakeShared<SMutableSearchComboBox::FFilteredOption>();
-						Option->ActualOption = Var.Tag;
-						Option->DisplayOption = Var.Tag;
-						Option->Parent = NodeOption;
-						TagComboOptionsSource.Add(Option);
-					}
-				}
-			}
 		}
 	}
 }
diff --git a/Engine/Plugins/Mutable/Source/CustomizableObjectEditor/Private/MuCOE/Widgets/SCustomizableObjectMacroLibraryList.cpp b/Engine/Plugins/Mutable/Source/CustomizableObjectEditor/Private/MuCOE/Widgets/SCustomizableObjectMacroLibraryList.cpp
--- a/Engine/Plugins/Mutable/Source/CustomizableObjectEditor/Private/MuCOE/Widgets/SCustomizableObjectMacroLibraryList.cpp
+++ b/Engine/Plugins/Mutable/Source/CustomizableObjectEditor/Private/MuCOE/Widgets/SCustomizableObjectMacroLibraryList.cpp
@@ -27,8 +27,6 @@ public:
 		RowItem = InRowItem;
 		OnRemoveMacro = Args._OnRemoveMacro;
 
-		FText MacroName = RowItem->Macro.IsValid() ? FText::FromName(RowItem->Macro->Name) : FText::FromString("#Error# - Null Macro");
-
 		this->ChildSlot
 			[
 				SNew(SHorizontalBox)
@@ -43,7 +41,7 @@ public:
 				.Padding(5.0f, 5.0f, 0.0f, 5.0f)
 				.AutoWidth()
 				[
-					SNew(STextBlock).Text(MacroName)
+					SNew(STextBlock).Text(this, &SMacroLibraryTreeRow::GetMacroName)
 				]
 
 				+ SHorizontalBox::Slot()
@@ -68,6 +66,13 @@ public:
 			, InOwnerTableView);
 	}
 
+
+	FText GetMacroName() const
+	{
+		return RowItem->Macro.IsValid() ? FText::FromName(RowItem->Macro->Name) : FText::FromString("#Error# - Null Macro");
+	}
+
+
 	void RemoveMacro()
 	{
 		if (RowItem && RowItem->Macro.IsValid())
@@ -97,50 +102,59 @@ void SCustomizableObjectMacroLibraryList::Construct(const FArguments& InArgs)
 	GenerateRowView();
 
 	SAssignNew(ListView, SListView<TSharedPtr<FMacroTreeEntry>>)
-		.SelectionMode(ESelectionMode::Single)
-		.ListItemsSource(&ListViewSource)
-		.OnGenerateRow_Lambda([this](TSharedPtr<FMacroTreeEntry> InItem, const TSharedRef<STableViewBase>& InOwnerTable) -> TSharedRef<ITableRow>
-		{
-			TSharedRef<SMacroLibraryTreeRow> Row = SNew(SMacroLibraryTreeRow, InOwnerTable, InItem)
-			.OnRemoveMacro(SMacroLibraryTreeRow::FOnRemoveMacroButtonClickedDelegate::CreateSP(this, &SCustomizableObjectMacroLibraryList::OnRemoveCurrentMacro));
+	.SelectionMode(ESelectionMode::Single)
+	.ListItemsSource(&ListViewSource)
+	.OnGenerateRow_Lambda([this](TSharedPtr<FMacroTreeEntry> InItem, const TSharedRef<STableViewBase>& InOwnerTable) -> TSharedRef<ITableRow>
+	{
+		TSharedRef<SMacroLibraryTreeRow> Row = SNew(SMacroLibraryTreeRow, InOwnerTable, InItem)
+		.OnRemoveMacro(SMacroLibraryTreeRow::FOnRemoveMacroButtonClickedDelegate::CreateSP(this, &SCustomizableObjectMacroLibraryList::OnRemoveCurrentMacro));
 
-			return Row;
-		})
-		.OnSelectionChanged_Lambda([this](TSharedPtr<FMacroTreeEntry> InItem, ESelectInfo::Type)
+		return Row;
+	})
+	.OnSelectionChanged_Lambda([this](TSharedPtr<FMacroTreeEntry> InItem, ESelectInfo::Type SelectInfo)
+	{
+		if (InItem.IsValid() && InItem->Macro.IsValid())
 		{
-			if (InItem.IsValid() && InItem->Macro.IsValid())
-			{
-				OnSelectMacro.ExecuteIfBound(InItem->Macro.Get());
-			}
-		});
+			OnSelectMacro.ExecuteIfBound(InItem->Macro.Get());
+		}
+	});
 
-		this->ChildSlot
+	// Select the firs macro on create the list view
+	if (ListView.IsValid() && ListViewSource.Num())
+	{
+		ListView->SetSelection(ListViewSource[0]);
+	}
+
+	this->ChildSlot
+	[
+		SNew(SVerticalBox)
+		+SVerticalBox::Slot()
+		.AutoHeight()
 		[
-			SNew(SVerticalBox)
-			+SVerticalBox::Slot()
-			.AutoHeight()
-			[
-				SNew(SPositiveActionButton)
-				.OnClicked_Lambda([this]() 
+			SNew(SPositiveActionButton)
+			.OnClicked_Lambda([this]() 
+			{
+				if (OnAddMacroButtonClicked.IsBound())
 				{
-					if (OnAddMacroButtonClicked.IsBound())
-					{
-						OnAddMacroButtonClicked.Execute();
-						GenerateRowView();
-						ListView->RequestListRefresh();
-						return FReply::Handled();
-					}
-
-					return FReply::Unhandled();
-				})
-			]
-
-			+ SVerticalBox::Slot()
-			.AutoHeight()
-			[
-				ListView.ToSharedRef()
-			]
-		];
+					OnAddMacroButtonClicked.Execute();
+					GenerateRowView();
+					
+					ListView->SetSelection(ListViewSource.Last());
+					ListView->RequestListRefresh();
+					
+					return FReply::Handled();
+				}
+
+				return FReply::Unhandled();
+			})
+		]
+
+		+ SVerticalBox::Slot()
+		.AutoHeight()
+		[
+			ListView.ToSharedRef()
+		]
+	];
 }
 
 
@@ -171,4 +185,19 @@ void SCustomizableObjectMacroLibraryList::OnRemoveCurrentMacro(UCustomizableObje
 	}
 }
 
+void SCustomizableObjectMacroLibraryList::SetSelectedMacro(UCustomizableObjectMacro* MacroToSelect)
+{
+	if (ListView.IsValid())
+	{
+		for (TSharedPtr<FMacroTreeEntry> MacroElement : ListViewSource)
+		{
+			if (MacroElement->Macro == MacroToSelect)
+			{
+				ListView->SetSelection(MacroElement);
+				break;
+			}
+		}
+	}
+}
+
 #undef LOCTEXT_NAMESPACE
diff --git a/Engine/Plugins/Mutable/Source/CustomizableObjectEditor/Private/MuCOE/Widgets/SCustomizableObjectMacroLibraryList.h b/Engine/Plugins/Mutable/Source/CustomizableObjectEditor/Private/MuCOE/Widgets/SCustomizableObjectMacroLibraryList.h
--- a/Engine/Plugins/Mutable/Source/CustomizableObjectEditor/Private/MuCOE/Widgets/SCustomizableObjectMacroLibraryList.h
+++ b/Engine/Plugins/Mutable/Source/CustomizableObjectEditor/Private/MuCOE/Widgets/SCustomizableObjectMacroLibraryList.h
@@ -35,6 +35,9 @@ public:
 	/** Callback to regenerate the List View when a Macro has been removed. It is also used to comunicate the Macro Library editor. */
 	void OnRemoveCurrentMacro(UCustomizableObjectMacro* MacroToDelete);
 
+	/** Set the Selected Macro */
+	void SetSelectedMacro(UCustomizableObjectMacro* MacroToSelect);
+
 private:
 
 	/** Pointer to the Macro Library that contains the macros. */
diff --git a/Engine/Plugins/Mutable/Source/CustomizableObjectEditor/Public/MuCOE/CustomizableObjectGraphEditorToolkit.h b/Engine/Plugins/Mutable/Source/CustomizableObjectEditor/Public/MuCOE/CustomizableObjectGraphEditorToolkit.h
deleted file mode 100644
diff --git a/Engine/Plugins/Mutable/Source/CustomizableObjectEditor/Public/MuCOE/ICustomizableObjectInstanceEditor.h b/Engine/Plugins/Mutable/Source/CustomizableObjectEditor/Public/MuCOE/ICustomizableObjectInstanceEditor.h
--- a/Engine/Plugins/Mutable/Source/CustomizableObjectEditor/Public/MuCOE/ICustomizableObjectInstanceEditor.h
+++ b/Engine/Plugins/Mutable/Source/CustomizableObjectEditor/Public/MuCOE/ICustomizableObjectInstanceEditor.h
@@ -23,7 +23,7 @@ class SCustomizableObjectEditorAdvancedPreviewSettings;
 /**
  * Public interface to Customizable Object Instance Editor
  */
-class ICustomizableObjectInstanceEditor : public FAssetEditorToolkit
+class ICustomizableObjectInstanceEditor
 {
 public:
 	virtual UCustomizableObjectInstance* GetPreviewInstance() = 0;
@@ -98,4 +98,6 @@ public:
 	virtual bool ShowLightingSettings() = 0;
 
 	virtual bool ShowProfileManagementOptions() = 0;
+
+	virtual const UObject* GetObjectBeingEdited() = 0;
 };
diff --git a/Engine/Plugins/Mutable/Source/CustomizableObjectEditor/Public/MuCOE/Nodes/CustomizableObjectNode.h b/Engine/Plugins/Mutable/Source/CustomizableObjectEditor/Public/MuCOE/Nodes/CustomizableObjectNode.h
--- a/Engine/Plugins/Mutable/Source/CustomizableObjectEditor/Public/MuCOE/Nodes/CustomizableObjectNode.h
+++ b/Engine/Plugins/Mutable/Source/CustomizableObjectEditor/Public/MuCOE/Nodes/CustomizableObjectNode.h
@@ -10,6 +10,7 @@
 #include "CustomizableObjectNode.generated.h"
 
 class FArchive;
+class FCustomizableObjectGraphEditorToolkit;
 class SWidget;
 class UCustomizableObjectNodeObject;
 class UCustomizableObjectNodeRemapPins;
@@ -17,6 +18,7 @@ class UCustomizableObjectNodeRemapPinsByName;
 class UCustomizableObjectNodeRemapPinsByPosition;
 class IDetailsView;
 class ICustomizableObjectEditor;
+
 struct FPropertyChangedEvent;
 
 /** Abstract base class for all pin data. */
@@ -151,7 +153,7 @@ public:
 	// Used to replace Ids of referenced nodes by their new ids after duplicating the CustomizableObject.
 	virtual void UpdateReferencedNodeId(const FGuid& NewGuid) {};
 
-	TSharedPtr<ICustomizableObjectEditor> GetGraphEditor() const;
+	TSharedPtr<FCustomizableObjectGraphEditorToolkit> GetGraphEditor() const;
 
 	virtual bool ProvidesCustomPinRelevancyTest() const { return false; } // Override and return true if IsPinRelevant is overridden.
 	virtual bool IsPinRelevant(const UEdGraphPin*) const { return true; } // Override if default pins allocated with AllocateDefaultPins do not provide enough information.
@@ -284,7 +286,8 @@ public:
 	/** returs true if this node can be added to a COMacro Graph */
 	virtual bool IsNodeSupportedInMacros() { return true; } 
 
-	bool IsMacroNode() const;
+	/** Returns true if this node is inside a Macro Graph */
+	bool IsInMacro() const;
 	
 	UPROPERTY()
 	FPostEditChangePropertyDelegate PostEditChangePropertyDelegate;
