[Mutable] Fixed incremental build error.

#rnx

diff --git a/Engine/Plugins/Mutable/Source/CustomizableObjectEditor/Private/MuCOE/CustomizableInstanceDetails.cpp b/Engine/Plugins/Mutable/Source/CustomizableObjectEditor/Private/MuCOE/CustomizableInstanceDetails.cpp
--- a/Engine/Plugins/Mutable/Source/CustomizableObjectEditor/Private/MuCOE/CustomizableInstanceDetails.cpp
+++ b/Engine/Plugins/Mutable/Source/CustomizableObjectEditor/Private/MuCOE/CustomizableInstanceDetails.cpp
@@ -71,10 +71,10 @@ void FCustomizableInstanceDetails::CustomizeDetails(const TSharedPtr<IDetailLayo
 	CustomInstance = Cast<UCustomizableObjectInstance>(DetailsView->GetSelectedObjects()[0].Get());
 	check(CustomInstance.IsValid());
 
-	TSharedPtr<IToolkit> FoundAssetEditor;
+	TSharedPtr<IToolkit> FoundAssetEditor = FToolkitManager::Get().FindEditorForAsset(CustomInstance.Get());
 
 	// Tab spawned in a COEInstanceEditor
-	if (FoundAssetEditor = FToolkitManager::Get().FindEditorForAsset(CustomInstance.Get())) 
+	if (FoundAssetEditor.IsValid()) 
 	{
 		if (TSharedPtr<FCustomizableObjectInstanceEditor> InstanceEditor = StaticCastSharedPtr<FCustomizableObjectInstanceEditor>(FoundAssetEditor))
 		{
@@ -82,10 +82,12 @@ void FCustomizableInstanceDetails::CustomizeDetails(const TSharedPtr<IDetailLayo
 		}
 	}
 
-	if (!FoundAssetEditor)
+	if (!FoundAssetEditor.IsValid())
 	{
+		FoundAssetEditor = FToolkitManager::Get().FindEditorForAsset(CustomInstance->GetCustomizableObject());
+
 		// Tab spawned in a COEditor
-		if (FoundAssetEditor = FToolkitManager::Get().FindEditorForAsset(CustomInstance->GetCustomizableObject()))
+		if (FoundAssetEditor.IsValid())
 		{
 			if (TSharedPtr<FCustomizableObjectEditor> ObjectEditor = StaticCastSharedPtr<FCustomizableObjectEditor>(FoundAssetEditor))
 			{
