[mutable] Various fixes and preventive changes.
- Fixed unsafe access to TextureFormatOddle module pointer. The pointer may be null if the "Texture Oddle" plugin is not enabled.
- Fixed case where the addition of more options to test would not affect the options being tested (CO Validator)

diff --git a/Engine/Plugins/Mutable/Source/MutableTools/Private/MuT/UnrealPixelFormatOverride.cpp b/Engine/Plugins/Mutable/Source/MutableTools/Private/MuT/UnrealPixelFormatOverride.cpp
--- a/Engine/Plugins/Mutable/Source/MutableTools/Private/MuT/UnrealPixelFormatOverride.cpp
+++ b/Engine/Plugins/Mutable/Source/MutableTools/Private/MuT/UnrealPixelFormatOverride.cpp
@@ -82,24 +82,24 @@ void PrepareUnrealCompression()
 			
 			if (bHasFormat)
 			{
-				ITextureFormatModule* TextureFormatModule = FModuleManager::LoadModulePtr<ITextureFormatModule>(*TextureCompressionFormat);
-			
-				ITextureFormat* TextureFormat = TextureFormatModule->GetTextureFormat();
-				if (TextureFormat)
+				if (ITextureFormatModule* TextureFormatModule = FModuleManager::LoadModulePtr<ITextureFormatModule>(*TextureCompressionFormat))
 				{
-					FString FormatPrefix = TextureFormat->GetAlternateTextureFormatPrefix();
-					check(!FormatPrefix.IsEmpty());
+					if (ITextureFormat* TextureFormat = TextureFormatModule->GetTextureFormat())
+					{
+						FString FormatPrefix = TextureFormat->GetAlternateTextureFormatPrefix();
+						check(!FormatPrefix.IsEmpty());
 				
-					FName NewFormatName(FormatPrefix + TextureFormatName.ToString());
+						FName NewFormatName(FormatPrefix + TextureFormatName.ToString());
 
-					// check that prefixed name is one we support
-					// only apply prefix if it is in list
-					TArray<FName> SupportedFormats;
-					TextureFormat->GetSupportedFormats(SupportedFormats);
+						// check that prefixed name is one we support
+						// only apply prefix if it is in list
+						TArray<FName> SupportedFormats;
+						TextureFormat->GetSupportedFormats(SupportedFormats);
 
-					if (SupportedFormats.Contains(NewFormatName))
-					{
-						SPrefixedMutableTextureFormatNameTable[FormatEnumIndex] = NewFormatName;
+						if (SupportedFormats.Contains(NewFormatName))
+						{
+							SPrefixedMutableTextureFormatNameTable[FormatEnumIndex] = NewFormatName;
+						}
 					}
 				}
 			}
diff --git a/Engine/Plugins/Mutable/Source/MutableValidation/Private/MuV/AssetValidator_CustomizableObjects.cpp b/Engine/Plugins/Mutable/Source/MutableValidation/Private/MuV/AssetValidator_CustomizableObjects.cpp
--- a/Engine/Plugins/Mutable/Source/MutableValidation/Private/MuV/AssetValidator_CustomizableObjects.cpp
+++ b/Engine/Plugins/Mutable/Source/MutableValidation/Private/MuV/AssetValidator_CustomizableObjects.cpp
@@ -148,6 +148,7 @@ EDataValidationResult UAssetValidator_CustomizableObjects::IsCustomizableObjectV
 	{
 		// Run Sync compilation -> Warning : Potentially long operation -------------
 		TSharedRef<FCompilationRequest> CompileRequest = MakeShared<FCompilationRequest>(*RootObject);
+		CompileRequest->Options = Options;
 		CompileRequest->bAsync = false;
 		ICustomizableObjectEditorModulePrivate::GetChecked().EnqueueCompileRequest(CompileRequest);
 		// --------------------------------------------------------------------------
