[Mutable] Fixed incremental build warning related to a pointer check.

#rnx

diff --git a/Engine/Plugins/Mutable/Source/CustomizableObjectEditor/Private/MuCOE/CustomizableObjectLayout.cpp b/Engine/Plugins/Mutable/Source/CustomizableObjectEditor/Private/MuCOE/CustomizableObjectLayout.cpp
--- a/Engine/Plugins/Mutable/Source/CustomizableObjectEditor/Private/MuCOE/CustomizableObjectLayout.cpp
+++ b/Engine/Plugins/Mutable/Source/CustomizableObjectEditor/Private/MuCOE/CustomizableObjectLayout.cpp
@@ -85,14 +85,15 @@ void UCustomizableObjectLayout::GenerateAutomaticBlocksFromUVs()
 {
 	UCustomizableObjectNode* Node = Cast<UCustomizableObjectNode>(GetOuter());
 	TSoftObjectPtr<const UObject> Mesh = GetMesh();
-	TSharedPtr<FCustomizableObjectEditor> Editor = StaticCastSharedPtr<FCustomizableObjectEditor>(Node->GetGraphEditor());
 
-	if (!Node || !Mesh || !Editor)
+	if (!Node || !Mesh)
 	{
 		return;
 	}
 
-	if (AutomaticBlocksStrategy == ECustomizableObjectLayoutAutomaticBlocksStrategy::Ignore)
+	TSharedPtr<FCustomizableObjectEditor> Editor = StaticCastSharedPtr<FCustomizableObjectEditor>(Node->GetGraphEditor());
+
+	if (AutomaticBlocksStrategy == ECustomizableObjectLayoutAutomaticBlocksStrategy::Ignore || !Editor)
 	{
 		return;
 	}
