[mutable] Added versions hash as part of the CO data DDC Key
- Removed current version from the hash (included in the new hash)
- Added computed hash to DDK key based on the enum that defines the various versions.
- Made enum class editor only to prevent it from spilling out into packages.

#rnx

diff --git a/Engine/Plugins/Mutable/Source/CustomizableObject/Internal/MuCO/CustomizableObjectPrivate.h b/Engine/Plugins/Mutable/Source/CustomizableObject/Internal/MuCO/CustomizableObjectPrivate.h
--- a/Engine/Plugins/Mutable/Source/CustomizableObject/Internal/MuCO/CustomizableObjectPrivate.h
+++ b/Engine/Plugins/Mutable/Source/CustomizableObject/Internal/MuCO/CustomizableObjectPrivate.h
@@ -1436,6 +1436,157 @@ public:
 };
 
 
+#if WITH_EDITORONLY_DATA 
+// This is a manual version number for the binary blobs in this asset.
+// Increasing it invalidates all the previously compiled models.
+UENUM()
+enum class ECustomizableObjectVersions : int32
+{
+	FirstEnumeratedVersion = 450,
+
+	DeterminisiticMeshVertexIds,
+
+	NumRuntimeReferencedTextures,
+	
+	DeterminisiticLayoutBlockIds,
+
+	BackoutDeterminisiticLayoutBlockIds,
+
+	FixWrappingProjectorLayoutBlockId,
+
+	MeshReferenceSupport,
+
+	ImproveMemoryUsageForStreamableBlocks,
+
+	FixClipMeshWithMeshCrash,
+
+	SkeletalMeshLODSettingsSupport,
+
+	RemoveCustomCurve,
+
+	AddEditorGamePlayTags,
+
+	AddedParameterThumbnailsToEditor,
+
+	ComponentsLODsRedesign,
+
+	ComponentsLODsRedesign2,
+
+	LayoutToPOD,
+
+	AddedRomFlags,
+
+	LayoutNodeCleanup,
+
+	AddSurfaceAndMeshMetadata,
+
+	TablesPropertyNameBug,
+
+	DataTablesParamTrackingForCompileOnlySelected,
+
+	CompilationOptimizationsMeshFormat,
+
+	ModelStreamableBulkData,
+
+	LayoutBlocksAsInt32,
+	
+	IntParameterOptionDataTable,
+
+	RemoveLODCountLimit,
+
+	IntParameterOptionDataTablePartialBackout,
+
+	IntParameterOptionDataTablePartialRestore,
+
+	CorrectlySerializeTableToParamNames,
+	
+	AddMaterialSlotNameIndexToSurfaceMetadata,
+
+	NodeComponentMesh,
+	
+	MoveEditNodesToModifiers,
+
+	DerivedDataCache,
+
+	ComponentsArray,
+
+	FixComponentNames,
+
+	AddedFaceCullStrategyToSomeOperations,
+
+	DDCParticipatingObjects,
+
+	GroupRomsBySource,
+	
+	RemovedGroupRomsBySource,
+
+	ReGroupRomsBySource,
+
+	UIMetadataGameplayTags,
+
+	TransformInMeshModifier,
+	
+	SurfaceMetadataSlotNameIndexToName,
+
+	BulkDataFilesNumFilesLimit,
+
+	RemoveModifiersHack,
+
+	SurfaceMetadataSerialized,
+
+	FixesForMeshSectionMultipleOutputs,
+
+	ImageParametersInServerBuilds,
+
+	RemovedUnnecessarySerializationVersioning,
+
+	AddTextureCompressionSettingCompilationInfo,
+
+	RestructureConstantImageData,
+
+	RestructureConstantMeshData,
+
+	RestructureRomData,
+
+	RestructureRomDataRemovingRomHash,
+
+	ModifiedRomCompiledDataSerialization,
+
+	ModelResourcesExtensionData,
+
+	LODsPerComponent,
+
+	LODsPerComponentTypeMismatch,
+
+	ImageHiResLODsUseLODGroupInfo,
+
+	MovedTableRowNoneGenerationToUnreal,
+
+	RemoveObsoletMeshInterpolateAndGeometryOp,
+
+	RemoveObsoleteDataTypesFromEnum,
+	ConvertModelResourcesToUObject,
+
+	RemoveObsoletImageGradientOp,
+
+	MeshReferencesExtendedForCompilation,
+
+	RemoveObsoleteBoolOps,
+
+	AddOverlayMaterials,
+
+	PrefetchHighQualityMipMaps,
+
+	AddedMeshParameterOp,
+
+	AddedMeshParameterOpForDDCPollution,
+	
+	// -----<new versions can be added above this line>--------
+	LastCustomizableObjectVersion
+};
+#endif
+
+
 UCLASS()
 class CUSTOMIZABLEOBJECT_API UCustomizableObjectPrivate : public UObject
 {
@@ -1699,156 +1850,9 @@ public:
 
 	UPROPERTY()
 	int32 ImageTiling = 0;
-#endif
-	
-	// This is a manual version number for the binary blobs in this asset.
-	// Increasing it invalidates all the previously compiled models.
-	enum ECustomizableObjectVersions
-	{
-		FirstEnumeratedVersion = 450,
-
-		DeterminisiticMeshVertexIds,
-
-		NumRuntimeReferencedTextures,
-		
-		DeterminisiticLayoutBlockIds,
-
-		BackoutDeterminisiticLayoutBlockIds,
-
-		FixWrappingProjectorLayoutBlockId,
-
-		MeshReferenceSupport,
-
-		ImproveMemoryUsageForStreamableBlocks,
-
-		FixClipMeshWithMeshCrash,
-
-		SkeletalMeshLODSettingsSupport,
-
-		RemoveCustomCurve,
-
-		AddEditorGamePlayTags,
-
-		AddedParameterThumbnailsToEditor,
-
-		ComponentsLODsRedesign,
-
-		ComponentsLODsRedesign2,
-
-		LayoutToPOD,
-
-		AddedRomFlags,
-
-		LayoutNodeCleanup,
-
-		AddSurfaceAndMeshMetadata,
-
-		TablesPropertyNameBug,
-
-		DataTablesParamTrackingForCompileOnlySelected,
-
-		CompilationOptimizationsMeshFormat,
-
-		ModelStreamableBulkData,
-
-		LayoutBlocksAsInt32,
-		
-		IntParameterOptionDataTable,
-
-		RemoveLODCountLimit,
-
-		IntParameterOptionDataTablePartialBackout,
-
-		IntParameterOptionDataTablePartialRestore,
-
-		CorrectlySerializeTableToParamNames,
-		
-		AddMaterialSlotNameIndexToSurfaceMetadata,
-
-		NodeComponentMesh,
-		
-		MoveEditNodesToModifiers,
-
-		DerivedDataCache,
-
-		ComponentsArray,
-
-		FixComponentNames,
 
-		AddedFaceCullStrategyToSomeOperations,
-
-		DDCParticipatingObjects,
-
-		GroupRomsBySource,
-		
-		RemovedGroupRomsBySource,
-
-		ReGroupRomsBySource,
-
-		UIMetadataGameplayTags,
-
-		TransformInMeshModifier,
-		
-		SurfaceMetadataSlotNameIndexToName,
-
-		BulkDataFilesNumFilesLimit,
-
-		RemoveModifiersHack,
-
-		SurfaceMetadataSerialized,
-
-		FixesForMeshSectionMultipleOutputs,
-
-		ImageParametersInServerBuilds,
-
-		RemovedUnnecessarySerializationVersioning,
-
-		AddTextureCompressionSettingCompilationInfo,
-
-		RestructureConstantImageData,
-
-		RestructureConstantMeshData,
-
-		RestructureRomData,
-
-		RestructureRomDataRemovingRomHash,
-
-		ModifiedRomCompiledDataSerialization,
-
-		ModelResourcesExtensionData,
-
-		LODsPerComponent,
-
-		LODsPerComponentTypeMismatch,
-
-		ImageHiResLODsUseLODGroupInfo,
-
-		MovedTableRowNoneGenerationToUnreal,
-
-		RemoveObsoletMeshInterpolateAndGeometryOp,
-
-		RemoveObsoleteDataTypesFromEnum,
-		ConvertModelResourcesToUObject,
-
-		RemoveObsoletImageGradientOp,
-
-		MeshReferencesExtendedForCompilation,
-
-		RemoveObsoleteBoolOps,
-
-		AddOverlayMaterials,
-
-		PrefetchHighQualityMipMaps,
-
-		AddedMeshParameterOp,
-
-		AddedMeshParameterOpForDDCPollution,
-
-		// -----<new versions can be added above this line>--------
-		LastCustomizableObjectVersion
-	};
-	
-	static constexpr int32 CurrentSupportedVersion = ECustomizableObjectVersions::LastCustomizableObjectVersion;
+	static constexpr int32 CurrentSupportedVersion = static_cast<int32>(ECustomizableObjectVersions::LastCustomizableObjectVersion);
+#endif
 	
 	static constexpr int32 DerivedDataVersion = 6;
 };
diff --git a/Engine/Plugins/Mutable/Source/CustomizableObject/Private/MuCO/CustomizableObject.cpp b/Engine/Plugins/Mutable/Source/CustomizableObject/Private/MuCO/CustomizableObject.cpp
--- a/Engine/Plugins/Mutable/Source/CustomizableObject/Private/MuCO/CustomizableObject.cpp
+++ b/Engine/Plugins/Mutable/Source/CustomizableObject/Private/MuCO/CustomizableObject.cpp
@@ -827,7 +827,7 @@ void UCustomizableObjectPrivate::SaveEmbeddedData(FArchive& Ar)
 
 	TSharedPtr<mu::FModel> Model = GetModel();
 
-	int32 InternalVersion = Model ? CurrentSupportedVersion : -1;
+	int32 InternalVersion = Model ? CurrentSupportedVersion : INDEX_NONE;
 	Ar << InternalVersion;
 
 	if (Model)
@@ -855,7 +855,7 @@ void UCustomizableObjectPrivate::LoadEmbeddedData(FArchive& Ar)
 
 	// If this fails, something went wrong with the packaging: we have data that belongs
 	// to a different version than the code.
-	if (ensure(CurrentSupportedVersion == InternalVersion))
+	if (ensure(InternalVersion != INDEX_NONE))
 	{		
 		// Load model
 		FUnrealMutableInputStream Stream(Ar);
@@ -1337,6 +1337,31 @@ const FString& UCustomizableObject::GetParameterName(int32 ParamIndex) const
 	return s_EmptyString;
 }
 
+#if WITH_EDITOR
+uint32 GetECustomizableObjectVersionEnumHash()
+{
+	static const uint32 VersionsHash = []() -> uint32
+	{
+		UEnum* Enum = StaticEnum<ECustomizableObjectVersions>();
+		check(Enum);
+		
+		const int32 StartIndex = Enum->GetIndexByValue(static_cast<int64>(ECustomizableObjectVersions::FirstEnumeratedVersion));
+		const int32 EndIndex = Enum->GetIndexByValue(static_cast<int64>(ECustomizableObjectVersions::LastCustomizableObjectVersion));
+		check(StartIndex < EndIndex);
+		
+		uint32 CombinedHash = 0;
+		for (int32 Index = StartIndex; Index <= EndIndex; Index++)
+		{
+			const FString VersionString = Enum->GetNameStringByIndex(Index);
+			CombinedHash = HashCombine( GetTypeHash(VersionString), CombinedHash);
+		}
+		
+		return CombinedHash;
+	}();
+
+	return VersionsHash;
+}
+#endif
 
 void UCustomizableObjectPrivate::UpdateParameterPropertiesFromModel(const TSharedPtr<mu::FModel>& Model)
 {
@@ -2970,8 +2995,8 @@ TArray<uint8> UCustomizableObjectPrivate::BuildDerivedDataKey(FCompilationOption
 	}
 	
 	{
-		int32 CurrentVersion = CurrentSupportedVersion;
-		Ar << CurrentVersion;	
+		uint32 VersionsHash = GetECustomizableObjectVersionEnumHash();
+		Ar << VersionsHash;
 	}
 
 	// Custom Version
