[Mutable] Simplify LoadAdditionalAssetsAsync Task Graph event.

A dummy event is no longer created in the constructor.

#rnx

diff --git a/Engine/Plugins/Mutable/Source/CustomizableObject/Internal/MuCO/CustomizableObjectInstancePrivate.h b/Engine/Plugins/Mutable/Source/CustomizableObject/Internal/MuCO/CustomizableObjectInstancePrivate.h
--- a/Engine/Plugins/Mutable/Source/CustomizableObject/Internal/MuCO/CustomizableObjectInstancePrivate.h
+++ b/Engine/Plugins/Mutable/Source/CustomizableObject/Internal/MuCO/CustomizableObjectInstancePrivate.h
@@ -230,14 +230,10 @@ public:
 	// \TODO: Create a key based on a struct instead of generating strings dynamically.
 	UPROPERTY(Transient)
 	TMap<FString, TWeakObjectPtr<UTexture2D>> TextureReuseCache;
-
-	UE::Tasks::FTaskEvent AssetAsyncLoadCompletionEvent = UE::Tasks::FTaskEvent(TEXT(""));
-
+	
 	// Only used in LiveUpdateMode to reuse core instances between updates and their temp data to speed up updates, but spend way more memory
 	mu::FInstance::FID LiveUpdateModeInstanceID = 0;
 
-	UCustomizableInstancePrivate();
-
 #if WITH_EDITOR
 	virtual void PostDuplicate(bool bDuplicateForPIE) override;
 
@@ -266,7 +262,7 @@ public:
 	 * If no StreamableManager is provided, it will load assets synchronously. */
 	UE::Tasks::FTask LoadAdditionalAssetsAndData(const TSharedRef<FUpdateContextPrivate>& OperationData);
 
-	void AdditionalAssetsAsyncLoaded(const TSharedRef<FUpdateContextPrivate> OperationData);
+	void AdditionalAssetsAsyncLoaded(TSharedRef<FUpdateContextPrivate> OperationData, UE::Tasks::FTaskEvent Event);
 
 	void TickUpdateCloseCustomizableObjects(UCustomizableObjectInstance& Public, FMutableInstanceUpdateMap& InOutRequestedUpdates);
 	void UpdateInstanceIfNotGenerated(UCustomizableObjectInstance& Public, FMutableInstanceUpdateMap& InOutRequestedUpdates);
diff --git a/Engine/Plugins/Mutable/Source/CustomizableObject/Private/MuCO/CustomizableObjectInstance.cpp b/Engine/Plugins/Mutable/Source/CustomizableObject/Private/MuCO/CustomizableObjectInstance.cpp
--- a/Engine/Plugins/Mutable/Source/CustomizableObject/Private/MuCO/CustomizableObjectInstance.cpp
+++ b/Engine/Plugins/Mutable/Source/CustomizableObject/Private/MuCO/CustomizableObjectInstance.cpp
@@ -335,16 +335,6 @@ void UCustomizableInstancePrivate::PrepareForUpdate(const TSharedRef<FUpdateCont
 }
 
 
-UCustomizableInstancePrivate::UCustomizableInstancePrivate()
-{
-	UE_AUTORTFM_ONCOMMIT(this)
-	{
-		// MakeCompletedEvent does not exist. Trigger the placeholder event outside of transactional work.
-		AssetAsyncLoadCompletionEvent.Trigger();
-	};
-}
-
-
 #if WITH_EDITOR
 void UCustomizableInstancePrivate::PostDuplicate(bool bDuplicateForPIE)
 {
@@ -5396,24 +5386,26 @@ UE::Tasks::FTask UCustomizableInstancePrivate::LoadAdditionalAssetsAndData(const
 		AssetsToStream.Add(StreamedData.GetPath().ToSoftObjectPath());
 #endif
 	}
+
+	TArray<UE::Tasks::FTask, TInlineAllocator<2>> Prerequisites;
 	
 	if (AssetsToStream.Num() > 0)
 	{
-		check(AssetAsyncLoadCompletionEvent.IsCompleted());
-		AssetAsyncLoadCompletionEvent = UE::Tasks::FTaskEvent(TEXT("AssetsStreamed"));
-
 #if WITH_EDITOR
 		// TODO: Remove with UE-217665 when the underlying bug in the ColorPicker is solved
 		// Disable the Slate throttling, otherwise the AsyncLoad may not complete until the editor window is clicked on due to a bug in
 		// some widgets such as the ColorPicker's throttling handling
 		FSlateThrottleManager::Get().DisableThrottle(true);
 #endif
+		
+		UE::Tasks::FTaskEvent Event(TEXT("AssetsStreamed"));
+		Prerequisites.Add(Event);
 
 		UCustomizableObjectSystemPrivate* PrivateSystem = UCustomizableObjectSystem::GetInstance()->GetPrivate();
-
+		
 		PrivateSystem->StreamableManager->RequestAsyncLoad(
 			AssetsToStream,
-			FStreamableDelegate::CreateUObject(this, &UCustomizableInstancePrivate::AdditionalAssetsAsyncLoaded, OperationData),
+			FStreamableDelegate::CreateUObject(this, &UCustomizableInstancePrivate::AdditionalAssetsAsyncLoaded, OperationData, Event),
 			CVarMutableHighPriorityLoading.GetValueOnAnyThread() ? FStreamableManager::AsyncLoadHighPriority : FStreamableManager::DefaultAsyncLoadPriority);
 	}
 	
@@ -5454,9 +5446,10 @@ UE::Tasks::FTask UCustomizableInstancePrivate::LoadAdditionalAssetsAndData(const
 	
 	// Stream files
 	UE::Tasks::FTask StreamingTask = StreamRequest.Stream();
-	
+	Prerequisites.Add(StreamingTask);
+
 	return UE::Tasks::Launch(TEXT("CaptureOperationData"), [OperationData]() {}, // Keep a reference to make sure allocated memory is always alive.
-	UE::Tasks::Prerequisites(AssetAsyncLoadCompletionEvent, StreamingTask),
+	Prerequisites,
 	UE::Tasks::ETaskPriority::Inherit);
 }
 
@@ -5486,13 +5479,13 @@ const TArray<TObjectPtr<UMaterialInterface>>* UCustomizableObjectInstance::GetOv
 }
 
 
-void UCustomizableInstancePrivate::AdditionalAssetsAsyncLoaded(const TSharedRef<FUpdateContextPrivate> OperationData)
+void UCustomizableInstancePrivate::AdditionalAssetsAsyncLoaded(const TSharedRef<FUpdateContextPrivate> OperationData, UE::Tasks::FTaskEvent Event)
 {
 	MUTABLE_CPUPROFILER_SCOPE(UCustomizableInstancePrivate::AdditionalAssetsAsyncLoaded);
 
 	check(IsInGameThread())
 
-	AssetAsyncLoadCompletionEvent.Trigger();
+	Event.Trigger();
 
 	UCustomizableObjectPrivate* CustomizableObjectPrivate = GetPublic()->GetCustomizableObject()->GetPrivate();
 
