[Mutable] Make ModelStreameableBulkData references in ModelBulkReader weak.

- This will remove possible stale data not been cleaned until PrepareStreamingForObject is called.

#rnx

diff --git a/Engine/Plugins/Mutable/Source/CustomizableObject/Internal/MuCO/UnrealMutableModelDiskStreamer.h b/Engine/Plugins/Mutable/Source/CustomizableObject/Internal/MuCO/UnrealMutableModelDiskStreamer.h
--- a/Engine/Plugins/Mutable/Source/CustomizableObject/Internal/MuCO/UnrealMutableModelDiskStreamer.h
+++ b/Engine/Plugins/Mutable/Source/CustomizableObject/Internal/MuCO/UnrealMutableModelDiskStreamer.h
@@ -87,7 +87,7 @@ protected:
 		TMap<FOperationID, FReadRequest> CurrentReadRequests;
 		TMap<uint32, TSharedPtr<IAsyncReadFileHandle>> ReadFileHandles;
 
-		TSharedPtr<FModelStreamableBulkData> ModelStreamableBulkData;
+		TWeakPtr<FModelStreamableBulkData> ModelStreamableBulkData;
 
 #if WITH_EDITORONLY_DATA
 		// DDC files streaming
diff --git a/Engine/Plugins/Mutable/Source/CustomizableObject/Private/MuCO/UnrealMutableModelDiskStreamer.cpp b/Engine/Plugins/Mutable/Source/CustomizableObject/Private/MuCO/UnrealMutableModelDiskStreamer.cpp
--- a/Engine/Plugins/Mutable/Source/CustomizableObject/Private/MuCO/UnrealMutableModelDiskStreamer.cpp
+++ b/Engine/Plugins/Mutable/Source/CustomizableObject/Private/MuCO/UnrealMutableModelDiskStreamer.cpp
@@ -57,7 +57,6 @@ FUnrealMutableModelBulkReader::~FUnrealMutableModelBulkReader()
 {
 }
 
-
 bool FUnrealMutableModelBulkReader::PrepareStreamingForObject(UCustomizableObject* CustomizableObject)
 {
 	if (!CustomizableObject)
@@ -66,7 +65,7 @@ bool FUnrealMutableModelBulkReader::PrepareStreamingForObject(UCustomizableObjec
 		return false;
 	}
 
-	// See if we can free previuously allocated resources
+	// See if we can free previuously allocated resources.
 	for (int32 ObjectIndex = 0; ObjectIndex < Objects.Num(); )
 	{
 		// Close open file handles
@@ -95,9 +94,11 @@ bool FUnrealMutableModelBulkReader::PrepareStreamingForObject(UCustomizableObjec
 	{
 		FObjectData NewData;
 		NewData.Model = TWeakPtr<const mu::FModel>(CustomizableObject->GetPrivate()->GetModel());
-		NewData.ModelStreamableBulkData = CustomizableObject->GetPrivate()->GetModelStreamableBulkData();
 
-		if (!NewData.ModelStreamableBulkData || NewData.ModelStreamableBulkData->ModelStreamables.IsEmpty())
+		const TSharedPtr<FModelStreamableBulkData> ModelStreamableBulkData = CustomizableObject->GetPrivate()->GetModelStreamableBulkData();
+		NewData.ModelStreamableBulkData = TWeakPtr<FModelStreamableBulkData>(ModelStreamableBulkData);
+
+		if (!ModelStreamableBulkData || ModelStreamableBulkData->ModelStreamables.IsEmpty())
 		{
 			UE_LOG(LogMutable, Warning, TEXT("Streaming: Customizable Object %s has no data to stream."), *CustomizableObject->GetName());
 
@@ -115,7 +116,7 @@ bool FUnrealMutableModelBulkReader::PrepareStreamingForObject(UCustomizableObjec
 		NewData.DDCKey = ModelResources.DDCKey;
 		NewData.DDCPolicy = ModelResources.DDCDefaultPolicy;
 #else
-		if (NewData.ModelStreamableBulkData->StreamableBulkData.IsEmpty())
+		if (ModelStreamableBulkData->StreamableBulkData.IsEmpty())
 		{
 			const UCustomizableObjectBulk* BulkData = CustomizableObject->GetPrivate()->GetStreamableBulkData();
 			if (!BulkData)
@@ -127,7 +128,6 @@ bool FUnrealMutableModelBulkReader::PrepareStreamingForObject(UCustomizableObjec
 			NewData.BulkFilePrefix = BulkData->GetBulkFilePrefix();
 		}
 #endif
-
 		Objects.Add(MoveTemp(NewData));
 	}
 
@@ -243,18 +243,25 @@ bool FUnrealMutableModelBulkReader::DoesBlockExist(const mu::FModel* Model, uint
 		return false;
 	}
 
-	const FMutableStreamableBlock* Block = ObjectData->ModelStreamableBulkData->ModelStreamables.Find(Key);
+	const TSharedPtr<FModelStreamableBulkData> PinnedModelStreamableBulkData = ObjectData->ModelStreamableBulkData.Pin();
+
+	if (!PinnedModelStreamableBulkData)
+	{
+		return false;
+	}
+
+	const FMutableStreamableBlock* Block = PinnedModelStreamableBulkData->ModelStreamables.Find(Key);
 
 	if (!Block)
 	{
 		return false;
 	}
 
-	if (ObjectData->ModelStreamableBulkData->StreamableBulkData.IsValidIndex(Block->FileId))
+	if (PinnedModelStreamableBulkData->StreamableBulkData.IsValidIndex(Block->FileId))
 	{	
 		MUTABLE_CPUPROFILER_SCOPE(FUnrealMutableModelBulkReader::DoesBlockExist_Bulk)
 
-		FByteBulkData& BulkData = ObjectData->ModelStreamableBulkData->StreamableBulkData[Block->FileId];
+		FByteBulkData& BulkData = PinnedModelStreamableBulkData->StreamableBulkData[Block->FileId];
 		return BulkData.DoesExist();
 	}
 	else
@@ -299,7 +306,18 @@ mu::FModelReader::FOperationID FUnrealMutableModelBulkReader::BeginReadBlock(con
 		return -1;
 	}
 
-	const FMutableStreamableBlock* Block = ObjectData->ModelStreamableBulkData->ModelStreamables.Find(Key);
+	const TSharedPtr<FModelStreamableBulkData> PinnedModelStreamableBulkData = ObjectData->ModelStreamableBulkData.Pin();	
+	if (!PinnedModelStreamableBulkData)
+	{
+		UE_LOG(LogMutable, Error, TEXT("Streamable Bulk Data not available!"));
+		if (CompletionCallback)
+		{
+			(*CompletionCallback)(false);
+		}
+		return -1;
+	}
+
+	const FMutableStreamableBlock* Block = PinnedModelStreamableBulkData->ModelStreamables.Find(Key);
 	if (!Block)
 	{
 		// File Handle not found! This shouldn't really happen.
@@ -315,9 +333,9 @@ mu::FModelReader::FOperationID FUnrealMutableModelBulkReader::BeginReadBlock(con
 
 	FOperationID Result = ++LastOperationID;
 
-	if(ObjectData->ModelStreamableBulkData->StreamableBulkData.IsValidIndex(Block->FileId))
+	if (PinnedModelStreamableBulkData->StreamableBulkData.IsValidIndex(Block->FileId))
 	{
-		FByteBulkData& BulkData = ObjectData->ModelStreamableBulkData->StreamableBulkData[Block->FileId];
+		FByteBulkData& BulkData = PinnedModelStreamableBulkData->StreamableBulkData[Block->FileId];
 
 		FBulkDataIORequestCallBack IOCallback;
 
