[mutable] Removed the Blueprintable tag from several basic CustomizableObject classes, as it made no sense and even crashed.
#rnx
#tests Editor

diff --git a/Engine/Plugins/Mutable/Source/CustomizableObject/Public/MuCO/CustomizableObjectInstance.h b/Engine/Plugins/Mutable/Source/CustomizableObject/Public/MuCO/CustomizableObjectInstance.h
--- a/Engine/Plugins/Mutable/Source/CustomizableObject/Public/MuCO/CustomizableObjectInstance.h
+++ b/Engine/Plugins/Mutable/Source/CustomizableObject/Public/MuCO/CustomizableObjectInstance.h
@@ -221,7 +221,7 @@ DECLARE_DYNAMIC_DELEGATE_TwoParams(FEachComponentAnimInstanceClassDelegate, FNam
 
 DECLARE_DELEGATE_TwoParams(FEachComponentAnimInstanceClassNativeDelegate, FName /*SlotIndex*/, TSubclassOf<UAnimInstance> /*AnimInstClass*/);
 
-UCLASS( Blueprintable, BlueprintType, HideCategories=(CustomizableObjectInstance) )
+UCLASS( BlueprintType, HideCategories=(CustomizableObjectInstance) )
 class CUSTOMIZABLEOBJECT_API UCustomizableObjectInstance : public UObject
 {
 	GENERATED_BODY()
diff --git a/Engine/Plugins/Mutable/Source/CustomizableObject/Public/MuCO/CustomizableObjectInstanceUsage.h b/Engine/Plugins/Mutable/Source/CustomizableObject/Public/MuCO/CustomizableObjectInstanceUsage.h
--- a/Engine/Plugins/Mutable/Source/CustomizableObject/Public/MuCO/CustomizableObjectInstanceUsage.h
+++ b/Engine/Plugins/Mutable/Source/CustomizableObject/Public/MuCO/CustomizableObjectInstanceUsage.h
@@ -23,7 +23,7 @@ DECLARE_DELEGATE(FCustomizableObjectInstanceUsageUpdatedDelegate);
 // UCustomizableObjectInstance and a USkeletalComponent so that the CustomizableObjectSystem takes care of updating it and its LODs, 
 // streaming, etc. It's a UObject, so it will be much cheaper than a UCustomizableComponent as it won't have to refresh its transforms
 // every time it's moved.
-UCLASS(Blueprintable, BlueprintType, ClassGroup = (CustomizableObject))
+UCLASS(BlueprintType, ClassGroup = (CustomizableObject))
 class CUSTOMIZABLEOBJECT_API UCustomizableObjectInstanceUsage : public UObject
 {
 	friend UCustomizableObjectInstanceUsagePrivate;
diff --git a/Engine/Plugins/Mutable/Source/CustomizableObject/Public/MuCO/CustomizableObjectSystem.h b/Engine/Plugins/Mutable/Source/CustomizableObject/Public/MuCO/CustomizableObjectSystem.h
--- a/Engine/Plugins/Mutable/Source/CustomizableObject/Public/MuCO/CustomizableObjectSystem.h
+++ b/Engine/Plugins/Mutable/Source/CustomizableObject/Public/MuCO/CustomizableObjectSystem.h
@@ -201,7 +201,7 @@ public:
 };
 
 
-UCLASS(Blueprintable, BlueprintType)
+UCLASS(BlueprintType)
 class CUSTOMIZABLEOBJECT_API UCustomizableObjectSystem : public UObject
 {
 	// Friends
