[mutable] Fixed being limited to just one dynamic mesh component. Extra components were being ignored silently.
#tests Editor
#rnx

diff --git a/Engine/Plugins/Mutable/Source/CustomizableObjectEditor/Private/MuCOE/GenerateMutableSource/GenerateMutableSourceComponent.cpp b/Engine/Plugins/Mutable/Source/CustomizableObjectEditor/Private/MuCOE/GenerateMutableSource/GenerateMutableSourceComponent.cpp
--- a/Engine/Plugins/Mutable/Source/CustomizableObjectEditor/Private/MuCOE/GenerateMutableSource/GenerateMutableSourceComponent.cpp
+++ b/Engine/Plugins/Mutable/Source/CustomizableObjectEditor/Private/MuCOE/GenerateMutableSource/GenerateMutableSourceComponent.cpp
@@ -554,7 +554,7 @@ void FirstPass(UCustomizableObjectNodeComponentMesh& Node, FMutableGraphGenerati
 		return Element.ComponentName == Node.ComponentName;
 	}))
 	{
-		FText Msg = FText::Format(LOCTEXT("ComponentNodeWithSameNameExists", "Error! It already exists a Mesh Component node with the same name in Customizable Object [{0}]"), FText::FromString(GraphTraversal::GetObject(*Result->Node)->GetName()));
+		FText Msg = FText::Format(LOCTEXT("ComponentNodeWithSameNameExists", "Error! A Mesh Component node with the same name already exists in the Customizable Object [{0}]"), FText::FromString(GraphTraversal::GetObject(*Result->Node)->GetName()));
 		GenerationContext.Log(Msg, &Node, EMessageSeverity::Error);
 		return;
 	}
diff --git a/Engine/Plugins/Mutable/Source/MutableTools/Private/MuT/CodeGenerator.cpp b/Engine/Plugins/Mutable/Source/MutableTools/Private/MuT/CodeGenerator.cpp
--- a/Engine/Plugins/Mutable/Source/MutableTools/Private/MuT/CodeGenerator.cpp
+++ b/Engine/Plugins/Mutable/Source/MutableTools/Private/MuT/CodeGenerator.cpp
@@ -1913,14 +1913,19 @@ namespace mu
 		// Add the components from child objects
 		FObjectGenerationResult::FAdditionalComponentKey ThisKey;
 		ThisKey.ObjectNode = InNode;
-		TArray<FObjectGenerationResult::FAdditionalComponentData>* ThisAdditionalComponents = OutResult.AdditionalComponents.Find(ThisKey);
-		if (LastCompOp && ThisAdditionalComponents)
+		TArray<TArray<FObjectGenerationResult::FAdditionalComponentData>> MultiAdditionalComponents;
+		OutResult.AdditionalComponents.MultiFind(ThisKey, MultiAdditionalComponents, true);
+
+		if (LastCompOp && !MultiAdditionalComponents.IsEmpty())
 		{
-			for (const FObjectGenerationResult::FAdditionalComponentData& Additional : *ThisAdditionalComponents)
+			for (TArray<FObjectGenerationResult::FAdditionalComponentData> ThisAdditionalComponents : MultiAdditionalComponents)
 			{
-				check(Additional.PlaceholderOp);
-				ASTOp::Replace(Additional.PlaceholderOp, LastCompOp);
-				LastCompOp = Additional.ComponentOp;
+				for (const FObjectGenerationResult::FAdditionalComponentData& Additional : ThisAdditionalComponents)
+				{
+					check(Additional.PlaceholderOp);
+					ASTOp::Replace(Additional.PlaceholderOp, LastCompOp);
+					LastCompOp = Additional.ComponentOp;
+				}
 			}
 		}
 
diff --git a/Engine/Plugins/Mutable/Source/MutableTools/Private/MuT/CodeGenerator.h b/Engine/Plugins/Mutable/Source/MutableTools/Private/MuT/CodeGenerator.h
--- a/Engine/Plugins/Mutable/Source/MutableTools/Private/MuT/CodeGenerator.h
+++ b/Engine/Plugins/Mutable/Source/MutableTools/Private/MuT/CodeGenerator.h
@@ -276,7 +276,7 @@ namespace mu
 				Ptr<ASTOp> PlaceholderOp;
 			};
 
-			TMap< FAdditionalComponentKey, TArray<FAdditionalComponentData> > AdditionalComponents;
+			TMultiMap< FAdditionalComponentKey, TArray<FAdditionalComponentData> > AdditionalComponents;
 		};
 
 		typedef TMap<FGeneratedObjectCacheKey, FObjectGenerationResult> FGeneratedObjectsMap;
