[mutable] Remove mesh format operation on the roots of some modifiers.
#rnx

diff --git a/Engine/Plugins/Mutable/Source/CustomizableObjectEditor/Private/MuCOE/GenerateMutableSource/GenerateMutableSourceModifier.cpp b/Engine/Plugins/Mutable/Source/CustomizableObjectEditor/Private/MuCOE/GenerateMutableSource/GenerateMutableSourceModifier.cpp
--- a/Engine/Plugins/Mutable/Source/CustomizableObjectEditor/Private/MuCOE/GenerateMutableSource/GenerateMutableSourceModifier.cpp
+++ b/Engine/Plugins/Mutable/Source/CustomizableObjectEditor/Private/MuCOE/GenerateMutableSource/GenerateMutableSourceModifier.cpp
@@ -331,17 +331,7 @@ mu::Ptr<mu::NodeModifier> GenerateMutableSourceModifier(const UEdGraphPin * Pin,
 				GenerationContext.LayoutGenerationFlags.Pop();
 			}
 
-			if (AddMeshNode)
-			{
-				mu::Ptr<mu::NodeMesh> MeshPtr = AddMeshNode;
-
-				mu::Ptr<mu::NodeMeshFormat> MeshFormat = new mu::NodeMeshFormat();
-				SetSurfaceFormat(GenerationContext.CompilationContext, MeshFormat->VertexBuffers, MeshFormat->IndexBuffers, MeshData);
-
-				MeshFormat->Source = MeshPtr;
-
-				SurfNode->LODs[LODIndex].MeshAdd = MeshFormat;
-			}
+			SurfNode->LODs[LODIndex].MeshAdd = AddMeshNode;
 
 			const int32 NumImages = TypedNodeExt->GetNumParameters(EMaterialParameterType::Texture);
 			SurfNode->LODs[LODIndex].Textures.SetNum(NumImages);
