[Mutable] Add Physics constants inspector to debugger.

#rnx

diff --git a/Engine/Plugins/Mutable/Source/CustomizableObjectEditor/Private/MuCOE/SMutableCodeViewer.cpp b/Engine/Plugins/Mutable/Source/CustomizableObjectEditor/Private/MuCOE/SMutableCodeViewer.cpp
--- a/Engine/Plugins/Mutable/Source/CustomizableObjectEditor/Private/MuCOE/SMutableCodeViewer.cpp
+++ b/Engine/Plugins/Mutable/Source/CustomizableObjectEditor/Private/MuCOE/SMutableCodeViewer.cpp
@@ -2666,6 +2666,14 @@ bool SMutableCodeViewer::IsConstantResourceUsedByOperation(const int32 IndexOnCo
 				}
 				break;
 			}
+		case mu::EDataType::PhysicsAsset:
+			{
+				if (OperationType == mu::EOpType::ME_CONSTANT)
+				{
+					bResourceLocated = IndexOnConstantsArray == InProgram.GetOpArgs<mu::OP::MeshConstantArgs>(OperationAddress).physicsBody;
+				}
+				break;
+			}
 
 		// Invalid types
 		case mu::EDataType::None:
@@ -3104,6 +3112,12 @@ void SMutableCodeViewer::PreviewMutableCurve(const FRichCurve& Curve)
 	PreviewCurveViewer->SetCurve(Curve);
 }
 
+// TODO: Implement physics viewer
+void SMutableCodeViewer::PreviewMutablePhysics(TSharedPtr<const mu::FPhysicsBody> Physics)
+{
+	UE_LOG(LogMutable, Warning, TEXT("Previewer for Mutable Physics not yet implemented"))
+}
+
 // TODO: Implement matrix viewer
 void SMutableCodeViewer::PreviewMutableMatrix(const FMatrix44f& Mat)
 {
diff --git a/Engine/Plugins/Mutable/Source/CustomizableObjectEditor/Private/MuCOE/SMutableCodeViewer.h b/Engine/Plugins/Mutable/Source/CustomizableObjectEditor/Private/MuCOE/SMutableCodeViewer.h
--- a/Engine/Plugins/Mutable/Source/CustomizableObjectEditor/Private/MuCOE/SMutableCodeViewer.h
+++ b/Engine/Plugins/Mutable/Source/CustomizableObjectEditor/Private/MuCOE/SMutableCodeViewer.h
@@ -798,6 +798,7 @@ public:
 	void PreviewMutableMesh (TSharedPtr<const mu::FMesh> InMeshPtr);
 	void PreviewMutableLayout(TSharedPtr<const mu::FLayout> Layout);
 	void PreviewMutableSkeleton(TSharedPtr<const mu::FSkeleton> Skeleton);
+	void PreviewMutablePhysics(TSharedPtr<const mu::FPhysicsBody> Physics);
 	void PreviewMutableString(const FString& InString);
 	void PreviewMutableProjector(const mu::FProjector* Projector);
 	void PreviewMutableMatrix(const FMatrix44f& Mat);
diff --git a/Engine/Plugins/Mutable/Source/CustomizableObjectEditor/Private/MuCOE/SMutableConstantsWidget.cpp b/Engine/Plugins/Mutable/Source/CustomizableObjectEditor/Private/MuCOE/SMutableConstantsWidget.cpp
--- a/Engine/Plugins/Mutable/Source/CustomizableObjectEditor/Private/MuCOE/SMutableConstantsWidget.cpp
+++ b/Engine/Plugins/Mutable/Source/CustomizableObjectEditor/Private/MuCOE/SMutableConstantsWidget.cpp
@@ -430,6 +430,27 @@ public:
 	}
 };
 
+class SMutableConstantPhysicsRow final : public STableRow<TSharedPtr<FMutableConstantPhysicsElement>>
+{
+public:
+	void Construct(const FArguments& Args, const TSharedRef<STableViewBase>& InOwnerTableView, const TSharedPtr<FMutableConstantPhysicsElement>& InRowElement)
+	{
+		check(InRowElement);
+		const FText PhysicsProxyText = FText::Format( LOCTEXT("PhysicsConstantProxyLabel", "{0}_PHYSICS "),InRowElement->IndexOnSourceVector);
+
+		this->ChildSlot
+		[
+			SNew(STextBlock)
+			.Text(PhysicsProxyText)
+		];
+
+		STableRow< TSharedPtr<FMutableConstantPhysicsElement> >::ConstructInternal(
+			STableRow::FArguments()
+			.ShowSelection(true)
+			, InOwnerTableView);
+	}
+};
+
 #pragma endregion SUPPORT_CLASSES
 
 
@@ -757,6 +778,26 @@ void SMutableConstantsWidget::Construct(const FArguments& InArgs,const mu::FProg
 				.SelectionMode(ESelectionMode::Single)
 			]
 		]
+
+
+		// Physics Constants
+		+ SScrollBox::Slot()
+		.Padding(0,InBetweenListsVerticalPadding)
+		[
+			SAssignNew(PhysicsExpandableArea,SExpandableArea)
+			.OnAreaExpansionChanged(this,&SMutableConstantsWidget::OnPhysicsRegionExpansionChanged)
+			.InitiallyCollapsed(true)
+			.AreaTitle(this,&SMutableConstantsWidget::OnDrawPhysicsAreaTitle)
+			.BodyContent()
+			[
+				SAssignNew(ConstantPhysicsSlate,STileView<TSharedPtr<FMutableConstantPhysicsElement>>)
+				.OnSelectionChanged(this,&SMutableConstantsWidget::OnSelectedPhysicsChanged)
+				.ListItemsSource(&ConstantPhysicsElements)
+				.ItemHeight(ProxyEntryHeight)
+				.OnGenerateTile(this,&SMutableConstantsWidget::OnGeneratePhysicsRow)
+				.SelectionMode(ESelectionMode::Single)
+			]
+		]
 	];
 
 	// Store all the expandable areas so they are later reachable using loops
@@ -769,6 +810,7 @@ void SMutableConstantsWidget::Construct(const FArguments& InArgs,const mu::FProg
 	ExpandableAreas.Add(CurvesExpandableArea);
 	ExpandableAreas.Add(MatricesExpandableArea);
 	ExpandableAreas.Add(SkeletonsExpandableArea);
+	ExpandableAreas.Add(PhysicsExpandableArea);
 }
 
 
@@ -863,6 +905,13 @@ TSharedRef<ITableRow> SMutableConstantsWidget::OnGenerateSkeletonRow(
 	return Row;
 }
 
+TSharedRef<ITableRow> SMutableConstantsWidget::OnGeneratePhysicsRow(
+	TSharedPtr<FMutableConstantPhysicsElement> MutableConstantPhysicsElement,
+	const TSharedRef<STableViewBase>& OwnerTable) const
+{
+	TSharedRef<SMutableConstantPhysicsRow> Row = SNew(SMutableConstantPhysicsRow,OwnerTable,MutableConstantPhysicsElement);
+	return Row;
+}
 
 void SMutableConstantsWidget::OnImageTableSortRequested(EColumnSortPriority::Type ColumnSortPriority, const FName& ColumnID,
                                                         EColumnSortMode::Type ColumnSortMode)
@@ -1112,6 +1161,15 @@ void SMutableConstantsWidget::OnSkeletonsRegionExpansionChanged(bool bExpanded)
 	}
 }
 
+void SMutableConstantsWidget::OnPhysicsRegionExpansionChanged(bool bExpanded)
+{
+	if (bExpanded)
+	{
+		ContractExpandableAreas(PhysicsExpandableArea);
+	}
+}
+
+
 void SMutableConstantsWidget::ContractExpandableAreas(const TSharedPtr<SExpandableArea>& InException)
 {
 #if UE_BUILD_DEBUG || UE_BUILD_DEVELOPMENT
@@ -1147,6 +1205,7 @@ void SMutableConstantsWidget::LoadConstantElements()
 	LoadConstantShapes();
 	LoadConstantCurves();
 	LoadConstantSkeletons();
+	LoadConstantPhysics();
 }
 
 void SMutableConstantsWidget::LoadConstantStrings()
@@ -1305,6 +1364,31 @@ void SMutableConstantsWidget::LoadConstantSkeletons()
 	ConstantSkeletonsFormattedSize = GenerateTextForSize(Stream.GetBufferSize());
 }
 
+void SMutableConstantsWidget::LoadConstantPhysics()
+{
+	check (MutableProgramPtr);
+	const int32 ConstantsCount = MutableProgramPtr->ConstantPhysicsBodies.Num();
+	ConstantPhysicsElements.Empty(ConstantsCount);
+
+	mu::FOutputMemoryStream Stream;
+	mu::FOutputArchive Archive{&Stream};
+	
+	for (int32 PhysicsIndex = 0; PhysicsIndex < ConstantsCount; PhysicsIndex++)
+	{
+		TSharedPtr<FMutableConstantPhysicsElement> ConstantPhysicsElement = MakeShared<FMutableConstantPhysicsElement>();
+		ConstantPhysicsElement->Physics = MutableProgramPtr->ConstantPhysicsBodies[PhysicsIndex];
+		ConstantPhysicsElement->IndexOnSourceVector = PhysicsIndex;
+		
+		ConstantPhysicsElement->Physics->Serialise(Archive);
+
+		ConstantPhysicsElements.Add(ConstantPhysicsElement);
+	}
+
+	// Cache the size in memory of the constants as a formatted text so it is able to be be later used by the UI
+	ConstantPhysicsFormattedSize = GenerateTextForSize(Stream.GetBufferSize());
+}
+
+
 void SMutableConstantsWidget::LoadConstantProjectors()
 {
 	check (MutableProgramPtr);
@@ -1441,6 +1525,10 @@ void SMutableConstantsWidget::ClearSelectedConstantItems(mu::EDataType Exception
 	{
 		ConstantSkeletonsSlate->ClearSelection();
 	}
+	if (ExceptionDataType != mu::EDataType::PhysicsAsset)
+	{
+		ConstantPhysicsSlate->ClearSelection();
+	}
 	if (ExceptionDataType != mu::EDataType::Image)
 	{
 		ImageListViewHandler->ClearSelection();
@@ -1596,6 +1684,22 @@ void SMutableConstantsWidget::OnSelectedSkeletonChanged(
 	}
 }
 
+void SMutableConstantsWidget::OnSelectedPhysicsChanged(
+	TSharedPtr<FMutableConstantPhysicsElement> MutableConstantPhysicsElement,
+	ESelectInfo::Type SelectionType) const
+{
+	if (MutableConstantPhysicsElement)
+	{
+		static constexpr mu::EDataType SlateDataType = mu::EDataType::PhysicsAsset;
+
+		MutableCodeViewerPtr->ClearSelectedTreeRow();
+		ClearSelectedConstantItems(SlateDataType);
+
+		MutableCodeViewerPtr->PreviewMutablePhysics(MutableConstantPhysicsElement->Physics);
+		MutableCodeViewerPtr->CacheAddressesRelatedWithConstantResource(SlateDataType,MutableConstantPhysicsElement->IndexOnSourceVector);
+	}
+}
+
 #pragma endregion 
 
 #pragma region Region title drawing callback methods
@@ -1646,6 +1750,11 @@ FText SMutableConstantsWidget::OnDrawSkeletonsAreaTitle() const
 	return FText::Format( LOCTEXT("SkeletonConstantsTitle", "Skeleton Constants ({0}) : {1} "), ConstantSkeletonElements.Num(), ConstantSkeletonsFormattedSize);
 }
 
+FText SMutableConstantsWidget::OnDrawPhysicsAreaTitle() const
+{
+	return FText::Format( LOCTEXT("PhysicsConstantsTitle", "Physics Constants ({0}) : {1} "), ConstantPhysicsElements.Num(), ConstantPhysicsFormattedSize);
+}
+
 #pragma endregion
 
 
diff --git a/Engine/Plugins/Mutable/Source/CustomizableObjectEditor/Private/MuCOE/SMutableConstantsWidget.h b/Engine/Plugins/Mutable/Source/CustomizableObjectEditor/Private/MuCOE/SMutableConstantsWidget.h
--- a/Engine/Plugins/Mutable/Source/CustomizableObjectEditor/Private/MuCOE/SMutableConstantsWidget.h
+++ b/Engine/Plugins/Mutable/Source/CustomizableObjectEditor/Private/MuCOE/SMutableConstantsWidget.h
@@ -70,6 +70,14 @@ struct FMutableConstantSkeletonElement : public FMutableConstantElement
 	TSharedPtr<const mu::FSkeleton> Skeleton;
 };
 
+/**
+* Cache object used for the generation of the ui elements related to the constant skeletons found on the model
+*/
+struct FMutableConstantPhysicsElement : public FMutableConstantElement
+{
+	TSharedPtr<const mu::FPhysicsBody> Physics;
+};
+
 /**
 * Cache object used for the generation of the ui elements related to the constant projectors found on the model
 */
@@ -153,6 +161,7 @@ private:
 	TSharedPtr<STileView<TSharedPtr<FMutableConstantShapeElement>>> ConstantShapesSlate;
 	TSharedPtr<STileView<TSharedPtr<FMutableConstantCurveElement>>> ConstantCurvesSlate;
 	TSharedPtr<STileView<TSharedPtr<FMutableConstantSkeletonElement>>> ConstantSkeletonsSlate;
+	TSharedPtr<STileView<TSharedPtr<FMutableConstantPhysicsElement>>> ConstantPhysicsSlate;
 	
 	/*
 	* Data backend for the lists of constants
@@ -171,6 +180,7 @@ private:
 	TArray<TSharedPtr<FMutableConstantShapeElement>> ConstantShapeElements;
 	TArray<TSharedPtr<FMutableConstantCurveElement>> ConstantCurveElements;
 	TArray<TSharedPtr<FMutableConstantSkeletonElement>> ConstantSkeletonElements;
+	TArray<TSharedPtr<FMutableConstantPhysicsElement>> ConstantPhysicsElements;
 
 	/**
 	 * Load up all the elements with the data found on the mu::FProgram object onto TArrays after parsing the data found.
@@ -186,6 +196,7 @@ private:
 	void LoadConstantShapes();
 	void LoadConstantCurves();
 	void LoadConstantSkeletons();
+	void LoadConstantPhysics();
 
 	/*
 	 * Proxy slates operation objects 
@@ -200,7 +211,7 @@ private:
 	void OnSelectedShapeChanged(TSharedPtr<FMutableConstantShapeElement> MutableConstantShapeElement, ESelectInfo::Type SelectionType) const;
 	void OnSelectedCurveChanged(TSharedPtr<FMutableConstantCurveElement> MutableConstantCurveElement, ESelectInfo::Type SelectionType) const;
 	void OnSelectedSkeletonChanged(TSharedPtr<FMutableConstantSkeletonElement> MutableConstantSkeletonElement, ESelectInfo::Type SelectionType) const;
-	
+	void OnSelectedPhysicsChanged(TSharedPtr<FMutableConstantPhysicsElement> MutableConstantPhysicsElement, ESelectInfo::Type SelectionType) const;
 	/*
 	* UI updating methods
 	*/
@@ -214,7 +225,7 @@ private:
 	TSharedRef<ITableRow> OnGenerateShapeRow(TSharedPtr<FMutableConstantShapeElement> MutableConstantShapeElement, const TSharedRef<STableViewBase>& OwnerTable) const;
 	TSharedRef<ITableRow> OnGenerateCurveRow(TSharedPtr<FMutableConstantCurveElement> MutableConstantCurveElement, const TSharedRef<STableViewBase>& OwnerTable) const;
 	TSharedRef<ITableRow> OnGenerateSkeletonRow(TSharedPtr<FMutableConstantSkeletonElement> MutableConstantSkeletonElement, const TSharedRef<STableViewBase>& OwnerTable) const;
-	
+	TSharedRef<ITableRow> OnGeneratePhysicsRow(TSharedPtr<FMutableConstantPhysicsElement> MutableConstantPhysicsElement, const TSharedRef<STableViewBase>& OwnerTable) const;
 	
 	/*
 	 * Image List sorting methods
@@ -255,6 +266,7 @@ private:
 	TSharedPtr<SExpandableArea> ShapesExpandableArea;
 	TSharedPtr<SExpandableArea> CurvesExpandableArea;
 	TSharedPtr<SExpandableArea> SkeletonsExpandableArea;
+	TSharedPtr<SExpandableArea> PhysicsExpandableArea;
 
 	/** Array with all expandable areas set on this object. Used for dynamic expansion/contraction */
 	TArray<TSharedPtr<SExpandableArea>> ExpandableAreas;
@@ -272,6 +284,7 @@ private:
 	void OnShapesRegionExpansionChanged(bool bExpanded);
 	void OnCurvesRegionExpansionChanged(bool bExpanded);
 	void OnSkeletonsRegionExpansionChanged(bool bExpanded);
+	void OnPhysicsRegionExpansionChanged(bool bExpanded);
 	
 	/** Method called each time one expandable area changes expansion state
 	 * @param InException - Pointer to the expandable area that will not be contracted while all the others will
@@ -292,6 +305,7 @@ private:
 	FText ConstantShapesFormattedSize;
 	FText ConstantCurvesFormattedSize;
 	FText ConstantSkeletonsFormattedSize;
+	FText ConstantPhysicsFormattedSize;
 	
 	/*
 	 * Callback methods used for drawing the titles of each of the constant expandable areas
@@ -306,6 +320,7 @@ private:
 	FText OnDrawShapesAreaTitle() const;
 	FText OnDrawCurvesAreaTitle() const;
 	FText OnDrawSkeletonsAreaTitle() const;
+	FText OnDrawPhysicsAreaTitle() const;
 	
 };
 
