[Mutable] Fix CIS compilation error on Linux and Mac builds.

#rnx

diff --git a/Engine/Plugins/Mutable/Source/CustomizableObjectEditor/Private/MuCOE/SMutableConstantsWidget.cpp b/Engine/Plugins/Mutable/Source/CustomizableObjectEditor/Private/MuCOE/SMutableConstantsWidget.cpp
--- a/Engine/Plugins/Mutable/Source/CustomizableObjectEditor/Private/MuCOE/SMutableConstantsWidget.cpp
+++ b/Engine/Plugins/Mutable/Source/CustomizableObjectEditor/Private/MuCOE/SMutableConstantsWidget.cpp
@@ -1379,7 +1379,7 @@ void SMutableConstantsWidget::LoadConstantPhysics()
 		ConstantPhysicsElement->Physics = MutableProgramPtr->ConstantPhysicsBodies[PhysicsIndex];
 		ConstantPhysicsElement->IndexOnSourceVector = PhysicsIndex;
 		
-		ConstantPhysicsElement->Physics->Serialise(Archive);
+		mu::FPhysicsBody::Serialise(ConstantPhysicsElement->Physics.Get(), Archive);
 
 		ConstantPhysicsElements.Add(ConstantPhysicsElement);
 	}
