[Mutable] Fix CIS compilation error.

#rnx

diff --git a/Engine/Plugins/Mutable/Source/CustomizableObjectEditor/Private/MuCOE/SMutableConstantsWidget.cpp b/Engine/Plugins/Mutable/Source/CustomizableObjectEditor/Private/MuCOE/SMutableConstantsWidget.cpp
--- a/Engine/Plugins/Mutable/Source/CustomizableObjectEditor/Private/MuCOE/SMutableConstantsWidget.cpp
+++ b/Engine/Plugins/Mutable/Source/CustomizableObjectEditor/Private/MuCOE/SMutableConstantsWidget.cpp
@@ -1378,8 +1378,9 @@ void SMutableConstantsWidget::LoadConstantPhysics()
 		TSharedPtr<FMutableConstantPhysicsElement> ConstantPhysicsElement = MakeShared<FMutableConstantPhysicsElement>();
 		ConstantPhysicsElement->Physics = MutableProgramPtr->ConstantPhysicsBodies[PhysicsIndex];
 		ConstantPhysicsElement->IndexOnSourceVector = PhysicsIndex;
-		
-		mu::FPhysicsBody::Serialise(ConstantPhysicsElement->Physics.Get(), Archive);
+	
+		//TODO: Fix me.
+		//mu::FPhysicsBody::Serialise(ConstantPhysicsElement->Physics.Get(), Archive);
 
 		ConstantPhysicsElements.Add(ConstantPhysicsElement);
 	}
