[mutable] Initial support for mesh conversion from runtime data for mesh parameters.
- Added mesh parameter type.
- Extended mesh parameter operation to have LOD and section information.
#rnx

diff --git a/Engine/Plugins/Mutable/Config/BaseMutable.ini b/Engine/Plugins/Mutable/Config/BaseMutable.ini
--- a/Engine/Plugins/Mutable/Config/BaseMutable.ini
+++ b/Engine/Plugins/Mutable/Config/BaseMutable.ini
@@ -2,6 +2,7 @@
 
 +StructRedirects=(OldName="/Script/CustomizableObject.ParameterUIData",NewName="/Script/CustomizableObject.MutableParameterData")
 +StructRedirects=(OldName="/Script/CustomizableObject.CompilationOptions",NewName="/Script/CustomizableObject.CompilationOptions_DEPRECATED")
++StructRedirects=(OldName="/Script/CustomizableObject.FCustomizableObjectTextureParameterValue",NewName="/Script/CustomizableObject.FCustomizableObjectAssetParameterValue")
 
 +PropertyRedirects=(OldName="GroupProjectorParameterImage.OptionImage", NewName="OptionTexture")
 +PropertyRedirects=(OldName="CustomizableObjectNodeGroupProjectorParameter.OptionImages", NewName="OptionTextures")
diff --git a/Engine/Plugins/Mutable/Source/CustomizableObject/Internal/MuCO/CustomizableObjectPrivate.h b/Engine/Plugins/Mutable/Source/CustomizableObject/Internal/MuCO/CustomizableObjectPrivate.h
--- a/Engine/Plugins/Mutable/Source/CustomizableObject/Internal/MuCO/CustomizableObjectPrivate.h
+++ b/Engine/Plugins/Mutable/Source/CustomizableObject/Internal/MuCO/CustomizableObjectPrivate.h
@@ -1580,6 +1580,8 @@ enum class ECustomizableObjectVersions : int32
 	AddedMeshParameterOp,
 
 	AddedMeshParameterOpForDDCPollution,
+
+	ExtendedMeshParameterArgumentsWithLODAndSection,
 	
 	// -----<new versions can be added above this line>--------
 	LastCustomizableObjectVersion
@@ -1854,6 +1856,6 @@ public:
 	static constexpr int32 CurrentSupportedVersion = static_cast<int32>(ECustomizableObjectVersions::LastCustomizableObjectVersion);
 #endif
 	
-	static constexpr int32 DerivedDataVersion = 6;
+	static constexpr int32 DerivedDataVersion = 0x0b410b41;
 };
 
diff --git a/Engine/Plugins/Mutable/Source/CustomizableObject/Internal/MuCO/CustomizableObjectSystemPrivate.h b/Engine/Plugins/Mutable/Source/CustomizableObject/Internal/MuCO/CustomizableObjectSystemPrivate.h
--- a/Engine/Plugins/Mutable/Source/CustomizableObject/Internal/MuCO/CustomizableObjectSystemPrivate.h
+++ b/Engine/Plugins/Mutable/Source/CustomizableObject/Internal/MuCO/CustomizableObjectSystemPrivate.h
@@ -901,9 +901,9 @@ public:
 	 * Used for stats that are costly to update. */
 	void UpdateStats();
 
-	void CacheTextureParameters(const TArray<FCustomizableObjectTextureParameterValue>& TextureParameters) const;
+	void CacheTextureParameters(const TArray<FCustomizableObjectAssetParameterValue>& TextureParameters) const;
 
-	void UnCacheTextureParameters(const TArray<FCustomizableObjectTextureParameterValue>& TextureParameters) const;
+	void UnCacheTextureParameters(const TArray<FCustomizableObjectAssetParameterValue>& TextureParameters) const;
 
 #if WITH_EDITOR
 	/** Add a CO to the pending - load list.COs need to wait until all related objects are fully loaded before being able
diff --git a/Engine/Plugins/Mutable/Source/CustomizableObject/Internal/MuCO/MutableMeshBufferUtils.h b/Engine/Plugins/Mutable/Source/CustomizableObject/Internal/MuCO/MutableMeshBufferUtils.h
--- a/Engine/Plugins/Mutable/Source/CustomizableObject/Internal/MuCO/MutableMeshBufferUtils.h
+++ b/Engine/Plugins/Mutable/Source/CustomizableObject/Internal/MuCO/MutableMeshBufferUtils.h
@@ -3,7 +3,13 @@
 #pragma once
 
 #include "HAL/Platform.h"
+#include "UObject/NameTypes.h"
+#include "Containers/Map.h"
+#include "Containers/Array.h"
+#include "UObject/ObjectPtr.h"
+#include "BoneIndices.h"
 
+class USkeletalMesh;
 namespace mu { class FMeshBufferSet; }
 
 /** Publicly accessible methods designed to prepare the Mutable Mesh buffers with the internal structure  required by Unreal and
@@ -76,4 +82,7 @@ namespace  MutableMeshBufferUtils
 	CUSTOMIZABLEOBJECT_API void SetupSkinWeightProfileBuffer(const int32& InCurrentVertexBuffer, const int32& MaxBoneIndexTypeSizeBytes,const int32& MaxBoneWeightTypeSizeBytes, const int32& MaxNumBonesPerVertex,
 		const int32 SemanticsIndex, mu::FMeshBufferSet& OutTargetVertexBuffers);
 	
+	/** */
+	CUSTOMIZABLEOBJECT_API void BuildRemappedBonesArray(const TMap<FName, bool>* BonesToRemove, TObjectPtr<const USkeletalMesh> InSkeletalMesh, int32 InLODIndex, const TArray<FBoneIndexType>& InRequiredBones, TArray<FBoneIndexType>& OutRemappedBones);
+
 }
diff --git a/Engine/Plugins/Mutable/Source/CustomizableObject/Internal/MuCO/UnrealConversionUtils.h b/Engine/Plugins/Mutable/Source/CustomizableObject/Internal/MuCO/UnrealConversionUtils.h
--- a/Engine/Plugins/Mutable/Source/CustomizableObject/Internal/MuCO/UnrealConversionUtils.h
+++ b/Engine/Plugins/Mutable/Source/CustomizableObject/Internal/MuCO/UnrealConversionUtils.h
@@ -21,6 +21,7 @@ namespace mu
 }
 struct FReferenceSkeleton;
 struct FMutableSurfaceMetadata;
+class UModelResources;
 
 class FSkeletalMeshLODRenderData;
 class USkeletalMesh;
@@ -134,4 +135,7 @@ namespace UnrealConversionUtils
 	);
 
 	void MorphTargetVertexInfoBuffers(FSkeletalMeshLODRenderData& LODResource, const USkeletalMesh& Owner, const mu::FMesh& MutableMesh, const TMap<uint32, FMorphTargetMeshData>& MorphTargetMeshData, int32 LODIndex);
+
+	CUSTOMIZABLEOBJECT_API void ConvertSkeletalMeshFromRuntimeData(USkeletalMesh& SkeletalMesh, int32 LODIndex, int32 SectionIndex, UModelResources* ModelResources, mu::FMesh& Result);
+
 }
diff --git a/Engine/Plugins/Mutable/Source/CustomizableObject/Private/MuCO/CustomizableObject.cpp b/Engine/Plugins/Mutable/Source/CustomizableObject/Private/MuCO/CustomizableObject.cpp
--- a/Engine/Plugins/Mutable/Source/CustomizableObject/Private/MuCO/CustomizableObject.cpp
+++ b/Engine/Plugins/Mutable/Source/CustomizableObject/Private/MuCO/CustomizableObject.cpp
@@ -1436,6 +1436,12 @@ void UCustomizableObjectPrivate::UpdateParameterPropertiesFromModel(const TShare
 				break;
 			}
 
+			case mu::EParameterType::Mesh:
+			{
+				Data.Type = EMutableParameterType::Mesh;
+				break;
+			}
+
 			default:
 				// Unhandled type?
 				check(false);
@@ -1862,14 +1868,14 @@ FName UCustomizableObject::GetTextureParameterDefaultValue(const FString& InPara
 	if (ParameterIndex == INDEX_NONE)
 	{
 		UE_LOG(LogMutable, Error, TEXT("Tried to access the default value of the nonexistent texture parameter [%s] in the CustomizableObject [%s]."), *InParameterName, *GetName());
-		return FCustomizableObjectTextureParameterValue::DEFAULT_PARAMETER_VALUE;
+		return FCustomizableObjectAssetParameterValue::DEFAULT_PARAMETER_VALUE;
 	}
 
 	const TSharedPtr<const mu::FModel>& Model = GetPrivate()->GetModel();
 	if (!Model)
 	{
 		checkNoEntry();
-		return FCustomizableObjectTextureParameterValue::DEFAULT_PARAMETER_VALUE;
+		return FCustomizableObjectAssetParameterValue::DEFAULT_PARAMETER_VALUE;
 	}
 	
 	return Model->GetImageDefaultValue(ParameterIndex);
diff --git a/Engine/Plugins/Mutable/Source/CustomizableObject/Private/MuCO/CustomizableObjectInstance.cpp b/Engine/Plugins/Mutable/Source/CustomizableObject/Private/MuCO/CustomizableObjectInstance.cpp
--- a/Engine/Plugins/Mutable/Source/CustomizableObject/Private/MuCO/CustomizableObjectInstance.cpp
+++ b/Engine/Plugins/Mutable/Source/CustomizableObject/Private/MuCO/CustomizableObjectInstance.cpp
@@ -794,12 +794,12 @@ bool UCustomizableObjectInstance::IsParameterDirty(const FString& ParamName, con
 		}		
 	case EMutableParameterType::Texture:
 		{
-			const FCustomizableObjectTextureParameterValue* Result = Descriptor.GetTextureParameters().FindByPredicate([&](const FCustomizableObjectTextureParameterValue& Value)
+			const FCustomizableObjectAssetParameterValue* Result = Descriptor.GetTextureParameters().FindByPredicate([&](const FCustomizableObjectAssetParameterValue& Value)
 				{
 					return Value.ParameterName == ParamName;
 				});
 			
-			const FCustomizableObjectTextureParameterValue* ResultCommited = GetPrivate()->CommittedDescriptor.GetTextureParameters().FindByPredicate([&](const FCustomizableObjectTextureParameterValue& Value)
+			const FCustomizableObjectAssetParameterValue* ResultCommited = GetPrivate()->CommittedDescriptor.GetTextureParameters().FindByPredicate([&](const FCustomizableObjectAssetParameterValue& Value)
 				{
 					return Value.ParameterName == ParamName;
 				});
@@ -3074,9 +3074,9 @@ bool UCustomizableInstancePrivate::MigrateProfileParametersToCurrentInstance(int
 		}
 	}
 
-	for (FCustomizableObjectTextureParameterValue& Parameter : TempProfile.TextureParameters)
+	for (FCustomizableObjectAssetParameterValue& Parameter : TempProfile.TextureParameters)
 	{
-		using ParamValType = FCustomizableObjectTextureParameterValue;
+		using ParamValType = FCustomizableObjectAssetParameterValue;
 		ParamValType* Found = Profile.TextureParameters.FindByPredicate(
 			[&Parameter](const ParamValType& P) { return P.ParameterName == Parameter.ParameterName; });
 
@@ -3086,6 +3086,18 @@ bool UCustomizableInstancePrivate::MigrateProfileParametersToCurrentInstance(int
 		}
 	}
 
+	for (FCustomizableObjectAssetParameterValue& Parameter : TempProfile.MeshParameters)
+	{
+		using ParamValType = FCustomizableObjectAssetParameterValue;
+		ParamValType* Found = Profile.MeshParameters.FindByPredicate(
+			[&Parameter](const ParamValType& P) { return P.ParameterName == Parameter.ParameterName; });
+
+		if (Found)
+		{
+			Parameter.ParameterValue = Found->ParameterValue;
+		}
+	}
+
 	for (FCustomizableObjectVectorParameterValue& Parameter : TempProfile.VectorParameters)
 	{
 		using ParamValType = FCustomizableObjectVectorParameterValue;
@@ -4889,6 +4901,11 @@ bool UCustomizableInstancePrivate::BuildOrCopyRenderData(const TSharedRef<FUpdat
 
 		for (int32 BufferIndex = 0; BufferIndex < NumBuffers; ++BufferIndex)
 		{
+			if (MutableMeshVertexBuffers.Buffers[BufferIndex].Channels.IsEmpty())
+			{
+				continue;
+			}
+
 			mu::EMeshBufferSemantic Semantic;
 			int32 SemanticIndex;
 			MutableMeshVertexBuffers.GetChannel(BufferIndex, 0, &Semantic, &SemanticIndex, nullptr, nullptr, nullptr);
diff --git a/Engine/Plugins/Mutable/Source/CustomizableObject/Private/MuCO/CustomizableObjectInstanceDescriptor.cpp b/Engine/Plugins/Mutable/Source/CustomizableObject/Private/MuCO/CustomizableObjectInstanceDescriptor.cpp
--- a/Engine/Plugins/Mutable/Source/CustomizableObject/Private/MuCO/CustomizableObjectInstanceDescriptor.cpp
+++ b/Engine/Plugins/Mutable/Source/CustomizableObject/Private/MuCO/CustomizableObjectInstanceDescriptor.cpp
@@ -210,7 +210,27 @@ void FCustomizableObjectInstanceDescriptor::SaveDescriptor(FArchive& Ar, bool bU
 			FName Value;
 			TArray<FName> RangeValues;
 
-			for (const FCustomizableObjectTextureParameterValue& P : TextureParameters)
+			for (const FCustomizableObjectAssetParameterValue& P : TextureParameters)
+			{
+				if (P.ParameterName == Name)
+				{
+					Value = P.ParameterValue;
+					RangeValues = P.ParameterRangeValues;
+					break;
+				}
+			}
+			Ar << Value;
+			Ar << RangeValues;
+
+			break;
+		}
+
+		case EMutableParameterType::Mesh:
+		{
+			FName Value;
+			TArray<FName> RangeValues;
+
+			for (const FCustomizableObjectAssetParameterValue& P : MeshParameters)
 			{
 				if (P.ParameterName == Name)
 				{
@@ -417,7 +437,7 @@ void FCustomizableObjectInstanceDescriptor::LoadDescriptor(FArchive& Ar)
 
 			break;
 		}
-			
+
 		case EMutableParameterType::Texture:
 		{
 			FName Value;
@@ -425,7 +445,27 @@ void FCustomizableObjectInstanceDescriptor::LoadDescriptor(FArchive& Ar)
 			Ar << Value;
 			Ar << RangeValues;
 
-			for (FCustomizableObjectTextureParameterValue& P : TextureParameters)
+			for (FCustomizableObjectAssetParameterValue& P : TextureParameters)
+			{
+				if (P.ParameterName == Name)
+				{
+					P.ParameterValue = Value;
+					P.ParameterRangeValues = RangeValues;
+					break;
+				}
+			}
+
+			break;
+		}
+
+		case EMutableParameterType::Mesh:
+		{
+			FName Value;
+			TArray<FName> RangeValues;
+			Ar << Value;
+			Ar << RangeValues;
+
+			for (FCustomizableObjectAssetParameterValue& P : MeshParameters)
 			{
 				if (P.ParameterName == Name)
 				{
@@ -696,7 +736,7 @@ TSharedPtr<mu::FParameters> FCustomizableObjectInstanceDescriptor::GetParameters
 
 		case mu::EParameterType::Image:
 		{
-			for (const FCustomizableObjectTextureParameterValue& TextureParameter : TextureParameters)
+			for (const FCustomizableObjectAssetParameterValue& TextureParameter : TextureParameters)
 			{
 				if (TextureParameter.ParameterName == Name || (Uid.IsValid() && TextureParameter.Id == Uid))
 				{
@@ -720,6 +760,32 @@ TSharedPtr<mu::FParameters> FCustomizableObjectInstanceDescriptor::GetParameters
 			break;
 		}
 
+		case mu::EParameterType::Mesh:
+		{
+			for (const FCustomizableObjectAssetParameterValue& MeshParameter : MeshParameters)
+			{
+				if (MeshParameter.ParameterName == Name || (Uid.IsValid() && MeshParameter.Id == Uid))
+				{
+					if (TSharedPtr<mu::FRangeIndex> RangeIdxPtr = MutableParameters->NewRangeIndex(ParamIndex))
+					{
+						for (int32 RangeIndex = 0; RangeIndex < MeshParameter.ParameterRangeValues.Num(); ++RangeIndex)
+						{
+							RangeIdxPtr->SetPosition(0, RangeIndex);
+							MutableParameters->SetMeshValue(ParamIndex, MeshParameter.ParameterRangeValues[RangeIndex], RangeIdxPtr.Get());
+						}
+					}
+					else
+					{
+						MutableParameters->SetMeshValue(ParamIndex, MeshParameter.ParameterValue);
+					}
+
+					break;
+				}
+			}
+
+			break;
+		}
+
 		default:
 			check(false); // Missing case.
 			break;
@@ -766,7 +832,8 @@ void FCustomizableObjectInstanceDescriptor::ReloadParameters()
 	TArray<FCustomizableObjectBoolParameterValue> OldBoolParameters = BoolParameters;
 	TArray<FCustomizableObjectIntParameterValue> OldIntParameters = IntParameters;
 	TArray<FCustomizableObjectFloatParameterValue> OldFloatParameters = FloatParameters;
-	TArray<FCustomizableObjectTextureParameterValue> OldTextureParameters = TextureParameters;
+	TArray<FCustomizableObjectAssetParameterValue> OldTextureParameters = TextureParameters;
+	TArray<FCustomizableObjectAssetParameterValue> OldMeshParameters = MeshParameters;
 	TArray<FCustomizableObjectVectorParameterValue> OldVectorParameters = VectorParameters;
 	TArray<FCustomizableObjectProjectorParameterValue> OldProjectorParameters = ProjectorParameters;
 	TArray<FCustomizableObjectTransformParameterValue> OldTransformParameters = TransformParameters;
@@ -775,6 +842,7 @@ void FCustomizableObjectInstanceDescriptor::ReloadParameters()
 	IntParameters.Reset();
 	FloatParameters.Reset();
 	TextureParameters.Reset();
+	MeshParameters.Reset();
 	VectorParameters.Reset();
 	ProjectorParameters.Reset();
 	TransformParameters.Reset();
@@ -1112,16 +1180,16 @@ void FCustomizableObjectInstanceDescriptor::ReloadParameters()
 
 		case mu::EParameterType::Image:
 		{
-			FCustomizableObjectTextureParameterValue Param;
+			FCustomizableObjectAssetParameterValue Param;
 			Param.ParameterName = Name;
 			Param.Id = Uid;
 
-			auto FindByNameAndUid = [&](const FCustomizableObjectTextureParameterValue& P)
+			auto FindByNameAndUid = [&](const FCustomizableObjectAssetParameterValue& P)
 			{
 				return P.ParameterName == Name || (Uid.IsValid() && P.Id == Uid);
 			};
 			
-			if (FCustomizableObjectTextureParameterValue* Result = OldTextureParameters.FindByPredicate(FindByNameAndUid))
+			if (FCustomizableObjectAssetParameterValue* Result = OldTextureParameters.FindByPredicate(FindByNameAndUid))
 			{	
 				if (TSharedPtr<mu::FRangeIndex> RangeIdxPtr = MutableParameters->NewRangeIndex(ParamIndex))
 				{
@@ -1167,6 +1235,64 @@ void FCustomizableObjectInstanceDescriptor::ReloadParameters()
 			break;
 		}
 
+
+		case mu::EParameterType::Mesh:
+		{
+			FCustomizableObjectAssetParameterValue Param;
+			Param.ParameterName = Name;
+			Param.Id = Uid;
+
+			auto FindByNameAndUid = [&](const FCustomizableObjectAssetParameterValue& P)
+				{
+					return P.ParameterName == Name || (Uid.IsValid() && P.Id == Uid);
+				};
+
+			if (FCustomizableObjectAssetParameterValue* Result = OldTextureParameters.FindByPredicate(FindByNameAndUid))
+			{
+				if (TSharedPtr<mu::FRangeIndex> RangeIdxPtr = MutableParameters->NewRangeIndex(ParamIndex))
+				{
+					Param.ParameterRangeValues = Result->ParameterRangeValues;
+				}
+				else
+				{
+					Param.ParameterValue = Result->ParameterValue;
+				}
+			}
+			else // Not found in Instance Parameters. Use Mutable Parameters.
+			{
+				if (TSharedPtr<mu::FRangeIndex> RangeIdxPtr = MutableParameters->NewRangeIndex(ParamIndex))
+				{
+					const int32 ValueCount = MutableParameters->GetValueCount(ParamIndex);
+
+					for (int32 ValueIndex = 0; ValueIndex < ValueCount; ++ValueIndex)
+					{
+						TSharedPtr<mu::FRangeIndex> RangeValueIdxPtr = MutableParameters->GetValueIndex(ParamIndex, ValueIndex);
+						int32 RangeIndex = RangeValueIdxPtr->GetPosition(0);
+
+						if (!Param.ParameterRangeValues.IsValidIndex(RangeIndex))
+						{
+							const int32 PreviousNum = Param.ParameterRangeValues.Num();
+							Param.ParameterRangeValues.AddUninitialized(RangeIndex + 1 - Param.ParameterRangeValues.Num());
+
+							for (int32 Index = PreviousNum; Index < Param.ParameterRangeValues.Num(); ++Index)
+							{
+								Param.ParameterRangeValues[Index] = FName();
+							}
+						}
+
+						Param.ParameterRangeValues[RangeIndex] = MutableParameters->GetMeshValue(ParamIndex, RangeValueIdxPtr.Get());
+					}
+				}
+				else
+				{
+					Param.ParameterValue = MutableParameters->GetMeshValue(ParamIndex);
+				}
+			}
+
+			MeshParameters.Add(Param);
+			break;
+		}
+
 		default:
 			check(false); // Missing case.
 			break;
@@ -1205,9 +1331,15 @@ const TArray<FCustomizableObjectFloatParameterValue>& FCustomizableObjectInstanc
 }
 
 
-const TArray<FCustomizableObjectTextureParameterValue>& FCustomizableObjectInstanceDescriptor::GetTextureParameters() const
+const TArray<FCustomizableObjectAssetParameterValue>& FCustomizableObjectInstanceDescriptor::GetTextureParameters() const
+{
+	return TextureParameters;
+}
+
+
+const TArray<FCustomizableObjectAssetParameterValue>& FCustomizableObjectInstanceDescriptor::GetMeshParameters() const
 {
-	return TextureParameters;	
+	return MeshParameters;
 }
 
 
@@ -2204,7 +2336,7 @@ int32 FCustomizableObjectInstanceDescriptor::AddValueToTextureRange(const FStrin
 	const int32 TextureParameterIndex = FindTypedParameterIndex(ParamName, EMutableParameterType::Texture);
 	if (TextureParameterIndex != INDEX_NONE)
 	{
-		FCustomizableObjectTextureParameterValue& TextureParameter = TextureParameters[TextureParameterIndex];
+		FCustomizableObjectAssetParameterValue& TextureParameter = TextureParameters[TextureParameterIndex];
 		return TextureParameter.ParameterRangeValues.Add(FName());		
 	}
 	
@@ -2271,7 +2403,7 @@ int32 FCustomizableObjectInstanceDescriptor::RemoveValueFromTextureRange(const F
 	const int32 TextureParameterIndex = FindTypedParameterIndex(ParamName, EMutableParameterType::Texture);
 	if (TextureParameterIndex != INDEX_NONE)
 	{
-		FCustomizableObjectTextureParameterValue& TextureParameter = TextureParameters[TextureParameterIndex];
+		FCustomizableObjectAssetParameterValue& TextureParameter = TextureParameters[TextureParameterIndex];
 		if (TextureParameter.ParameterRangeValues.Num() > 0)
 		{
 			TextureParameter.ParameterRangeValues.Pop();
@@ -2287,7 +2419,7 @@ int32 FCustomizableObjectInstanceDescriptor::RemoveValueFromTextureRange(const F
 	const int32 TextureParameterIndex = FindTypedParameterIndex(ParamName, EMutableParameterType::Texture);
 	if (TextureParameterIndex != INDEX_NONE)
 	{
-		FCustomizableObjectTextureParameterValue& TextureParameter = TextureParameters[TextureParameterIndex];
+		FCustomizableObjectAssetParameterValue& TextureParameter = TextureParameters[TextureParameterIndex];
 		if (TextureParameter.ParameterRangeValues.Num() > 0)
 		{
 			TextureParameter.ParameterRangeValues.RemoveAt(RangeIndex);
@@ -2442,7 +2574,7 @@ void FCustomizableObjectInstanceDescriptor::SetRandomValuesFromStream(const FRan
 
 	if (const int32 NumPossibleValues = PossibleValues.Num())
 	{
-		for (FCustomizableObjectTextureParameterValue& TextureParameter : TextureParameters)
+		for (FCustomizableObjectAssetParameterValue& TextureParameter : TextureParameters)
 		{
 			TextureParameter.ParameterValue = PossibleValues[NumPossibleValues * InStream.GetFraction()];
 		
@@ -2534,7 +2666,7 @@ void FCustomizableObjectInstanceDescriptor::SetDefaultValue(int32 ParamIndex)
 		
 		if (ensure(TextureParameters.IsValidIndex(TypedIndex)))
 		{
-			FCustomizableObjectTextureParameterValue& TextureParameter = TextureParameters[TypedIndex];
+			FCustomizableObjectAssetParameterValue& TextureParameter = TextureParameters[TypedIndex];
 			TextureParameter.ParameterValue = DefaultValue;
 			TextureParameter.ParameterRangeValues.Empty();
 		}
diff --git a/Engine/Plugins/Mutable/Source/CustomizableObject/Private/MuCO/CustomizableObjectSystem.cpp b/Engine/Plugins/Mutable/Source/CustomizableObject/Private/MuCO/CustomizableObjectSystem.cpp
--- a/Engine/Plugins/Mutable/Source/CustomizableObject/Private/MuCO/CustomizableObjectSystem.cpp
+++ b/Engine/Plugins/Mutable/Source/CustomizableObject/Private/MuCO/CustomizableObjectSystem.cpp
@@ -3337,6 +3337,10 @@ namespace impl
 		
 		UCustomizableObjectSystemPrivate* CustomizableObjectSystemPrivateData = System->GetPrivate();
 
+		// Unregister the current CO from the external resource provider.
+		check(IsInGameThread());
+		CustomizableObjectSystemPrivateData->ResourceProvider->SetCurrentObject(nullptr);
+
 		// Next Task: Release Mutable. We need this regardless if we cancel or not
 		//-------------------------------------------------------------		
 		const  TSharedPtr<mu::FSystem> MutableSystem = CustomizableObjectSystemPrivateData->MutableSystem;
@@ -3631,6 +3635,8 @@ namespace impl
 	{
 		MUTABLE_CPUPROFILER_SCOPE(Task_Game_StartUpdate)
 
+		check(IsInGameThread());
+
 #if WITH_EDITOR
 if (CVarMutableLogObjectMemoryOnUpdate.GetValueOnAnyThread())
 {
@@ -3780,6 +3786,13 @@ if (CVarMutableLogObjectMemoryOnUpdate.GetValueOnAnyThread())
 
 		Operation->Model = CustomizableObject->GetPrivate()->GetModel().ToSharedRef();
 
+		// Register the current CO into the external resource provider for runtime meshes and images.
+		// It is needed to get access to the bone map when converting meshes at runtime, for example.
+		if (SystemPrivateData->ResourceProvider)
+		{
+			SystemPrivateData->ResourceProvider->SetCurrentObject(CustomizableObject);
+		}
+
 		UE::Tasks::FTask CacheRuntimeTexturesEvent = UE::Tasks::MakeCompletedTask<void>();
 
 #if WITH_EDITOR
@@ -4436,9 +4449,9 @@ void UCustomizableObjectSystem::UnregisterImageProvider(UCustomizableSystemImage
 }
 
 
-void UCustomizableObjectSystemPrivate::CacheTextureParameters(const TArray<FCustomizableObjectTextureParameterValue>& TextureParameters) const
+void UCustomizableObjectSystemPrivate::CacheTextureParameters(const TArray<FCustomizableObjectAssetParameterValue>& TextureParameters) const
 {
-	for (const FCustomizableObjectTextureParameterValue& TextureParameter : TextureParameters)
+	for (const FCustomizableObjectAssetParameterValue& TextureParameter : TextureParameters)
 	{
 		ResourceProvider->CacheImage(TextureParameter.ParameterValue, false);
 
@@ -4450,9 +4463,9 @@ void UCustomizableObjectSystemPrivate::CacheTextureParameters(const TArray<FCust
 }
 
 
-void UCustomizableObjectSystemPrivate::UnCacheTextureParameters(const TArray<FCustomizableObjectTextureParameterValue>& TextureParameters) const
+void UCustomizableObjectSystemPrivate::UnCacheTextureParameters(const TArray<FCustomizableObjectAssetParameterValue>& TextureParameters) const
 {
-	for (const FCustomizableObjectTextureParameterValue& TextureParameter : TextureParameters)
+	for (const FCustomizableObjectAssetParameterValue& TextureParameter : TextureParameters)
 	{
 		ResourceProvider->UnCacheImage(TextureParameter.ParameterValue, false);
 
diff --git a/Engine/Plugins/Mutable/Source/CustomizableObject/Private/MuCO/DescriptorHash.cpp b/Engine/Plugins/Mutable/Source/CustomizableObject/Private/MuCO/DescriptorHash.cpp
--- a/Engine/Plugins/Mutable/Source/CustomizableObject/Private/MuCO/DescriptorHash.cpp
+++ b/Engine/Plugins/Mutable/Source/CustomizableObject/Private/MuCO/DescriptorHash.cpp
@@ -29,7 +29,12 @@ FDescriptorHash::FDescriptorHash(const FCustomizableObjectInstanceDescriptor& De
 		Hash = HashCombine(Hash, GetTypeHash(Value));
 	}
 
-	for (const FCustomizableObjectTextureParameterValue& Value : Descriptor.TextureParameters)
+	for (const FCustomizableObjectAssetParameterValue& Value : Descriptor.TextureParameters)
+	{
+		Hash = HashCombine(Hash, GetTypeHash(Value));
+	}
+
+	for (const FCustomizableObjectAssetParameterValue& Value : Descriptor.MeshParameters)
 	{
 		Hash = HashCombine(Hash, GetTypeHash(Value));
 	}
diff --git a/Engine/Plugins/Mutable/Source/CustomizableObject/Private/MuCO/UnrealConversionUtils.cpp b/Engine/Plugins/Mutable/Source/CustomizableObject/Private/MuCO/UnrealConversionUtils.cpp
--- a/Engine/Plugins/Mutable/Source/CustomizableObject/Private/MuCO/UnrealConversionUtils.cpp
+++ b/Engine/Plugins/Mutable/Source/CustomizableObject/Private/MuCO/UnrealConversionUtils.cpp
@@ -4,15 +4,18 @@
 
 #include "MuCO/CustomizableObjectInstance.h"
 #include "MuCO/UnrealPortabilityHelpers.h"
-#include "Engine/SkeletalMesh.h"
-#include "GPUSkinVertexFactory.h"
 #include "MuCO/CustomizableObject.h"
 #include "MuCO/CustomizableObjectPrivate.h"
 #include "MuCO/CustomizableObjectSystemPrivate.h"
+#include "MuCO/MutableMeshBufferUtils.h"
 #include "MuR/Skeleton.h"
 #include "MuR/Mesh.h"
 #include "MuR/MeshBufferSet.h"
 #include "MuR/MutableTrace.h"
+#include "Animation/Skeleton.h"
+#include "Engine/SkeletalMesh.h"
+#include "GPUSkinVertexFactory.h"
+#include "Rendering/SkeletalMeshRenderData.h"
 
 class USkeleton;
 
@@ -802,4 +805,330 @@ namespace UnrealConversionUtils
 		const float ErrorTolerance = Owner.GetLODInfo(LODIndex)->MorphTargetPositionErrorTolerance;
 		LODResource.MorphTargetVertexInfoBuffers.InitMorphResourcesStreaming(LODResource.RenderSections, MorphsToInit, LODResource.StaticVertexBuffers.StaticMeshVertexBuffer.GetNumVertices(), ErrorTolerance);
 	}
+
+
+	void ConvertSkeletalMeshFromRuntimeData(USkeletalMesh& SkeletalMesh, int32 LODIndex, int32 SectionIndex, UModelResources* ModelResources, mu::FMesh& Result)
+	{
+		MUTABLE_CPUPROFILER_SCOPE(ConvertSkeletalMeshFromRuntimeData);
+
+		FSkeletalMeshRenderData* Data = SkeletalMesh.GetResourceForRendering();
+
+		if (!Data->LODRenderData.IsValidIndex(LODIndex))
+		{
+			// This may happen if we are requesting a LOD that is not present in the provided mesh.
+			// TODO: an empty mesh is returned, but should we use the last LOD instead?
+			UE_LOG(LogMutable, Warning, TEXT("Trying to convert mesh [%s] LOD %d but it only has %d LODs"), 
+				*SkeletalMesh.GetName(), LODIndex, Data->LODRenderData.Num() );
+
+			return;
+		}
+
+		FSkeletalMeshLODRenderData& LOD = Data->LODRenderData[LODIndex];
+
+		if (!LOD.RenderSections.IsValidIndex(SectionIndex))
+		{
+			// This may happen in the provided mesh in the parameter value doesn't have the expected number of sections
+			UE_LOG(LogMutable, Warning, TEXT("Trying to convert mesh [%s] section %d in LOD %d but it only has %d sections"),
+				*SkeletalMesh.GetName(), SectionIndex, LODIndex, LOD.RenderSections.Num());
+
+			return;
+		}
+
+		const FSkelMeshRenderSection& Section = LOD.RenderSections[SectionIndex];
+
+		// Skeleton data
+		bool bBoneMapModified = false;
+		TArray<FBoneIndexType> BoneMap;
+		TArray<FBoneIndexType> RemappedBoneMapIndices;
+
+		bool bIgnoreSkeleton = false;
+		if (!bIgnoreSkeleton)
+		{
+			USkeleton* Skeleton = SkeletalMesh.GetSkeleton();
+			if (!Skeleton)
+			{
+				ensure(false);
+				return;
+			}
+
+			// Add the skeleton to the list of referenced skeletons and add its index to the mesh
+			if (ModelResources)
+			{
+				const int32 SkeletonID = ModelResources->Skeletons.AddUnique(Skeleton);
+				Result.AddSkeletonID(SkeletonID);
+			}
+
+			const TArray<uint16>& SourceRequiredBones = LOD.RequiredBones;
+
+			// Build RequiredBones array
+			TArray<FBoneIndexType> RequiredBones;
+			RequiredBones.Reserve(SourceRequiredBones.Num());
+
+			for (const FBoneIndexType& RequiredBoneIndex : SourceRequiredBones)
+			{
+				RequiredBones.AddUnique(RequiredBoneIndex);
+			}
+
+			// Rebuild BoneMap
+			const TArray<uint16>& SourceBoneMap = Section.BoneMap;
+			const int32 NumBonesInBoneMap = SourceBoneMap.Num();
+
+			for (int32 BoneIndex = 0; BoneIndex < NumBonesInBoneMap; ++BoneIndex)
+			{
+				const FBoneIndexType FinalBoneIndex = SourceBoneMap[BoneIndex];
+
+				const int32 BoneMapIndex = BoneMap.AddUnique(FinalBoneIndex);
+				RemappedBoneMapIndices.Add(BoneMapIndex);
+
+				bBoneMapModified = bBoneMapModified || SourceBoneMap[BoneIndex] != FinalBoneIndex;
+			}
+
+			// Create the skeleton, poses, and BoneMap for this mesh
+			TSharedPtr<mu::FSkeleton> MutableSkeleton = MakeShared<mu::FSkeleton>();
+			Result.SetSkeleton(MutableSkeleton);
+
+			const int32 NumRequiredBones = RequiredBones.Num();
+			Result.SetBonePoseCount(NumRequiredBones);
+			MutableSkeleton->SetBoneCount(NumRequiredBones);
+
+			// MutableBoneMap will not keep an index to the Skeleton, but to the BoneName
+			TArray<mu::FBoneName> MutableBoneMap;
+			MutableBoneMap.SetNum(BoneMap.Num());
+
+			TArray<FMatrix> ComposedRefPoseMatrices;
+			ComposedRefPoseMatrices.SetNum(NumRequiredBones);
+
+			const TArray<FMeshBoneInfo>& RefBoneInfo = SkeletalMesh.GetRefSkeleton().GetRefBoneInfo();
+			for (int32 BoneIndex = 0; BoneIndex < NumRequiredBones; ++BoneIndex)
+			{
+				const int32 RefSkeletonBoneIndex = RequiredBones[BoneIndex];
+
+				const FMeshBoneInfo& BoneInfo = RefBoneInfo[RefSkeletonBoneIndex];
+				const int32 ParentBoneIndex = RequiredBones.Find(BoneInfo.ParentIndex);
+
+				// Set bone hierarchy
+
+				// Get the bone id
+				mu::FBoneName BoneName;
+				if (ModelResources)
+				{
+					// See if the bone exists in the compiled CO data
+					const FString BoneNameString = BoneInfo.Name.ToString().ToLower();
+					uint32* BoneId = ModelResources->BoneNamesMap.Find(BoneNameString);
+
+					if (BoneId)
+					{
+						BoneName = *BoneId;
+					}
+					else
+					{
+						// Go the slow way and add it to get a unique ID that the core can work with.
+						uint32 NewBoneId = CityHash32(reinterpret_cast<const char*>(*BoneNameString), BoneNameString.Len() * sizeof(FString::ElementType));
+
+						// See if the hash collides with an existing bone
+						bool bUnique = false;
+						while (!bUnique)
+						{
+							bUnique = true;
+							TMap<FString, uint32>::TConstIterator MapIterator = ModelResources->BoneNamesMap.CreateConstIterator();
+							while (MapIterator)
+							{
+								if (MapIterator.Value() == NewBoneId)
+								{
+									bUnique = false;
+									break;
+								}
+								++MapIterator;
+							}
+
+							if (!bUnique)
+							{
+								NewBoneId++;
+							}
+						}
+
+						ModelResources->BoneNamesMap.Add(BoneNameString, NewBoneId);
+						BoneName = NewBoneId;
+					}
+				}
+
+				MutableSkeleton->SetBoneName(BoneIndex, BoneName);
+				MutableSkeleton->SetBoneParent(BoneIndex, ParentBoneIndex);
+
+				// Debug. Will not be serialized
+				MutableSkeleton->SetDebugName(BoneIndex, BoneInfo.Name);
+
+				// BoneMap: Convert RefSkeletonBoneIndex to BoneId
+				const int32 BoneMapIndex = BoneMap.Find(RefSkeletonBoneIndex);
+				if (BoneMapIndex != INDEX_NONE)
+				{
+					MutableBoneMap[BoneMapIndex] = BoneName;
+				}
+
+				if (ParentBoneIndex >= 0)
+				{
+					ComposedRefPoseMatrices[BoneIndex] = SkeletalMesh.GetRefPoseMatrix(RefSkeletonBoneIndex) * ComposedRefPoseMatrices[ParentBoneIndex];
+				}
+				else
+				{
+					ComposedRefPoseMatrices[BoneIndex] = SkeletalMesh.GetRefPoseMatrix(RefSkeletonBoneIndex);
+				}
+
+				// Set bone pose
+				FTransform3f BoneTransform;
+				BoneTransform.SetFromMatrix(FMatrix44f(ComposedRefPoseMatrices[BoneIndex]));
+
+				mu::EBoneUsageFlags BoneUsageFlags = mu::EBoneUsageFlags::None;
+				EnumAddFlags(BoneUsageFlags, BoneMapIndex != INDEX_NONE ? mu::EBoneUsageFlags::Skinning : mu::EBoneUsageFlags::None);
+				EnumAddFlags(BoneUsageFlags, ParentBoneIndex == INDEX_NONE ? mu::EBoneUsageFlags::Root : mu::EBoneUsageFlags::None);
+
+				Result.SetBonePose(BoneIndex, BoneName, BoneTransform, BoneUsageFlags);
+			}
+
+			Result.SetBoneMap(MutableBoneMap);
+		}
+
+
+
+		// Mesh data
+		int32 VertexCount = Section.GetNumVertices();
+
+		mu::FMeshBufferSet& Vertices = Result.GetVertexBuffers();
+		Vertices.SetElementCount(VertexCount, mu::EMemoryInitPolicy::Zeroed);
+		Vertices.SetBufferCount(5);
+
+		// Position buffer
+		int32 CurrentVertexBuffer = 0;
+		MutableMeshBufferUtils::SetupVertexPositionsBuffer(CurrentVertexBuffer, Vertices);
+		check(Vertices.Buffers[CurrentVertexBuffer].ElementSize * VertexCount <= LOD.StaticVertexBuffers.PositionVertexBuffer.GetAllocatedSize());
+		FMemory::Memcpy(Vertices.GetBufferData(CurrentVertexBuffer), LOD.StaticVertexBuffers.PositionVertexBuffer.GetVertexData(), Vertices.Buffers[CurrentVertexBuffer].ElementSize * VertexCount);
+		++CurrentVertexBuffer;
+
+		// Tangent buffer
+		{
+			bool bIsHighPrecision = LOD.StaticVertexBuffers.StaticMeshVertexBuffer.GetUseHighPrecisionTangentBasis();
+
+			const mu::EMeshBufferSemantic Semantics[2] = { mu::EMeshBufferSemantic::Tangent, mu::EMeshBufferSemantic::Normal };
+			const int32 SemanticIndices[2] = { 0, 0 };
+			const int32 Components[2] = { 4, 4 };
+			mu::EMeshBufferFormat Formats[2] = { mu::EMeshBufferFormat::PackedDirS8, mu::EMeshBufferFormat::PackedDirS8_W_TangentSign };
+			int32 Offsets[2] = { 0, 4 };
+			int32 ElementSize = 8;
+
+			if (bIsHighPrecision)
+			{
+				// Not really supported
+				ensure(false);
+				Formats[0] = mu::EMeshBufferFormat::Int16;
+				Formats[1] = mu::EMeshBufferFormat::Int16;
+				Offsets[1] = 8;
+				ElementSize = 16;
+			}
+			Vertices.SetBuffer(CurrentVertexBuffer, ElementSize, 2, Semantics, SemanticIndices, Formats, Components, Offsets);
+
+			check((int32)Vertices.Buffers[CurrentVertexBuffer].ElementSize * VertexCount <= LOD.StaticVertexBuffers.StaticMeshVertexBuffer.GetTangentSize());
+			FMemory::Memcpy(Vertices.GetBufferData(CurrentVertexBuffer), LOD.StaticVertexBuffers.StaticMeshVertexBuffer.GetTangentData(), ElementSize * VertexCount);
+
+			++CurrentVertexBuffer;
+		}
+
+		// Texture coords buffer
+		{
+			int32 NumTexCoords = LOD.GetNumTexCoords();
+
+			constexpr int32 MaxChannelCount = 4;
+			mu::EMeshBufferSemantic Semantics[MaxChannelCount];
+			int32 SemanticIndices[MaxChannelCount];
+			mu::EMeshBufferFormat Formats[MaxChannelCount];
+			int32 Components[MaxChannelCount];
+			int32 Offsets[MaxChannelCount];
+
+			int32 UVSize = LOD.StaticVertexBuffers.StaticMeshVertexBuffer.GetUseFullPrecisionUVs() ? 8 : 4;
+			const int32 ElementSize = 8 * NumTexCoords; // UVSize * NumTexCoords;
+
+			for (int32 UV = 0; UV < NumTexCoords; ++UV)
+			{
+				Semantics[UV] = mu::EMeshBufferSemantic::TexCoords;
+				SemanticIndices[UV] = UV;
+				Formats[UV] = mu::EMeshBufferFormat::Float32; //  (UVSize == 8) ? mu::EMeshBufferFormat::Float32 : mu::EMeshBufferFormat::Float16;
+				Components[UV] = 2;
+				Offsets[UV] = UVSize * UV;
+			}
+			Vertices.SetBuffer(CurrentVertexBuffer, ElementSize, NumTexCoords, Semantics, SemanticIndices, Formats, Components, Offsets);
+
+			if (UVSize == 8)
+			{
+				FMemory::Memcpy(Vertices.GetBufferData(CurrentVertexBuffer), LOD.StaticVertexBuffers.StaticMeshVertexBuffer.GetTexCoordData(), Vertices.Buffers[CurrentVertexBuffer].ElementSize * VertexCount);
+			}
+			else
+			{
+				float* Target = (float*)(Vertices.GetBufferData(CurrentVertexBuffer));
+				const FFloat16* Source = (const FFloat16*)LOD.StaticVertexBuffers.StaticMeshVertexBuffer.GetTexCoordData();
+				for (int32 v = 0; v < VertexCount * 2 * NumTexCoords; ++v)
+				{
+					*Target = *Source;
+					++Target;
+					++Source;
+				}
+			}
+			++CurrentVertexBuffer;
+		}
+
+		// Skin buffer
+		if (Section.MaxBoneInfluences > 0)
+		{
+			const FSkinWeightVertexBuffer* SkinBuffer = LOD.GetSkinWeightVertexBuffer();
+			const int32 MaxBoneIndexTypeSizeBytes = SkinBuffer->GetBoneIndexByteSize();
+			const int32 MaxBoneWeightTypeSizeBytes = SkinBuffer->GetBoneWeightByteSize();
+			const int32 MaxBonesPerVertex = SkinBuffer->GetMaxBoneInfluences();
+			MutableMeshBufferUtils::SetupSkinBuffer(CurrentVertexBuffer, MaxBoneIndexTypeSizeBytes, MaxBoneWeightTypeSizeBytes, MaxBonesPerVertex, Vertices);
+			check(Vertices.Buffers[CurrentVertexBuffer].ElementSize <= SkinBuffer->GetVertexDataSize());
+			FMemory::Memcpy(Vertices.GetBufferData(CurrentVertexBuffer),
+				SkinBuffer->GetDataVertexBuffer()->GetWeightData(),
+				Vertices.Buffers[CurrentVertexBuffer].ElementSize * VertexCount);
+
+			++CurrentVertexBuffer;
+		}
+
+		// Colour buffer
+		// TODO
+		//if (MeshData.bHasVertexColors)
+		//{
+		//	MutableMeshBufferUtils::SetupVertexColorBuffer(CurrentVertexBuffer, Vertices);
+		//	++CurrentVertexBuffer;
+		//}
+
+		// Indices
+		{
+			int32 FaceCount = Section.NumTriangles;
+			int32 ElementSize = LOD.MultiSizeIndexContainer.GetDataTypeSize();
+
+			mu::FMeshBufferSet& Indices = Result.GetIndexBuffers();
+			Indices.SetBufferCount(1);
+			Indices.SetElementCount(FaceCount * 3);
+			constexpr int32 ChannelCount = 1;
+			const mu::EMeshBufferSemantic Semantics[ChannelCount] = { mu::EMeshBufferSemantic::VertexIndex };
+			const int32 SemanticIndices[ChannelCount] = { 0 };
+			mu::EMeshBufferFormat Formats[ChannelCount] = { ElementSize == 4 ? mu::EMeshBufferFormat::UInt32 : mu::EMeshBufferFormat::UInt16 };
+			const int32 Components[ChannelCount] = { 1 };
+			const int32 Offsets[ChannelCount] = { 0 };
+			Indices.SetBuffer(0, ElementSize, ChannelCount, Semantics, SemanticIndices, Formats, Components, Offsets);
+
+			FMemory::Memcpy(Indices.GetBufferData(0), LOD.MultiSizeIndexContainer.GetIndexBuffer()->GetPointerTo(0), Indices.GetElementCount() * ElementSize);
+		}
+
+		// Assign an ID prefix to the mesh. 
+		// TODO: Hash all data?
+		// TODO: Ensure unique?
+		uint32 MeshID = HashCombine(GetTypeHash(LODIndex), GetTypeHash(SkeletalMesh.GetFullName()));
+		MeshID = HashCombine(MeshID, GetTypeHash(SectionIndex));
+		Result.MeshIDPrefix = MeshID;
+
+		FString DebugLog;
+		Result.Log(DebugLog, 8);
+		UE_LOG(LogMutable, Warning, TEXT("[%s]"), *DebugLog);
+
+		Result.EnsureSurfaceData();
+	}
+
 }
diff --git a/Engine/Plugins/Mutable/Source/CustomizableObject/Private/MuCO/UnrealMutableImageProvider.cpp b/Engine/Plugins/Mutable/Source/CustomizableObject/Private/MuCO/UnrealMutableImageProvider.cpp
--- a/Engine/Plugins/Mutable/Source/CustomizableObject/Private/MuCO/UnrealMutableImageProvider.cpp
+++ b/Engine/Plugins/Mutable/Source/CustomizableObject/Private/MuCO/UnrealMutableImageProvider.cpp
@@ -2,21 +2,26 @@
 
 #include "MuCO/UnrealMutableImageProvider.h"
 
-#include "ImageCoreUtils.h"
 #include "MuCO/CustomizableObjectSystem.h"
 #include "MuCO/CustomizableObject.h"
 #include "MuCO/CustomizableObjectPrivate.h"
 #include "MuCO/UnrealToMutableTextureConversionUtils.h"
-#include "TextureResource.h"
+#include "MuCO/UnrealConversionUtils.h"
 #include "MuCO/LoadUtils.h"
+#include "MuCO/MutableMeshBufferUtils.h"
 #include "MuR/Parameters.h"
 #include "MuR/ImageTypes.h"
 #include "MuR/Model.h"
 #include "MuR/Mesh.h"
+#include "Engine/SkeletalMesh.h"
+#include "Animation/BoneReference.h"
+#include "Animation/Skeleton.h"
 #include "ProfilingDebugging/IoStoreTrace.h"
+#include "ReferenceSkeleton.h"
+#include "ImageCoreUtils.h"
+#include "TextureResource.h"
 
 
-//-------------------------------------------------------------------------------------------------
 namespace
 {
 	void ConvertTextureUnrealPlatformToMutable(mu::FImage* OutResult, UTexture2D* Texture, uint8 MipmapsToSkip)
@@ -472,51 +477,70 @@ mu::FExtendedImageDesc FUnrealMutableResourceProvider::GetImageDesc(FName Id)
 }
 
 
-TTuple<UE::Tasks::FTask, TFunction<void()>> FUnrealMutableResourceProvider::GetMeshAsync(FName Id, TFunction<void(TSharedPtr<mu::FMesh>)>& ResultCallback)
+TTuple<UE::Tasks::FTask, TFunction<void()>> FUnrealMutableResourceProvider::GetMeshAsync(FName Id, int32 LODIndex, int32 SectionIndex, TFunction<void(TSharedPtr<mu::FMesh>)>& ResultCallback)
 {
-	// TODO: This is a stub function where Mesh FParameters would be supported.
-
 	// Thread: worker
 	MUTABLE_CPUPROFILER_SCOPE(FUnrealMutableImageProvider::GetMeshAsync);
 
-	auto TrivialReturn = []() -> TTuple<UE::Tasks::FTask, TFunction<void()>>
-		{
-			return MakeTuple(UE::Tasks::MakeCompletedTask<void>(), []() -> void {});
-		};
-
-	// No USkeletalMesh was provided, cannot do anything, just provide a dummy mesh
-	UE_LOG(LogMutable, Warning, TEXT("No USkeletalMesh was provided for an application-specific mesh parameter."));
-
-	TSharedPtr<mu::FMesh> Dummy = MakeShared<mu::FMesh>();
-	ResultCallback( Dummy );
-	return Invoke(TrivialReturn);
-}
+	TSharedPtr<mu::FMesh> Result = MakeShared<mu::FMesh>();
 
+	UE::Tasks::FTaskEvent Completion(TEXT("MutableMeshParameterLoadInGameThread"));
+	ExecuteOnGameThread(TEXT("MeshParameterLoad"),
+		[Id, Completion, LODIndex, SectionIndex, Result, CustomizableObject=CurrentCustomizableObject] () mutable
+		{
+			check(IsInGameThread());
 
-TTuple<UE::Tasks::FTask, TFunction<void()>> FUnrealMutableResourceProvider::GetReferencedMeshAsync(const void* ModelPtr, int32 Id, TFunction<void(TSharedPtr<mu::FMesh>)>& ResultCallback)
-{
-	// TODO: This is a stub function where referenced Meshes would be supported.
+			MUTABLE_CPUPROFILER_SCOPE(MutableMeshParameterLoadInGameThread);
 
-	// Thread: worker
-	MUTABLE_CPUPROFILER_SCOPE(FUnrealMutableImageProvider::GetReferencedMeshAsync);
+			// TODO: Do it async
+			USkeletalMesh* SkeletalMesh = nullptr;
+			{
+				MUTABLE_CPUPROFILER_SCOPE(ActualLoad);
+				SkeletalMesh = Cast<USkeletalMesh>(FSoftObjectPath(Id.ToString()).TryLoad());
+			}
 
-	auto TrivialReturn = []() -> TTuple<UE::Tasks::FTask, TFunction<void()>>
-		{
-			return MakeTuple(UE::Tasks::MakeCompletedTask<void>(), []() -> void {});
-		};
+			if (!SkeletalMesh)
+			{
+				UE_LOG(LogMutable, Error, TEXT("Failed to load the skeletal mesh [%s] set for a mesh parameter. Please check that it was packaged."), *Id.ToString());
+			}
+			else
+			{
+				TStrongObjectPtr<UCustomizableObject> CO = CustomizableObject.Pin();
+				check(CO);
 
-	TSharedPtr<mu::FMesh> Mesh = MakeShared<mu::FMesh>();
+				if (CO)
+				{
+					UModelResources* ModelResources = CO->GetPrivate()->GetModelResources();
+					check(ModelResources);
+					if (ModelResources)
+					{
+						UnrealConversionUtils::ConvertSkeletalMeshFromRuntimeData(*SkeletalMesh, LODIndex, SectionIndex, ModelResources, *Result.Get());
+					}
+				}
+			}
 
-	// Not supported outside editor yet.
-	UE_LOG(LogMutable, Warning, TEXT("Failed to get referenced mesh. Not implemented yet."));
+			// Signal
+			Completion.Trigger();
+		}
+	);
 
-	ResultCallback(Mesh);
-	return Invoke(TrivialReturn);
+	return MakeTuple(
+		// Some post-game conversion stuff can happen here in a worker thread
+		UE::Tasks::Launch(TEXT("MutableMeshParameterLoadPostGame"),
+			[ResultCallback, Result]()
+			{
+				ResultCallback(Result);
+			},
+			Completion ),
 
+		// Cleanup code that will be called after the result is received in calling code.
+		[]()
+		{
+		}
+	);
 }
 
 
-
 void FUnrealMutableResourceProvider::CacheImage(FName Id, bool bUser)
 {
 	if (Id == NAME_None)
@@ -850,6 +874,13 @@ void FUnrealMutableResourceProvider::AddReferencedObjects(FReferenceCollector& C
 }
 
 
+void FUnrealMutableResourceProvider::SetCurrentObject(const TWeakObjectPtr<UCustomizableObject>& InObject)
+{
+	check(IsInGameThread());
+	CurrentCustomizableObject = InObject;
+}
+
+
 FUnrealMutableResourceProvider::FUnrealMutableImageInfo::FUnrealMutableImageInfo(const TSharedPtr<mu::FImage>& InImage)
 {
 	check(IsInGameThread())
diff --git a/Engine/Plugins/Mutable/Source/CustomizableObject/Private/MuCO/UnrealMutableImageProvider.h b/Engine/Plugins/Mutable/Source/CustomizableObject/Private/MuCO/UnrealMutableImageProvider.h
--- a/Engine/Plugins/Mutable/Source/CustomizableObject/Private/MuCO/UnrealMutableImageProvider.h
+++ b/Engine/Plugins/Mutable/Source/CustomizableObject/Private/MuCO/UnrealMutableImageProvider.h
@@ -2,6 +2,7 @@
 
 #pragma once
 
+#include "MuCO/UnrealToMutableTextureConversionUtils.h"
 #include "MuR/Image.h"
 #include "MuR/System.h"
 #include "UObject/GCObject.h"
@@ -9,8 +10,6 @@
 #include "Containers/Map.h"
 #include "Containers/Queue.h"
 #include "PixelFormat.h"
-#include "MuCO/UnrealToMutableTextureConversionUtils.h"
-
 #include "Tasks/Task.h"
 
 class UCustomizableObject;
@@ -27,8 +26,7 @@ public:
 	virtual TTuple<UE::Tasks::FTask, TFunction<void()>> GetImageAsync(FName Id, uint8 MipmapsToSkip, TFunction<void(TSharedPtr<mu::FImage>)>& ResultCallback) override;
 	virtual TTuple<UE::Tasks::FTask, TFunction<void()>> GetReferencedImageAsync(const void* ModelPtr, int32 Id, uint8 MipmapsToSkip, TFunction<void(TSharedPtr<mu::FImage>)>& ResultCallback) override;
 	virtual mu::FExtendedImageDesc GetImageDesc(FName Id) override;
-	virtual TTuple<UE::Tasks::FTask, TFunction<void()>> GetMeshAsync(FName Id, TFunction<void(TSharedPtr<mu::FMesh>)>& ResultCallback) override;
-	virtual TTuple<UE::Tasks::FTask, TFunction<void()>> GetReferencedMeshAsync(const void* ModelPtr, int32 Id, TFunction<void(TSharedPtr<mu::FMesh>)>& ResultCallback) override;
+	virtual TTuple<UE::Tasks::FTask, TFunction<void()>> GetMeshAsync(FName Id, int32 LODIndex, int32 SectionIndex, TFunction<void(TSharedPtr<mu::FMesh>)>& ResultCallback) override;
 
 	
 	// Own interface	
@@ -61,6 +59,9 @@ public:
 		return TEXT("FUnrealMutableImageProvider");
 	}
 
+	// Register and unregister the current CO that we are working on. This is called by internal code when building instances.
+	void SetCurrentObject( const TWeakObjectPtr<UCustomizableObject>& );
+
 private:
 	struct FUnrealMutableImageInfo
 	{
@@ -87,6 +88,9 @@ private:
 		int32 ReferencesSystem = 0;
 	};
 
+	/** Runtime reference to the CO being updated. It is needed for example for mesh conversion, to access to the bone id map.  */
+	TWeakObjectPtr<UCustomizableObject> CurrentCustomizableObject;
+
 	static inline const mu::FImageDesc DUMMY_IMAGE_DESC = 
 			mu::FImageDesc {mu::FImageSize(32, 32), mu::EImageFormat::RGBA_UByte, 1};
 
diff --git a/Engine/Plugins/Mutable/Source/CustomizableObject/Public/MuCO/CustomizableObject.h b/Engine/Plugins/Mutable/Source/CustomizableObject/Public/MuCO/CustomizableObject.h
--- a/Engine/Plugins/Mutable/Source/CustomizableObject/Public/MuCO/CustomizableObject.h
+++ b/Engine/Plugins/Mutable/Source/CustomizableObject/Public/MuCO/CustomizableObject.h
@@ -89,7 +89,11 @@ struct FProfileParameterDat
 
 	////
 	UPROPERTY()
-	TArray<FCustomizableObjectTextureParameterValue> TextureParameters;
+	TArray<FCustomizableObjectAssetParameterValue> TextureParameters;
+
+	////
+	UPROPERTY()
+	TArray<FCustomizableObjectAssetParameterValue> MeshParameters;
 
 	////
 	UPROPERTY()
diff --git a/Engine/Plugins/Mutable/Source/CustomizableObject/Public/MuCO/CustomizableObjectInstance.h b/Engine/Plugins/Mutable/Source/CustomizableObject/Public/MuCO/CustomizableObjectInstance.h
--- a/Engine/Plugins/Mutable/Source/CustomizableObject/Public/MuCO/CustomizableObjectInstance.h
+++ b/Engine/Plugins/Mutable/Source/CustomizableObject/Public/MuCO/CustomizableObjectInstance.h
@@ -723,7 +723,7 @@ private:
 	TArray<FCustomizableObjectFloatParameterValue> FloatParameters_DEPRECATED;
 
 	UPROPERTY()																																						
-	TArray<FCustomizableObjectTextureParameterValue> TextureParameters_DEPRECATED;
+	TArray<FCustomizableObjectAssetParameterValue> TextureParameters_DEPRECATED;
 
 	UPROPERTY()
 	TArray<FCustomizableObjectVectorParameterValue> VectorParameters_DEPRECATED;
diff --git a/Engine/Plugins/Mutable/Source/CustomizableObject/Public/MuCO/CustomizableObjectInstanceDescriptor.h b/Engine/Plugins/Mutable/Source/CustomizableObject/Public/MuCO/CustomizableObjectInstanceDescriptor.h
--- a/Engine/Plugins/Mutable/Source/CustomizableObject/Public/MuCO/CustomizableObjectInstanceDescriptor.h
+++ b/Engine/Plugins/Mutable/Source/CustomizableObject/Public/MuCO/CustomizableObjectInstanceDescriptor.h
@@ -87,7 +87,9 @@ struct CUSTOMIZABLEOBJECT_API FCustomizableObjectInstanceDescriptor
 
 	const TArray<FCustomizableObjectFloatParameterValue>& GetFloatParameters() const;
 	
-	const TArray<FCustomizableObjectTextureParameterValue>& GetTextureParameters() const;
+	const TArray<FCustomizableObjectAssetParameterValue>& GetTextureParameters() const;
+
+	const TArray<FCustomizableObjectAssetParameterValue>& GetMeshParameters() const;
 
 	const TArray<FCustomizableObjectVectorParameterValue>& GetVectorParameters() const;
 
@@ -310,10 +312,13 @@ struct CUSTOMIZABLEOBJECT_API FCustomizableObjectInstanceDescriptor
 	
 	UPROPERTY()
 	TArray<FCustomizableObjectFloatParameterValue> FloatParameters;
-	
+
 	UPROPERTY()
-	TArray<FCustomizableObjectTextureParameterValue> TextureParameters;
-	
+	TArray<FCustomizableObjectAssetParameterValue> TextureParameters;
+
+	UPROPERTY()
+	TArray<FCustomizableObjectAssetParameterValue> MeshParameters;
+
 	UPROPERTY()
 	TArray<FCustomizableObjectVectorParameterValue> VectorParameters;
 	
diff --git a/Engine/Plugins/Mutable/Source/CustomizableObject/Public/MuCO/CustomizableObjectParameterTypeDefinitions.h b/Engine/Plugins/Mutable/Source/CustomizableObject/Public/MuCO/CustomizableObjectParameterTypeDefinitions.h
--- a/Engine/Plugins/Mutable/Source/CustomizableObject/Public/MuCO/CustomizableObjectParameterTypeDefinitions.h
+++ b/Engine/Plugins/Mutable/Source/CustomizableObject/Public/MuCO/CustomizableObjectParameterTypeDefinitions.h
@@ -19,6 +19,7 @@ enum class EMutableParameterType : uint8
 	Projector	UMETA(DisplayName = "Projector"),
 	Texture		UMETA(DisplayName = "Texture"),
 	Transform	UMETA(DisplayName = "Transform"),
+	Mesh		UMETA(DisplayName = "Mesh"),
 };
 
 
@@ -156,7 +157,7 @@ inline uint32 GetTypeHash(const FCustomizableObjectFloatParameterValue& Key)
 
 
 USTRUCT(BlueprintType)
-struct FCustomizableObjectTextureParameterValue
+struct FCustomizableObjectAssetParameterValue
 {
 	GENERATED_USTRUCT_BODY()
 
@@ -175,8 +176,10 @@ struct FCustomizableObjectTextureParameterValue
 	TArray<FName> ParameterRangeValues;
 };
 
+using FCustomizableObjectTextureParameterValue = FCustomizableObjectAssetParameterValue;
 
-inline uint32 GetTypeHash(const FCustomizableObjectTextureParameterValue& Key)
+
+inline uint32 GetTypeHash(const FCustomizableObjectAssetParameterValue& Key)
 {
 	uint32 Hash = GetTypeHash(Key.ParameterName);
 	Hash = HashCombine(Hash, GetTypeHash(Key.ParameterValue));
diff --git a/Engine/Plugins/Mutable/Source/CustomizableObjectEditor/Private/MuCOE/CustomizableInstanceDetails.cpp b/Engine/Plugins/Mutable/Source/CustomizableObjectEditor/Private/MuCOE/CustomizableInstanceDetails.cpp
--- a/Engine/Plugins/Mutable/Source/CustomizableObjectEditor/Private/MuCOE/CustomizableInstanceDetails.cpp
+++ b/Engine/Plugins/Mutable/Source/CustomizableObjectEditor/Private/MuCOE/CustomizableInstanceDetails.cpp
@@ -1328,7 +1328,7 @@ void FCustomizableInstanceDetails::GenerateTextureParameterOptions()
 {
 	UCustomizableObject* CustomizableObject = CustomInstance->GetCustomizableObject();
 
-	for (const FCustomizableObjectTextureParameterValue& TextureParameter : CustomInstance->GetPrivate()->GetDescriptor().GetTextureParameters())
+	for (const FCustomizableObjectAssetParameterValue& TextureParameter : CustomInstance->GetPrivate()->GetDescriptor().GetTextureParameters())
 	{
 		// Get default values.
 		const FName& DefaultValue = CustomizableObject->GetTextureParameterDefaultValue(TextureParameter.ParameterName);
diff --git a/Engine/Plugins/Mutable/Source/CustomizableObjectEditor/Private/MuCOE/CustomizableObjectEditorModule.cpp b/Engine/Plugins/Mutable/Source/CustomizableObjectEditor/Private/MuCOE/CustomizableObjectEditorModule.cpp
--- a/Engine/Plugins/Mutable/Source/CustomizableObjectEditor/Private/MuCOE/CustomizableObjectEditorModule.cpp
+++ b/Engine/Plugins/Mutable/Source/CustomizableObjectEditor/Private/MuCOE/CustomizableObjectEditorModule.cpp
@@ -56,6 +56,7 @@
 #include "MuCOE/Nodes/CustomizableObjectNodeModifierClipWithUVMask.h"
 #include "MuCOE/Nodes/CustomizableObjectNodeModifierClipDeform.h"
 #include "MuCOE/Nodes/CustomizableObjectNodeModifierTransformInMesh.h"
+#include "MuCOE/Nodes/CustomizableObjectNodeMeshParameter.h"
 #include "MuCOE/Nodes/CustomizableObjectNodeMeshMorph.h"
 #include "MuCOE/Nodes/CustomizableObjectNodeMeshMorphDetails.h"
 #include "MuCOE/Nodes/CustomizableObjectNodeMeshReshapeCommon.h"
@@ -907,6 +908,11 @@ TMap<FName, FGuid> FCustomizableObjectEditorModule::GetParticipatingObjects(cons
 			AddObjectSoft(TransformNode->ReferenceTexture);
 		}
 
+		else if (UCustomizableObjectNodeMeshParameter* MeshParamNode = Cast<UCustomizableObjectNodeMeshParameter>(&Node))
+		{
+			AddObjectSoft(MeshParamNode->DefaultValue.ToSoftObjectPath());
+		}
+
 		else if (UCustomizableObjectNodeMaterial* MatNode = Cast<UCustomizableObjectNodeMaterial>(&Node))
 		{
 			AddObjectSoft(MatNode->GetMaterial());
diff --git a/Engine/Plugins/Mutable/Source/CustomizableObjectEditor/Private/MuCOE/Nodes/CustomizableObjectNodeMeshParameter.h b/Engine/Plugins/Mutable/Source/CustomizableObjectEditor/Private/MuCOE/Nodes/CustomizableObjectNodeMeshParameter.h
--- a/Engine/Plugins/Mutable/Source/CustomizableObjectEditor/Private/MuCOE/Nodes/CustomizableObjectNodeMeshParameter.h
+++ b/Engine/Plugins/Mutable/Source/CustomizableObjectEditor/Private/MuCOE/Nodes/CustomizableObjectNodeMeshParameter.h
@@ -21,7 +21,7 @@ public:
 
 	/** Default value of the parameter. */
 	UPROPERTY(EditAnywhere, Category=CustomizableObject)
-	TObjectPtr<USkeletalMesh> DefaultValue;
+	TSoftObjectPtr<USkeletalMesh> DefaultValue;
 
 	UPROPERTY(EditAnywhere, Category = CustomizableObject)
 	int32 DefaultValueSection = 0;
diff --git a/Engine/Plugins/Mutable/Source/CustomizableObjectEditor/Private/MuCOE/SMutableCodeViewer.cpp b/Engine/Plugins/Mutable/Source/CustomizableObjectEditor/Private/MuCOE/SMutableCodeViewer.cpp
--- a/Engine/Plugins/Mutable/Source/CustomizableObjectEditor/Private/MuCOE/SMutableCodeViewer.cpp
+++ b/Engine/Plugins/Mutable/Source/CustomizableObjectEditor/Private/MuCOE/SMutableCodeViewer.cpp
@@ -5,7 +5,7 @@
 #include "Framework/MultiBox/MultiBoxBuilder.h"
 #include "Framework/Views/TableViewMetadata.h"
 #include "Misc/Paths.h"
-#include "MuCO/CustomizableObject.h"
+#include "MuCO/UnrealConversionUtils.h"
 #include "MuCO/UnrealToMutableTextureConversionUtils.h"
 #include "MuCOE/SMutableBoolViewer.h"
 #include "MuCOE/SMutableColorViewer.h"
@@ -22,9 +22,13 @@
 #include "MuCOE/SMutableSkeletonViewer.h"
 #include "MuCOE/SMutableStringViewer.h"
 #include "MuCOE/UnrealEditorPortabilityHelpers.h"
-#include "MuR/SystemPrivate.h"
+#include "MuCOE/Widgets/MutableExpanderArrow.h"
+#include "MuCO/CustomizableObject.h"
+#include "MuCO/CustomizableObjectPrivate.h"
+#include "MuCO/UnrealToMutableTextureConversionUtils.h"
 #include "MuT/ErrorLog.h"
 #include "MuT/TypeInfo.h"
+#include "MuR/SystemPrivate.h"
 #include "Widgets/SNullWidget.h"
 #include "Widgets/Colors/SColorBlock.h"
 #include "Widgets/Input/SButton.h"
@@ -32,10 +36,8 @@
 #include "Widgets/Input/SNumericEntryBox.h"
 #include "Widgets/Layout/SScrollBox.h"
 #include "Widgets/Views/STreeView.h"
-#include "MuCOE/Widgets/MutableExpanderArrow.h"
-
-
 #include "Widgets/Input/SSearchBox.h"
+#include "Engine/SkeletalMesh.h"
 #include "Internationalization/Regex.h"
 
 class FExtender;
@@ -2783,59 +2785,27 @@ namespace
 
 			ResultCallback(ResultImage);
 
-			auto TrivialReturn = []() -> TTuple<UE::Tasks::FTask, TFunction<void()>>
-			{
-				return MakeTuple(UE::Tasks::MakeCompletedTask<void>(), []() -> void {});
-			};
-
-			return Invoke(TrivialReturn);
+			return MakeTuple(UE::Tasks::MakeCompletedTask<void>(), []() -> void {});
 		}
 
 
-		virtual TTuple<UE::Tasks::FTask, TFunction<void()>> GetMeshAsync(FName Id, TFunction<void(TSharedPtr<mu::FMesh>)>& ResultCallback) override
+		virtual TTuple<UE::Tasks::FTask, TFunction<void()>> GetMeshAsync(FName Id, int32 LODIndex, int32 SectionIndex, TFunction<void(TSharedPtr<mu::FMesh>)>& ResultCallback) override
 		{
-			// TODO: This is a stub function where Mesh FParameters would be supported.
-			ensure(false);
-
 			// Thread: worker
 			MUTABLE_CPUPROFILER_SCOPE(FUnrealMutableImageProvider::GetMeshAsync);
 
-			auto TrivialReturn = []() -> TTuple<UE::Tasks::FTask, TFunction<void()>>
-				{
-					return MakeTuple(UE::Tasks::MakeCompletedTask<void>(), []() -> void {});
-				};
-
-			// No USkeletalMesh was provided, cannot do anything, just provide a dummy mesh
-			UE_LOG(LogMutable, Warning, TEXT("No USkeletalMesh was provided for an application-specific mesh parameter."));
+			TSharedPtr<mu::FMesh> Result = MakeShared<mu::FMesh>();
 
-			TSharedPtr<mu::FMesh> Dummy = MakeShared<mu::FMesh>();
-			ResultCallback(Dummy);
-			return Invoke(TrivialReturn);
-		}
+			// TODO: Do it async
+			USkeletalMesh* SkeletalMesh = Cast<USkeletalMesh>(FSoftObjectPath(Id.ToString()).TryLoad());
+			check(SkeletalMesh);
 
+			UnrealConversionUtils::ConvertSkeletalMeshFromRuntimeData(*SkeletalMesh, LODIndex, SectionIndex, nullptr, *Result.Get());
 
-		virtual TTuple<UE::Tasks::FTask, TFunction<void()>> GetReferencedMeshAsync(const void* ModelPtr, int32 Id, TFunction<void(TSharedPtr<mu::FMesh>)>& ResultCallback) override
-		{
-			// TODO: This is a stub function where referenced Meshes would be supported.
-
-			// Thread: worker
-			MUTABLE_CPUPROFILER_SCOPE(FUnrealMutableImageProvider::GetReferencedMeshAsync);
-
-			auto TrivialReturn = []() -> TTuple<UE::Tasks::FTask, TFunction<void()>>
-				{
-					return MakeTuple(UE::Tasks::MakeCompletedTask<void>(), []() -> void {});
-				};
-
-			TSharedPtr<mu::FMesh> Mesh = MakeShared<mu::FMesh>();
-
-			// Not supported outside editor yet.
-			UE_LOG(LogMutable, Warning, TEXT("Failed to get referenced mesh. Not implemented yet."));
-
-			ResultCallback(Mesh);
-			return Invoke(TrivialReturn);
+			ResultCallback(Result);
 
+			return MakeTuple(UE::Tasks::MakeCompletedTask<void>(), []() -> void {});
 		}
-
 	};
 }
 
@@ -3211,7 +3181,6 @@ void FMutableCodeTreeElement::GenerateLabelText()
 	case mu::EOpType::CO_PARAMETER:
 	case mu::EOpType::PR_PARAMETER:
 	case mu::EOpType::IM_PARAMETER:
-	case mu::EOpType::ME_PARAMETER:
 	case mu::EOpType::ST_PARAMETER:
 	{
 		mu::OP::ParameterArgs Args = Program.GetOpArgs<mu::OP::ParameterArgs>(MutableOperation);
@@ -3220,6 +3189,14 @@ void FMutableCodeTreeElement::GenerateLabelText()
 		break;
 	}
 
+	case mu::EOpType::ME_PARAMETER:
+	{
+		mu::OP::MeshParameterArgs Args = Program.GetOpArgs<mu::OP::MeshParameterArgs>(MutableOperation);
+		OpName += FString::Printf( TEXT(" LOD %d Section %d of "), Args.LOD, Args.Section );
+		OpName += Program.Parameters[int32(Args.variable)].Name;
+		break;
+	}
+
 	case mu::EOpType::IM_SWIZZLE:
 	{
 		mu::OP::ImageSwizzleArgs Args = Program.GetOpArgs<mu::OP::ImageSwizzleArgs>(MutableOperation);
diff --git a/Engine/Plugins/Mutable/Source/MutableRuntime/Internal/MuR/OpMeshMerge.h b/Engine/Plugins/Mutable/Source/MutableRuntime/Internal/MuR/OpMeshMerge.h
--- a/Engine/Plugins/Mutable/Source/MutableRuntime/Internal/MuR/OpMeshMerge.h
+++ b/Engine/Plugins/Mutable/Source/MutableRuntime/Internal/MuR/OpMeshMerge.h
@@ -60,19 +60,16 @@ namespace mu
 			if (FirstCount && SecondCount)
 			{
 				check(!FirstIndexBuffer.Channels.IsEmpty());
-				check(FirstIndexBuffer.Channels == SecondIndexBuffer.Channels);
-				check(FirstIndexBuffer.ElementSize == SecondIndexBuffer.ElementSize);
 
 				// We need to know the total number of vertices in case we need to adjust the index buffer format.
-				const uint64 totalVertexCount = pFirst->GetVertexBuffers().GetElementCount() + pSecond->GetVertexBuffers().GetElementCount();
-				const uint64 maxValueBits = GetMeshFormatData(pFirst->GetIndexBuffers().Buffers[0].Channels[0].Format).MaxValueBits;
-				const uint64 maxSupportedVertices = uint64(1) << maxValueBits;
+				const uint64 TotalVertexCount = pFirst->GetVertexBuffers().GetElementCount() + pSecond->GetVertexBuffers().GetElementCount();
+				const uint64 MaxValueBits = GetMeshFormatData(pFirst->GetIndexBuffers().Buffers[0].Channels[0].Format).MaxValueBits;
+				const uint64 MaxSupportedVertices = uint64(1) << MaxValueBits;
 				
-				if (totalVertexCount > maxSupportedVertices)
+				if (TotalVertexCount > MaxSupportedVertices)
 				{
-					IndexBufferFormat = totalVertexCount > MAX_uint16 ? EMeshBufferFormat::UInt32 : EMeshBufferFormat::UInt16;
+					IndexBufferFormat = TotalVertexCount > MAX_uint16 ? EMeshBufferFormat::UInt32 : EMeshBufferFormat::UInt16;
 				}
-
 			}
 			
 			if (IndexBufferFormat != EMeshBufferFormat::None)
diff --git a/Engine/Plugins/Mutable/Source/MutableRuntime/Internal/MuR/Operations.h b/Engine/Plugins/Mutable/Source/MutableRuntime/Internal/MuR/Operations.h
--- a/Engine/Plugins/Mutable/Source/MutableRuntime/Internal/MuR/Operations.h
+++ b/Engine/Plugins/Mutable/Source/MutableRuntime/Internal/MuR/Operations.h
@@ -416,10 +416,17 @@ namespace mu
 			int32 physicsBody;
         };
 
-        struct ParameterArgs
-        {
-            ADDRESS variable;
-        };
+		struct ParameterArgs
+		{
+			ADDRESS variable;
+		};
+
+		struct MeshParameterArgs
+		{
+			ADDRESS variable;
+			uint8 LOD = 0;
+			uint8 Section = 0;
+		};
 
         struct ConditionalArgs
         {
diff --git a/Engine/Plugins/Mutable/Source/MutableRuntime/Internal/MuR/ParametersPrivate.h b/Engine/Plugins/Mutable/Source/MutableRuntime/Internal/MuR/ParametersPrivate.h
--- a/Engine/Plugins/Mutable/Source/MutableRuntime/Internal/MuR/ParametersPrivate.h
+++ b/Engine/Plugins/Mutable/Source/MutableRuntime/Internal/MuR/ParametersPrivate.h
@@ -250,12 +250,12 @@ namespace Private
 	using FParamFloatType = float;
 	using FParamColorType = FVector4f;
 	using FParamProjectorType = Private::TIndirectObject<FProjector>;
-	using FParamImageType = FName;
+	using FParamAssetType = FName;
 	using FParamStringType = Private::TIndirectObject<FString>;
 	using FParamMatrixType = Private::TIndirectObject<FMatrix44f>;
 	
 	using FParameterValue = TVariant<
-            FParamBoolType, FParamIntType, FParamFloatType, FParamColorType, FParamProjectorType, FParamImageType, FParamStringType, FParamMatrixType>;
+            FParamBoolType, FParamIntType, FParamFloatType, FParamColorType, FParamProjectorType, FParamAssetType, FParamStringType, FParamMatrixType>;
 
     // static_assert to track ParameterValue size changes. It is ok to change if needed.
     static_assert(sizeof(FParameterValue) == 8*4, "ParameterValue size has changed.");
diff --git a/Engine/Plugins/Mutable/Source/MutableRuntime/Private/MuR/CodeRunner.cpp b/Engine/Plugins/Mutable/Source/MutableRuntime/Private/MuR/CodeRunner.cpp
--- a/Engine/Plugins/Mutable/Source/MutableRuntime/Private/MuR/CodeRunner.cpp
+++ b/Engine/Plugins/Mutable/Source/MutableRuntime/Private/MuR/CodeRunner.cpp
@@ -161,7 +161,6 @@ namespace mu
 	}
 
 
-    //---------------------------------------------------------------------------------------------
 	FProgramCache& CodeRunner::GetMemory()
     {
 		return *System->WorkingMemoryManager.CurrentInstanceCache;
@@ -197,7 +196,7 @@ namespace mu
 	}
 
 
-	TTuple<UE::Tasks::FTask, TFunction<void()>> CodeRunner::LoadExternalMeshAsync(FExternalResourceId Id, TFunction<void(TSharedPtr<FMesh>)>& ResultCallback)
+	TTuple<UE::Tasks::FTask, TFunction<void()>> CodeRunner::LoadExternalMeshAsync(FExternalResourceId Id, int32 LODIndex, int32 SectionIndex, TFunction<void(TSharedPtr<FMesh>)>& ResultCallback)
 	{
 		MUTABLE_CPUPROFILER_SCOPE(LoadExternalImageAsync);
 
@@ -208,12 +207,12 @@ namespace mu
 			if (Id.ReferenceResourceId < 0)
 			{
 				// It's a parameter mesh
-				return System->ExternalResourceProvider->GetMeshAsync(Id.ParameterId, ResultCallback);
+				return System->ExternalResourceProvider->GetMeshAsync(Id.ParameterId, LODIndex, SectionIndex, ResultCallback);
 			}
 			else
 			{
 				// It's a mesh reference
-				return System->ExternalResourceProvider->GetReferencedMeshAsync(Model.Get(), Id.ReferenceResourceId, ResultCallback);
+				check(false);
 			}
 		}
 		else
@@ -226,7 +225,6 @@ namespace mu
 	}
 
 	
-    //---------------------------------------------------------------------------------------------
 	FExtendedImageDesc CodeRunner::GetExternalImageDesc(FName Id)
 	{
 		MUTABLE_CPUPROFILER_SCOPE(GetExternalImageDesc);
diff --git a/Engine/Plugins/Mutable/Source/MutableRuntime/Private/MuR/CodeRunner.h b/Engine/Plugins/Mutable/Source/MutableRuntime/Private/MuR/CodeRunner.h
--- a/Engine/Plugins/Mutable/Source/MutableRuntime/Private/MuR/CodeRunner.h
+++ b/Engine/Plugins/Mutable/Source/MutableRuntime/Private/MuR/CodeRunner.h
@@ -149,7 +149,7 @@ namespace  mu
  	    mu::FExtendedImageDesc GetExternalImageDesc(FName Id);
 
 		/** Load an external mesh asynchronously, returns an event to wait for complition and a cleanup function that must be called once the event has completed. */
-		TTuple<UE::Tasks::FTask, TFunction<void()>> LoadExternalMeshAsync(FExternalResourceId Id, TFunction<void(TSharedPtr<FMesh>)>& ResultCallback);
+		TTuple<UE::Tasks::FTask, TFunction<void()>> LoadExternalMeshAsync(FExternalResourceId Id, int32 LODIndex, int32 SectionIndex, TFunction<void(TSharedPtr<FMesh>)>& ResultCallback);
 
 		/** Settings that may affect the execution of some operations, like image conversion quality. */
 		FSettings Settings;
diff --git a/Engine/Plugins/Mutable/Source/MutableRuntime/Private/MuR/CodeRunnerTasked.cpp b/Engine/Plugins/Mutable/Source/MutableRuntime/Private/MuR/CodeRunnerTasked.cpp
--- a/Engine/Plugins/Mutable/Source/MutableRuntime/Private/MuR/CodeRunnerTasked.cpp
+++ b/Engine/Plugins/Mutable/Source/MutableRuntime/Private/MuR/CodeRunnerTasked.cpp
@@ -2352,7 +2352,7 @@ namespace mu
 	{
 	public:
 
-		FMeshExternalLoadTask(const FScheduledOp&, CodeRunner::FExternalResourceId);
+		FMeshExternalLoadTask(const FScheduledOp&, CodeRunner::FExternalResourceId, int32 InLODIndex, int32 InSectionIndex);
 
 		// FIssuedTask interface
 		virtual bool Prepare(CodeRunner*, bool& bOutFailed) override;
@@ -2361,6 +2361,8 @@ namespace mu
 	private:
 		uint8 MipmapsToSkip;
 		CodeRunner::FExternalResourceId Id;
+		int32 LODIndex = 0;
+		int32 SectionIndex = 0;
 
 		TSharedPtr<FMesh> Result;
 
@@ -2368,9 +2370,11 @@ namespace mu
 	};
 
 
-	FMeshExternalLoadTask::FMeshExternalLoadTask(const FScheduledOp& InOp, CodeRunner::FExternalResourceId InId)
+	FMeshExternalLoadTask::FMeshExternalLoadTask(const FScheduledOp& InOp, CodeRunner::FExternalResourceId InId, int32 InLODIndex, int32 InSectionIndex)
 		: FIssuedTask(InOp)
 		, Id(InId)
+		, LODIndex(InLODIndex)
+		, SectionIndex(InSectionIndex)
 	{
 	}
 
@@ -2392,7 +2396,7 @@ namespace mu
 				Result = InResult;
 			};
 
-		Tie(Event, ExternalCleanUpFunc) = Runner->LoadExternalMeshAsync(Id, ResultCallback);
+		Tie(Event, ExternalCleanUpFunc) = Runner->LoadExternalMeshAsync(Id, LODIndex, SectionIndex, ResultCallback);
 
 		// return false indicating there is no work to do so Event is not overriden by a DoWork task.
 		return false;
@@ -2519,6 +2523,23 @@ namespace mu
 			break;
 		}
 
+		case EOpType::ME_PARAMETER:
+		{
+			OP::MeshParameterArgs Args = Program.GetOpArgs<OP::MeshParameterArgs>(item.At);
+			TSharedPtr<FRangeIndex> Index = BuildCurrentOpRangeIndex(item, Params, Model.Get(), Args.variable);
+
+			const FName Id = Params->GetMeshValue(Args.variable, Index.Get());
+
+			int32 LODIndex = Args.LOD;
+			int32 SectionIndex = Args.Section;
+
+			CodeRunner::FExternalResourceId FullId;
+			FullId.ParameterId = Id;
+			Issued = MakeShared<FMeshExternalLoadTask>(item, FullId, LODIndex, SectionIndex);
+
+			break;
+		}
+
 		case EOpType::ME_REFERENCE:
 		{
 			OP::ResourceReferenceArgs Args = Model->GetPrivate()->Program.GetOpArgs<OP::ResourceReferenceArgs>(item.At);
@@ -2530,7 +2551,12 @@ namespace mu
 
 				FExternalResourceId FullId;
 				FullId.ReferenceResourceId = Args.ID;
-				Issued = MakeShared<FMeshExternalLoadTask>(item, FullId);
+
+				// TODO
+				int32 LODIndex = 0;
+				int32 SectionIndex = 0;
+
+				Issued = MakeShared<FMeshExternalLoadTask>(item, FullId, LODIndex, SectionIndex);
 			}
 
 			break;
diff --git a/Engine/Plugins/Mutable/Source/MutableRuntime/Private/MuR/Mesh.cpp b/Engine/Plugins/Mutable/Source/MutableRuntime/Private/MuR/Mesh.cpp
--- a/Engine/Plugins/Mutable/Source/MutableRuntime/Private/MuR/Mesh.cpp
+++ b/Engine/Plugins/Mutable/Source/MutableRuntime/Private/MuR/Mesh.cpp
@@ -1549,7 +1549,8 @@ namespace
 						}
                         case EMeshBufferFormat::Float16: 
 						{
-							Out += FString::Printf(TEXT("%d"), *(const uint16*)ChanDataPtr); 
+							float Converted = *((const FFloat16*)ChanDataPtr);
+							Out += FString::Printf(TEXT("%.3f"), Converted); 
 							ChanDataPtr += 2; 
 							break;
 						}
diff --git a/Engine/Plugins/Mutable/Source/MutableRuntime/Private/MuR/Model.cpp b/Engine/Plugins/Mutable/Source/MutableRuntime/Private/MuR/Model.cpp
--- a/Engine/Plugins/Mutable/Source/MutableRuntime/Private/MuR/Model.cpp
+++ b/Engine/Plugins/Mutable/Source/MutableRuntime/Private/MuR/Model.cpp
@@ -377,7 +377,7 @@ namespace mu
             return {};
         }
 		
-        return m_pD->Program.Parameters[Index].DefaultValue.Get<FParamImageType>();
+        return m_pD->Program.Parameters[Index].DefaultValue.Get<FParamAssetType>();
     }
 
 
diff --git a/Engine/Plugins/Mutable/Source/MutableRuntime/Private/MuR/Parameters.cpp b/Engine/Plugins/Mutable/Source/MutableRuntime/Private/MuR/Parameters.cpp
--- a/Engine/Plugins/Mutable/Source/MutableRuntime/Private/MuR/Parameters.cpp
+++ b/Engine/Plugins/Mutable/Source/MutableRuntime/Private/MuR/Parameters.cpp
@@ -592,6 +592,74 @@ namespace mu
     }
 
 
+	FName FParameters::GetMeshValue(int32 Index, const FRangeIndex* Pos) const
+	{
+		check(Index >= 0 && Index < (int32)m_pD->Values.Num());
+		check(GetType(Index) == EParameterType::Mesh);
+
+		// Early out in case of invalid parameters
+		if (Index < 0
+			||
+			Index >= (int32)m_pD->Values.Num()
+			||
+			GetType(Index) != EParameterType::Mesh)
+		{
+			return {};
+		}
+
+		// Single value case
+		if (!Pos)
+		{
+			return m_pD->Values[Index].Get<FParamAssetType>();
+		}
+
+		// Multivalue case
+		check(Pos->Parameter == Index);
+
+		if (Index < int32(m_pD->MultiValues.Num()))
+		{
+			const TMap< TArray<int32>, FParameterValue >& m = m_pD->MultiValues[Index];
+			const FParameterValue* it = m.Find(Pos->Values);
+			if (it)
+			{
+				return it->Get<FParamAssetType>();
+			}
+		}
+
+		// Multivalue parameter, but no multivalue set. Return single value.
+		return m_pD->Values[Index].Get<FParamAssetType>();
+
+	}
+
+
+	void FParameters::SetMeshValue(int32 Index, FName id, const FRangeIndex* Pos)
+	{
+		check(Index >= 0 && Index < (int32)m_pD->Values.Num());
+		check(GetType(Index) == EParameterType::Mesh);
+
+		// Single value case
+		if (!Pos)
+		{
+			m_pD->Values[Index].Set<FParamAssetType>(id);
+		}
+
+		// Multivalue case
+		else
+		{
+			check(Pos->Parameter == Index);
+
+			if (Index >= int32(m_pD->MultiValues.Num()))
+			{
+				m_pD->MultiValues.SetNum(Index + 1);
+			}
+
+			TMap< TArray<int32>, FParameterValue >& m = m_pD->MultiValues[Index];
+			FParameterValue& it = m.FindOrAdd(Pos->Values);
+			it.Set<FParamAssetType>(id);
+		}
+	}
+
+
 	FName FParameters::GetImageValue( int32 Index, const FRangeIndex* Pos ) const
     {
         check( Index >= 0 && Index < (int32)m_pD->Values.Num() );
@@ -610,7 +678,7 @@ namespace mu
 		// Single value case
 		if (!Pos)
 		{
-			return m_pD->Values[Index].Get<FParamImageType>();
+			return m_pD->Values[Index].Get<FParamAssetType>();
 		}
 
 		// Multivalue case
@@ -622,12 +690,12 @@ namespace mu
 			const FParameterValue* it = m.Find(Pos->Values);
 			if (it)
 			{
-				return it->Get<FParamImageType>();
+				return it->Get<FParamAssetType>();
 			}
 		}
 
 		// Multivalue parameter, but no multivalue set. Return single value.
-		return m_pD->Values[Index].Get<FParamImageType>();
+		return m_pD->Values[Index].Get<FParamAssetType>();
 
     }
 
@@ -640,7 +708,7 @@ namespace mu
 		// Single value case
 		if (!Pos)
 		{
-			m_pD->Values[Index].Set<FParamImageType>(id);
+			m_pD->Values[Index].Set<FParamAssetType>(id);
 		}
 
 		// Multivalue case
@@ -655,7 +723,7 @@ namespace mu
 
 			TMap< TArray<int32>, FParameterValue >& m = m_pD->MultiValues[Index];
 			FParameterValue& it = m.FindOrAdd(Pos->Values);
-			it.Set<FParamImageType>(id);
+			it.Set<FParamAssetType>(id);
 		}
     }
 
diff --git a/Engine/Plugins/Mutable/Source/MutableRuntime/Private/MuR/System.cpp b/Engine/Plugins/Mutable/Source/MutableRuntime/Private/MuR/System.cpp
--- a/Engine/Plugins/Mutable/Source/MutableRuntime/Private/MuR/System.cpp
+++ b/Engine/Plugins/Mutable/Source/MutableRuntime/Private/MuR/System.cpp
@@ -1881,7 +1881,7 @@ namespace mu
 			case EParameterType::Image:
 			{
 				DataSize = sizeof(FName);
-				Blob.Serialize((void*)&Params->GetPrivate()->Values[ParamIndex].Get<FParamImageType>(), DataSize);
+				Blob.Serialize((void*)&Params->GetPrivate()->Values[ParamIndex].Get<FParamAssetType>(), DataSize);
 
 				// Multi-values
 				if (Params->GetPrivate()->HasMultipleValues(ParamIndex))
@@ -1890,7 +1890,7 @@ namespace mu
 					AddMultiValueKeys(Blob, Multi);
 					for (const TPair< TArray<int32>, FParameterValue >& v : Multi)
 					{
-						Blob.Serialize((void*)&v.Value.Get<FParamImageType>(), DataSize);
+						Blob.Serialize((void*)&v.Value.Get<FParamAssetType>(), DataSize);
 					}
 				}
 				break;
diff --git a/Engine/Plugins/Mutable/Source/MutableRuntime/Public/MuR/Mesh.h b/Engine/Plugins/Mutable/Source/MutableRuntime/Public/MuR/Mesh.h
--- a/Engine/Plugins/Mutable/Source/MutableRuntime/Public/MuR/Mesh.h
+++ b/Engine/Plugins/Mutable/Source/MutableRuntime/Public/MuR/Mesh.h
@@ -563,7 +563,7 @@ namespace mu
 		void ResetBufferIndices();
 
 		/** Debug: get a text representation of the mesh */
-		void Log(FString& Out, int32 VertrexLimit);
+		void Log(FString& Out, int32 VertexLimit);
     };
 
 
diff --git a/Engine/Plugins/Mutable/Source/MutableRuntime/Public/MuR/Parameters.h b/Engine/Plugins/Mutable/Source/MutableRuntime/Public/MuR/Parameters.h
--- a/Engine/Plugins/Mutable/Source/MutableRuntime/Public/MuR/Parameters.h
+++ b/Engine/Plugins/Mutable/Source/MutableRuntime/Public/MuR/Parameters.h
@@ -40,8 +40,11 @@ namespace mu
         /** 3D Projector type, defining a position, scale and orientation.Basically used for projected decals. */
         Projector,
 
-        /** An externally provided image. */
-        Image,
+		/** An externally provided image. */
+		Image,
+
+		/** An externally provided mesh. */
+		Mesh,
 
         /** A text string. */
         String,
@@ -308,12 +311,18 @@ namespace mu
 	    //! \return The externalId specified when setting the image value (\see SetImageValue)
 		FName GetImageValue( int32 Index, const FRangeIndex* Pos = nullptr) const;
 
+		/** */
+		FName GetMeshValue(int32 Index, const FRangeIndex* Pos = nullptr) const;
+
         //! If the parameter is of the image type, set its value.
         //! \param index Index of the parameter from 0 to GetCount()-1
         //! \param externalId Application-specific id used to identify this image during replication.
 		//! \param pos Only for multidimensional parameters: relevant position to set in the ranges
 		void SetImageValue( int32 Index, FName externalId, const FRangeIndex* = nullptr);
 
+		/** */
+		void SetMeshValue(int32 Index, FName externalId, const FRangeIndex* = nullptr);
+
         //! Return the value of a float parameter.
         //! \pre The parameter specified by index is a T_FLOAT.
         //! \param index Index of the parameter from 0 to GetCount()-1
diff --git a/Engine/Plugins/Mutable/Source/MutableRuntime/Public/MuR/System.h b/Engine/Plugins/Mutable/Source/MutableRuntime/Public/MuR/System.h
--- a/Engine/Plugins/Mutable/Source/MutableRuntime/Public/MuR/System.h
+++ b/Engine/Plugins/Mutable/Source/MutableRuntime/Public/MuR/System.h
@@ -82,8 +82,7 @@ namespace mu
 		virtual mu::FExtendedImageDesc GetImageDesc(FName Id) = 0;
 
 		/** Returns the completion event and a cleanup function that must be called once event is completed. */
-		virtual TTuple<UE::Tasks::FTask, TFunction<void()>> GetMeshAsync(FName Id, TFunction<void(TSharedPtr<FMesh>)>& ResultCallback) = 0;
-		virtual TTuple<UE::Tasks::FTask, TFunction<void()>> GetReferencedMeshAsync(const void* ModelPtr, int32 Id, TFunction<void(TSharedPtr<FMesh>)>& ResultCallback) { check(false); return {}; }
+		virtual TTuple<UE::Tasks::FTask, TFunction<void()>> GetMeshAsync(FName Id, int32 InLODIndex, int32 InSectionIndex, TFunction<void(TSharedPtr<FMesh>)>& ResultCallback) = 0;
     };
 
 
diff --git a/Engine/Plugins/Mutable/Source/MutableTools/Private/MuT/ASTOpParameter.cpp b/Engine/Plugins/Mutable/Source/MutableTools/Private/MuT/ASTOpParameter.cpp
--- a/Engine/Plugins/Mutable/Source/MutableTools/Private/MuT/ASTOpParameter.cpp
+++ b/Engine/Plugins/Mutable/Source/MutableTools/Private/MuT/ASTOpParameter.cpp
@@ -38,6 +38,8 @@ namespace mu
 			const ASTOpParameter* Other = static_cast<const ASTOpParameter*>(&OtherUntyped);
 			return Type == Other->Type &&
 				Parameter == Other->Parameter &&
+				LODIndex == Other->LODIndex &&
+				SectionIndex == Other->SectionIndex &&
 				Ranges == Other->Ranges;
 		}
 		return false;
@@ -49,6 +51,8 @@ namespace mu
 		Ptr<ASTOpParameter> n = new ASTOpParameter();
 		n->Type = Type;
 		n->Parameter = Parameter;
+		n->LODIndex = LODIndex;
+		n->SectionIndex = SectionIndex;
 		for (const auto& c : Ranges)
 		{
 			n->Ranges.Emplace(n.get(), MapChild(c.rangeSize.child()), c.rangeName, c.rangeUID);
@@ -67,6 +71,8 @@ namespace mu
 	{
 		uint64 res = std::hash<uint64>()(uint64(Type));
 		hash_combine(res, Parameter.Type);
+		hash_combine(res, LODIndex);
+		hash_combine(res, SectionIndex);
 		hash_combine(res, Parameter.Name.Len());
 		return res;
 	}
@@ -83,8 +89,6 @@ namespace mu
 		// Already linked?
 		if (!linkedAddress)
 		{
-			OP::ParameterArgs Args;
-			FMemory::Memzero(Args);
 
 			LinkedParameterIndex = Program.Parameters.Find(Parameter);
 
@@ -92,21 +96,45 @@ namespace mu
 			// parameters map in the CodeGenerator_FirstPass.
 			check(LinkedParameterIndex != INDEX_NONE);
 
-			Args.variable = (OP::ADDRESS)LinkedParameterIndex;
+			linkedAddress = (OP::ADDRESS)Program.OpAddress.Num();
+			Program.OpAddress.Add((uint32)Program.ByteCode.Num());
+			AppendCode(Program.ByteCode, Type);
 
-			for (const auto& d : Ranges)
+			if (Type == EOpType::ME_PARAMETER)
 			{
-				OP::ADDRESS sizeAt = 0;
-				uint16 rangeId = 0;
-				LinkRange(Program, d, sizeAt, rangeId);
-				Program.Parameters[LinkedParameterIndex].Ranges.Add(rangeId);
+				OP::MeshParameterArgs Args;
+				FMemory::Memzero(Args);
+				Args.variable = (OP::ADDRESS)LinkedParameterIndex;
+				Args.LOD = LODIndex;
+				Args.Section = SectionIndex;
+
+				for (const auto& d : Ranges)
+				{
+					OP::ADDRESS sizeAt = 0;
+					uint16 rangeId = 0;
+					LinkRange(Program, d, sizeAt, rangeId);
+					Program.Parameters[LinkedParameterIndex].Ranges.Add(rangeId);
+				}
+
+				AppendCode(Program.ByteCode, Args);
 			}
+			else
+			{
+				OP::ParameterArgs Args;
+				FMemory::Memzero(Args);
+				Args.variable = (OP::ADDRESS)LinkedParameterIndex;
 
-			linkedAddress = (OP::ADDRESS)Program.OpAddress.Num();
+				for (const auto& d : Ranges)
+				{
+					OP::ADDRESS sizeAt = 0;
+					uint16 rangeId = 0;
+					LinkRange(Program, d, sizeAt, rangeId);
+					Program.Parameters[LinkedParameterIndex].Ranges.Add(rangeId);
+				}
+
+				AppendCode(Program.ByteCode, Args);
+			}
 
-			Program.OpAddress.Add((uint32)Program.ByteCode.Num());
-			AppendCode(Program.ByteCode, Type);
-			AppendCode(Program.ByteCode, Args);
 		}
 	}
 
diff --git a/Engine/Plugins/Mutable/Source/MutableTools/Private/MuT/ASTOpParameter.h b/Engine/Plugins/Mutable/Source/MutableTools/Private/MuT/ASTOpParameter.h
--- a/Engine/Plugins/Mutable/Source/MutableTools/Private/MuT/ASTOpParameter.h
+++ b/Engine/Plugins/Mutable/Source/MutableTools/Private/MuT/ASTOpParameter.h
@@ -26,16 +26,23 @@ namespace mu
 		//!
 		FParameterDesc Parameter;
 
+		/** Used by some parameter types (Mesh) to specify which subdata from the actual parameter value should be used in the operation. */
+		int32 LODIndex = 0;
+		int32 SectionIndex = 0;
+
 		//** Ranges adding dimensions to this parameter. */
 		TArray<FRangeData> Ranges;
 
-		/** Index of the parameter in the program parameter list. Generated ar link time. */
+		/** Index of the parameter in the final program parameter list. 
+		* This is generated at link time, because the parameters may be reordered. 
+		*/
 		int32 LinkedParameterIndex = -1;
 
 	public:
 
 		~ASTOpParameter() override;
 
+		// ASTOp interface.
 		virtual EOpType GetOpType() const override;
 		virtual uint64 Hash() const override;
 		virtual void ForEachChild(const TFunctionRef<void(ASTChild&)>) override;
diff --git a/Engine/Plugins/Mutable/Source/MutableTools/Private/MuT/CodeGenerator_Bool.cpp b/Engine/Plugins/Mutable/Source/MutableTools/Private/MuT/CodeGenerator_Bool.cpp
--- a/Engine/Plugins/Mutable/Source/MutableTools/Private/MuT/CodeGenerator_Bool.cpp
+++ b/Engine/Plugins/Mutable/Source/MutableTools/Private/MuT/CodeGenerator_Bool.cpp
@@ -92,7 +92,7 @@ namespace mu
 		{
 			UE::TUniqueLock Lock(FirstPass.ParameterNodes.Mutex);
 
-			Found = FirstPass.ParameterNodes.Map.Find(Node);
+			Found = FirstPass.ParameterNodes.GenericParametersCache.Find(Node);
 			if (!Found)
 			{
 				FParameterDesc param;
@@ -106,7 +106,7 @@ namespace mu
 				op->Type = EOpType::BO_PARAMETER;
 				op->Parameter = param;
 
-				FirstPass.ParameterNodes.Map.Add(Node, op);
+				FirstPass.ParameterNodes.GenericParametersCache.Add(Node, op);
 			}
 			else
 			{
diff --git a/Engine/Plugins/Mutable/Source/MutableTools/Private/MuT/CodeGenerator_Color.cpp b/Engine/Plugins/Mutable/Source/MutableTools/Private/MuT/CodeGenerator_Color.cpp
--- a/Engine/Plugins/Mutable/Source/MutableTools/Private/MuT/CodeGenerator_Color.cpp
+++ b/Engine/Plugins/Mutable/Source/MutableTools/Private/MuT/CodeGenerator_Color.cpp
@@ -147,7 +147,7 @@ namespace mu
 		{
 			UE::TUniqueLock Lock(FirstPass.ParameterNodes.Mutex);
 
-			Found = FirstPass.ParameterNodes.Map.Find(Typed);
+			Found = FirstPass.ParameterNodes.GenericParametersCache.Find(Typed);
 			if (!Found)
 			{
 				FParameterDesc param;
@@ -168,7 +168,7 @@ namespace mu
 				op->Type = EOpType::CO_PARAMETER;
 				op->Parameter = param;
 
-				FirstPass.ParameterNodes.Map.Add(Typed, op);
+				FirstPass.ParameterNodes.GenericParametersCache.Add(Typed, op);
 			}
 			else
 			{
diff --git a/Engine/Plugins/Mutable/Source/MutableTools/Private/MuT/CodeGenerator_FirstPass.cpp b/Engine/Plugins/Mutable/Source/MutableTools/Private/MuT/CodeGenerator_FirstPass.cpp
--- a/Engine/Plugins/Mutable/Source/MutableTools/Private/MuT/CodeGenerator_FirstPass.cpp
+++ b/Engine/Plugins/Mutable/Source/MutableTools/Private/MuT/CodeGenerator_FirstPass.cpp
@@ -617,7 +617,7 @@ namespace mu
 			Node->Type==NodeObjectGroup::CS_ONE_OR_NONE )
         {
 			UE::TUniqueLock Lock(ParameterNodes.Mutex);
-			Ptr<ASTOpParameter>* Found = ParameterNodes.Map.Find(Node);
+			Ptr<ASTOpParameter>* Found = ParameterNodes.GenericParametersCache.Find(Node);
 
 			if (!Found)
 			{
@@ -639,7 +639,7 @@ namespace mu
 					op->Parameter.PossibleValues.Add(nullValue);
 				}
 
-				ParameterNodes.Map.Add(Node, op);
+				ParameterNodes.GenericParametersCache.Add(Node, op);
 
 				enumOp = op;
 			}
@@ -670,7 +670,7 @@ namespace mu
 						{
 							// Create a new boolean parameter			
 							UE::TUniqueLock Lock(ParameterNodes.Mutex);
-							Ptr<ASTOpParameter>* Found = ParameterNodes.Map.Find(ChildNode);
+							Ptr<ASTOpParameter>* Found = ParameterNodes.GenericParametersCache.Find(ChildNode);
 
 							if (!Found)
 							{
@@ -683,7 +683,7 @@ namespace mu
 								op->Parameter.Type = EParameterType::Bool;
 								op->Parameter.DefaultValue.Set<FParamBoolType>(false);
 
-								ParameterNodes.Map.Add(ChildNode, op);
+								ParameterNodes.GenericParametersCache.Add(ChildNode, op);
 
 								paramOp = op;
 							}
diff --git a/Engine/Plugins/Mutable/Source/MutableTools/Private/MuT/CodeGenerator_FirstPass.h b/Engine/Plugins/Mutable/Source/MutableTools/Private/MuT/CodeGenerator_FirstPass.h
--- a/Engine/Plugins/Mutable/Source/MutableTools/Private/MuT/CodeGenerator_FirstPass.h
+++ b/Engine/Plugins/Mutable/Source/MutableTools/Private/MuT/CodeGenerator_FirstPass.h
@@ -23,6 +23,7 @@
 #include "MuT/NodeSurfaceNew.h"
 #include "MuT/NodeSurfaceVariation.h"
 #include "MuT/NodeSurfaceSwitch.h"
+#include "MuT/NodeMeshParameter.h"
 #include "MuT/ASTOpParameter.h"
 #include "MuT/ErrorLog.h"
 #include "Async/Mutex.h"
@@ -247,8 +248,14 @@ namespace mu
 		/** FParameters added for every node. */
 		struct FSafeParameterNodes
 		{
+			/** Lock this mutex before accessing any data in the struct. */
 			UE::FMutex Mutex;
-			TMap< Ptr<const Node>, Ptr<ASTOpParameter> > Map;
+
+			/** Cache for generic parameters. */
+			TMap< Ptr<const Node>, Ptr<ASTOpParameter> > GenericParametersCache;
+
+			/** For mesh parameters we generate a different ASTOp for each LOD. */
+			TMap< Ptr<const NodeMeshParameter>, TArray<Ptr<ASTOpParameter>> > MeshParametersCache;
 		};
 		FSafeParameterNodes ParameterNodes;
 
diff --git a/Engine/Plugins/Mutable/Source/MutableTools/Private/MuT/CodeGenerator_Image.cpp b/Engine/Plugins/Mutable/Source/MutableTools/Private/MuT/CodeGenerator_Image.cpp
--- a/Engine/Plugins/Mutable/Source/MutableTools/Private/MuT/CodeGenerator_Image.cpp
+++ b/Engine/Plugins/Mutable/Source/MutableTools/Private/MuT/CodeGenerator_Image.cpp
@@ -257,7 +257,7 @@ namespace mu
 		Ptr<ASTOpParameter>* Found = nullptr;
 		{
 			UE::TUniqueLock Lock(FirstPass.ParameterNodes.Mutex);
-			Found = FirstPass.ParameterNodes.Map.Find(Node);
+			Found = FirstPass.ParameterNodes.GenericParametersCache.Find(Node);
 			if (!Found)
 			{
 				op = new ASTOpParameter();
@@ -267,9 +267,9 @@ namespace mu
 				bool bParseOk = FGuid::Parse(Node->UID, op->Parameter.UID);
 				check(bParseOk);
 				op->Parameter.Type = EParameterType::Image;
-				op->Parameter.DefaultValue.Set<FParamImageType>(Node->DefaultValue);
+				op->Parameter.DefaultValue.Set<FParamAssetType>(Node->DefaultValue);
 
-				FirstPass.ParameterNodes.Map.Add(Node, op);
+				FirstPass.ParameterNodes.GenericParametersCache.Add(Node, op);
 			}
 			else
 			{
diff --git a/Engine/Plugins/Mutable/Source/MutableTools/Private/MuT/CodeGenerator_Matrix.cpp b/Engine/Plugins/Mutable/Source/MutableTools/Private/MuT/CodeGenerator_Matrix.cpp
--- a/Engine/Plugins/Mutable/Source/MutableTools/Private/MuT/CodeGenerator_Matrix.cpp
+++ b/Engine/Plugins/Mutable/Source/MutableTools/Private/MuT/CodeGenerator_Matrix.cpp
@@ -63,7 +63,7 @@ namespace mu
 		Ptr<ASTOpParameter>* Found = nullptr;
 		{
 			UE::TUniqueLock Lock(FirstPass.ParameterNodes.Mutex);
-			Found = FirstPass.ParameterNodes.Map.Find(Typed.get());
+			Found = FirstPass.ParameterNodes.GenericParametersCache.Find(Typed.get());
 			if (!Found)
 			{
 				FParameterDesc param;
@@ -77,7 +77,7 @@ namespace mu
 				op->Type = EOpType::MA_PARAMETER;
 				op->Parameter = param;
 
-				FirstPass.ParameterNodes.Map.Add(Typed.get(), op);
+				FirstPass.ParameterNodes.GenericParametersCache.Add(Typed.get(), op);
 			}
 			else
 			{
diff --git a/Engine/Plugins/Mutable/Source/MutableTools/Private/MuT/CodeGenerator_Mesh.cpp b/Engine/Plugins/Mutable/Source/MutableTools/Private/MuT/CodeGenerator_Mesh.cpp
--- a/Engine/Plugins/Mutable/Source/MutableTools/Private/MuT/CodeGenerator_Mesh.cpp
+++ b/Engine/Plugins/Mutable/Source/MutableTools/Private/MuT/CodeGenerator_Mesh.cpp
@@ -1810,49 +1810,78 @@ namespace mu
 
 	FMeshTask CodeGenerator::GenerateMesh_Parameter(const FMeshGenerationStaticOptions& StaticOptions, FMeshOptionsTask InOptions, const NodeMeshParameter* Node)
 	{
-		FMeshGenerationResult Result;
-		return UE::Tasks::MakeCompletedTask<FMeshGenerationResult>(Result);
-
-		// Add concurrency protection before enabling.
-		//Ptr<ASTOpParameter>* Found = FirstPass.ParameterNodes.Find(InNode);
-		//if (Found)
-		//{
-		//	check((*Found)->Type==EOpType::ME_PARAMETER);
-		//	FMeshGenerationResult Result;
-		//	Result.BaseMeshOp = *Found;
-		//	Result.MeshOp = *Found;
-		//	// Rest of the parameters?
-		//	//Result.LayoutOps ...
-		//	// ... maybe a special cache for these parameters?
-		//	return UE::Tasks::MakeCompletedTask<FMeshGenerationResult>(Result);
-		//}
-
-
-		//// LOD ? 
-		////int32 LODIndex = StaticOptions.LODIndex;
-
-		//// To task?
-		////FMeshGenerationDynamicOptions = InOptions.GetResult();
-
-		//Ptr<ASTOpParameter> Op = new ASTOpParameter();
-		//Op->Type = EOpType::ME_PARAMETER;
-
-		//Op->Parameter.Name = InNode->Name;
-		//bool bParseOk = FGuid::Parse(InNode->UID, Op->Parameter.UID);
-		////check(bParseOk);
-		//Op->Parameter.Type = EParameterType::Image;
-		//Op->Parameter.DefaultValue.Set<FParamImageType>(InNode->DefaultValue);
-
-		//// Generate the code for the ranges
-		//for (int32 a = 0; a < InNode->Ranges.Num(); ++a)
-		//{
-		//	FRangeGenerationResult rangeResult;
-		//	GenerateRange(rangeResult, StaticOptions, InNode->Ranges[a]);
-		//	Op->Ranges.Emplace(Op.get(), rangeResult.sizeOp, rangeResult.rangeName, rangeResult.rangeUID);
-		//}
-
-		//FirstPass.ParameterNodes.Add(InNode, Op);
+		// Local pipe because we call GenerateRange in there.
+		FMeshTask MeshParamTask = this->LocalPipe.Launch(TEXT("MutableMeshParameter"),
+			[InOptions, Node, StaticOptions, this]() mutable
+			{
+				Ptr<ASTOpParameter> Op;
+				{
+					UE::TUniqueLock Lock(FirstPass.ParameterNodes.Mutex);
+
+					Ptr<ASTOpParameter> Found;					
+					TArray<Ptr<ASTOpParameter>>& ArrayFound = FirstPass.ParameterNodes.MeshParametersCache.FindOrAdd(Node);
+					if (ArrayFound.IsValidIndex(StaticOptions.LODIndex))
+					{
+						Found = ArrayFound[StaticOptions.LODIndex];
+					}
+
+					if (Found)
+					{
+						check(Found->Type == EOpType::ME_PARAMETER);
+
+						FMeshGenerationResult Result;
+						Result.BaseMeshOp = Found;
+						Result.MeshOp = Found;
+						// Rest of the parameters?
+						//Result.LayoutOps ...
+						// ... maybe a special cache for these parameters?
+						return Result;
+					}
+
+					FParameterDesc param;
+					param.Name = Node->Name;
+					bool bParseOk = FGuid::Parse(Node->UID, param.UID);
+					check(bParseOk);
+					param.Type = EParameterType::Mesh;
+					param.DefaultValue.Set<FParamAssetType>(Node->DefaultValue);
+
+					Op = new ASTOpParameter();
+					Op->Type = EOpType::ME_PARAMETER;
+					Op->Parameter = param;
+					Op->LODIndex = StaticOptions.LODIndex;
+					Op->SectionIndex = Node->SectionIndex;
+
+					if (!ArrayFound.IsValidIndex(StaticOptions.LODIndex))
+					{
+						ArrayFound.SetNumZeroed(StaticOptions.LODIndex+1);
+					}
+
+					ArrayFound[StaticOptions.LODIndex] = Op;
+				}
+
+				int32 LODIndex = StaticOptions.LODIndex;
+
+				FMeshGenerationDynamicOptions DynamicOptions = InOptions.GetResult();
+
+				// Generate the code for the ranges
+				for (int32 RangeIndex = 0; RangeIndex < Node->Ranges.Num(); ++RangeIndex)
+				{
+					FRangeGenerationResult RangeResult;
+					GenerateRange(RangeResult, StaticOptions, Node->Ranges[RangeIndex]);
+					Op->Ranges.Emplace(Op.get(), RangeResult.sizeOp, RangeResult.rangeName, RangeResult.rangeUID);
+				}
+
+				FMeshGenerationResult Result;
+				Result.MeshOp = Op;
+				// Rest of the parameters?
+				//Result.LayoutOps ...
+
+				return Result;
+			},
+			InOptions
+		);
 
+		return MeshParamTask;
 	}
 
 }
diff --git a/Engine/Plugins/Mutable/Source/MutableTools/Private/MuT/CodeGenerator_Projector.cpp b/Engine/Plugins/Mutable/Source/MutableTools/Private/MuT/CodeGenerator_Projector.cpp
--- a/Engine/Plugins/Mutable/Source/MutableTools/Private/MuT/CodeGenerator_Projector.cpp
+++ b/Engine/Plugins/Mutable/Source/MutableTools/Private/MuT/CodeGenerator_Projector.cpp
@@ -94,7 +94,7 @@ namespace mu
 		Ptr<ASTOpParameter>* Found = nullptr;
 		{
 			UE::TUniqueLock Lock(FirstPass.ParameterNodes.Mutex);
-			Found = FirstPass.ParameterNodes.Map.Find(Node);
+			Found = FirstPass.ParameterNodes.GenericParametersCache.Find(Node);
 			if (!Found)
 			{
 				FParameterDesc param;
@@ -117,7 +117,7 @@ namespace mu
 				op->Type = EOpType::PR_PARAMETER;
 				op->Parameter = param;
 
-				FirstPass.ParameterNodes.Map.Add(Node, op);
+				FirstPass.ParameterNodes.GenericParametersCache.Add(Node, op);
 			}
 			else
 			{
diff --git a/Engine/Plugins/Mutable/Source/MutableTools/Private/MuT/CodeGenerator_Scalar.cpp b/Engine/Plugins/Mutable/Source/MutableTools/Private/MuT/CodeGenerator_Scalar.cpp
--- a/Engine/Plugins/Mutable/Source/MutableTools/Private/MuT/CodeGenerator_Scalar.cpp
+++ b/Engine/Plugins/Mutable/Source/MutableTools/Private/MuT/CodeGenerator_Scalar.cpp
@@ -132,7 +132,7 @@ namespace mu
 		Ptr<ASTOpParameter>* Found = nullptr;		
 		{
 			UE::TUniqueLock Lock(FirstPass.ParameterNodes.Mutex);
-			Found = FirstPass.ParameterNodes.Map.Find(Node);
+			Found = FirstPass.ParameterNodes.GenericParametersCache.Find(Node);
 			if (!Found)
 			{
 				FParameterDesc param;
@@ -146,7 +146,7 @@ namespace mu
 				op->Type = EOpType::SC_PARAMETER;
 				op->Parameter = param;
 
-				FirstPass.ParameterNodes.Map.Add(Node, op);
+				FirstPass.ParameterNodes.GenericParametersCache.Add(Node, op);
 			}
 			else
 			{
@@ -176,7 +176,7 @@ namespace mu
 		Ptr<ASTOpParameter>* Found = nullptr;
 		{
 			UE::TUniqueLock Lock(FirstPass.ParameterNodes.Mutex);
-			Found = FirstPass.ParameterNodes.Map.Find(Node);
+			Found = FirstPass.ParameterNodes.GenericParametersCache.Find(Node);
 			if (!Found)
 			{
 				FParameterDesc param;
@@ -200,7 +200,7 @@ namespace mu
 				op->Type = EOpType::NU_PARAMETER;
 				op->Parameter = param;
 
-				FirstPass.ParameterNodes.Map.Add(Node, op);
+				FirstPass.ParameterNodes.GenericParametersCache.Add(Node, op);
 			}
 			else
 			{
diff --git a/Engine/Plugins/Mutable/Source/MutableTools/Private/MuT/CodeGenerator_String.cpp b/Engine/Plugins/Mutable/Source/MutableTools/Private/MuT/CodeGenerator_String.cpp
--- a/Engine/Plugins/Mutable/Source/MutableTools/Private/MuT/CodeGenerator_String.cpp
+++ b/Engine/Plugins/Mutable/Source/MutableTools/Private/MuT/CodeGenerator_String.cpp
@@ -81,7 +81,7 @@ namespace mu
 		Ptr<ASTOpParameter>* Found = nullptr;
 		{
 			UE::TUniqueLock Lock(FirstPass.ParameterNodes.Mutex);
-			Found = FirstPass.ParameterNodes.Map.Find(Node);
+			Found = FirstPass.ParameterNodes.GenericParametersCache.Find(Node);
 			if (!Found)
 			{
 				FParameterDesc param;
@@ -90,13 +90,12 @@ namespace mu
 				check(bParseOk);
 				param.Type = EParameterType::Float;
 				param.DefaultValue.Set<FParamStringType>(Node->DefaultValue);
-				param.DefaultValue.Set<FParamStringType>(Node->DefaultValue);
 
 				op = new ASTOpParameter();
 				op->Type = EOpType::ST_PARAMETER;
 				op->Parameter = param;
 
-				FirstPass.ParameterNodes.Map.Add(Node, op);
+				FirstPass.ParameterNodes.GenericParametersCache.Add(Node, op);
 			}
 			else
 			{
diff --git a/Engine/Plugins/Mutable/Source/MutableTools/Private/MuT/Compiler.cpp b/Engine/Plugins/Mutable/Source/MutableTools/Private/MuT/Compiler.cpp
--- a/Engine/Plugins/Mutable/Source/MutableTools/Private/MuT/Compiler.cpp
+++ b/Engine/Plugins/Mutable/Source/MutableTools/Private/MuT/Compiler.cpp
@@ -453,39 +453,47 @@ namespace mu
 			{
 				UE::TUniqueLock Lock(Generator.FirstPass.ParameterNodes.Mutex);
 
+				// Gather the parameter list from the non-optimized data, so that we have them all even if they are optimized out
+				{
+					const int32 ParameterCount = 
+						Generator.FirstPass.ParameterNodes.GenericParametersCache.Num()
+						+
+						Generator.FirstPass.ParameterNodes.MeshParametersCache.Num();
+
+					Parameters.Reserve(ParameterCount);
+					for (const TPair<Ptr<const Node>, Ptr<ASTOpParameter>>& Entry : Generator.FirstPass.ParameterNodes.GenericParametersCache)
+					{
+						Parameters.Add(Entry.Value->Parameter);
+					}
+					for (const TPair<Ptr<const NodeMeshParameter>, TArray<Ptr<ASTOpParameter>>>& Entry : Generator.FirstPass.ParameterNodes.MeshParametersCache)
+					{
+						for (const Ptr<ASTOpParameter>& ArrayEntry : Entry.Value)
+						{
+							if (ArrayEntry)
+							{
+								Parameters.AddUnique(ArrayEntry->Parameter);
+							}
+						}
+					}
+				}
+
 				// Sort the parameters as deterministically as possible.
-				Generator.FirstPass.ParameterNodes.Map.ValueStableSort([](const Ptr<ASTOpParameter>& A, const Ptr<ASTOpParameter>& B)->bool
+				Parameters.StableSort([](const FParameterDesc& A, const FParameterDesc& B)->bool
 					{
-						if (A->Parameter.Name < B->Parameter.Name)
+						if (A.Name < B.Name)
 						{
 							return true;
 						}
-						if (A->Parameter.Name > B->Parameter.Name)
+						if (A.Name > B.Name)
 						{
 							return false;
 						}
-						return A->Parameter.UID < B->Parameter.UID;
+						return A.UID < B.UID;
 					});
 
-				// Set the parameter list from the non-optimized data, so that we have them all even if they are optimized out
-				{
-					const int32 ParameterCount = Generator.FirstPass.ParameterNodes.Map.Num();
-					Parameters.Reserve(ParameterCount);
-					for (const TPair<Ptr<const Node>, Ptr<ASTOpParameter>>& Entry : Generator.FirstPass.ParameterNodes.Map)
-					{
-						Parameters.Add(Entry.Value->Parameter);
-					}
-				}
-        	
         		// Check for repeated parameter names and report them to the user
 				if (Parameters.Num() > 1)
 				{
-					TArray<Ptr<const Node>> ParameterNodes;
-					Generator.FirstPass.ParameterNodes.Map.GenerateKeyArray(ParameterNodes);
-				
-					// Should be impossible as both arrays come from the same data.
-					check(ParameterNodes.Num() == Parameters.Num());
-
 					TArray<FString> RepeatedParameterNames;
 				
 					for (int32 Index = 1; Index < Parameters.Num(); ++Index)
@@ -495,12 +503,39 @@ namespace mu
 					
 						if (CurrentParameterName.Compare(PreviousParameterName) == 0 && !RepeatedParameterNames.Contains(CurrentParameterName))
 						{
-							// Report repeated parameter name
-							const FString WarningText = FString::Printf(TEXT("Repeated parameter found : \"%s\". Please use a different name."), *CurrentParameterName);
-							const void* MessageContext = ParameterNodes[Index]->GetMessageContext();
-							GenErrorLog->Add(WarningText, ELMT_WARNING, MessageContext);
+							// Find the nodes defining these parameters
+							TArray<const void*, TInlineAllocator<4>> MessageContexts;
+							{
+								for (const TPair<Ptr<const Node>, Ptr<ASTOpParameter>>& Entry : Generator.FirstPass.ParameterNodes.GenericParametersCache)
+								{
+									if (Entry.Value->Parameter.Name.Compare(CurrentParameterName))
+									{
+										MessageContexts.AddUnique(Entry.Key->GetMessageContext());
+									}
+								}
+								for (const TPair<Ptr<const NodeMeshParameter>, TArray<Ptr<ASTOpParameter>>>& Entry : Generator.FirstPass.ParameterNodes.MeshParametersCache)
+								{
+									for (const Ptr<ASTOpParameter>& ArrayEntry : Entry.Value)
+									{
+										if (ArrayEntry && ArrayEntry->Parameter.Name.Compare(CurrentParameterName))
+										{
+											MessageContexts.AddUnique(Entry.Key->GetMessageContext());
+										}
+									}
+								}
+							}
 
-							RepeatedParameterNames.Add(CurrentParameterName);
+							// Report repeated parameter name
+							if (MessageContexts.Num() >= 2)
+							{
+								const FString WarningText = FString::Printf(TEXT("Repeated parameter found : \"%s\". Please use a different name."), *CurrentParameterName);
+								GenErrorLog->Add(WarningText, ELMT_WARNING, MessageContexts[0], MessageContexts[1]);
+								RepeatedParameterNames.Add(CurrentParameterName);
+							}
+							else
+							{
+								ensure(false);
+							}
 						}
 					}
 				}
diff --git a/Engine/Plugins/Mutable/Source/MutableTools/Public/MuT/NodeMeshParameter.h b/Engine/Plugins/Mutable/Source/MutableTools/Public/MuT/NodeMeshParameter.h
--- a/Engine/Plugins/Mutable/Source/MutableTools/Public/MuT/NodeMeshParameter.h
+++ b/Engine/Plugins/Mutable/Source/MutableTools/Public/MuT/NodeMeshParameter.h
@@ -19,6 +19,9 @@ namespace mu
 		/** Full asset name for the default value of the parameter. */
 		FName DefaultValue;
 
+		/** TODO */
+		uint8 SectionIndex = 0;
+
 		/** Name of the parameter */
 		FString Name;
 
