[Mutable] Fix hitch when loading UModelResources in the Editor.
+ Fix missing StreamedResources when generating mesh constants.

#rnx

diff --git a/Engine/Plugins/Mutable/Source/CustomizableObject/Internal/MuCO/CustomizableObjectPrivate.h b/Engine/Plugins/Mutable/Source/CustomizableObject/Internal/MuCO/CustomizableObjectPrivate.h
--- a/Engine/Plugins/Mutable/Source/CustomizableObject/Internal/MuCO/CustomizableObjectPrivate.h
+++ b/Engine/Plugins/Mutable/Source/CustomizableObject/Internal/MuCO/CustomizableObjectPrivate.h
@@ -1583,6 +1583,8 @@ enum class ECustomizableObjectVersions : int32
 
 	ExtendedMeshParameterArgumentsWithLODAndSection,
 	
+	AddAssetUserDataEditor,
+	
 	// -----<new versions can be added above this line>--------
 	LastCustomizableObjectVersion
 };
diff --git a/Engine/Plugins/Mutable/Source/CustomizableObject/Internal/MuCO/CustomizableObjectResourceDataTypes.h b/Engine/Plugins/Mutable/Source/CustomizableObject/Internal/MuCO/CustomizableObjectResourceDataTypes.h
--- a/Engine/Plugins/Mutable/Source/CustomizableObject/Internal/MuCO/CustomizableObjectResourceDataTypes.h
+++ b/Engine/Plugins/Mutable/Source/CustomizableObject/Internal/MuCO/CustomizableObjectResourceDataTypes.h
@@ -21,4 +21,11 @@ public:
 	// Valid in Cooked builds
 	UPROPERTY()
 	TObjectPtr<UAssetUserData> AssetUserData;
+
+#if WITH_EDITORONLY_DATA
+	// Valid in Editor
+	UPROPERTY()
+	TSoftObjectPtr<UAssetUserData> AssetUserDataEditor;
+#endif
+
 };
\ No newline at end of file
diff --git a/Engine/Plugins/Mutable/Source/CustomizableObject/Private/MuCO/CustomizableObject.cpp b/Engine/Plugins/Mutable/Source/CustomizableObject/Private/MuCO/CustomizableObject.cpp
--- a/Engine/Plugins/Mutable/Source/CustomizableObject/Private/MuCO/CustomizableObject.cpp
+++ b/Engine/Plugins/Mutable/Source/CustomizableObject/Private/MuCO/CustomizableObject.cpp
@@ -45,6 +45,7 @@
 #include "DerivedDataCache.h"
 #include "DerivedDataCacheKey.h"
 #include "DerivedDataRequestOwner.h"
+#include "MuCO/LoadUtils.h"
 #endif
 
 
@@ -551,14 +552,15 @@ void UModelResources::InitCookData()
 		if (FCustomizableObjectAssetUserData* AUDResource = Container->Data.Data.GetMutablePtr<FCustomizableObjectAssetUserData>())
 		{
 			// Find or duplicate the AUD replacing the outer
-			const FString AssetName = GetNameSafe(AUDResource->AssetUserData);
-			check(AUDResource->AssetUserData != nullptr);
+			UAssetUserData* SourceAssetUserData = MutablePrivate::LoadObject(AUDResource->AssetUserDataEditor);
+			const FString AssetName = GetNameSafe(SourceAssetUserData);
+			check(SourceAssetUserData != nullptr);
 
 			UAssetUserData* AssetUserData = FindObject<UAssetUserData>(Container, *AssetName);
 			if (!AssetUserData)
 			{
 				// AUD may be private objects within meshes. Duplicate changing the outer to avoid including meshes into the builds.
-				AssetUserData = DuplicateObject<UAssetUserData>(AUDResource->AssetUserData, Container, FName(*AssetName));
+				AssetUserData = DuplicateObject<UAssetUserData>(SourceAssetUserData, Container, FName(*AssetName));
 			}
 			
 			AUDResource->AssetUserData = AssetUserData;
diff --git a/Engine/Plugins/Mutable/Source/CustomizableObject/Private/MuCO/CustomizableObjectInstance.cpp b/Engine/Plugins/Mutable/Source/CustomizableObject/Private/MuCO/CustomizableObjectInstance.cpp
--- a/Engine/Plugins/Mutable/Source/CustomizableObject/Private/MuCO/CustomizableObjectInstance.cpp
+++ b/Engine/Plugins/Mutable/Source/CustomizableObject/Private/MuCO/CustomizableObjectInstance.cpp
@@ -5329,6 +5329,12 @@ UE::Tasks::FTask UCustomizableInstancePrivate::LoadAdditionalAssetsAndData(const
 				if (!ModelResources.StreamedResourceDataEditor.IsValidIndex(ResourceIndex))
 				{
 					UE_LOG(LogMutable, Error, TEXT("Invalid streamed resource index. Max Index [%d]. Resource Index [%d]."), ModelResources.StreamedResourceDataEditor.Num(), ResourceIndex);
+					continue;
+				}
+
+				if (const FCustomizableObjectAssetUserData* AUDResource = ModelResources.StreamedResourceDataEditor[ResourceIndex].Data.GetPtr<FCustomizableObjectAssetUserData>())
+				{
+					AssetsToStream.AddUnique(AUDResource->AssetUserDataEditor.ToSoftObjectPath());
 				}
 #else
 				if (!ModelResources.StreamedResourceData.IsValidIndex(ResourceIndex))
@@ -5554,7 +5560,7 @@ void UCustomizableInstancePrivate::AdditionalAssetsAsyncLoaded(const TSharedRef<
 			{
 				if (const FCustomizableObjectAssetUserData* AUDResource = ModelResources.StreamedResourceDataEditor[ResourceIndex].Data.GetPtr<FCustomizableObjectAssetUserData>())
 				{
-					ComponentData.AssetUserDataArray.Add(AUDResource->AssetUserData);
+					ComponentData.AssetUserDataArray.Add(MutablePrivate::LoadObject(AUDResource->AssetUserDataEditor)); // Already loaded
 				}
 			}
 #else
diff --git a/Engine/Plugins/Mutable/Source/CustomizableObjectEditor/Private/MuCOE/GenerateMutableSource/GenerateMutableSource.cpp b/Engine/Plugins/Mutable/Source/CustomizableObjectEditor/Private/MuCOE/GenerateMutableSource/GenerateMutableSource.cpp
--- a/Engine/Plugins/Mutable/Source/CustomizableObjectEditor/Private/MuCOE/GenerateMutableSource/GenerateMutableSource.cpp
+++ b/Engine/Plugins/Mutable/Source/CustomizableObjectEditor/Private/MuCOE/GenerateMutableSource/GenerateMutableSource.cpp
@@ -313,7 +313,7 @@ int32 FMutableGraphGenerationContext::AddStreamedResource(UAssetUserData& AssetU
 	}
 
 	FCustomizableObjectAssetUserData ResourceData;
-	ResourceData.AssetUserData = &AssetUserData;
+	ResourceData.AssetUserDataEditor = &AssetUserData;
 
 	const int32 ResourceIndex = StreamedResourceData.Num();
 
diff --git a/Engine/Plugins/Mutable/Source/MutableTools/Private/MuT/CodeGenerator_Mesh.cpp b/Engine/Plugins/Mutable/Source/MutableTools/Private/MuT/CodeGenerator_Mesh.cpp
--- a/Engine/Plugins/Mutable/Source/MutableTools/Private/MuT/CodeGenerator_Mesh.cpp
+++ b/Engine/Plugins/Mutable/Source/MutableTools/Private/MuT/CodeGenerator_Mesh.cpp
@@ -1227,6 +1227,8 @@ namespace mu
 					Mesh = TaglessMesh;
 				}
 
+				Mesh->StreamedResources.Append(ReferenceMesh->StreamedResources);
+
 				FMeshGenerationDynamicOptions Options = InOptions.GetResult();
 				bool bIsOverridingLayouts = !Options.OverrideLayouts.IsEmpty();
 
