[mutable] Remove missing function declaration.
#rnx

diff --git a/Engine/Plugins/Mutable/Source/CustomizableObject/Internal/MuCO/MutableMeshBufferUtils.h b/Engine/Plugins/Mutable/Source/CustomizableObject/Internal/MuCO/MutableMeshBufferUtils.h
--- a/Engine/Plugins/Mutable/Source/CustomizableObject/Internal/MuCO/MutableMeshBufferUtils.h
+++ b/Engine/Plugins/Mutable/Source/CustomizableObject/Internal/MuCO/MutableMeshBufferUtils.h
@@ -81,8 +81,5 @@ namespace  MutableMeshBufferUtils
 	 */
 	CUSTOMIZABLEOBJECT_API void SetupSkinWeightProfileBuffer(const int32& InCurrentVertexBuffer, const int32& MaxBoneIndexTypeSizeBytes,const int32& MaxBoneWeightTypeSizeBytes, const int32& MaxNumBonesPerVertex,
 		const int32 SemanticsIndex, mu::FMeshBufferSet& OutTargetVertexBuffers);
-	
-	/** */
-	CUSTOMIZABLEOBJECT_API void BuildRemappedBonesArray(const TMap<FName, bool>* BonesToRemove, TObjectPtr<const USkeletalMesh> InSkeletalMesh, int32 InLODIndex, const TArray<FBoneIndexType>& InRequiredBones, TArray<FBoneIndexType>& OutRemappedBones);
 
 }
