[Mutable] Fix crash when generating mesh constants if mesh == ReferenceMesh.

#rnx

diff --git a/Engine/Plugins/Mutable/Source/MutableTools/Private/MuT/CodeGenerator_Mesh.cpp b/Engine/Plugins/Mutable/Source/MutableTools/Private/MuT/CodeGenerator_Mesh.cpp
--- a/Engine/Plugins/Mutable/Source/MutableTools/Private/MuT/CodeGenerator_Mesh.cpp
+++ b/Engine/Plugins/Mutable/Source/MutableTools/Private/MuT/CodeGenerator_Mesh.cpp
@@ -1227,7 +1227,10 @@ namespace mu
 					Mesh = TaglessMesh;
 				}
 
-				Mesh->StreamedResources.Append(ReferenceMesh->StreamedResources);
+				for (uint64 StreamedResource : ReferenceMesh->StreamedResources)
+				{
+					Mesh->StreamedResources.AddUnique(StreamedResource);
+				}
 
 				FMeshGenerationDynamicOptions Options = InOptions.GetResult();
 				bool bIsOverridingLayouts = !Options.OverrideLayouts.IsEmpty();
