[Mutable] Fix physics size calculation in constant debugger inspector.

- Remove inline specifier from physics serialization function declarations. This could cause compilation errors in some cases.
- Other naming changes to better follow UE naming conventions.

#rnx

diff --git a/Engine/Plugins/Mutable/Source/CustomizableObjectEditor/Private/MuCOE/SMutableConstantsWidget.cpp b/Engine/Plugins/Mutable/Source/CustomizableObjectEditor/Private/MuCOE/SMutableConstantsWidget.cpp
--- a/Engine/Plugins/Mutable/Source/CustomizableObjectEditor/Private/MuCOE/SMutableConstantsWidget.cpp
+++ b/Engine/Plugins/Mutable/Source/CustomizableObjectEditor/Private/MuCOE/SMutableConstantsWidget.cpp
@@ -1379,8 +1379,7 @@ void SMutableConstantsWidget::LoadConstantPhysics()
 		ConstantPhysicsElement->Physics = MutableProgramPtr->ConstantPhysicsBodies[PhysicsIndex];
 		ConstantPhysicsElement->IndexOnSourceVector = PhysicsIndex;
 	
-		//TODO: Fix me.
-		//mu::FPhysicsBody::Serialise(ConstantPhysicsElement->Physics.Get(), Archive);
+		ConstantPhysicsElement->Physics->Serialise(Archive);
 
 		ConstantPhysicsElements.Add(ConstantPhysicsElement);
 	}
diff --git a/Engine/Plugins/Mutable/Source/MutableRuntime/Private/MuR/PhysicsBody.cpp b/Engine/Plugins/Mutable/Source/MutableRuntime/Private/MuR/PhysicsBody.cpp
--- a/Engine/Plugins/Mutable/Source/MutableRuntime/Private/MuR/PhysicsBody.cpp
+++ b/Engine/Plugins/Mutable/Source/MutableRuntime/Private/MuR/PhysicsBody.cpp
@@ -8,175 +8,158 @@ namespace mu
 {
 
 
-	//-------------------------------------------------------------------------------------------------
-	void FPhysicsBody::Serialise(const FPhysicsBody* p, FOutputArchive& arch)
+	void FPhysicsBody::Serialise(const FPhysicsBody* p, FOutputArchive& Arch)
 	{
-		arch << *p;
+		Arch << *p;
 	}
 
 
-	//-------------------------------------------------------------------------------------------------
-	TSharedPtr<FPhysicsBody> FPhysicsBody::StaticUnserialise(FInputArchive& arch)
+	TSharedPtr<FPhysicsBody> FPhysicsBody::StaticUnserialise(FInputArchive& Arch)
 	{
 		TSharedPtr<FPhysicsBody> pResult = MakeShared<FPhysicsBody>();
-		arch >> *pResult;
+		Arch >> *pResult;
 		return pResult;
 	}
 
 
-	//-------------------------------------------------------------------------------------------------
-	void FBodyShape::Serialise(FOutputArchive& arch) const
+	void FBodyShape::Serialise(FOutputArchive& Arch) const
 	{
-		arch << Name;
-		arch << Flags;
+		Arch << Name;
+		Arch << Flags;
 	}
 
 
-	//-------------------------------------------------------------------------------------------------
-	void FBodyShape::Unserialise(FInputArchive& arch)
+	void FBodyShape::Unserialise(FInputArchive& Arch)
 	{
-		arch >> Name;
-		arch >> Flags;
+		Arch >> Name;
+		Arch >> Flags;
 	}
 
 
-	//-------------------------------------------------------------------------------------------------
-	void FSphereBody::Serialise(FOutputArchive& arch) const
+	void FSphereBody::Serialise(FOutputArchive& Arch) const
 	{
-		FBodyShape::Serialise(arch);
+		FBodyShape::Serialise(Arch);
 
-		arch << Position;
-		arch << Radius;
+		Arch << Position;
+		Arch << Radius;
 	}
 
 
-	//-------------------------------------------------------------------------------------------------
-	void FSphereBody::Unserialise(FInputArchive& arch)
+	void FSphereBody::Unserialise(FInputArchive& Arch)
 	{
-		FBodyShape::Unserialise(arch);
+		FBodyShape::Unserialise(Arch);
 
-		arch >> Position;
-		arch >> Radius;
+		Arch >> Position;
+		Arch >> Radius;
 	}
 
 
-	//-------------------------------------------------------------------------------------------------
-	void FBoxBody::Serialise(FOutputArchive& arch) const
+	void FBoxBody::Serialise(FOutputArchive& Arch) const
 	{
-		FBodyShape::Serialise( arch );
+		FBodyShape::Serialise( Arch );
 
-		arch << Position;
-		arch << Orientation;
-		arch << Size;
+		Arch << Position;
+		Arch << Orientation;
+		Arch << Size;
 	}
 
 
-	//-------------------------------------------------------------------------------------------------
-	void FBoxBody::Unserialise(FInputArchive& arch)
+	void FBoxBody::Unserialise(FInputArchive& Arch)
 	{
-		FBodyShape::Unserialise( arch );
+		FBodyShape::Unserialise( Arch );
 
-		arch >> Position;
-		arch >> Orientation;
-		arch >> Size;
+		Arch >> Position;
+		Arch >> Orientation;
+		Arch >> Size;
 	}
 
 	
-	//-------------------------------------------------------------------------------------------------
-	void FSphylBody::Serialise(FOutputArchive& arch) const
+	void FSphylBody::Serialise(FOutputArchive& Arch) const
 	{
-		FBodyShape::Serialise( arch );
+		FBodyShape::Serialise( Arch );
 
-		arch << Position;
-		arch << Orientation;
-		arch << Radius;
-		arch << Length;
+		Arch << Position;
+		Arch << Orientation;
+		Arch << Radius;
+		Arch << Length;
 	}
 
 
-	//-------------------------------------------------------------------------------------------------
-	void FSphylBody::Unserialise(FInputArchive& arch)
+	void FSphylBody::Unserialise(FInputArchive& Arch)
 	{
-		FBodyShape::Unserialise( arch );
+		FBodyShape::Unserialise( Arch );
 
-		arch >> Position;
-		arch >> Orientation;
-		arch >> Radius;
-		arch >> Length;
+		Arch >> Position;
+		Arch >> Orientation;
+		Arch >> Radius;
+		Arch >> Length;
 	}
 
 
-	//-------------------------------------------------------------------------------------------------
-	void FTaperedCapsuleBody::Serialise(FOutputArchive& arch) const
+	void FTaperedCapsuleBody::Serialise(FOutputArchive& Arch) const
 	{
-		FBodyShape::Serialise( arch );
+		FBodyShape::Serialise( Arch );
 
-		arch << Position;
-		arch << Orientation;
-		arch << Radius0;
-		arch << Radius1;
-		arch << Length;
+		Arch << Position;
+		Arch << Orientation;
+		Arch << Radius0;
+		Arch << Radius1;
+		Arch << Length;
 	}
 
 
-	//-------------------------------------------------------------------------------------------------
-	void FTaperedCapsuleBody::Unserialise(FInputArchive& arch)
+	void FTaperedCapsuleBody::Unserialise(FInputArchive& Arch)
 	{
-		FBodyShape::Unserialise( arch );
+		FBodyShape::Unserialise( Arch );
 
-		arch >> Position;
-		arch >> Orientation;
-		arch >> Radius0;
-		arch >> Radius1;
-		arch >> Length;
+		Arch >> Position;
+		Arch >> Orientation;
+		Arch >> Radius0;
+		Arch >> Radius1;
+		Arch >> Length;
 	}
 
 
-	//-------------------------------------------------------------------------------------------------
-	void FConvexBody::Serialise(FOutputArchive& arch) const
+	void FConvexBody::Serialise(FOutputArchive& Arch) const
 	{
-		FBodyShape::Serialise( arch );
+		FBodyShape::Serialise( Arch );
 
-		arch << Vertices;
-		arch << Indices;
-		arch << Transform;
+		Arch << Vertices;
+		Arch << Indices;
+		Arch << Transform;
 	}
 
 
-	//-------------------------------------------------------------------------------------------------
-	void FConvexBody::Unserialise(FInputArchive& arch)
+	void FConvexBody::Unserialise(FInputArchive& Arch)
 	{
-		FBodyShape::Unserialise( arch );
+		FBodyShape::Unserialise( Arch );
 
-		arch >> Vertices;
-		arch >> Indices;
-		arch >> Transform;
+		Arch >> Vertices;
+		Arch >> Indices;
+		Arch >> Transform;
 	}
 
 
-	//-------------------------------------------------------------------------------------------------
-	void FPhysicsBodyAggregate::Serialise(FOutputArchive& arch) const
+	void FPhysicsBodyAggregate::Serialise(FOutputArchive& Arch) const
 	{
-		arch << Spheres;
-		arch << Boxes;
-		arch << Convex;
-		arch << Sphyls;
-		arch << TaperedCapsules;
+		Arch << Spheres;
+		Arch << Boxes;
+		Arch << Convex;
+		Arch << Sphyls;
+		Arch << TaperedCapsules;
 	}
 
 
-	//-------------------------------------------------------------------------------------------------
-	void FPhysicsBodyAggregate::Unserialise(FInputArchive& arch)
+	void FPhysicsBodyAggregate::Unserialise(FInputArchive& Arch)
 	{
-		arch >> Spheres;
-		arch >> Boxes;
-		arch >> Convex;
-		arch >> Sphyls;
-		arch >> TaperedCapsules;
+		Arch >> Spheres;
+		Arch >> Boxes;
+		Arch >> Convex;
+		Arch >> Sphyls;
+		Arch >> TaperedCapsules;
 	}
 
 
-	//-------------------------------------------------------------------------------------------------
 	TSharedPtr<FPhysicsBody> FPhysicsBody::Clone() const
 	{
 		TSharedPtr<FPhysicsBody> pResult = MakeShared<FPhysicsBody>();
@@ -193,19 +176,16 @@ namespace mu
 		return pResult;
 	}
 
-	//-------------------------------------------------------------------------------------------------
 	void FPhysicsBody::SetCustomId(int32 InCustomId)
 	{
 		CustomId = InCustomId;
 	}
 
-	//-------------------------------------------------------------------------------------------------
 	int32 FPhysicsBody::GetCustomId() const
 	{
 		return CustomId;
 	}
 
-	//-------------------------------------------------------------------------------------------------
 	void FPhysicsBody::SetBodyCount(int32 Count)
 	{
 		Bodies.SetNum(Count);
@@ -213,27 +193,23 @@ namespace mu
 		BodiesCustomIds.Init(-1, Count);
 	}
 
-	//-------------------------------------------------------------------------------------------------
 	int32 FPhysicsBody::GetBodyCount() const
 	{
 		return Bodies.Num();
 	}
 
-	//-------------------------------------------------------------------------------------------------
 	const FBoneName& FPhysicsBody::GetBodyBoneId(int32 B) const
 	{
 		check(BoneIds.IsValidIndex(B));
 		return BoneIds[B];
 	}
 
-	//-------------------------------------------------------------------------------------------------
 	void FPhysicsBody::SetBodyBoneId(int32 B, const FBoneName& BoneId)
 	{
 		check(BoneIds.IsValidIndex(B));
 		BoneIds[B] = BoneId;
 	}
 
-	//-------------------------------------------------------------------------------------------------
 	int32 FPhysicsBody::GetBodyCustomId(int32 B) const
 	{
 		check(B >= 0 && B < BodiesCustomIds.Num());
@@ -241,7 +217,6 @@ namespace mu
 		return BodiesCustomIds[B];
 	}
 
-	//-------------------------------------------------------------------------------------------------
 	void FPhysicsBody::SetBodyCustomId(int32 B, int32 BodyCustomId)
 	{
 		check(B >= 0 && B < BodiesCustomIds.Num());
@@ -249,67 +224,56 @@ namespace mu
 		BodiesCustomIds[B] = BodyCustomId;
 	}
 
-	//-------------------------------------------------------------------------------------------------
 	int32 FPhysicsBody::GetSphereCount(int32 B) const
 	{
 		return Bodies[B].Spheres.Num();
 	}
 
-	//-------------------------------------------------------------------------------------------------
 	int32 FPhysicsBody::GetBoxCount(int32 B) const
 	{
 		return Bodies[B].Boxes.Num();
 	}
 
-	//-------------------------------------------------------------------------------------------------
 	int32 FPhysicsBody::GetConvexCount(int32 B) const
 	{
 		return Bodies[B].Convex.Num();
 	}
 
-	//-------------------------------------------------------------------------------------------------
 	int32 FPhysicsBody::GetSphylCount(int32 B) const
 	{
 		return Bodies[B].Sphyls.Num();
 	}
 
-	//-------------------------------------------------------------------------------------------------
 	int32 FPhysicsBody::GetTaperedCapsuleCount(int32 B) const
 	{
 		return Bodies[B].TaperedCapsules.Num();
 	}
 
-	//-------------------------------------------------------------------------------------------------
 	void FPhysicsBody::SetSphereCount(int32 B, int32 Count)
 	{
 		Bodies[B].Spheres.SetNum(Count);
 	}
 
-	//-------------------------------------------------------------------------------------------------
 	void FPhysicsBody::SetBoxCount(int32 B, int32 Count)
 	{
 		Bodies[B].Boxes.SetNum(Count);
 	}
 
-	//-------------------------------------------------------------------------------------------------
 	void FPhysicsBody::SetConvexCount(int32 B, int32 Count)
 	{
 		Bodies[B].Convex.SetNum(Count);
 	}
 
-	//-------------------------------------------------------------------------------------------------
 	void FPhysicsBody::SetSphylCount(int32 B, int32 Count)
 	{
 		Bodies[B].Sphyls.SetNum(Count);
 	}
 
-	//-------------------------------------------------------------------------------------------------
 	void FPhysicsBody::SetTaperedCapsuleCount(int32 B, int32 Count)
 	{
 		Bodies[B].TaperedCapsules.SetNum(Count);
 	}
 
-	//-------------------------------------------------------------------------------------------------
 	void FPhysicsBody::SetSphere(
 		int32 B, int32 I,
 		FVector3f Position, float Radius)
@@ -321,7 +285,6 @@ namespace mu
 		Bodies[B].Spheres[I].Radius = Radius;
 	}
 
-	//-------------------------------------------------------------------------------------------------
 	void FPhysicsBody::SetBox(
 		int32 B, int32 I,
 		FVector3f Position, FQuat4f Orientation, FVector3f Size)
@@ -334,7 +297,6 @@ namespace mu
 		Bodies[B].Boxes[I].Size = Size;
 	}
 
-	//-------------------------------------------------------------------------------------------------
 	void FPhysicsBody::SetConvexMesh(
 		int32 B, int32 I,
 		TArrayView<const FVector3f> Vertices, TArrayView<const int32> Indices)
@@ -346,7 +308,6 @@ namespace mu
 		Bodies[B].Convex[I].Indices = TArray<int32>(Indices);
 	}
 
-	//-------------------------------------------------------------------------------------------------
 	void FPhysicsBody::SetConvexTransform(
 		int32 B, int32 I,
 		const FTransform3f& Transform)
@@ -358,7 +319,6 @@ namespace mu
 	}
 
 
-	//-------------------------------------------------------------------------------------------------
 	void FPhysicsBody::SetSphyl(
 		int32 B, int32 I,
 		FVector3f Position, FQuat4f Orientation, float Radius, float Length)
@@ -372,7 +332,6 @@ namespace mu
 		Bodies[B].Sphyls[I].Length = Length;
 	}
 
-	//-------------------------------------------------------------------------------------------------
 	void FPhysicsBody::SetSphereFlags(int32 B, int32 I, uint32 Flags)
 	{
 		check(B >= 0 && B < Bodies.Num());
@@ -381,7 +340,6 @@ namespace mu
 		Bodies[B].Spheres[I].Flags = Flags;
 	}
 
-	//-------------------------------------------------------------------------------------------------
 	void FPhysicsBody::SetBoxFlags(int32 B, int32 I, uint32 Flags)
 	{
 		check(B >= 0 && B < Bodies.Num());
@@ -390,7 +348,6 @@ namespace mu
 		Bodies[B].Boxes[I].Flags = Flags;
 	}
 
-	//-------------------------------------------------------------------------------------------------
 	void FPhysicsBody::SetConvexFlags(int32 B, int32 I, uint32 Flags)
 	{
 		check(B >= 0 && B < Bodies.Num());
@@ -399,7 +356,6 @@ namespace mu
 		Bodies[B].Convex[I].Flags = Flags;
 	}
 
-	//-------------------------------------------------------------------------------------------------
 	void FPhysicsBody::SetSphylFlags(int32 B, int32 I, uint32 Flags)
 	{
 		check(B >= 0 && B < Bodies.Num());
@@ -408,7 +364,6 @@ namespace mu
 		Bodies[B].Sphyls[I].Flags = Flags;
 	}
 
-	//-------------------------------------------------------------------------------------------------
 	void FPhysicsBody::SetTaperedCapsuleFlags(int32 B, int32 I, uint32 Flags)
 	{
 		check(B >= 0 && B < Bodies.Num());
@@ -417,7 +372,6 @@ namespace mu
 		Bodies[B].TaperedCapsules[I].Flags = Flags;
 	}
 
-	//-------------------------------------------------------------------------------------------------
 	void FPhysicsBody::SetSphereName(int32 B, int32 I, const char* Name)
 	{
 		check(B >= 0 && B < Bodies.Num());
@@ -426,7 +380,6 @@ namespace mu
 		Bodies[B].Spheres[I].Name = Name;
 	}
 
-	//-------------------------------------------------------------------------------------------------
 	void FPhysicsBody::SetBoxName(int32 B, int32 I, const char* Name)
 	{
 		check(B >= 0 && B < Bodies.Num());
@@ -435,7 +388,6 @@ namespace mu
 		Bodies[B].Boxes[I].Name = Name;
 	}
 
-	//-------------------------------------------------------------------------------------------------
 	void FPhysicsBody::SetConvexName(int32 B, int32 I, const char* Name)
 	{
 		check(B >= 0 && B < Bodies.Num());
@@ -444,7 +396,6 @@ namespace mu
 		Bodies[B].Convex[I].Name = Name;
 	}
 
-	//-------------------------------------------------------------------------------------------------
 	void FPhysicsBody::SetSphylName(int32 B, int32 I, const char* Name)
 	{
 		check(B >= 0 && B < Bodies.Num());
@@ -453,7 +404,6 @@ namespace mu
 		Bodies[B].Sphyls[I].Name = Name;
 	}
 
-	//-------------------------------------------------------------------------------------------------
 	void FPhysicsBody::SetTaperedCapsuleName(int32 B, int32 I, const char* Name)
 	{
 		check(B >= 0 && B < Bodies.Num());
@@ -462,7 +412,6 @@ namespace mu
 		Bodies[B].TaperedCapsules[I].Name = Name;
 	}
 
-	//-------------------------------------------------------------------------------------------------
 	void FPhysicsBody::SetTaperedCapsule(
 		int32 B, int32 I,
 		FVector3f Position, FQuat4f Orientation,
@@ -478,7 +427,6 @@ namespace mu
 		Bodies[B].TaperedCapsules[I].Length = Length;
 	}
 
-	//-------------------------------------------------------------------------------------------------
 	void FPhysicsBody::GetSphere(
 		int32 B, int32 I,
 		FVector3f& OutPosition, float& OutRadius) const
@@ -490,7 +438,6 @@ namespace mu
 		OutRadius = Bodies[B].Spheres[I].Radius;
 	}
 
-	//-------------------------------------------------------------------------------------------------
 	void FPhysicsBody::GetBox(
 		int32 B, int32 I,
 		FVector3f& OutPosition, FQuat4f& OutOrientation, FVector3f& OutSize) const
@@ -503,7 +450,6 @@ namespace mu
 		OutSize = Bodies[B].Boxes[I].Size;
 	}
 
-	//-------------------------------------------------------------------------------------------------
 	void FPhysicsBody::GetConvex(
 		int32 B, int32 I,
 		TArrayView<const FVector3f>& OutVertices, TArrayView<const int32>& OutIndices,
@@ -535,7 +481,6 @@ namespace mu
 		OutTransform = Bodies[B].Convex[I].Transform;
 	}
 
-	//-------------------------------------------------------------------------------------------------
 	void FPhysicsBody::GetSphyl(
 		int32 B, int32 I,
 		FVector3f& OutPosition, FQuat4f& OutOrientation, float& OutRadius, float& OutLength) const
@@ -549,7 +494,6 @@ namespace mu
 		OutLength = Bodies[B].Sphyls[I].Length;
 	}
 
-	//-------------------------------------------------------------------------------------------------
 	void FPhysicsBody::GetTaperedCapsule(
 		int32 B, int32 I,
 		FVector3f& OutPosition, FQuat4f& OutOrientation,
@@ -566,7 +510,6 @@ namespace mu
 	}
 
 
-	//-------------------------------------------------------------------------------------------------
 	uint32 FPhysicsBody::GetSphereFlags(int32 B, int32 I) const
 	{
 		check(B >= 0 && B < Bodies.Num());
@@ -575,7 +518,6 @@ namespace mu
 		return Bodies[B].Spheres[I].Flags;
 	}
 
-	//-------------------------------------------------------------------------------------------------
 	uint32 FPhysicsBody::GetBoxFlags(int32 B, int32 I) const
 	{
 		check(B >= 0 && B < Bodies.Num());
@@ -584,7 +526,6 @@ namespace mu
 		return Bodies[B].Boxes[I].Flags;
 	}
 
-	//-------------------------------------------------------------------------------------------------
 	uint32 FPhysicsBody::GetConvexFlags(int32 B, int32 I) const
 	{
 		check(B >= 0 && B < Bodies.Num());
@@ -593,7 +534,6 @@ namespace mu
 		return Bodies[B].Convex[I].Flags;
 	}
 
-	//-------------------------------------------------------------------------------------------------
 	uint32 FPhysicsBody::GetSphylFlags(int32 B, int32 I) const
 	{
 		check(B >= 0 && B < Bodies.Num());
@@ -602,7 +542,6 @@ namespace mu
 		return Bodies[B].Sphyls[I].Flags;
 	}
 
-	//-------------------------------------------------------------------------------------------------
 	uint32 FPhysicsBody::GetTaperedCapsuleFlags(int32 B, int32 I) const
 	{
 		check(B >= 0 && B < Bodies.Num());
@@ -611,7 +550,6 @@ namespace mu
 		return Bodies[B].TaperedCapsules[I].Flags;
 	}
 
-	//-------------------------------------------------------------------------------------------------
 	const FString& FPhysicsBody::GetSphereName(int32 B, int32 I) const
 	{
 		check(B >= 0 && B < Bodies.Num());
@@ -620,7 +558,6 @@ namespace mu
 		return Bodies[B].Spheres[I].Name;
 	}
 
-	//-------------------------------------------------------------------------------------------------
 	const FString& FPhysicsBody::GetBoxName(int32 B, int32 I) const
 	{
 		check(B >= 0 && B < Bodies.Num());
@@ -629,7 +566,6 @@ namespace mu
 		return Bodies[B].Boxes[I].Name;
 	}
 
-	//-------------------------------------------------------------------------------------------------
 	const FString& FPhysicsBody::GetConvexName(int32 B, int32 I) const
 	{
 		check(B >= 0 && B < Bodies.Num());
@@ -638,7 +574,6 @@ namespace mu
 		return Bodies[B].Convex[I].Name;
 	}
 
-	//-------------------------------------------------------------------------------------------------
 	const FString& FPhysicsBody::GetSphylName(int32 B, int32 I) const
 	{
 		check(B >= 0 && B < Bodies.Num());
@@ -647,7 +582,6 @@ namespace mu
 		return Bodies[B].Sphyls[I].Name;
 	}
 
-	//-------------------------------------------------------------------------------------------------
 	const FString& FPhysicsBody::GetTaperedCapsuleName(int32 B, int32 I) const
 	{
 		check(B >= 0 && B < Bodies.Num());
@@ -657,24 +591,21 @@ namespace mu
 	}
 
 
-	//-------------------------------------------------------------------------------------------------
-	void FPhysicsBody::Serialise(FOutputArchive& arch) const
+	void FPhysicsBody::Serialise(FOutputArchive& Arch) const
 	{
-		arch << CustomId;
-		arch << Bodies;
-		arch << BoneIds;
-		arch << BodiesCustomIds;
-		arch << bBodiesModified;
+		Arch << CustomId;
+		Arch << Bodies;
+		Arch << BoneIds;
+		Arch << BodiesCustomIds;
+		Arch << bBodiesModified;
 	}
 
-
-	//-------------------------------------------------------------------------------------------------
-	void FPhysicsBody::Unserialise(FInputArchive& arch)
+	void FPhysicsBody::Unserialise(FInputArchive& Arch)
 	{
-		arch >> CustomId;        	
-		arch >> Bodies;
-		arch >> BoneIds;
-		arch >> BodiesCustomIds;
-		arch >> bBodiesModified;
+		Arch >> CustomId;        	
+		Arch >> Bodies;
+		Arch >> BoneIds;
+		Arch >> BodiesCustomIds;
+		Arch >> bBodiesModified;
 	}
 }
diff --git a/Engine/Plugins/Mutable/Source/MutableRuntime/Public/MuR/PhysicsBody.h b/Engine/Plugins/Mutable/Source/MutableRuntime/Public/MuR/PhysicsBody.h
--- a/Engine/Plugins/Mutable/Source/MutableRuntime/Public/MuR/PhysicsBody.h
+++ b/Engine/Plugins/Mutable/Source/MutableRuntime/Public/MuR/PhysicsBody.h
@@ -30,11 +30,8 @@ namespace mu
 			return Flags == Other.Flags && Name == Other.Name;
 		}
 	
-		//!
-		inline void Serialise( FOutputArchive& arch ) const;
-
-		//!
-		inline void Unserialise( FInputArchive& arch );
+		void Serialise(FOutputArchive& Arch) const;
+		void Unserialise(FInputArchive& Arch);
 	};
 
 	struct FSphereBody : FBodyShape
@@ -47,11 +44,8 @@ namespace mu
 			return FBodyShape::operator==(Other) && Position == Other.Position && Radius == Other.Radius;
 		}
 		 
-		//!
-		inline void Serialise( FOutputArchive& arch ) const;
-
-		//!
-		inline void Unserialise( FInputArchive& arch );
+		void Serialise(FOutputArchive& Arch) const;
+		void Unserialise(FInputArchive& Arch);
 	};
 	
 	struct FBoxBody : FBodyShape
@@ -65,11 +59,8 @@ namespace mu
 			return FBodyShape::operator==(Other) && Position == Other.Position && Orientation == Other.Orientation && Size == Other.Size;
 		}
 
-		//!
-		inline void Serialise( FOutputArchive& arch ) const;
-
-		//!
-		inline void Unserialise( FInputArchive& arch );
+		void Serialise(FOutputArchive& Arch) const;
+		void Unserialise(FInputArchive& Arch);
 	};
 
 	struct FSphylBody : FBodyShape
@@ -88,11 +79,8 @@ namespace mu
 				Length == Other.Length;
 		}
 		 
-		//!
-		inline void Serialise( FOutputArchive& arch ) const;
-
-		//!
-		inline void Unserialise( FInputArchive& arch );
+		void Serialise(FOutputArchive& Arch) const;
+		void Unserialise(FInputArchive& Arch);
 	};
 	
 	struct FTaperedCapsuleBody : FBodyShape
@@ -113,11 +101,8 @@ namespace mu
 				Length == Other.Length;
 		}
 		 
-		//!
-		inline void Serialise( FOutputArchive& arch ) const;
-
-		//!
-		inline void Unserialise( FInputArchive& arch );
+		void Serialise(FOutputArchive& Arch) const;
+		void Unserialise(FInputArchive& Arch);
 	};
 	
 	struct FConvexBody : FBodyShape
@@ -135,11 +120,8 @@ namespace mu
 				Transform.Equals(Other.Transform);
 		}
 
-		//!
-		inline void Serialise( FOutputArchive& arch ) const;
-
-		//!
-		inline void Unserialise( FInputArchive& arch );
+		void Serialise(FOutputArchive& Arch) const;
+		void Unserialise(FInputArchive& Arch);
 	};
 	
 	struct FPhysicsBodyAggregate
@@ -159,11 +141,8 @@ namespace mu
 					Convex == Other.Convex;
 		}
 
-		 //!
-		inline void Serialise( FOutputArchive& ) const;
-
-		//!
-		inline void Unserialise( FInputArchive& );
+		void Serialise(FOutputArchive& Arch) const;
+		void Unserialise(FInputArchive& Arch);
 	};
 
 	class MUTABLERUNTIME_API FPhysicsBody
@@ -290,14 +269,10 @@ namespace mu
 		TArray<FPhysicsBodyAggregate> Bodies;
 		TArray<int32> BodiesCustomIds;
 
-
 		bool bBodiesModified = false;
 
-        //!
-        inline void Serialise( FOutputArchive& arch ) const;
-
-		//!
-        inline void Unserialise( FInputArchive& arch );
+        void Serialise(FOutputArchive& Arch) const;
+        void Unserialise(FInputArchive& Arch);
 
 		//!
 		inline bool operator==(const FPhysicsBody& Other) const
