[Mutable] Add mesh content filter options to mesh runs.

- Meshes now are split in 4 roms containing, geometry, pose, phyiscs and metadata repectively.
- Now geometry is filtered for the first generation of non resident lods when streaming is enabled.

#rnx

diff --git a/Engine/Plugins/Mutable/Source/CustomizableObject/Internal/MuCO/CustomizableObjectPrivate.h b/Engine/Plugins/Mutable/Source/CustomizableObject/Internal/MuCO/CustomizableObjectPrivate.h
--- a/Engine/Plugins/Mutable/Source/CustomizableObject/Internal/MuCO/CustomizableObjectPrivate.h
+++ b/Engine/Plugins/Mutable/Source/CustomizableObject/Internal/MuCO/CustomizableObjectPrivate.h
@@ -1565,6 +1565,7 @@ enum class ECustomizableObjectVersions : int32
 	RemoveObsoletMeshInterpolateAndGeometryOp,
 
 	RemoveObsoleteDataTypesFromEnum,
+
 	ConvertModelResourcesToUObject,
 
 	RemoveObsoletImageGradientOp,
@@ -1585,6 +1586,8 @@ enum class ECustomizableObjectVersions : int32
 	
 	AddAssetUserDataEditor,
 	
+	MeshDataRomSplit,
+
 	// -----<new versions can be added above this line>--------
 	LastCustomizableObjectVersion
 };
@@ -1858,6 +1861,6 @@ public:
 	static constexpr int32 CurrentSupportedVersion = static_cast<int32>(ECustomizableObjectVersions::LastCustomizableObjectVersion);
 #endif
 	
-	static constexpr int32 DerivedDataVersion = 0x0b410b41;
+	static constexpr int32 DerivedDataVersion = 0x58abf629;
 };
 
diff --git a/Engine/Plugins/Mutable/Source/CustomizableObject/Private/MuCO/CustomizableObjectInstance.cpp b/Engine/Plugins/Mutable/Source/CustomizableObject/Private/MuCO/CustomizableObjectInstance.cpp
--- a/Engine/Plugins/Mutable/Source/CustomizableObject/Private/MuCO/CustomizableObjectInstance.cpp
+++ b/Engine/Plugins/Mutable/Source/CustomizableObject/Private/MuCO/CustomizableObjectInstance.cpp
@@ -4200,7 +4200,7 @@ void UCustomizableInstancePrivate::BuildOrCopyMorphTargetsData(const TSharedRef<
 			MorphLODModels.SetNum(ComponentMorphTargets->RealTimeMorphsLODData[MorphTargetIndex].Num());
 			
 			// Streamed LODs
-			const int32 FirstGeneratedLOD = FMath::Max(OperationData->GetRequestedLODs()[ComponentName], OperationData->GetMinLOD()[ComponentName]);
+			const int32 FirstGeneratedLOD = OperationData->FirstLODAvailable[ComponentName];
 			for (int32 LODIndex = FirstGeneratedLOD; LODIndex < OperationData->FirstResidentLOD[ComponentName]; ++LODIndex)
 			{
 				// Copy data required for streaming
diff --git a/Engine/Plugins/Mutable/Source/CustomizableObject/Private/MuCO/CustomizableObjectMeshUpdate.cpp b/Engine/Plugins/Mutable/Source/CustomizableObject/Private/MuCO/CustomizableObjectMeshUpdate.cpp
--- a/Engine/Plugins/Mutable/Source/CustomizableObject/Private/MuCO/CustomizableObjectMeshUpdate.cpp
+++ b/Engine/Plugins/Mutable/Source/CustomizableObject/Private/MuCO/CustomizableObjectMeshUpdate.cpp
@@ -183,7 +183,9 @@ namespace impl
 		}
 
 		const mu::FResourceID MeshID = OperationData->MeshIDs[LODIndex];
-		UE::Tasks::TTask<TSharedPtr<const mu::FMesh>> GetMeshTask = OperationData->System->GetMesh(InstanceID, MeshID);
+
+		constexpr mu::EMeshContentFlags MeshContentFilter = mu::EMeshContentFlags::AllFlags;
+		UE::Tasks::TTask<TSharedPtr<const mu::FMesh>> GetMeshTask = OperationData->System->GetMesh(InstanceID, MeshID, MeshContentFilter);
 
 		UE::Tasks::AddNested(UE::Tasks::Launch(TEXT("Task_MutableGetMeshes_GetMesh_Post"), [=]() mutable
 			{
@@ -390,4 +392,4 @@ void FCustomizableObjectMeshStreamIn::ConvertMesh(const FContext& Context)
 
 	// Clear MeshUpdate data
 	OperationData = nullptr;
-}
\ No newline at end of file
+}
diff --git a/Engine/Plugins/Mutable/Source/CustomizableObject/Private/MuCO/CustomizableObjectSystem.cpp b/Engine/Plugins/Mutable/Source/CustomizableObject/Private/MuCO/CustomizableObjectSystem.cpp
--- a/Engine/Plugins/Mutable/Source/CustomizableObject/Private/MuCO/CustomizableObjectSystem.cpp
+++ b/Engine/Plugins/Mutable/Source/CustomizableObject/Private/MuCO/CustomizableObjectSystem.cpp
@@ -128,6 +128,11 @@ TAutoConsoleVariable<bool> CVarIgnoreFirstAvailableLODCalculation(
 	false,
 	TEXT("If set to true, ignores the first available LOD calculation to set the generated tetxure size."));
 
+TAutoConsoleVariable<bool> CVarForceGeometryOnFirstGeneration(
+	TEXT("mutable.ForceGeometryOnFirstGeneration"),
+	false,
+	TEXT("If set to true, forces geometry generation on firts generation even if the LOD will be streamed."));
+
 #if WITH_EDITOR
 bool bEnableLODManagmentInEditor = false;
 
@@ -2092,6 +2097,7 @@ namespace impl
 	{
 		int32 LODIndex;
 		mu::FResourceID MeshID;
+		mu::EMeshContentFlags ContentFilter;
 	};
 
 
@@ -2396,7 +2402,7 @@ namespace impl
 				MorphsData[MorphIndex].SetNum(NumLODsInThisComponent);
 			}
 			
-			const int32 FirstGeneratedLOD = FMath::Max(OperationData->GetRequestedLODs()[ComponentName], OperationData->GetMinLOD()[ComponentName]);
+			const int32 FirstGeneratedLOD = OperationData->FirstResidentLOD[ComponentName]; 
 			for (int32 LODIndex = FirstGeneratedLOD; LODIndex < NumLODsInThisComponent; ++LODIndex)
 			{
 				const FInstanceUpdateData::FLOD& LOD = UpdateData.LODs[Component.FirstLOD+LODIndex];
@@ -2787,7 +2793,8 @@ namespace impl
 		}
 		
 		const FGetMeshData& MeshData = (*GetMeshesData)[GetMeshIndex];
-		UE::Tasks::TTask<TSharedPtr<const mu::FMesh>> GetMeshTask = OperationData->MutableSystem->GetMesh(OperationData->InstanceID, MeshData.MeshID);
+		UE::Tasks::TTask<TSharedPtr<const mu::FMesh>> GetMeshTask = 
+				OperationData->MutableSystem->GetMesh(OperationData->InstanceID, MeshData.MeshID, MeshData.ContentFilter);
 		
 		UE::Tasks::AddNested(UE::Tasks::Launch(TEXT("Task_MutableGetMeshes_GetMesh_Post"), [=]()
 		{
@@ -2840,8 +2847,10 @@ namespace impl
 				return;
 			}
 
-			const TSharedRef<TArray<FGetMeshData>> GetMeshesData = MakeShared<TArray<FGetMeshData>>();
+			const bool bForceGeometryGeneration = CVarForceGeometryOnFirstGeneration.GetValueOnAnyThread();
 		
+			const TSharedRef<TArray<FGetMeshData>> GetMeshesData = MakeShared<TArray<FGetMeshData>>();
+			
 			OperationData->InstanceUpdateData.Components.SetNum(OperationData->NumInstanceComponents);
 			for (int32 InstanceComponentIndex = 0; InstanceComponentIndex < OperationData->NumInstanceComponents; ++InstanceComponentIndex)
 			{
@@ -2869,6 +2878,8 @@ namespace impl
 
 				Component.OverlayMaterial = Instance->GetOverlayMaterial(InstanceComponentIndex);
 
+				const uint8 FirstResidentLOD = OperationData->FirstResidentLOD[ComponentName];
+				
 				for (int32 MutableLODIndex = 0; MutableLODIndex < Component.LODCount; ++MutableLODIndex)
 				{
 					// If the LOD is not generated we still add an empty one to keep indexes aligned.
@@ -2894,7 +2905,17 @@ namespace impl
 						{
 							LOD.bGenerated = true;
 
-							GetMeshesData->Add({ UpdateDataLODIndex, LOD.MeshID});
+							mu::EMeshContentFlags MeshContentFilter = mu::EMeshContentFlags::AllFlags; 
+							
+							if (!bForceGeometryGeneration)
+							{
+								if (MutableLODIndex < FirstResidentLOD)
+								{
+									EnumRemoveFlags(MeshContentFilter, mu::EMeshContentFlags::GeometryData);
+								}
+							}
+
+							GetMeshesData->Add({ UpdateDataLODIndex, LOD.MeshID, MeshContentFilter });
 						}
 					}
 				}
diff --git a/Engine/Plugins/Mutable/Source/CustomizableObject/Private/MuCO/UnrealConversionUtils.cpp b/Engine/Plugins/Mutable/Source/CustomizableObject/Private/MuCO/UnrealConversionUtils.cpp
--- a/Engine/Plugins/Mutable/Source/CustomizableObject/Private/MuCO/UnrealConversionUtils.cpp
+++ b/Engine/Plugins/Mutable/Source/CustomizableObject/Private/MuCO/UnrealConversionUtils.cpp
@@ -344,7 +344,9 @@ namespace UnrealConversionUtils
 		MUTABLE_CPUPROFILER_SCOPE(CopyMutableVertexBuffers);
 
 		const mu::FMeshBufferSet& MutableMeshVertexBuffers = MutableMesh->GetVertexBuffers();
-		const int32 NumVertices = MutableMeshVertexBuffers.GetElementCount();
+		const int32 NumVertices = MutableMeshVertexBuffers.IsDescriptor()
+				? 0 
+				: MutableMeshVertexBuffers.GetElementCount();
 
 		//const ESkeletalMeshVertexFlags BuildFlags = SkeletalMesh->GetVertexBufferFlags();
 		//const bool bUseFullPrecisionUVs = EnumHasAllFlags(BuildFlags, ESkeletalMeshVertexFlags::UseFullPrecisionUVs);
@@ -458,6 +460,12 @@ namespace UnrealConversionUtils
 			return false;
 		}
 
+		if (InMutableMesh->GetIndexBuffers().IsDescriptor())
+		{
+			UE_LOG(LogMutable, Error, TEXT("UCustomizableInstancePrivateData::BuildSkeletalMeshRenderData is converting a mesh descriptor."));
+			return false;
+		}
+
 		const uint8* DataPtr = InMutableMesh->GetIndexBuffers().GetBufferData(0);
 		const int32 ElementSize = InMutableMesh->GetIndexBuffers().GetElementSize(0);
 
diff --git a/Engine/Plugins/Mutable/Source/CustomizableObjectEditor/Private/MuCOE/SMutableCodeViewer.cpp b/Engine/Plugins/Mutable/Source/CustomizableObjectEditor/Private/MuCOE/SMutableCodeViewer.cpp
--- a/Engine/Plugins/Mutable/Source/CustomizableObjectEditor/Private/MuCOE/SMutableCodeViewer.cpp
+++ b/Engine/Plugins/Mutable/Source/CustomizableObjectEditor/Private/MuCOE/SMutableCodeViewer.cpp
@@ -2575,7 +2575,7 @@ bool SMutableCodeViewer::IsConstantResourceUsedByOperation(const int32 IndexOnCo
 				if (OperationType == mu::EOpType::ME_CONSTANT)
 				{
 					mu::OP::MeshConstantArgs Args = InProgram.GetOpArgs<mu::OP::MeshConstantArgs>(OperationAddress);
-					mu::FConstantResourceIndex ResourceIndex = *reinterpret_cast<const mu::FConstantResourceIndex*>(&Args.value);
+					mu::FConstantResourceIndex ResourceIndex = *reinterpret_cast<const mu::FConstantResourceIndex*>(&Args.Value);
 					
 					if (!ResourceIndex.Streamable)
 					{
@@ -2664,7 +2664,7 @@ bool SMutableCodeViewer::IsConstantResourceUsedByOperation(const int32 IndexOnCo
 			{
 				if (OperationType == mu::EOpType::ME_CONSTANT)
 				{
-					bResourceLocated = IndexOnConstantsArray == InProgram.GetOpArgs<mu::OP::MeshConstantArgs>(OperationAddress).skeleton;
+					bResourceLocated = IndexOnConstantsArray == InProgram.GetOpArgs<mu::OP::MeshConstantArgs>(OperationAddress).Skeleton;
 				}
 				break;
 			}
@@ -2672,7 +2672,7 @@ bool SMutableCodeViewer::IsConstantResourceUsedByOperation(const int32 IndexOnCo
 			{
 				if (OperationType == mu::EOpType::ME_CONSTANT)
 				{
-					bResourceLocated = IndexOnConstantsArray == InProgram.GetOpArgs<mu::OP::MeshConstantArgs>(OperationAddress).physicsBody;
+					bResourceLocated = IndexOnConstantsArray == InProgram.GetOpArgs<mu::OP::MeshConstantArgs>(OperationAddress).PhysicsBody;
 				}
 				break;
 			}
@@ -2868,7 +2868,7 @@ void SMutableCodeViewer::Tick(const FGeometry& AllottedGeometry, const double In
 	case mu::EDataType::Mesh:
 	{
 		check(PreviewMeshViewer);
-		TSharedPtr<const mu::FMesh> MutableMesh = System->GetPrivate()->BuildMesh(MutableModel, PreviewParameters.Get(), SelectedOperationAddress);
+		TSharedPtr<const mu::FMesh> MutableMesh = System->GetPrivate()->BuildMesh(MutableModel, PreviewParameters.Get(), SelectedOperationAddress, mu::EMeshContentFlags::AllFlags);
 		PreviewMeshViewer->SetMesh(MutableMesh);
 		break;
 	}
diff --git a/Engine/Plugins/Mutable/Source/CustomizableObjectEditor/Private/MuCOE/SMutableInstanceViewer.cpp b/Engine/Plugins/Mutable/Source/CustomizableObjectEditor/Private/MuCOE/SMutableInstanceViewer.cpp
--- a/Engine/Plugins/Mutable/Source/CustomizableObjectEditor/Private/MuCOE/SMutableInstanceViewer.cpp
+++ b/Engine/Plugins/Mutable/Source/CustomizableObjectEditor/Private/MuCOE/SMutableInstanceViewer.cpp
@@ -312,7 +312,12 @@ void SMutableInstanceViewer::SetInstance(const TSharedPtr<const mu::FInstance>&
 					// Generate the mesh to get some data
 					{
 						mu::OP::ADDRESS MeshRootAddress = mu::GetResourceIDRoot(MeshId);
-						TSharedPtr<const mu::FMesh> MutableMesh = System.GetPrivate()->BuildMesh(Model, Parameters.Get(), MeshRootAddress);
+						TSharedPtr<const mu::FMesh> MutableMesh = System.GetPrivate()->BuildMesh(
+								Model, 
+								Parameters.Get(), 
+								MeshRootAddress, 
+								mu::EMeshContentFlags::AllFlags);
+
 						if (MutableMesh && MutableMesh->IsReference())
 						{
 							NewLOD->MeshId = FText::Format(
diff --git a/Engine/Plugins/Mutable/Source/MutableRuntime/Internal/MuR/ModelPrivate.h b/Engine/Plugins/Mutable/Source/MutableRuntime/Internal/MuR/ModelPrivate.h
--- a/Engine/Plugins/Mutable/Source/MutableRuntime/Internal/MuR/ModelPrivate.h
+++ b/Engine/Plugins/Mutable/Source/MutableRuntime/Internal/MuR/ModelPrivate.h
@@ -93,6 +93,15 @@ namespace mu
 	};
 	MUTABLE_DEFINE_POD_SERIALISABLE(FImageLODRange);
 
+	struct FMeshContentRange
+	{
+		uint32            FirstIndex   : 24;
+		EMeshContentFlags ContentFlags : 8 ;
+
+		uint32 MeshIDPrefix = 0;
+	};
+	MUTABLE_DEFINE_POD_SERIALISABLE(FMeshContentRange);
+
 	struct FExtensionDataConstant
 	{
 		// This should always be valid, but if the state is Unloaded it won't be usable.
@@ -252,9 +261,15 @@ namespace mu
 		//! Constant image mip chain indices: ranges in this array are defined in FImageLODRange and the indices here refer to ConstantImageLODs.
 		TArray<FConstantResourceIndex> ConstantImageLODIndices;
 
+		/**  */
+		TArray<FConstantResourceIndex> ConstantMeshContentIndices;
+
 		//! Constant image data.
 		TArray<FImageLODRange> ConstantImages;
 
+		/** Constant mesh data */
+		TArray<FMeshContentRange> ConstantMeshes;
+
 		/** Constant mesh data is split in 2 sets: ConstantMeshesPermanent constains data that is always loaded.
 		* Index with FConstantResourceIndex::Index, when Streamable is 0.
 		*/
@@ -322,6 +337,8 @@ namespace mu
 			arch << ConstantImageLODIndices;
 			arch << ConstantImages;
 			arch << ConstantMeshesPermanent;
+			arch << ConstantMeshContentIndices;
+			arch << ConstantMeshes;
 			arch << ConstantExtensionData;
 			arch << ConstantStrings;
             arch << ConstantLayouts;
@@ -348,6 +365,8 @@ namespace mu
 			arch >> ConstantImageLODIndices;
 			arch >> ConstantImages;
 			arch >> ConstantMeshesPermanent;
+			arch >> ConstantMeshContentIndices;
+			arch >> ConstantMeshes;
 			arch >> ConstantExtensionData;
 			arch >> ConstantStrings;
             arch >> ConstantLayouts;
@@ -704,24 +723,208 @@ namespace mu
 				res = Result;
 			}
 		}
-
-        void GetConstant(FConstantResourceIndex Index, TSharedPtr<const FMesh>& OutMesh) const
+		
+		template<class CreateMeshFunc>
+        void GetConstant(
+				uint32 MeshConstantIndex, 
+				int32 SkeletonConstantIndex, 
+				int32 PhysicsConstantIndex, 
+				TSharedPtr<const FMesh>& OutMesh, 
+				EMeshContentFlags FilterContentFlags, 
+				CreateMeshFunc&& CreateMesh) const
         {
-			if (!Index.Streamable)
+			MUTABLE_CPUPROFILER_SCOPE(GetConstant_Mesh)
+
+			FMeshContentRange MeshContentRange = ConstantMeshes[MeshConstantIndex];
+
+			TSharedPtr<const FMesh> EmptyMesh = nullptr;
+			auto GetMeshAtResourceIndex = [&](FConstantResourceIndex ResourceIndex) -> TSharedPtr<const FMesh>
 			{
-				if (ConstantMeshesPermanent.IsValidIndex(Index.Index))
+				MUTABLE_CPUPROFILER_SCOPE(GetConstant_Mesh_GetMesh)
+
+				if (!ResourceIndex.Streamable)
 				{
-					OutMesh = ConstantMeshesPermanent[Index.Index];
+					if (ConstantMeshesPermanent.IsValidIndex(ResourceIndex.Index))
+					{
+						return ConstantMeshesPermanent[ResourceIndex.Index];
+					}
+
+					if (!EmptyMesh)
+					{
+						EmptyMesh = MakeShared<FMesh>();
+					}
+					
+					return EmptyMesh;
+				}
+				else
+				{
+					const TSharedPtr<const FMesh>* Found = ConstantMeshesStreamed.Find(ResourceIndex.Index);
+					if (Found)
+					{
+						return *Found;
+					}
+
+					if (!EmptyMesh)
+					{
+						EmptyMesh = MakeShared<FMesh>();
+					}
+
+					return EmptyMesh;
 				}
+			};
+
+			TSharedPtr<const FMesh> GeometryMesh = nullptr;
+			TSharedPtr<const FMesh> PoseMesh     = nullptr;
+			TSharedPtr<const FMesh> PhysicsMesh  = nullptr;
+			TSharedPtr<const FMesh> MetaDataMesh = nullptr;
+
+			int32 MeshRomCurrentIndex = MeshContentRange.FirstIndex;
+			if (EnumHasAnyFlags(FilterContentFlags, EMeshContentFlags::GeometryData) &&
+				EnumHasAnyFlags(MeshContentRange.ContentFlags, EMeshContentFlags::GeometryData))
+			{
+				const FConstantResourceIndex ResourceIndex = ConstantMeshContentIndices[MeshRomCurrentIndex];
+				GeometryMesh = GetMeshAtResourceIndex(ResourceIndex);
 			}
-			else
+			MeshRomCurrentIndex += (int32)EnumHasAnyFlags(MeshContentRange.ContentFlags, EMeshContentFlags::GeometryData);
+
+			if (EnumHasAnyFlags(FilterContentFlags, EMeshContentFlags::PoseData) &&
+				EnumHasAnyFlags(MeshContentRange.ContentFlags, EMeshContentFlags::PoseData))
 			{
-				const TSharedPtr<const FMesh>* Found = ConstantMeshesStreamed.Find(Index.Index);
-				if (Found)
+				const FConstantResourceIndex ResourceIndex = ConstantMeshContentIndices[MeshRomCurrentIndex];
+				PoseMesh = GetMeshAtResourceIndex(ResourceIndex);
+			}
+			MeshRomCurrentIndex += (int32)EnumHasAnyFlags(MeshContentRange.ContentFlags, EMeshContentFlags::PoseData);
+
+			if (EnumHasAnyFlags(FilterContentFlags, EMeshContentFlags::PhysicsData) &&
+				EnumHasAnyFlags(MeshContentRange.ContentFlags, EMeshContentFlags::PhysicsData))
+			{
+				const FConstantResourceIndex ResourceIndex = ConstantMeshContentIndices[MeshRomCurrentIndex];
+				PhysicsMesh = GetMeshAtResourceIndex(ResourceIndex);
+			}
+			MeshRomCurrentIndex += (int32)EnumHasAnyFlags(MeshContentRange.ContentFlags, EMeshContentFlags::PhysicsData);
+
+			if (EnumHasAnyFlags(FilterContentFlags, EMeshContentFlags::MetaData) &&
+				EnumHasAnyFlags(MeshContentRange.ContentFlags, EMeshContentFlags::MetaData))
+			{
+				const FConstantResourceIndex ResourceIndex = ConstantMeshContentIndices[MeshRomCurrentIndex];
+				MetaDataMesh = GetMeshAtResourceIndex(ResourceIndex);
+			}
+			MeshRomCurrentIndex += (int32)EnumHasAnyFlags(MeshContentRange.ContentFlags, EMeshContentFlags::MetaData);
+
+			check(MeshRomCurrentIndex - MeshContentRange.FirstIndex == FMath::CountBits((uint64)MeshContentRange.ContentFlags));
+
+			uint32 MeshBudgetReserve = 0;
+			MeshBudgetReserve += GeometryMesh ? GeometryMesh->GetDataSize() : 0;
+			MeshBudgetReserve += PoseMesh     ? PoseMesh->GetDataSize()     : 0;
+			MeshBudgetReserve += PhysicsMesh  ? PhysicsMesh->GetDataSize()  : 0;
+			MeshBudgetReserve += MetaDataMesh ? MetaDataMesh->GetDataSize() : 0;
+
+			TSharedPtr<FMesh> Result = nullptr;
+			if (GeometryMesh)
+			{
+				MUTABLE_CPUPROFILER_SCOPE(GetConstant_Mesh_Geoemtry)
+
+				Result = CreateMesh(MeshBudgetReserve);
+				Result->CopyFrom(*GeometryMesh);
+			}
+
+			if (PoseMesh)
+			{
+				MUTABLE_CPUPROFILER_SCOPE(GetConstant_Mesh_Pose)
+
+				if (!Result)
+				{
+					Result = CreateMesh(MeshBudgetReserve);
+					Result->CopyFrom(*PoseMesh);
+				}
+				else
+				{
+					Result->BonePoses = PoseMesh->BonePoses; 
+					Result->BoneMap   = PoseMesh->BoneMap;
+
+					for (const TPair<EMeshBufferType, FMeshBufferSet>& AdditionalBuffer : PoseMesh->AdditionalBuffers)
+					{
+						const bool bIsPoseBufferType = 	
+								AdditionalBuffer.Key == EMeshBufferType::SkeletonDeformBinding;
+
+						if (bIsPoseBufferType)
+						{
+							Result->AdditionalBuffers.Emplace(AdditionalBuffer);
+						}
+					}	
+				}
+			}
+
+			if (PhysicsMesh)
+			{
+				MUTABLE_CPUPROFILER_SCOPE(GetConstant_Mesh_Physics)
+				
+				if (!Result)
+				{
+					Result = CreateMesh(MeshBudgetReserve);
+					Result->CopyFrom(*PhysicsMesh);
+				}
+				else
+				{
+					for (const TPair<EMeshBufferType, FMeshBufferSet>& AdditionalBuffer : PhysicsMesh->AdditionalBuffers)
+					{
+						const bool bIsPhysicsBufferType = 
+								AdditionalBuffer.Key == EMeshBufferType::PhysicsBodyDeformBinding   ||
+								AdditionalBuffer.Key == EMeshBufferType::PhysicsBodyDeformSelection ||
+								AdditionalBuffer.Key == EMeshBufferType::PhysicsBodyDeformOffsets;
+
+						if (bIsPhysicsBufferType)
+						{
+							Result->AdditionalBuffers.Emplace(AdditionalBuffer);
+						}
+					}
+				}
+			}
+
+			if (MetaDataMesh)
+			{
+				MUTABLE_CPUPROFILER_SCOPE(GetConstant_Mesh_Metadata)
+
+				if (!Result)
 				{
-					OutMesh = *Found;
+					Result = CreateMesh(MeshBudgetReserve);
+					Result->CopyFrom(*MetaDataMesh);
+				}
+				else
+				{
+					// Only in case the geoemtry has been filtered, add the meta data descriptors.
+					if (EnumHasAnyFlags(MeshContentRange.ContentFlags, EMeshContentFlags::GeometryData) &&
+						!EnumHasAnyFlags(FilterContentFlags, EMeshContentFlags::GeometryData))
+					{
+						check(MetaDataMesh->VertexBuffers.IsDescriptor());
+						check(MetaDataMesh->IndexBuffers.IsDescriptor());
+
+						Result->VertexBuffers = MetaDataMesh->VertexBuffers;
+						Result->IndexBuffers = MetaDataMesh->IndexBuffers;
+						Result->Surfaces = MetaDataMesh->Surfaces;
+					}
+
+					Result->Tags = MetaDataMesh->Tags; 
+					Result->SkeletonIDs = MetaDataMesh->SkeletonIDs; 
+					Result->StreamedResources = MetaDataMesh->StreamedResources; 
 				}
 			}
+
+			Result->MeshIDPrefix = MeshContentRange.MeshIDPrefix;
+
+			check(ConstantSkeletons.Num() > SkeletonConstantIndex);
+			if (ConstantSkeletons.IsValidIndex(SkeletonConstantIndex))
+			{
+				Result->Skeleton = ConstantSkeletons[SkeletonConstantIndex];
+			}
+
+			check(ConstantPhysicsBodies.Num() > PhysicsConstantIndex);
+			if (ConstantPhysicsBodies.IsValidIndex(PhysicsConstantIndex))
+			{
+				Result->PhysicsBody = ConstantPhysicsBodies[PhysicsConstantIndex];
+			}
+
+			OutMesh = Result;
 		}
 
 		void GetExtensionDataConstant(int32 ConstantIndex, TSharedPtr<const FExtensionData>& Result) const
@@ -785,7 +988,7 @@ namespace mu
     {
     public:
 
-        FProgram Program;
+		mu::FProgram Program;
 
     	void UnloadRoms();
 
diff --git a/Engine/Plugins/Mutable/Source/MutableRuntime/Internal/MuR/OpMeshMerge.h b/Engine/Plugins/Mutable/Source/MutableRuntime/Internal/MuR/OpMeshMerge.h
--- a/Engine/Plugins/Mutable/Source/MutableRuntime/Internal/MuR/OpMeshMerge.h
+++ b/Engine/Plugins/Mutable/Source/MutableRuntime/Internal/MuR/OpMeshMerge.h
@@ -39,6 +39,12 @@ namespace mu
 
 			const int32 FirstCount = pFirst->GetIndexBuffers().GetElementCount();
 			const int32 SecondCount = pSecond->GetIndexBuffers().GetElementCount();
+
+			if (pFirst->IndexBuffers.IsDescriptor() || pSecond->IndexBuffers.IsDescriptor())
+			{
+				EnumAddFlags(Result->IndexBuffers.Flags, EMeshBufferSetFlags::IsDescriptor);
+			}
+			
 			Result->GetIndexBuffers().SetElementCount(FirstCount + SecondCount);
 
 			check(pFirst->GetIndexBuffers().GetBufferCount() <= 1);
@@ -96,74 +102,85 @@ namespace mu
 				ResultIndexBuffer.ElementSize = SecondIndexBuffer.ElementSize;
 			}
 
-			ResultIndexBuffer.Data.SetNum(ResultIndexBuffer.ElementSize * (FirstCount + SecondCount));
 
 			check(ResultIndexBuffer.Channels.Num() == 1);
 			check(ResultIndexBuffer.Channels[0].Semantic == EMeshBufferSemantic::VertexIndex);
 
-			if (!ResultIndexBuffer.Data.IsEmpty())
+			if (!Result->IndexBuffers.IsDescriptor())
 			{
-				if (FirstCount)
+
+				ResultIndexBuffer.Data.SetNum(ResultIndexBuffer.ElementSize * (FirstCount + SecondCount));
+			
+				if (!ResultIndexBuffer.Data.IsEmpty())
 				{
-					if (IndexBufferFormat == EMeshBufferFormat::None
-						|| IndexBufferFormat == FirstIndexBuffer.Channels[0].Format)
+					if (FirstCount)
 					{
-						FMemory::Memcpy(&ResultIndexBuffer.Data[0],
-							&FirstIndexBuffer.Data[0],
-							FirstIndexBuffer.ElementSize * FirstCount);
-					}
-					else
-					{
-						// Conversion required
-						const uint8_t* pSource = &FirstIndexBuffer.Data[0];
-						uint8_t* pDest = &ResultIndexBuffer.Data[0];
-						switch (IndexBufferFormat)
+						if (IndexBufferFormat == EMeshBufferFormat::None
+							|| IndexBufferFormat == FirstIndexBuffer.Channels[0].Format)
 						{
-						case EMeshBufferFormat::UInt32:
+							FMemory::Memcpy(&ResultIndexBuffer.Data[0],
+								&FirstIndexBuffer.Data[0],
+								FirstIndexBuffer.ElementSize * FirstCount);
+						}
+						else
 						{
-							switch (FirstIndexBuffer.Channels[0].Format)
+							// Conversion required
+							const uint8_t* pSource = &FirstIndexBuffer.Data[0];
+							uint8_t* pDest = &ResultIndexBuffer.Data[0];
+							switch (IndexBufferFormat)
 							{
-							case EMeshBufferFormat::UInt16:
+							case EMeshBufferFormat::UInt32:
 							{
-								for (int32 v = 0; v < FirstCount; ++v)
+								switch (FirstIndexBuffer.Channels[0].Format)
 								{
-									*(uint32_t*)pDest = *(const uint16*)pSource;
-									pSource += FirstIndexBuffer.ElementSize;
-									pDest += ResultIndexBuffer.ElementSize;
+								case EMeshBufferFormat::UInt16:
+								{
+									for (int32 v = 0; v < FirstCount; ++v)
+									{
+										*(uint32_t*)pDest = *(const uint16*)pSource;
+										pSource += FirstIndexBuffer.ElementSize;
+										pDest += ResultIndexBuffer.ElementSize;
+									}
+									break;
 								}
-								break;
-							}
 
-							case EMeshBufferFormat::UInt8:
-							{
-								for (int32 v = 0; v < FirstCount; ++v)
+								case EMeshBufferFormat::UInt8:
 								{
-									*(uint32_t*)pDest = *(const uint8_t*)pSource;
-									pSource += FirstIndexBuffer.ElementSize;
-									pDest += ResultIndexBuffer.ElementSize;
+									for (int32 v = 0; v < FirstCount; ++v)
+									{
+										*(uint32_t*)pDest = *(const uint8_t*)pSource;
+										pSource += FirstIndexBuffer.ElementSize;
+										pDest += ResultIndexBuffer.ElementSize;
+									}
+									break;
 								}
-								break;
-							}
 
-							default:
-								checkf(false, TEXT("Format not supported."));
+								default:
+									checkf(false, TEXT("Format not supported."));
+									break;
+								}
 								break;
 							}
-							break;
-						}
 
-						case EMeshBufferFormat::UInt16:
-						{
-							switch (FirstIndexBuffer.Channels[0].Format)
+							case EMeshBufferFormat::UInt16:
 							{
+								switch (FirstIndexBuffer.Channels[0].Format)
+								{
 
-							case EMeshBufferFormat::UInt8:
-							{
-								for (int32 v = 0; v < FirstCount; ++v)
+								case EMeshBufferFormat::UInt8:
 								{
-									*(uint16*)pDest = *(const uint8_t*)pSource;
-									pSource += FirstIndexBuffer.ElementSize;
-									pDest += ResultIndexBuffer.ElementSize;
+									for (int32 v = 0; v < FirstCount; ++v)
+									{
+										*(uint16*)pDest = *(const uint8_t*)pSource;
+										pSource += FirstIndexBuffer.ElementSize;
+										pDest += ResultIndexBuffer.ElementSize;
+									}
+									break;
+								}
+
+								default:
+									checkf(false, TEXT("Format not supported."));
+									break;
 								}
 								break;
 							}
@@ -172,85 +189,35 @@ namespace mu
 								checkf(false, TEXT("Format not supported."));
 								break;
 							}
-							break;
-						}
-
-						default:
-							checkf(false, TEXT("Format not supported."));
-							break;
 						}
 					}
-				}
-
-				if (SecondCount)
-				{
-					const uint8_t* pSource = &SecondIndexBuffer.Data[0];
-					uint8_t* pDest = &ResultIndexBuffer.Data[ResultIndexBuffer.ElementSize * FirstCount];
 
-					uint32_t firstVertexCount = pFirst->GetVertexBuffers().GetElementCount();
-
-					if (IndexBufferFormat == EMeshBufferFormat::None
-						|| IndexBufferFormat == SecondIndexBuffer.Channels[0].Format)
+					if (SecondCount)
 					{
-						switch (SecondIndexBuffer.Channels[0].Format)
-						{
-						case EMeshBufferFormat::Int32:
-						case EMeshBufferFormat::UInt32:
-						case EMeshBufferFormat::NInt32:
-						case EMeshBufferFormat::NUInt32:
-						{
-							for (int32 v = 0; v < SecondCount; ++v)
-							{
-								*(uint32_t*)pDest = firstVertexCount + *(const uint32_t*)pSource;
-								pSource += SecondIndexBuffer.ElementSize;
-								pDest += ResultIndexBuffer.ElementSize;
-							}
-							break;
-						}
+						const uint8_t* pSource = &SecondIndexBuffer.Data[0];
+						uint8_t* pDest = &ResultIndexBuffer.Data[ResultIndexBuffer.ElementSize * FirstCount];
 
-						case EMeshBufferFormat::Int16:
-						case EMeshBufferFormat::UInt16:
-						case EMeshBufferFormat::NInt16:
-						case EMeshBufferFormat::NUInt16:
-						{
-							for (int32 v = 0; v < SecondCount; ++v)
-							{
-								*(uint16*)pDest = uint16(firstVertexCount) + *(const uint16*)pSource;
-								pSource += SecondIndexBuffer.ElementSize;
-								pDest += ResultIndexBuffer.ElementSize;
-							}
-							break;
-						}
+						uint32_t firstVertexCount = pFirst->GetVertexBuffers().GetElementCount();
 
-						case EMeshBufferFormat::Int8:
-						case EMeshBufferFormat::UInt8:
-						case EMeshBufferFormat::NInt8:
-						case EMeshBufferFormat::NUInt8:
+						if (IndexBufferFormat == EMeshBufferFormat::None
+							|| IndexBufferFormat == SecondIndexBuffer.Channels[0].Format)
 						{
-							for (int32 v = 0; v < SecondCount; ++v)
+							switch (SecondIndexBuffer.Channels[0].Format)
 							{
-								*(uint8_t*)pDest = uint8_t(firstVertexCount) + *(const uint8_t*)pSource;
-								pSource += SecondIndexBuffer.ElementSize;
-								pDest += ResultIndexBuffer.ElementSize;
+							case EMeshBufferFormat::Int32:
+							case EMeshBufferFormat::UInt32:
+							case EMeshBufferFormat::NInt32:
+							case EMeshBufferFormat::NUInt32:
+							{
+								for (int32 v = 0; v < SecondCount; ++v)
+								{
+									*(uint32_t*)pDest = firstVertexCount + *(const uint32_t*)pSource;
+									pSource += SecondIndexBuffer.ElementSize;
+									pDest += ResultIndexBuffer.ElementSize;
+								}
+								break;
 							}
-							break;
-						}
-
-						default:
-							checkf(false, TEXT("Format not supported."));
-							break;
-						}
-					}
-					else
-					{
-						// Format conversion required
-						switch (IndexBufferFormat)
-						{
 
-						case EMeshBufferFormat::UInt32:
-						{
-							switch (SecondIndexBuffer.Channels[0].Format)
-							{
 							case EMeshBufferFormat::Int16:
 							case EMeshBufferFormat::UInt16:
 							case EMeshBufferFormat::NInt16:
@@ -258,7 +225,7 @@ namespace mu
 							{
 								for (int32 v = 0; v < SecondCount; ++v)
 								{
-									*(uint32_t*)pDest = uint32_t(firstVertexCount) + *(const uint16*)pSource;
+									*(uint16*)pDest = uint16(firstVertexCount) + *(const uint16*)pSource;
 									pSource += SecondIndexBuffer.ElementSize;
 									pDest += ResultIndexBuffer.ElementSize;
 								}
@@ -272,7 +239,7 @@ namespace mu
 							{
 								for (int32 v = 0; v < SecondCount; ++v)
 								{
-									*(uint32_t*)pDest = uint32_t(firstVertexCount) + *(const uint8_t*)pSource;
+									*(uint8_t*)pDest = uint8_t(firstVertexCount) + *(const uint8_t*)pSource;
 									pSource += SecondIndexBuffer.ElementSize;
 									pDest += ResultIndexBuffer.ElementSize;
 								}
@@ -283,40 +250,84 @@ namespace mu
 								checkf(false, TEXT("Format not supported."));
 								break;
 							}
-
-							break;
 						}
-
-						case EMeshBufferFormat::UInt16:
+						else
 						{
-							switch (SecondIndexBuffer.Channels[0].Format)
+							// Format conversion required
+							switch (IndexBufferFormat)
 							{
-							case EMeshBufferFormat::Int8:
-							case EMeshBufferFormat::UInt8:
-							case EMeshBufferFormat::NInt8:
-							case EMeshBufferFormat::NUInt8:
+
+							case EMeshBufferFormat::UInt32:
 							{
-								for (int32 v = 0; v < SecondCount; ++v)
+								switch (SecondIndexBuffer.Channels[0].Format)
 								{
-									*(uint16*)pDest = uint16(firstVertexCount) + *(const uint8_t*)pSource;
-									pSource += SecondIndexBuffer.ElementSize;
-									pDest += ResultIndexBuffer.ElementSize;
+								case EMeshBufferFormat::Int16:
+								case EMeshBufferFormat::UInt16:
+								case EMeshBufferFormat::NInt16:
+								case EMeshBufferFormat::NUInt16:
+								{
+									for (int32 v = 0; v < SecondCount; ++v)
+									{
+										*(uint32_t*)pDest = uint32_t(firstVertexCount) + *(const uint16*)pSource;
+										pSource += SecondIndexBuffer.ElementSize;
+										pDest += ResultIndexBuffer.ElementSize;
+									}
+									break;
+								}
+
+								case EMeshBufferFormat::Int8:
+								case EMeshBufferFormat::UInt8:
+								case EMeshBufferFormat::NInt8:
+								case EMeshBufferFormat::NUInt8:
+								{
+									for (int32 v = 0; v < SecondCount; ++v)
+									{
+										*(uint32_t*)pDest = uint32_t(firstVertexCount) + *(const uint8_t*)pSource;
+										pSource += SecondIndexBuffer.ElementSize;
+										pDest += ResultIndexBuffer.ElementSize;
+									}
+									break;
 								}
+
+								default:
+									checkf(false, TEXT("Format not supported."));
+									break;
+								}
+
 								break;
 							}
 
-							default:
-								checkf(false, TEXT("Format not supported."));
+							case EMeshBufferFormat::UInt16:
+							{
+								switch (SecondIndexBuffer.Channels[0].Format)
+								{
+								case EMeshBufferFormat::Int8:
+								case EMeshBufferFormat::UInt8:
+								case EMeshBufferFormat::NInt8:
+								case EMeshBufferFormat::NUInt8:
+								{
+									for (int32 v = 0; v < SecondCount; ++v)
+									{
+										*(uint16*)pDest = uint16(firstVertexCount) + *(const uint8_t*)pSource;
+										pSource += SecondIndexBuffer.ElementSize;
+										pDest += ResultIndexBuffer.ElementSize;
+									}
+									break;
+								}
+
+								default:
+									checkf(false, TEXT("Format not supported."));
+									break;
+								}
+
 								break;
 							}
 
-							break;
-						}
-
-						default:
-							checkf(false, TEXT("Format not supported."));
-							break;
+							default:
+								checkf(false, TEXT("Format not supported."));
+								break;
 
+							}
 						}
 					}
 				}
@@ -534,6 +545,7 @@ namespace mu
 
 		// Pose
 		//---------------------------
+		if (Result->GetSkeleton())
 		{
 			MUTABLE_CPUPROFILER_SCOPE(Pose);
 
@@ -913,12 +925,17 @@ namespace mu
 
 			// Step 2: Fill the result buffers
 			//-----------------------------------------------------------------------
+			if (pFirst->VertexBuffers.IsDescriptor() || pSecond->VertexBuffers.IsDescriptor())
+			{
+				EnumAddFlags(Result->VertexBuffers.Flags, EMeshBufferSetFlags::IsDescriptor);
+			}
+			Result->VertexBuffers.SetElementCount(FirstVertexCount + SecondVertexCount);
+			
+			if (!Result->VertexBuffers.IsDescriptor())
 			{
 				MUTABLE_CPUPROFILER_SCOPE(ResultFill);
-
+				
 				// We have the final result vertex buffers structure: allocate the memory for them.
-				Result->VertexBuffers.SetElementCount(FirstVertexCount + SecondVertexCount);
-
 				int32 ResultBufferCount = Result->GetVertexBuffers().GetBufferCount();
 				for (int32 ResultBufferIndex = 0; ResultBufferIndex < ResultBufferCount; ++ResultBufferIndex)
 				{
diff --git a/Engine/Plugins/Mutable/Source/MutableRuntime/Internal/MuR/Operations.h b/Engine/Plugins/Mutable/Source/MutableRuntime/Internal/MuR/Operations.h
--- a/Engine/Plugins/Mutable/Source/MutableRuntime/Internal/MuR/Operations.h
+++ b/Engine/Plugins/Mutable/Source/MutableRuntime/Internal/MuR/Operations.h
@@ -405,15 +405,15 @@ namespace mu
         struct MeshConstantArgs
         {
             // Index of the mesh in the mesh constant array
-            ADDRESS value;
+            ADDRESS Value;
 
             // If not negative, index of the skeleton to set to the mesh from the skeleton
             // constant array.
-            int32 skeleton;
+            int32 Skeleton;
 
             // If not negative, index of the physics body to set to the mesh from the physics body
             // constant array.
-			int32 physicsBody;
+			int32 PhysicsBody;
         };
 
 		struct ParameterArgs
diff --git a/Engine/Plugins/Mutable/Source/MutableRuntime/Internal/MuR/SystemPrivate.h b/Engine/Plugins/Mutable/Source/MutableRuntime/Internal/MuR/SystemPrivate.h
--- a/Engine/Plugins/Mutable/Source/MutableRuntime/Internal/MuR/SystemPrivate.h
+++ b/Engine/Plugins/Mutable/Source/MutableRuntime/Internal/MuR/SystemPrivate.h
@@ -52,9 +52,8 @@ namespace mu
 #endif
 	}
 
-	constexpr uint64 AllParametersMask = TNumericLimits<uint64> ::Max();
-
-
+	constexpr uint64 AllParametersMask = TNumericLimits<uint64>::Max();
+ 
 	/** ExecutinIndex stores the location inside all ranges of the execution of a specific
 	* operation. The first integer on each pair is the dimension/range index in the program
 	* array of ranges, and the second integer is the value inside that range.
@@ -162,10 +161,9 @@ namespace mu
 
 	inline uint32 GetTypeHash(const FScheduledOp& Op)
 	{
-		return HashCombine(::GetTypeHash(Op.At), HashCombine(::GetTypeHash(Op.Stage), ::GetTypeHash(Op.ExecutionIndex)));
+		return HashCombineFast(::GetTypeHash(Op.At), HashCombineFast((uint32)Op.Stage), (uint32)Op.ExecutionIndex);
 	}
 
-
 	/** A cache address is the operation plus the context of execution(iteration indices, etc...). */
 	struct FCacheAddress
 	{
@@ -175,36 +173,29 @@ namespace mu
 		uint8 ExecutionOptions = 0;
 		FScheduledOp::EType Type = FScheduledOp::EType::Full;
 
-		FCacheAddress() {}
+		FCacheAddress() 
+		{
+		}
 
-		FCacheAddress(OP::ADDRESS InAt, uint16 InExecutionIndex, uint8 InExecutionOptions)
+		FCacheAddress(OP::ADDRESS InAt, uint16 InExecutionIndex, uint8 InExecutionOptions, FScheduledOp::EType InType = FScheduledOp::EType::Full)
+			: At(InAt), ExecutionIndex(InExecutionIndex), ExecutionOptions(InExecutionOptions), Type(InType)
 		{
-			At = InAt;
-			ExecutionIndex = InExecutionIndex;
-			ExecutionOptions = InExecutionOptions;
 		}
 
 		FCacheAddress(OP::ADDRESS InAt, const FScheduledOp& Item)
+			: At(InAt), ExecutionIndex(Item.ExecutionIndex), ExecutionOptions(Item.ExecutionOptions), Type(Item.Type)
 		{
-			At = InAt;
-			ExecutionIndex = Item.ExecutionIndex;
-			ExecutionOptions = Item.ExecutionOptions;
-			Type = Item.Type;
 		}
 
 		FCacheAddress(const FScheduledOp& Item)
+			: At(Item.At), ExecutionIndex(Item.ExecutionIndex), ExecutionOptions(Item.ExecutionOptions), Type(Item.Type)
 		{
-			At = Item.At;
-			ExecutionIndex = Item.ExecutionIndex;
-			ExecutionOptions = Item.ExecutionOptions;
-			Type = Item.Type;
 		}
 	};
 
-
-	inline uint32 GetTypeHash(const FCacheAddress& a)
+	inline uint32 GetTypeHash(const FCacheAddress& Address)
 	{
-		return HashCombine(::GetTypeHash(a.At), a.ExecutionIndex);
+		return HashCombineFast(::GetTypeHash(Address.At), (uint32)Address.ExecutionIndex);
 	}
 
 
@@ -1949,8 +1940,8 @@ namespace mu
 		MUTABLERUNTIME_API float BuildScalar(const TSharedPtr<const FModel>&, const FParameters*, OP::ADDRESS) ;
 		MUTABLERUNTIME_API FVector4f BuildColour(const TSharedPtr<const FModel>&, const FParameters*, OP::ADDRESS) ;
 		MUTABLERUNTIME_API TSharedPtr<const String> BuildString(const TSharedPtr<const FModel>&, const FParameters*, OP::ADDRESS) ;
-		MUTABLERUNTIME_API TSharedPtr<const FImage> BuildImage(const TSharedPtr<const FModel>&, const FParameters*, OP::ADDRESS, int32 MipsToSkip, int32 LOD) ;
-		MUTABLERUNTIME_API TSharedPtr<const FMesh> BuildMesh(const TSharedPtr<const FModel>&, const FParameters*, OP::ADDRESS) ;
+		MUTABLERUNTIME_API TSharedPtr<const FImage> BuildImage(const TSharedPtr<const FModel>&, const FParameters*, OP::ADDRESS, int32 MipsToSkip, int32 LOD);
+		MUTABLERUNTIME_API TSharedPtr<const FMesh> BuildMesh(const TSharedPtr<const FModel>&, const FParameters*, OP::ADDRESS, EMeshContentFlags MeshContentFilter);
 		MUTABLERUNTIME_API TSharedPtr<const FInstance> BuildInstance(const TSharedPtr<const FModel>&, const FParameters*, OP::ADDRESS);
 		MUTABLERUNTIME_API TSharedPtr<const FLayout> BuildLayout(const TSharedPtr<const FModel>&, const FParameters*, OP::ADDRESS) ;
     	MUTABLERUNTIME_API FProjector BuildProjector(const TSharedPtr<const FModel>&, const FParameters*, OP::ADDRESS) ;
diff --git a/Engine/Plugins/Mutable/Source/MutableRuntime/Private/MuR/CodeRunner.cpp b/Engine/Plugins/Mutable/Source/MutableRuntime/Private/MuR/CodeRunner.cpp
--- a/Engine/Plugins/Mutable/Source/MutableRuntime/Private/MuR/CodeRunner.cpp
+++ b/Engine/Plugins/Mutable/Source/MutableRuntime/Private/MuR/CodeRunner.cpp
@@ -119,8 +119,8 @@ namespace mu
 		EExecutionStrategy InExecutionStrategy,
 		const TSharedPtr<const FModel>& InModel,
 		const FParameters* InParams,
-		OP::ADDRESS at,
-		uint32 InLodMask, uint8 executionOptions, int32 InImageLOD, FScheduledOp::EType Type )
+		OP::ADDRESS At,
+		uint32 InLodMask, uint8 ExecutionOptions, int32 InImageLOD, FScheduledOp::EType Type )
 		: Settings(InSettings)
 		, RunnerCompletionEvent(TEXT("CodeRunnerCompletioneEventInit"))
 		, ExecutionStrategy(InExecutionStrategy)
@@ -130,9 +130,9 @@ namespace mu
 		, LODMask(InLodMask)
 	{
 		MUTABLE_CPUPROFILER_SCOPE(CodeRunner_Create)
+		const FProgram& Program = Model->GetPrivate()->Program;
+		ScheduledStagePerOp.resize(Program.OpAddress.Num());
 
-		const FProgram& program = Model->GetPrivate()->Program;
-		ScheduledStagePerOp.resize(program.OpAddress.Num());
 		
 		if (Type == FScheduledOp::EType::ImageDesc)
 		{
@@ -143,7 +143,7 @@ namespace mu
 		// We will read this in the end, so make sure we keep it.
    		if (Type == FScheduledOp::EType::Full)
    		{
-			GetMemory().IncreaseHitCount(FCacheAddress(at, 0, executionOptions));
+			GetMemory().IncreaseHitCount(FCacheAddress(At, 0, ExecutionOptions));
 		}
     
 		// Start with a completed Event. This is checked at StartRun() to make sure StartRun is not called while there is 
@@ -153,11 +153,11 @@ namespace mu
 		ImageLOD = InImageLOD;
 	
 		// Push the root operation
-		FScheduledOp rootOp;
-		rootOp.At = at;
-		rootOp.ExecutionOptions = executionOptions;
-		rootOp.Type = Type;
-		AddOp(rootOp);
+		FScheduledOp RootOp;
+		RootOp.At = At;
+		RootOp.ExecutionOptions = ExecutionOptions;
+		RootOp.Type = Type;
+		AddOp(RootOp);
 	}
 
 
@@ -978,7 +978,7 @@ namespace mu
 
 
     //---------------------------------------------------------------------------------------------
-    bool CodeRunner::RunCode_ConstantResource(const FScheduledOp& Item, const FModel* InModel )
+    bool CodeRunner::RunCode_ConstantResource(const FScheduledOp& Item, const FModel* InModel)
     {
 		MUTABLE_CPUPROFILER_SCOPE(RunCode_Constant);
 
@@ -991,34 +991,21 @@ namespace mu
         case EOpType::ME_CONSTANT:
         {
 			OP::MeshConstantArgs Args = Program.GetOpArgs<OP::MeshConstantArgs>(Item.At);
+			EMeshContentFlags MeshContentFlags = (EMeshContentFlags)Item.ExecutionOptions;
 
-            FConstantResourceIndex ResourceIndex = *reinterpret_cast<const FConstantResourceIndex*>(&Args.value);
-
-            TSharedPtr<const FMesh> SourceConst;
-			Program.GetConstant(ResourceIndex, SourceConst);
-
-			check(SourceConst);
-			TSharedPtr<FMesh> Source = CreateMesh(SourceConst->GetDataSize());
-			Source->CopyFrom(*SourceConst);
-
-            // Set the separate skeleton if necessary
-            if (Args.skeleton >= 0)
-            {
-                check(Program.ConstantSkeletons.Num() > Args.skeleton);
-                TSharedPtr<const FSkeleton> pSkeleton = Program.ConstantSkeletons[Args.skeleton];
-                Source->SetSkeleton(pSkeleton);
-            }
-
-			if (Args.physicsBody >= 0)
+			TSharedPtr<const FMesh> Source;
+			Program.GetConstant(Args.Value, Args.Skeleton, Args.PhysicsBody, Source, MeshContentFlags, 
+					[this](int32 BudgetReserve){ return CreateMesh(BudgetReserve); });
+			
+			if (!Source)
 			{
-                check(Program.ConstantPhysicsBodies.Num() > Args.physicsBody);
-                TSharedPtr<const FPhysicsBody> pPhysicsBody = Program.ConstantPhysicsBodies[Args.physicsBody];
-                Source->SetPhysicsBody(pPhysicsBody);
+				return false;
 			}
 
             StoreMesh(Item, Source);
+
 			//UE_LOG(LogMutableCore, Log, TEXT("Set mesh constant %d."), Item.At);
-            break;
+			break;
         }
 
         case EOpType::IM_CONSTANT:
@@ -1071,11 +1058,12 @@ namespace mu
 		return true;
     }
 
-    //---------------------------------------------------------------------------------------------
 	void CodeRunner::RunCode_Mesh(const FScheduledOp& Item, const FModel* InModel )
     {
 		MUTABLE_CPUPROFILER_SCOPE(RunCode_Mesh);
 
+		check((EMeshContentFlags)Item.ExecutionOptions != EMeshContentFlags::None); 
+
 		const FProgram& Program = Model->GetPrivate()->Program;
 
 		EOpType type = Program.GetOpType(Item.At);
@@ -1118,34 +1106,49 @@ namespace mu
 			{
 			case 0:
 			{
-				AddOp(FScheduledOp(Item.At, Item, 1),
-					FScheduledOp(Args.Mesh, Item),
-					FScheduledOp(Args.Layout, Item));
+				if (!EnumHasAnyFlags(EMeshContentFlags(Item.ExecutionOptions), EMeshContentFlags::GeometryData))
+				{
+					AddOp(FScheduledOp(Item.At, Item, 1),
+						FScheduledOp(Args.Mesh, Item));
+				}
+				else
+				{
+					AddOp(FScheduledOp(Item.At, Item, 1),
+						FScheduledOp(Args.Mesh, Item),
+						FScheduledOp(Args.Layout, Item));
+				}
 				break;
 			}
             case 1:
             {
            		MUTABLE_CPUPROFILER_SCOPE(ME_APPLYLAYOUT)
             		
-                TSharedPtr<const FMesh> Base = LoadMesh(FCacheAddress(Args.Mesh, Item));
+				if (!EnumHasAnyFlags(EMeshContentFlags(Item.ExecutionOptions), EMeshContentFlags::GeometryData))
+				{
+					TSharedPtr<const FMesh> Base = LoadMesh(FCacheAddress(Args.Mesh, Item));
+					StoreMesh(Item, Base);
+				}
+				else
+				{
+					TSharedPtr<const FMesh> Base = LoadMesh(FCacheAddress(Args.Mesh, Item));
 
-                if (Base)
-                {
-					TSharedPtr<FMesh> Result = CloneOrTakeOver(Base);
+					if (Base)
+					{
+						TSharedPtr<FMesh> Result = CloneOrTakeOver(Base);
 
-                    TSharedPtr<const FLayout> Layout = LoadLayout(FCacheAddress(Args.Layout, Item));
-                    int32 texCoordsSet = Args.Channel;
+						TSharedPtr<const FLayout> Layout = LoadLayout(FCacheAddress(Args.Layout, Item));
+						int32 texCoordsSet = Args.Channel;
 
-                    MeshApplyLayout(Result.Get(), Layout.Get(), texCoordsSet);
-					
-					StoreMesh(Item, Result);
-                }
-				else
-				{
-					StoreMesh(Item, nullptr);
+						MeshApplyLayout(Result.Get(), Layout.Get(), texCoordsSet);
+						
+						StoreMesh(Item, Result);
+					}
+					else
+					{
+						StoreMesh(Item, nullptr);
+					}
 				}
 
-
                 break;
             }
 
@@ -1172,15 +1175,23 @@ namespace mu
             {
             case 0:
 			{
-				if (BaseAt && TargetAt)
+				if (!EnumHasAnyFlags(EMeshContentFlags(Item.ExecutionOptions), EMeshContentFlags::GeometryData))
 				{
-					AddOp(FScheduledOp(Item.At, Item, 1),
-						FScheduledOp(BaseAt, Item),
-						FScheduledOp(TargetAt, Item));
+                	TSharedPtr<const FMesh> Result = CreateMesh();
+					StoreMesh(Item, Result);
 				}
 				else
 				{
-					StoreMesh(Item, nullptr);
+					if (BaseAt && TargetAt)
+					{
+						AddOp(FScheduledOp(Item.At, Item, 1),
+							FScheduledOp(BaseAt, Item),
+							FScheduledOp(TargetAt, Item));
+					}
+					else
+					{
+						StoreMesh(Item, nullptr);
+					}
 				}
 				break;
 			}
@@ -1188,8 +1199,10 @@ namespace mu
             {
        	        MUTABLE_CPUPROFILER_SCOPE(ME_DIFFERENCE)
 
-                TSharedPtr<const FMesh> Base = LoadMesh(FCacheAddress(BaseAt,Item));
-                TSharedPtr<const FMesh> pTarget = LoadMesh(FCacheAddress(TargetAt,Item));
+				check(EnumHasAnyFlags(EMeshContentFlags(Item.ExecutionOptions), EMeshContentFlags::GeometryData))
+
+				TSharedPtr<const FMesh> Base = LoadMesh(FCacheAddress(BaseAt,Item));
+				TSharedPtr<const FMesh> pTarget = LoadMesh(FCacheAddress(TargetAt,Item));
 
 				TArray<EMeshBufferSemantic, TInlineAllocator<8>> Semantics;
 				TArray<int32, TInlineAllocator<8>> SemanticIndices;
@@ -1202,8 +1215,8 @@ namespace mu
 				FMemory::Memcpy(&NumChannels, data, sizeof(uint8)); 
 				data += sizeof(uint8);
 
-                for (uint8 i = 0; i < NumChannels; ++i)
-                {
+				for (uint8 i = 0; i < NumChannels; ++i)
+				{
 					uint8 Semantic = 0;
 					FMemory::Memcpy(&Semantic, data, sizeof(uint8)); 
 					data += sizeof(uint8);
@@ -1214,18 +1227,19 @@ namespace mu
 
 					Semantics.Add(EMeshBufferSemantic(Semantic));
 					SemanticIndices.Add(SemanticIndex);
-                }
+				}
 
 				TSharedPtr<FMesh> Result = CreateMesh();
 				bool bOutSuccess = false;
-                MeshDifference(Result.Get(), Base.Get(), pTarget.Get(),
-                               NumChannels, Semantics.GetData(), SemanticIndices.GetData(),
-                               bIgnoreTextureCoords != 0, bOutSuccess);
+				MeshDifference(Result.Get(), Base.Get(), pTarget.Get(),
+							   NumChannels, Semantics.GetData(), SemanticIndices.GetData(),
+							   bIgnoreTextureCoords != 0, bOutSuccess);
 				Release(Base);
 				Release(pTarget);
 
-                StoreMesh(Item, Result);
-                break;
+				StoreMesh(Item, Result);
+
+				break;
             }
 
             default:
@@ -1254,69 +1268,87 @@ namespace mu
 			switch (Item.Stage)
             {
             case 0:
-                if (BaseAt)
-                {
-                    AddOp(FScheduledOp(Item.At, Item, 1),
-                           FScheduledOp(FactorAt, Item));
-                }
-                else
-                {
-                    StoreMesh(Item, nullptr);
-                }
+			{
+				if (BaseAt)
+				{
+					if (!EnumHasAnyFlags(EMeshContentFlags(Item.ExecutionOptions), EMeshContentFlags::GeometryData))
+					{
+						AddOp(FScheduledOp(Item.At, Item, 1),
+							FScheduledOp(BaseAt, Item));
+					}
+					else
+					{
+						AddOp(FScheduledOp(Item.At, Item, 1),
+							FScheduledOp(FactorAt, Item));
+					}
+				}
+				else
+				{
+					StoreMesh(Item, nullptr);
+				}
+				
                 break;
-
+			}
             case 1:
             {
                 MUTABLE_CPUPROFILER_SCOPE(ME_MORPH_1)
 
-                float Factor = LoadScalar(FCacheAddress(FactorAt, Item));
-
-                // Factor goes from -1 to 1 across all targets. [0 - 1] represents positive morphs, while [-1, 0) represent negative morphs.
-				Factor = FMath::Clamp(Factor, -1.0f, 1.0f); // Is the factor not in range [-1, 1], it will index a non existing morph.
+				if (!EnumHasAnyFlags(EMeshContentFlags(Item.ExecutionOptions), EMeshContentFlags::GeometryData))
+				{
+                	TSharedPtr<const FMesh> Base = LoadMesh(FCacheAddress(BaseAt, Item));
+					StoreMesh(Item, Base);
+				}
+				else	
+				{
+					float Factor = LoadScalar(FCacheAddress(FactorAt, Item));
 
-                FScheduledOpData OpHeapData;
-				OpHeapData.Interpolate.Bifactor = Factor;
-				uint32 DataAddress = uint32(HeapData.Add(OpHeapData));
+					// Factor goes from -1 to 1 across all targets. [0 - 1] represents positive morphs, while [-1, 0) represent negative morphs.
+					Factor = FMath::Clamp(Factor, -1.0f, 1.0f); // Is the factor not in range [-1, 1], it will index a non existing morph.
 
-                // No morph
-				if (FMath::IsNearlyZero(Factor))
-                {                        
-                    AddOp(FScheduledOp(Item.At, Item, 2, DataAddress),
-						FScheduledOp(BaseAt, Item));
-                }
-                // The Morph, partial or full
-                else
-                {
-                    // We will need the base again
-                    AddOp(FScheduledOp(Item.At, Item, 2, DataAddress),
-						FScheduledOp(BaseAt, Item),
-						FScheduledOp(TargetAt, Item));
-                }
+					FScheduledOpData OpHeapData;
+					OpHeapData.Interpolate.Bifactor = Factor;
+					uint32 DataAddress = uint32(HeapData.Add(OpHeapData));
 
+					// No morph
+					if (FMath::IsNearlyZero(Factor))
+					{                        
+						AddOp(FScheduledOp(Item.At, Item, 2, DataAddress),
+							FScheduledOp(BaseAt, Item));
+					}
+					// The Morph, partial or full
+					else
+					{
+						// We will need the base again
+						AddOp(FScheduledOp(Item.At, Item, 2, DataAddress),
+							FScheduledOp(BaseAt, Item),
+							FScheduledOp(TargetAt, Item));
+					}
+				}
                 break;
             }
 
             case 2:
             {
        		    MUTABLE_CPUPROFILER_SCOPE(ME_MORPH_2)
+				check(EnumHasAnyFlags(EMeshContentFlags(Item.ExecutionOptions), EMeshContentFlags::GeometryData))
+				
+				TSharedPtr<const FMesh> BaseMesh = LoadMesh(FCacheAddress(BaseAt, Item));
 
-                TSharedPtr<const FMesh> BaseMesh = LoadMesh(FCacheAddress(BaseAt, Item));
+				// Factor from 0 to 1 between the two targets
+				const FScheduledOpData& OpHeapData = HeapData[Item.CustomState];
+				float Factor = OpHeapData.Interpolate.Bifactor;
 
-                // Factor from 0 to 1 between the two targets
-                const FScheduledOpData& OpHeapData = HeapData[Item.CustomState];
-                float Factor = OpHeapData.Interpolate.Bifactor;
-
-                if (BaseMesh)
-                {
+				if (BaseMesh)
+				{
 					// No morph
 					if (FMath::IsNearlyZero(Factor))
-                    {
+					{
 						StoreMesh(Item, BaseMesh);
-                    }
+					}
 					// The Morph, partial or full
-                    else 
-                    {
-                        TSharedPtr<const FMesh> MorphMesh = LoadMesh(FCacheAddress(TargetAt, Item));
+					else 
+					{
+						TSharedPtr<const FMesh> MorphMesh = LoadMesh(FCacheAddress(TargetAt, Item));
 						
 						if (MorphMesh)
 						{
@@ -1330,13 +1362,12 @@ namespace mu
 						{
 							StoreMesh(Item, BaseMesh);
 						}
-                    }
-                }
+					}
+				}
 				else
 				{
 					StoreMesh(Item, nullptr);
 				}
-
                 break;
             }
 
@@ -1434,45 +1465,59 @@ namespace mu
             {
             case 0:
 			{
-				AddOp(FScheduledOp(Item.At, Item, 1),
-					FScheduledOp(Args.source, Item),
-					FScheduledOp(Args.clip, Item));
+				if (!EnumHasAnyFlags(EMeshContentFlags(Item.ExecutionOptions), EMeshContentFlags::GeometryData))
+				{
+					AddOp(FScheduledOp(Item.At, Item, 1),
+						FScheduledOp(Args.source, Item));
+				}
+				else
+				{
+					AddOp(FScheduledOp(Item.At, Item, 1),
+						FScheduledOp(Args.source, Item),
+						FScheduledOp(Args.clip, Item));
+				}
 				break;
 			}
             case 1:
             {
             	MUTABLE_CPUPROFILER_SCOPE(ME_MASKCLIPMESH_1)
-            		
-                TSharedPtr<const FMesh> Source = LoadMesh(FCacheAddress(Args.source, Item));
-                TSharedPtr<const FMesh> pClip = LoadMesh(FCacheAddress(Args.clip, Item));
-
-                // Only if both are valid.
-                if (Source.Get() && pClip.Get())
-                {
-					TSharedPtr<FMesh> Result = CreateMesh();
+				if (!EnumHasAnyFlags(EMeshContentFlags(Item.ExecutionOptions), EMeshContentFlags::GeometryData))
+				{
+					TSharedPtr<const FMesh> Source = LoadMesh(FCacheAddress(Args.source, Item));
+					StoreMesh(Item, Source);
+				}
+				else
+				{
+					TSharedPtr<const FMesh> Source = LoadMesh(FCacheAddress(Args.source, Item));
+					TSharedPtr<const FMesh> pClip = LoadMesh(FCacheAddress(Args.clip, Item));
 
-					bool bOutSuccess = false;
-                    MeshMaskClipMesh(Result.Get(), Source.Get(), pClip.Get(), bOutSuccess);
-					
-					Release(Source);
-					Release(pClip);
-					if (!bOutSuccess)
+					// Only if both are valid.
+					if (Source.Get() && pClip.Get())
 					{
-						Release(Result);
-						StoreMesh(Item, nullptr);
+						TSharedPtr<FMesh> Result = CreateMesh();
+
+						bool bOutSuccess = false;
+						MeshMaskClipMesh(Result.Get(), Source.Get(), pClip.Get(), bOutSuccess);
+						
+						Release(Source);
+						Release(pClip);
+						if (!bOutSuccess)
+						{
+							Release(Result);
+							StoreMesh(Item, nullptr);
+						}
+						else
+						{
+							StoreMesh(Item, Result);
+						}
 					}
 					else
 					{
-						StoreMesh(Item, Result);
+						Release(Source);
+						Release(pClip);
+						StoreMesh(Item, nullptr);
 					}
-                }
-				else
-				{
-					Release(Source);
-					Release(pClip);
-					StoreMesh(Item, nullptr);
 				}
-
                 break;
             }
 
@@ -1490,71 +1535,86 @@ namespace mu
 			{
 			case 0:
 			{
-				AddOp(FScheduledOp(Item.At, Item, 1),
-					FScheduledOp(Args.Source, Item),
-					FScheduledOp(Args.UVSource, Item),
-					FScheduledOp(Args.MaskImage, Item),
-					FScheduledOp(Args.MaskLayout, Item));
+				if (!EnumHasAnyFlags(EMeshContentFlags(Item.ExecutionOptions), EMeshContentFlags::GeometryData))
+				{
+					AddOp(FScheduledOp(Item.At, Item, 1),
+						FScheduledOp(Args.Source, Item));
+				}
+				else
+				{
+					AddOp(FScheduledOp(Item.At, Item, 1),
+						FScheduledOp(Args.Source, Item),
+						FScheduledOp(Args.UVSource, Item),
+						FScheduledOp(Args.MaskImage, Item),
+						FScheduledOp(Args.MaskLayout, Item));
+				}
 				break;
 			}
 			case 1:
 			{
 				MUTABLE_CPUPROFILER_SCOPE(ME_MASKCLIPUVMASK_1);
 
-				TSharedPtr<const FMesh> Source = LoadMesh(FCacheAddress(Args.Source, Item));
-				TSharedPtr<const FMesh> UVSource = LoadMesh(FCacheAddress(Args.UVSource, Item));
-				TSharedPtr<const FImage> MaskImage = LoadImage(FCacheAddress(Args.MaskImage, Item));
-				TSharedPtr<const FLayout> MaskLayout = LoadLayout(FCacheAddress(Args.MaskLayout, Item));
-
-				// Only if both are valid.
-				if (Source.Get() && MaskImage.Get())
+				if (!EnumHasAnyFlags(EMeshContentFlags(Item.ExecutionOptions), EMeshContentFlags::GeometryData))
 				{
-					TSharedPtr<FMesh> Result = CreateMesh();
-
-					bool bOutSuccess = false;
-					MakeMeshMaskFromUVMask(Result.Get(), Source.Get(), UVSource.Get(), MaskImage.Get(), Args.LayoutIndex, bOutSuccess);
+					TSharedPtr<const FMesh> Source = LoadMesh(FCacheAddress(Args.Source, Item));
+					StoreMesh(Item, Source);
+				}
+				else
+				{
+					TSharedPtr<const FMesh> Source = LoadMesh(FCacheAddress(Args.Source, Item));
+					TSharedPtr<const FMesh> UVSource = LoadMesh(FCacheAddress(Args.UVSource, Item));
+					TSharedPtr<const FImage> MaskImage = LoadImage(FCacheAddress(Args.MaskImage, Item));
+					TSharedPtr<const FLayout> MaskLayout = LoadLayout(FCacheAddress(Args.MaskLayout, Item));
 
-					Release(Source);
-					Release(UVSource);
-					Release(MaskImage);
-					if (!bOutSuccess)
+					// Only if both are valid.
+					if (Source.Get() && MaskImage.Get())
 					{
-						Release(Result);
-						StoreMesh(Item, nullptr);
-					}
-					else
-					{
-						StoreMesh(Item, Result);
+						TSharedPtr<FMesh> Result = CreateMesh();
+
+						bool bOutSuccess = false;
+						MakeMeshMaskFromUVMask(Result.Get(), Source.Get(), UVSource.Get(), MaskImage.Get(), Args.LayoutIndex, bOutSuccess);
+
+						Release(Source);
+						Release(UVSource);
+						Release(MaskImage);
+						if (!bOutSuccess)
+						{
+							Release(Result);
+							StoreMesh(Item, nullptr);
+						}
+						else
+						{
+							StoreMesh(Item, Result);
+						}
 					}
-				}
-				else if (Source.Get() && MaskLayout.Get())
-				{					
-					TSharedPtr<FMesh> Result = CreateMesh();
+					else if (Source.Get() && MaskLayout.Get())
+					{					
+						TSharedPtr<FMesh> Result = CreateMesh();
 
-					bool bOutSuccess = false;
-					MakeMeshMaskFromLayout(Result.Get(), Source.Get(), UVSource.Get(), MaskLayout.Get(), Args.LayoutIndex, bOutSuccess);
+						bool bOutSuccess = false;
+						MakeMeshMaskFromLayout(Result.Get(), Source.Get(), UVSource.Get(), MaskLayout.Get(), Args.LayoutIndex, bOutSuccess);
 
-					Release(Source);
-					Release(UVSource);
-					Release(MaskImage);
-					if (!bOutSuccess)
-					{
-						Release(Result);
-						StoreMesh(Item, nullptr);
+						Release(Source);
+						Release(UVSource);
+						Release(MaskImage);
+						if (!bOutSuccess)
+						{
+							Release(Result);
+							StoreMesh(Item, nullptr);
+						}
+						else
+						{
+							StoreMesh(Item, Result);
+						}
 					}
 					else
 					{
-						StoreMesh(Item, Result);
+						Release(Source);
+						Release(UVSource);
+						Release(MaskImage);
+						StoreMesh(Item, nullptr);
 					}
 				}
-				else
-				{
-					Release(Source);
-					Release(UVSource);
-					Release(MaskImage);
-					StoreMesh(Item, nullptr);
-				}
-
 				break;
 			}
 
@@ -1565,53 +1625,68 @@ namespace mu
 			break;
 		}
 
-        case EOpType::ME_MASKDIFF:
-        {
+		case EOpType::ME_MASKDIFF:
+		{
 			OP::MeshMaskDiffArgs Args = Program.GetOpArgs<OP::MeshMaskDiffArgs>(Item.At);
-            switch (Item.Stage)
-            {
-            case 0:
+			switch (Item.Stage)
 			{
-				AddOp(FScheduledOp(Item.At, Item, 1),
-					FScheduledOp(Args.Source, Item),
-					FScheduledOp(Args.Fragment, Item));
+			case 0:
+			{
+				if (!EnumHasAnyFlags(EMeshContentFlags(Item.ExecutionOptions), EMeshContentFlags::GeometryData))
+				{
+					AddOp(FScheduledOp(Item.At, Item, 1),
+						FScheduledOp(Args.Source, Item));
+				}
+				else
+				{
+					AddOp(FScheduledOp(Item.At, Item, 1),
+						FScheduledOp(Args.Source, Item),
+						FScheduledOp(Args.Fragment, Item));
+				}
+
 				break;
 			}
-            case 1:
-            {
-           		MUTABLE_CPUPROFILER_SCOPE(ME_MASKDIFF_1)
-            		
-                TSharedPtr<const FMesh> Source = LoadMesh(FCacheAddress(Args.Source, Item));
-                TSharedPtr<const FMesh> pClip = LoadMesh(FCacheAddress(Args.Fragment, Item));
+			case 1:
+			{
+				MUTABLE_CPUPROFILER_SCOPE(ME_MASKDIFF_1)
+				if (!EnumHasAnyFlags(EMeshContentFlags(Item.ExecutionOptions), EMeshContentFlags::GeometryData))
+				{
+					TSharedPtr<const FMesh> Source = LoadMesh(FCacheAddress(Args.Source, Item));
+					StoreMesh(Item, Source);
+				}
+				else
+				{
+					TSharedPtr<const FMesh> Source = LoadMesh(FCacheAddress(Args.Source, Item));
+					TSharedPtr<const FMesh> pClip = LoadMesh(FCacheAddress(Args.Fragment, Item));
 
-                // Only if both are valid.
-                if (Source.Get() && pClip.Get())
-                {
-					TSharedPtr<FMesh> Result = CreateMesh();
+					// Only if both are valid.
+					if (Source.Get() && pClip.Get())
+					{
+						TSharedPtr<FMesh> Result = CreateMesh();
 
-					bool bOutSuccess = false;
-                    MeshMaskDiff(Result.Get(), Source.Get(), pClip.Get(), bOutSuccess);
+						bool bOutSuccess = false;
+						MeshMaskDiff(Result.Get(), Source.Get(), pClip.Get(), bOutSuccess);
 
-					Release(Source);
-					Release(pClip);
+						Release(Source);
+						Release(pClip);
 
-					if (!bOutSuccess)
-					{
-						Release(Result);
-						StoreMesh(Item, nullptr);
+						if (!bOutSuccess)
+						{
+							Release(Result);
+							StoreMesh(Item, nullptr);
+						}
+						else
+						{
+							StoreMesh(Item, Result);
+						}
 					}
 					else
 					{
-						StoreMesh(Item, Result);
+						Release(Source);
+						Release(pClip);
+						StoreMesh(Item, nullptr);
 					}
-                }
-				else
-				{
-					Release(Source);
-					Release(pClip);
-					StoreMesh(Item, nullptr);
 				}
-
                 break;
             }
 
@@ -1643,10 +1718,9 @@ namespace mu
 			}
             case 1:
             {
-            	MUTABLE_CPUPROFILER_SCOPE(ME_FORMAT_1)
-            		
-                TSharedPtr<const FMesh> Source = LoadMesh(FCacheAddress(Args.source,Item));
-                TSharedPtr<const FMesh> Format = LoadMesh(FCacheAddress(Args.format,Item));
+				MUTABLE_CPUPROFILER_SCOPE(ME_FORMAT_1)
+				TSharedPtr<const FMesh> Source = LoadMesh(FCacheAddress(Args.source,Item));
+				TSharedPtr<const FMesh> Format = LoadMesh(FCacheAddress(Args.format,Item));
 
 				if (Source && Source->IsReference())
 				{
@@ -1747,46 +1821,54 @@ namespace mu
 				}
 				break;
 			}
-            case 1:
-            {
+			case 1:
+			{
 				MUTABLE_CPUPROFILER_SCOPE(ME_EXTRACTLAYOUTBLOCK_1)
 
-                TSharedPtr<const FMesh> Source = LoadMesh(FCacheAddress(source, Item));
-
-                // Access with memcpy necessary for unaligned memory access issues.
-                uint64 blocks[512];
-				check(blockCount< 512);
-				FMemory::Memcpy(blocks, data, sizeof(uint64)*FMath::Min(512,int32(blockCount)));
-
-				if (Source)
+				if (!EnumHasAnyFlags(EMeshContentFlags(Item.ExecutionOptions), EMeshContentFlags::GeometryData))
+				{
+					TSharedPtr<const FMesh> SourceMesh = LoadMesh(FCacheAddress(source, Item));
+					StoreMesh(Item, SourceMesh);
+				}
+				else
 				{
-					TSharedPtr<FMesh> Result = CreateMesh();
-					bool bOutSuccess;
+					TSharedPtr<const FMesh> Source = LoadMesh(FCacheAddress(source, Item));
 
-					if (blockCount > 0)
-					{
-						MeshExtractLayoutBlock(Result.Get(), Source.Get(), LayoutIndex, blockCount, blocks, bOutSuccess);
-					}
-					else
-					{
-						MeshExtractLayoutBlock(Result.Get(), Source.Get(), LayoutIndex, bOutSuccess);
-					}
+					// Access with memcpy necessary for unaligned memory access issues.
+					uint64 blocks[512];
+					check(blockCount< 512);
+					FMemory::Memcpy(blocks, data, sizeof(uint64)*FMath::Min(512,int32(blockCount)));
 
-					if (!bOutSuccess)
+					if (Source)
 					{
-						Release(Result);
-						StoreMesh(Item, Source);
+						TSharedPtr<FMesh> Result = CreateMesh();
+						bool bOutSuccess;
+
+						if (blockCount > 0)
+						{
+							MeshExtractLayoutBlock(Result.Get(), Source.Get(), LayoutIndex, blockCount, blocks, bOutSuccess);
+						}
+						else
+						{
+							MeshExtractLayoutBlock(Result.Get(), Source.Get(), LayoutIndex, bOutSuccess);
+						}
+
+						if (!bOutSuccess)
+						{
+							Release(Result);
+							StoreMesh(Item, Source);
+						}
+						else
+						{
+							Release(Source);
+							StoreMesh(Item, Result);
+						}
 					}
 					else
 					{
-						Release(Source);
-						StoreMesh(Item, Result);
+						StoreMesh(Item, nullptr);
 					}
 				}
-				else
-				{
-					StoreMesh(Item, nullptr);
-				}
                 break;
             }
 
@@ -1815,28 +1897,36 @@ namespace mu
 				}
 				break;
 			}
-            case 1:
-            {
-           		MUTABLE_CPUPROFILER_SCOPE(ME_TRANSFORM_1)
-            		
-                TSharedPtr<const FMesh> Source = LoadMesh(FCacheAddress(Args.source,Item));
-
-                const FMatrix44f& mat = Program.ConstantMatrices[Args.matrix];
-
-				TSharedPtr<FMesh> Result = CreateMesh(Source ? Source->GetDataSize() : 0);
-
-				bool bOutSuccess = false;
-                MeshTransform(Result.Get(), Source.Get(), mat, bOutSuccess);
+			case 1:
+			{
+				MUTABLE_CPUPROFILER_SCOPE(ME_TRANSFORM_1)
 
-				if (!bOutSuccess)
+				if (!EnumHasAnyFlags(EMeshContentFlags(Item.ExecutionOptions), EMeshContentFlags::GeometryData))
 				{
-					Release(Result);
+                	TSharedPtr<const FMesh> Source = LoadMesh(FCacheAddress(Args.source, Item));
 					StoreMesh(Item, Source);
 				}
 				else
 				{
-					Release(Source);
-					StoreMesh(Item, Result);
+					TSharedPtr<const FMesh> Source = LoadMesh(FCacheAddress(Args.source,Item));
+
+					const FMatrix44f& mat = Program.ConstantMatrices[Args.matrix];
+
+					TSharedPtr<FMesh> Result = CreateMesh(Source ? Source->GetDataSize() : 0);
+
+					bool bOutSuccess = false;
+					MeshTransform(Result.Get(), Source.Get(), mat, bOutSuccess);
+
+					if (!bOutSuccess)
+					{
+						Release(Result);
+						StoreMesh(Item, Source);
+					}
+					else
+					{
+						Release(Source);
+						StoreMesh(Item, Result);
+					}
 				}
                 break;
             }
@@ -1864,95 +1954,104 @@ namespace mu
 				{
 					StoreMesh(Item, nullptr);
 				}
+
 				break;
 			}
             case 1:
             {
            		MUTABLE_CPUPROFILER_SCOPE(ME_CLIPMORPHPLANE_1)
-            		
-                TSharedPtr<const FMesh> Source = LoadMesh(FCacheAddress(Args.Source, Item));
-
-                check(Args.MorphShape < (uint32)InModel->GetPrivate()->Program.ConstantShapes.Num());
 
-                // Should be an ellipse
-                const FShape& MorphShape = Program.ConstantShapes[Args.MorphShape];
+				if (!EnumHasAnyFlags(EMeshContentFlags(Item.ExecutionOptions), EMeshContentFlags::GeometryData))
+				{
+                	TSharedPtr<const FMesh> Source = LoadMesh(FCacheAddress(Args.Source, Item));
+					StoreMesh(Item, Source);
+				}
+				else
+				{
 
-                const FVector3f& Origin = MorphShape.position;
-                const FVector3f& Normal = MorphShape.up;
+					TSharedPtr<const FMesh> Source = LoadMesh(FCacheAddress(Args.Source, Item));
 
-				bool bRemoveFaceIfAllVerticesCulled = Args.FaceCullStrategy==EFaceCullStrategy::AllVerticesCulled;
+					check(Args.MorphShape < (uint32)InModel->GetPrivate()->Program.ConstantShapes.Num());
 
-                if (Args.VertexSelectionType == EClipVertexSelectionType::Shape)
-                {
-                    check(Args.VertexSelectionShapeOrBone < (uint32)InModel->GetPrivate()->Program.ConstantShapes.Num());
+					// Should be an ellipse
+					const FShape& MorphShape = Program.ConstantShapes[Args.MorphShape];
 
-                    // Should be None or an axis aligned box
-                    const FShape& SelectionShape = Program.ConstantShapes[Args.VertexSelectionShapeOrBone];
+					const FVector3f& Origin = MorphShape.position;
+					const FVector3f& Normal = MorphShape.up;
 
-					TSharedPtr<FMesh> Result = CreateMesh(Source ? Source->GetDataSize() : 0);
+					bool bRemoveFaceIfAllVerticesCulled = Args.FaceCullStrategy==EFaceCullStrategy::AllVerticesCulled;
 
-					bool bOutSuccess = false;
-					MeshClipMorphPlane(Result.Get(), Source.Get(), Origin, Normal, Args.Dist, Args.Factor, MorphShape.size[0], MorphShape.size[1], MorphShape.size[2], SelectionShape, bRemoveFaceIfAllVerticesCulled, bOutSuccess, nullptr, -1);
-					
-					if (!bOutSuccess)
+					if (Args.VertexSelectionType == EClipVertexSelectionType::Shape)
 					{
-						Release(Result);
-						StoreMesh(Item, Source);
-					}
-					else
-					{
-						Release(Source);
-						StoreMesh(Item, Result);
+						check(Args.VertexSelectionShapeOrBone < (uint32)InModel->GetPrivate()->Program.ConstantShapes.Num());
+
+						// Should be None or an axis aligned box
+						const FShape& SelectionShape = Program.ConstantShapes[Args.VertexSelectionShapeOrBone];
+
+						TSharedPtr<FMesh> Result = CreateMesh(Source ? Source->GetDataSize() : 0);
+
+						bool bOutSuccess = false;
+						MeshClipMorphPlane(Result.Get(), Source.Get(), Origin, Normal, Args.Dist, Args.Factor, MorphShape.size[0], MorphShape.size[1], MorphShape.size[2], SelectionShape, bRemoveFaceIfAllVerticesCulled, bOutSuccess, nullptr, -1);
+						
+						if (!bOutSuccess)
+						{
+							Release(Result);
+							StoreMesh(Item, Source);
+						}
+						else
+						{
+							Release(Source);
+							StoreMesh(Item, Result);
+						}
 					}
-                }
 
-				else if (Args.VertexSelectionType == EClipVertexSelectionType::BoneHierarchy)
-				{
-					FShape SelectionShape;
-					SelectionShape.type = (uint8)FShape::Type::None;
+					else if (Args.VertexSelectionType == EClipVertexSelectionType::BoneHierarchy)
+					{
+						FShape SelectionShape;
+						SelectionShape.type = (uint8)FShape::Type::None;
 
-					TSharedPtr<FMesh> Result = CreateMesh(Source->GetDataSize());
+						TSharedPtr<FMesh> Result = CreateMesh(Source->GetDataSize());
 
-					check(Args.VertexSelectionShapeOrBone <= MAX_uint32);
-					const FBoneName Bone(Args.VertexSelectionShapeOrBone);
+						check(Args.VertexSelectionShapeOrBone <= MAX_uint32);
+						const FBoneName Bone(Args.VertexSelectionShapeOrBone);
 
-					bool bOutSuccess = false;
-					MeshClipMorphPlane(Result.Get(), Source.Get(), Origin, Normal, Args.Dist, Args.Factor, MorphShape.size[0], MorphShape.size[1], MorphShape.size[2], SelectionShape, bRemoveFaceIfAllVerticesCulled, bOutSuccess, &Bone, Args.MaxBoneRadius);
+						bool bOutSuccess = false;
+						MeshClipMorphPlane(Result.Get(), Source.Get(), Origin, Normal, Args.Dist, Args.Factor, MorphShape.size[0], MorphShape.size[1], MorphShape.size[2], SelectionShape, bRemoveFaceIfAllVerticesCulled, bOutSuccess, &Bone, Args.MaxBoneRadius);
 
-					if (!bOutSuccess)
-					{
-						Release(Result);
-						StoreMesh(Item, Source);
+						if (!bOutSuccess)
+						{
+							Release(Result);
+							StoreMesh(Item, Source);
+						}
+						else
+						{
+							Release(Source);
+							StoreMesh(Item, Result);
+						}
 					}
 					else
 					{
-						Release(Source);
-						StoreMesh(Item, Result);
-					}
-                }
-                else
-                {
-                    // No vertex selection
-                    FShape SelectionShape;
-                    SelectionShape.type = (uint8)FShape::Type::None;
+						// No vertex selection
+						FShape SelectionShape;
+						SelectionShape.type = (uint8)FShape::Type::None;
 
-					TSharedPtr<FMesh> Result = CreateMesh(Source ? Source->GetDataSize() : 0);
+						TSharedPtr<FMesh> Result = CreateMesh(Source ? Source->GetDataSize() : 0);
 
-					bool bOutSuccess = false;
-					MeshClipMorphPlane(Result.Get(), Source.Get(), Origin, Normal, Args.Dist, Args.Factor, MorphShape.size[0], MorphShape.size[1], MorphShape.size[2], SelectionShape, bRemoveFaceIfAllVerticesCulled, bOutSuccess, nullptr, -1.0f);
+						bool bOutSuccess = false;
+						MeshClipMorphPlane(Result.Get(), Source.Get(), Origin, Normal, Args.Dist, Args.Factor, MorphShape.size[0], MorphShape.size[1], MorphShape.size[2], SelectionShape, bRemoveFaceIfAllVerticesCulled, bOutSuccess, nullptr, -1.0f);
 
-					if (!bOutSuccess)
-					{
-						Release(Result);
-						StoreMesh(Item, Source);
-					}
-					else
-					{
-						Release(Source);
-						StoreMesh(Item, Result);
+						if (!bOutSuccess)
+						{
+							Release(Result);
+							StoreMesh(Item, Source);
+						}
+						else
+						{
+							Release(Source);
+							StoreMesh(Item, Result);
+						}
 					}
-                }
-
+				}
                 break;
             }
 
@@ -1973,9 +2072,17 @@ namespace mu
 			{
 				if (Args.Source)
 				{
-					AddOp(FScheduledOp(Item.At, Item, 1),
-						FScheduledOp(Args.Source, Item),
-						FScheduledOp(Args.ClipMesh, Item));
+					if (!EnumHasAnyFlags(EMeshContentFlags(Item.ExecutionOptions), EMeshContentFlags::GeometryData))
+					{
+						AddOp(FScheduledOp(Item.At, Item, 1),
+							FScheduledOp(Args.Source, Item));
+					}
+					else
+					{
+						AddOp(FScheduledOp(Item.At, Item, 1),
+							FScheduledOp(Args.Source, Item),
+							FScheduledOp(Args.ClipMesh, Item));
+					}
 				}
 				else
 				{
@@ -1987,36 +2094,44 @@ namespace mu
             case 1:
             {
 				MUTABLE_CPUPROFILER_SCOPE(ME_CLIPWITHMESH_1)
+				
+				if (!EnumHasAnyFlags(EMeshContentFlags(Item.ExecutionOptions), EMeshContentFlags::GeometryData))
+				{
+					TSharedPtr<const FMesh> Source = LoadMesh(FCacheAddress(Args.Source, Item));
+					StoreMesh(Item, Source);
+				}
+				else
+				{
 
-                TSharedPtr<const FMesh> Source = LoadMesh(FCacheAddress(Args.Source, Item));
-                TSharedPtr<const FMesh> Clip = LoadMesh(FCacheAddress(Args.ClipMesh, Item));
-
-                // Only if both are valid.
-                if (Source && Clip)
-                {
-					TSharedPtr<FMesh> Result = CreateMesh(Source->GetDataSize());
-
-					bool bOutSuccess = false;
-                    MeshClipWithMesh(Result.Get(), Source.Get(), Clip.Get(), bOutSuccess);
+					TSharedPtr<const FMesh> Source = LoadMesh(FCacheAddress(Args.Source, Item));
+					TSharedPtr<const FMesh> Clip = LoadMesh(FCacheAddress(Args.ClipMesh, Item));
 
-					Release(Clip);
-					if (!bOutSuccess)
+					// Only if both are valid.
+					if (Source && Clip)
 					{
-						Release(Result);
-						StoreMesh(Item, Source);
+						TSharedPtr<FMesh> Result = CreateMesh(Source->GetDataSize());
+
+						bool bOutSuccess = false;
+						MeshClipWithMesh(Result.Get(), Source.Get(), Clip.Get(), bOutSuccess);
+
+						Release(Clip);
+						if (!bOutSuccess)
+						{
+							Release(Result);
+							StoreMesh(Item, Source);
+						}
+						else
+						{
+							Release(Source);
+							StoreMesh(Item, Result);
+						}
 					}
 					else
 					{
-						Release(Source);
-						StoreMesh(Item, Result);
+						Release(Clip);
+						StoreMesh(Item, Source);
 					}
-                }
-                else
-                {
-					Release(Clip);
-					StoreMesh(Item, Source);
-                }
-
+				}
                 break;
             }
 
@@ -2035,9 +2150,17 @@ namespace mu
 			{
 				if (Args.mesh)
 				{
-					AddOp(FScheduledOp(Item.At, Item, 1),
-						FScheduledOp(Args.mesh, Item),
-						FScheduledOp(Args.clipShape, Item));
+					if (!EnumHasAnyFlags(EMeshContentFlags(Item.ExecutionOptions), EMeshContentFlags::GeometryData))
+					{
+						AddOp(FScheduledOp(Item.At, Item, 1),
+							FScheduledOp(Args.mesh, Item));
+					}
+					else
+					{
+						AddOp(FScheduledOp(Item.At, Item, 1),
+							FScheduledOp(Args.mesh, Item),
+							FScheduledOp(Args.clipShape, Item));
+					}
 				}
 				else
 				{
@@ -2049,37 +2172,44 @@ namespace mu
 			{
 				MUTABLE_CPUPROFILER_SCOPE(ME_CLIPDEFORM_1)
 
-				TSharedPtr<const FMesh> BaseMesh = LoadMesh(FCacheAddress(Args.mesh, Item));
-				TSharedPtr<const FMesh> ClipShape = LoadMesh(FCacheAddress(Args.clipShape, Item));
-
-				if (BaseMesh && ClipShape)
+				if (!EnumHasAnyFlags(EMeshContentFlags(Item.ExecutionOptions), EMeshContentFlags::GeometryData))
+				{
+					TSharedPtr<const FMesh> BaseMesh = LoadMesh(FCacheAddress(Args.mesh, Item));
+					StoreMesh(Item, BaseMesh);
+				}
+				else
 				{
-					TSharedPtr<FMesh> Result = CreateMesh(BaseMesh->GetDataSize());
+					TSharedPtr<const FMesh> BaseMesh = LoadMesh(FCacheAddress(Args.mesh, Item));
+					TSharedPtr<const FMesh> ClipShape = LoadMesh(FCacheAddress(Args.clipShape, Item));
 
-					bool bRemoveIfAllVerticesCulled = Args.FaceCullStrategy == EFaceCullStrategy::AllVerticesCulled;
+					if (BaseMesh && ClipShape)
+					{
+						TSharedPtr<FMesh> Result = CreateMesh(BaseMesh->GetDataSize());
 
-					bool bOutSuccess = false;
-					MeshClipDeform(Result.Get(), BaseMesh.Get(), ClipShape.Get(), Args.clipWeightThreshold, bRemoveIfAllVerticesCulled, bOutSuccess);
+						bool bRemoveIfAllVerticesCulled = Args.FaceCullStrategy == EFaceCullStrategy::AllVerticesCulled;
 
-					Release(ClipShape);
+						bool bOutSuccess = false;
+						MeshClipDeform(Result.Get(), BaseMesh.Get(), ClipShape.Get(), Args.clipWeightThreshold, bRemoveIfAllVerticesCulled, bOutSuccess);
 
-					if (!bOutSuccess)
-					{
-						Release(Result);
-						StoreMesh(Item, BaseMesh);
+						Release(ClipShape);
+
+						if (!bOutSuccess)
+						{
+							Release(Result);
+							StoreMesh(Item, BaseMesh);
+						}
+						else
+						{
+							Release(BaseMesh);
+							StoreMesh(Item, Result);
+						}
 					}
 					else
 					{
-						Release(BaseMesh);
-						StoreMesh(Item, Result);
+						Release(ClipShape);
+						StoreMesh(Item, BaseMesh);
 					}
 				}
-				else
-				{
-					Release(ClipShape);
-					StoreMesh(Item, BaseMesh);
-				}
-
 				break;
 			}
 
@@ -2157,15 +2287,35 @@ namespace mu
 			OP::MeshBindShapeArgs Args = Program.GetOpArgs<OP::MeshBindShapeArgs>(Item.At);
 			const uint8* Data = Program.GetOpArgsPointer(Item.At);
 
+			constexpr uint8 ShapeContentFilter = (uint8)(EMeshContentFlags::GeometryData | EMeshContentFlags::PoseData);
+			const EShapeBindingMethod BindingMethod = static_cast<EShapeBindingMethod>(Args.bindingMethod); 
+
 			switch (Item.Stage)
 			{
 			case 0:
 			{
 				if (Args.mesh)
 				{
-					AddOp(FScheduledOp(Item.At, Item, 1),
-						FScheduledOp(Args.mesh, Item),
-						FScheduledOp(Args.shape, Item));
+					if (BindingMethod == EShapeBindingMethod::ReshapeClosestProject)
+					{
+						AddOp(FScheduledOp(Item.At, Item, 1),
+							FScheduledOp(Args.mesh, Item),
+							FScheduledOp::FromOpAndOptions(Args.shape, Item, ShapeContentFilter));
+					}
+					else
+					{
+						if (!EnumHasAnyFlags(EMeshContentFlags(Item.ExecutionOptions), EMeshContentFlags::GeometryData))
+						{
+							AddOp(FScheduledOp(Item.At, Item, 1),
+								FScheduledOp(Args.mesh, Item));
+						}
+						else
+						{
+							AddOp(FScheduledOp(Item.At, Item, 1),
+								FScheduledOp(Args.mesh, Item),
+								FScheduledOp::FromOpAndOptions(Args.shape, Item, ShapeContentFilter));
+						}
+					}
 				}
 				else
 				{
@@ -2176,13 +2326,12 @@ namespace mu
 			case 1:
 			{
 				MUTABLE_CPUPROFILER_SCOPE(ME_BINDSHAPE_1)
-				TSharedPtr<const FMesh> BaseMesh = LoadMesh(FCacheAddress(Args.mesh, Item));
-				TSharedPtr<const FMesh> Shape = LoadMesh(FCacheAddress(Args.shape, Item));
 				
-				EShapeBindingMethod BindingMethod = static_cast<EShapeBindingMethod>(Args.bindingMethod); 
-
 				if (BindingMethod == EShapeBindingMethod::ReshapeClosestProject)
-				{ 
+				{
+					TSharedPtr<const FMesh> BaseMesh = LoadMesh(FCacheAddress(Args.mesh, Item));
+					TSharedPtr<const FMesh> Shape = 
+					LoadMesh(FScheduledOp::FromOpAndOptions(Args.shape, Item, ShapeContentFilter));
 					// Bones are stored after the Args
 					Data += sizeof(Args);
 
@@ -2205,7 +2354,27 @@ namespace mu
 					FMemory::Memcpy(PhysicsToDeform.GetData(), Data, NumPhysicsBodies * sizeof(FBoneName));
 					Data += NumPhysicsBodies * sizeof(FBoneName);
 
-					const EMeshBindShapeFlags BindFlags = static_cast<EMeshBindShapeFlags>(Args.flags);
+					EMeshBindShapeFlags BindFlags = static_cast<EMeshBindShapeFlags>(Args.flags);
+					EMeshContentFlags MeshContentFilter = (EMeshContentFlags)Item.ExecutionOptions;
+
+					if (!EnumHasAnyFlags(MeshContentFilter, EMeshContentFlags::GeometryData))
+					{
+						EnumRemoveFlags(BindFlags, 
+								EMeshBindShapeFlags::EnableRigidParts |
+								EMeshBindShapeFlags::ReshapeVertices  |
+								EMeshBindShapeFlags::ApplyLaplacian   |
+								EMeshBindShapeFlags::RecomputeNormals);
+					}
+
+					if (!EnumHasAnyFlags(MeshContentFilter, EMeshContentFlags::PhysicsData))
+					{
+						EnumRemoveFlags(BindFlags, EMeshBindShapeFlags::ReshapePhysicsVolumes);
+					}
+
+					if (!EnumHasAnyFlags(MeshContentFilter, EMeshContentFlags::PoseData))
+					{
+						EnumRemoveFlags(BindFlags, EMeshBindShapeFlags::ReshapeSkeleton);
+					}
 
 					FMeshBindColorChannelUsages ColorChannelUsages;
 					FMemory::Memcpy(&ColorChannelUsages, &Args.ColorUsage, sizeof(ColorChannelUsages));
@@ -2229,7 +2398,6 @@ namespace mu
 						{
 							TSharedPtr<FMesh> BindMeshNoVertsResult = CloneOrTakeOver(BaseMesh);
 							BindMeshNoVertsResult->AdditionalBuffers = MoveTemp(BindMeshResult->AdditionalBuffers);
-
 							Release(BaseMesh);
 							Release(BindMeshResult);
 							StoreMesh(Item, BindMeshNoVertsResult);
@@ -2243,21 +2411,33 @@ namespace mu
 				}	
 				else
 				{
-					TSharedPtr<FMesh> Result = CreateMesh(BaseMesh ? BaseMesh->GetDataSize() : 0);
-
-					bool bOutSuccess = false;
-					MeshBindShapeClipDeform(Result.Get(), BaseMesh.Get(), Shape.Get(), BindingMethod, bOutSuccess);
-
-					Release(Shape);
-					if (!bOutSuccess)
+					if (!EnumHasAnyFlags(EMeshContentFlags(Item.ExecutionOptions), EMeshContentFlags::GeometryData))
 					{
-						Release(Result);
+						TSharedPtr<const FMesh> BaseMesh = LoadMesh(FCacheAddress(Args.mesh, Item));
 						StoreMesh(Item, BaseMesh);
 					}
 					else
 					{
-						Release(BaseMesh);
-						StoreMesh(Item, Result);
+						TSharedPtr<const FMesh> BaseMesh = LoadMesh(FCacheAddress(Args.mesh, Item));
+						TSharedPtr<const FMesh> Shape = 
+								LoadMesh(FScheduledOp::FromOpAndOptions(Args.shape, Item, ShapeContentFilter));
+
+						TSharedPtr<FMesh> Result = CreateMesh(BaseMesh ? BaseMesh->GetDataSize() : 0);
+
+						bool bOutSuccess = false;
+						MeshBindShapeClipDeform(Result.Get(), BaseMesh.Get(), Shape.Get(), BindingMethod, bOutSuccess);
+
+						Release(Shape);
+						if (!bOutSuccess)
+						{
+							Release(Result);
+							StoreMesh(Item, BaseMesh);
+						}
+						else
+						{
+							Release(BaseMesh);
+							StoreMesh(Item, Result);
+						}
 					}
 				}
 
@@ -2298,7 +2478,28 @@ namespace mu
 				TSharedPtr<const FMesh> BaseMesh = LoadMesh(FCacheAddress(Args.mesh, Item));
 				TSharedPtr<const FMesh> Shape = LoadMesh(FCacheAddress(Args.shape, Item));
 
-				const EMeshBindShapeFlags ReshapeFlags = static_cast<EMeshBindShapeFlags>(Args.flags);
+				EMeshBindShapeFlags ReshapeFlags = static_cast<EMeshBindShapeFlags>(Args.flags);
+				const EMeshContentFlags MeshContentFilter = (EMeshContentFlags)Item.ExecutionOptions;
+
+				if (!EnumHasAnyFlags(MeshContentFilter, EMeshContentFlags::GeometryData))
+				{
+					EnumRemoveFlags(ReshapeFlags, 
+							EMeshBindShapeFlags::EnableRigidParts |
+							EMeshBindShapeFlags::ReshapeVertices  |
+							EMeshBindShapeFlags::ApplyLaplacian   |
+							EMeshBindShapeFlags::RecomputeNormals);
+				}
+
+				if (!EnumHasAnyFlags(MeshContentFilter, EMeshContentFlags::PhysicsData))
+				{
+					EnumRemoveFlags(ReshapeFlags, EMeshBindShapeFlags::ReshapePhysicsVolumes);
+				}
+
+				if (!EnumHasAnyFlags(MeshContentFilter, EMeshContentFlags::PoseData))
+				{
+					EnumRemoveFlags(ReshapeFlags, EMeshBindShapeFlags::ReshapeSkeleton);
+				}
+
 				const bool bReshapeVertices = EnumHasAnyFlags(ReshapeFlags, EMeshBindShapeFlags::ReshapeVertices);
 
 				TSharedPtr<FMesh> ReshapedMeshResult = CreateMesh(BaseMesh ? BaseMesh->GetDataSize() : 0);
@@ -2360,11 +2561,11 @@ namespace mu
 			{
 				if (Args.Morph)
 				{
-					AddOp(FScheduledOp(Item.At, Item, 1), 
+					AddOp(FScheduledOp(Item.At, Item, 1),
 						FScheduledOp(Args.Morph, Item),
 						FScheduledOp(Args.Reshape, Item));
 				}
-				else 
+				else
 				{
 					StoreMesh(Item, nullptr);
 				}
@@ -2542,8 +2743,16 @@ namespace mu
 			{
 				if (source)
 				{
-					// Request the conditions
-					AddOp(FScheduledOp(Item.At, Item, 1), conditions);
+					if (!EnumHasAnyFlags(EMeshContentFlags(Item.ExecutionOptions), EMeshContentFlags::GeometryData))
+					{
+						AddOp(FScheduledOp(Item.At, Item, 1), 
+								FScheduledOp(source, Item));
+					}
+					else
+					{
+						// Request the conditions
+						AddOp(FScheduledOp(Item.At, Item, 1), conditions);
+					}
 				}
 				else
 				{
@@ -2555,28 +2764,36 @@ namespace mu
             {
         		MUTABLE_CPUPROFILER_SCOPE(ME_REMOVEMASK_1)
 
-                // Request the source and the necessary masks
-                // \todo: store condition values in heap?
-                TArray<FScheduledOp> deps;
-                deps.Emplace( source, Item );
-                for( int32 r=0; source && r<conditions.Num(); ++r )
+				if (!EnumHasAnyFlags(EMeshContentFlags(Item.ExecutionOptions), EMeshContentFlags::GeometryData))
                 {
-                    // If there is no expression, we'll assume true.
-                    bool value = true;
-                    if (conditions[r].At)
-                    {
-                        value = LoadBool(FCacheAddress(conditions[r].At, Item));
-                    }
+                	TSharedPtr<const FMesh> Source = LoadMesh(FCacheAddress(source, Item));
+					StoreMesh(Item, Source);
+				}
+				else
+				{
+					// Request the source and the necessary masks
+					// \todo: store condition values in heap?
+					TArray<FScheduledOp> deps;
+					deps.Emplace( source, Item );
+					for( int32 r=0; source && r<conditions.Num(); ++r )
+					{
+						// If there is no expression, we'll assume true.
+						bool value = true;
+						if (conditions[r].At)
+						{
+							value = LoadBool(FCacheAddress(conditions[r].At, Item));
+						}
 
-                    if (value)
-                    {
-                        deps.Emplace(masks[r], Item);
-                    }
-                }
+						if (value)
+						{
+							deps.Emplace(masks[r], Item);
+						}
+					}
 
-                if (source)
-                {
-					AddOp(FScheduledOp(Item.At, Item, 2), deps);
+					if (source)
+					{
+						AddOp(FScheduledOp(Item.At, Item, 2), deps);
+					}
 				}
                 break;
             }
@@ -2681,37 +2898,44 @@ namespace mu
 			{
 				MUTABLE_CPUPROFILER_SCOPE(ME_ADDTAGS_2)
 
-				TSharedPtr<const FMesh> SourceMesh = LoadMesh(FCacheAddress(Source, Item));
-
-				if (!SourceMesh)
+				if (!EnumHasAnyFlags(EMeshContentFlags(Item.ExecutionOptions), EMeshContentFlags::MetaData))
 				{
-					StoreMesh(Item, nullptr);
+					TSharedPtr<const FMesh> SourceMesh = LoadMesh(FCacheAddress(Source, Item));
+					StoreMesh(Item, SourceMesh);
 				}
 				else
 				{
-					TSharedPtr<FMesh> Result = CloneOrTakeOver(SourceMesh);
-
-					// Decode the tags
-					uint16 TagCount;
-					FMemory::Memcpy(&TagCount, Data, sizeof(uint16));
-					Data += sizeof(uint16);
+					TSharedPtr<const FMesh> SourceMesh = LoadMesh(FCacheAddress(Source, Item));
 
-					int32 FirstMeshTagIndex = Result->Tags.Num();
-					Result->Tags.SetNum(FirstMeshTagIndex+TagCount);
-					for (uint16 TagIndex = 0; TagIndex < TagCount; ++TagIndex)
+					if (!SourceMesh)
 					{
-						OP::ADDRESS TagConstant;
-						FMemory::Memcpy(&TagConstant, Data, sizeof(OP::ADDRESS));
-						Data += sizeof(OP::ADDRESS);
-
-						check(TagConstant < (uint32)InModel->GetPrivate()->Program.ConstantStrings.Num());
-						const FString& Name = Program.ConstantStrings[TagConstant];
-						Result->Tags[FirstMeshTagIndex+TagIndex] = Name;
+						StoreMesh(Item, nullptr);
 					}
+					else
+					{
+						TSharedPtr<FMesh> Result = CloneOrTakeOver(SourceMesh);
 
-					StoreMesh(Item, Result);
-				}
+						// Decode the tags
+						uint16 TagCount;
+						FMemory::Memcpy(&TagCount, Data, sizeof(uint16));
+						Data += sizeof(uint16);
 
+						int32 FirstMeshTagIndex = Result->Tags.Num();
+						Result->Tags.SetNum(FirstMeshTagIndex+TagCount);
+						for (uint16 TagIndex = 0; TagIndex < TagCount; ++TagIndex)
+						{
+							OP::ADDRESS TagConstant;
+							FMemory::Memcpy(&TagConstant, Data, sizeof(OP::ADDRESS));
+							Data += sizeof(OP::ADDRESS);
+
+							check(TagConstant < (uint32)InModel->GetPrivate()->Program.ConstantStrings.Num());
+							const FString& Name = Program.ConstantStrings[TagConstant];
+							Result->Tags[FirstMeshTagIndex+TagIndex] = Name;
+						}
+
+						StoreMesh(Item, Result);
+					}
+				}
 				break;
 			}
 
@@ -2731,9 +2955,17 @@ namespace mu
 			{
 				if (Args.Mesh)
 				{
-					AddOp(FScheduledOp(Item.At, Item, 1),
-						FScheduledOp(Args.Mesh, Item),
-						FScheduledOp(Args.Projector, Item));
+					if (!EnumHasAnyFlags(EMeshContentFlags(Item.ExecutionOptions), EMeshContentFlags::GeometryData))
+					{
+						AddOp(FScheduledOp(Item.At, Item, 1),
+							FScheduledOp(Args.Mesh, Item));
+					}
+					else
+					{
+						AddOp(FScheduledOp(Item.At, Item, 1),
+							FScheduledOp(Args.Mesh, Item),
+							FScheduledOp(Args.Projector, Item));
+					}
 				}
 				else
 				{
@@ -2741,39 +2973,46 @@ namespace mu
 				}
 				break;
 			}
-            case 1:
-            {
+			case 1:
+			{
 				MUTABLE_CPUPROFILER_SCOPE(ME_PROJECT_1)
 
-                TSharedPtr<const FMesh> Mesh = LoadMesh(FCacheAddress(Args.Mesh,Item));
-                const FProjector Projector = LoadProjector(FCacheAddress(Args.Projector, Item));
+				if (!EnumHasAnyFlags(EMeshContentFlags(Item.ExecutionOptions), EMeshContentFlags::GeometryData))
+				{
+					TSharedPtr<const FMesh> Mesh = LoadMesh(FCacheAddress(Args.Mesh, Item));
+					StoreMesh(Item, Mesh);
+				}
+				else
+				{
+					TSharedPtr<const FMesh> Mesh = LoadMesh(FCacheAddress(Args.Mesh, Item));
+					const FProjector Projector = LoadProjector(FCacheAddress(Args.Projector, Item));
 
-                // Only if both are valid.
-                if (Mesh && Mesh->GetVertexBuffers().GetBufferCount() > 0)
-                {
-					TSharedPtr<FMesh> Result = CreateMesh();
+					// Only if both are valid.
+					if (Mesh && Mesh->GetVertexBuffers().GetBufferCount() > 0)
+					{
+						TSharedPtr<FMesh> Result = CreateMesh();
 
-					bool bOutSuccess = false;
-					MeshProject(Result.Get(), Mesh.Get(), Projector, bOutSuccess);
+						bool bOutSuccess = false;
+						MeshProject(Result.Get(), Mesh.Get(), Projector, bOutSuccess);
 
-					if (!bOutSuccess)
-					{
-						Release(Result);
-						StoreMesh(Item, Mesh);
+						if (!bOutSuccess)
+						{
+							Release(Result);
+							StoreMesh(Item, Mesh);
+						}
+						else
+						{
+							// Result->GetPrivate()->CheckIntegrity();
+							Release(Mesh);
+							StoreMesh(Item, Result);
+						}
 					}
 					else
-					{	
-//						Result->GetPrivate()->CheckIntegrity();
+					{
 						Release(Mesh);
-						StoreMesh(Item, Result);
+						StoreMesh(Item, nullptr);
 					}
-                }
-				else
-				{
-					Release(Mesh);
-					StoreMesh(Item, nullptr);
 				}
-
                 break;
             }
 
@@ -2847,10 +3086,18 @@ namespace mu
 			{
 				if (Args.sourceMesh)
 				{
-					AddOp(FScheduledOp(Item.At, Item, 1),
-						FScheduledOp(Args.sourceMesh, Item),
-						FScheduledOp(Args.boundingMesh, Item),
-						FScheduledOp(Args.matrix, Item));
+					if (!EnumHasAnyFlags(EMeshContentFlags(Item.ExecutionOptions), EMeshContentFlags::GeometryData))
+					{
+						AddOp(FScheduledOp(Item.At, Item, 1),
+							FScheduledOp(Args.sourceMesh, Item));
+					}
+					else
+					{
+						AddOp(FScheduledOp(Item.At, Item, 1),
+							FScheduledOp(Args.sourceMesh, Item),
+							FScheduledOp(Args.boundingMesh, Item),
+							FScheduledOp(Args.matrix, Item));
+					}
 				}
 				else
 				{
@@ -2861,35 +3108,43 @@ namespace mu
 			case 1:
 			{
 				MUTABLE_CPUPROFILER_SCOPE(ME_TRANSFORMWITHMESH_1)
-            
-				TSharedPtr<const FMesh> SourceMesh = LoadMesh(FCacheAddress(Args.sourceMesh,Item));
-				TSharedPtr<const FMesh> BoundingMesh = LoadMesh(FCacheAddress(Args.boundingMesh, Item));
-				const FMatrix44f& Transform = LoadMatrix(FCacheAddress(Args.matrix, Item));
 
-				if (SourceMesh)
+				if (!EnumHasAnyFlags(EMeshContentFlags(Item.ExecutionOptions), EMeshContentFlags::GeometryData))
 				{
-					TSharedPtr<FMesh> Result = CreateMesh(SourceMesh->GetDataSize());
-
-					bool bOutSuccess = false;
-					MeshTransformWithMesh(Result.Get(), SourceMesh.Get(), BoundingMesh.Get(), Transform, bOutSuccess);
-					Release(BoundingMesh);
+					TSharedPtr<const FMesh> SourceMesh = LoadMesh(FCacheAddress(Args.sourceMesh, Item));
+					StoreMesh(Item, SourceMesh);
+				}
+				else
+				{
+					TSharedPtr<const FMesh> SourceMesh = LoadMesh(FCacheAddress(Args.sourceMesh,Item));
+					TSharedPtr<const FMesh> BoundingMesh = LoadMesh(FCacheAddress(Args.boundingMesh, Item));
+					const FMatrix44f& Transform = LoadMatrix(FCacheAddress(Args.matrix, Item));
 
-					if (!bOutSuccess)
+					if (SourceMesh)
 					{
-						Release(Result);
-						StoreMesh(Item, SourceMesh);
+						TSharedPtr<FMesh> Result = CreateMesh(SourceMesh->GetDataSize());
+
+						bool bOutSuccess = false;
+						MeshTransformWithMesh(Result.Get(), SourceMesh.Get(), BoundingMesh.Get(), Transform, bOutSuccess);
+						Release(BoundingMesh);
+
+						if (!bOutSuccess)
+						{
+							Release(Result);
+							StoreMesh(Item, SourceMesh);
+						}
+						else
+						{
+							Release(SourceMesh);
+							StoreMesh(Item, Result);
+						}
 					}
 					else
 					{
-						Release(SourceMesh);
-						StoreMesh(Item, Result);
+						Release(BoundingMesh);
+						StoreMesh(Item, SourceMesh);
 					}
 				}
-				else
-				{
-					Release(BoundingMesh);
-					StoreMesh(Item, SourceMesh);
-				}
 				break;
 			}
 
@@ -2899,15 +3154,15 @@ namespace mu
 			break;
 		}
 
-        default:
-            if (type!=EOpType::NONE)
-            {
-                // Operation not implemented
-                check( false );
-            }
-            break;
-        }
-    }
+		default:
+			if (type!=EOpType::NONE)
+			{
+				// Operation not implemented
+				check( false );
+			}
+			break;
+		}
+	}
 
     //---------------------------------------------------------------------------------------------
     void CodeRunner::RunCode_Image(const FScheduledOp& Item, const FParameters* pParams, const FModel* InModel )
@@ -4286,19 +4541,21 @@ namespace mu
         case EOpType::IM_RASTERMESH:
         {
 			OP::ImageRasterMeshArgs Args = Program.GetOpArgs<OP::ImageRasterMeshArgs>(Item.At);
+
+			constexpr uint8 MeshContentFilter = (uint8)(EMeshContentFlags::GeometryData | EMeshContentFlags::PoseData);
             switch (Item.Stage)
             {
             case 0:
 				if (Args.image)
 				{
 					AddOp(FScheduledOp(Item.At, Item, 1),
-						FScheduledOp::FromOpAndOptions(Args.mesh, Item, 0),
+						FScheduledOp::FromOpAndOptions(Args.mesh, Item, MeshContentFilter),
 						FScheduledOp::FromOpAndOptions(Args.projector, Item, 0));
 				}
 				else
 				{
 					AddOp(FScheduledOp(Item.At, Item, 1),
-						FScheduledOp::FromOpAndOptions(Args.mesh, Item, 0));
+						FScheduledOp::FromOpAndOptions(Args.mesh, Item, MeshContentFilter));
 				}
                 break;
 
@@ -4306,7 +4563,9 @@ namespace mu
 			{
 				MUTABLE_CPUPROFILER_SCOPE(IM_RASTERMESH_1)
 
-				TSharedPtr<const FMesh> Mesh = LoadMesh(FScheduledOp::FromOpAndOptions(Args.mesh, Item, 0));
+				
+				TSharedPtr<const FMesh> Mesh = 
+						LoadMesh(FScheduledOp::FromOpAndOptions(Args.mesh, Item, MeshContentFilter));
 
 				// If no image, we are generating a flat mesh UV raster. This is the final stage in this case.
 				if (!Args.image)
@@ -5784,16 +6043,19 @@ namespace mu
 		case EOpType::LA_FROMMESH:
 		{
 			OP::LayoutFromMeshArgs Args = Program.GetOpArgs<OP::LayoutFromMeshArgs>(Item.At);
+
+			constexpr uint8 MeshContentFilter = (uint8)(EMeshContentFlags::AllFlags);
 			switch (Item.Stage)
 			{
 			case 0:
 				AddOp(FScheduledOp(Item.At, Item, 1),
-					FScheduledOp(Args.Mesh, Item));
+					FScheduledOp::FromOpAndOptions(Args.Mesh, Item, MeshContentFilter));
 				break;
 
 			case 1:
 			{
-				TSharedPtr<const FMesh> Mesh = LoadMesh(FCacheAddress(Args.Mesh, Item));
+				TSharedPtr<const FMesh> Mesh = LoadMesh(
+						FScheduledOp::FromOpAndOptions(Args.Mesh, Item, MeshContentFilter));
 
 				TSharedPtr<const FLayout> Result = LayoutFromMesh_RemoveBlocks(Mesh.Get(), Args.LayoutIndex);
 
diff --git a/Engine/Plugins/Mutable/Source/MutableRuntime/Private/MuR/CodeRunner.h b/Engine/Plugins/Mutable/Source/MutableRuntime/Private/MuR/CodeRunner.h
--- a/Engine/Plugins/Mutable/Source/MutableRuntime/Private/MuR/CodeRunner.h
+++ b/Engine/Plugins/Mutable/Source/MutableRuntime/Private/MuR/CodeRunner.h
@@ -47,9 +47,8 @@ namespace  mu
 				class FSystem::Private*, 
 				EExecutionStrategy,
 				const TSharedPtr<const FModel>&, 
-				const FParameters* pParams,
-				OP::ADDRESS at, uint32 lodMask, uint8 executionOptions, int32 InImageLOD, FScheduledOp::EType);
-
+				const FParameters* Params,
+				OP::ADDRESS At, uint32 LodMask, uint8 ExecutionOptions, int32 InImageLOD, FScheduledOp::EType);
 		// Private constructor to prevent stack allocation. In general we can not call AsShared() if the lifetime is
 		// bounded.
 		explicit CodeRunner(FPrivateToken, 
@@ -57,8 +56,8 @@ namespace  mu
 			class FSystem::Private*, 
 			EExecutionStrategy,
 			const TSharedPtr<const FModel>&, 
-			const FParameters* pParams,
-			OP::ADDRESS at, uint32 lodMask, uint8 executionOptions, int32 InImageLOD, FScheduledOp::EType);
+			const FParameters* Params,
+			OP::ADDRESS At, uint32 LodMask, uint8 ExecutionOptions, int32 InImageLOD, FScheduledOp::EType);
 
     protected:
 		struct FProfileContext
@@ -114,7 +113,7 @@ namespace  mu
 
         void RunCode_Conditional(const FScheduledOp&, const FModel* );
         void RunCode_Switch(const FScheduledOp&, const FModel* );
-        void RunCode_Instance(const FScheduledOp&, const FModel*, uint32 LodMask );
+        void RunCode_Instance(const FScheduledOp&, const FModel*, uint32 LodMask);
         void RunCode_InstanceAddResource(const FScheduledOp&, const TSharedPtr<const FModel>& Model, const FParameters* );
 
 		/** Return false incase of failure. */
@@ -132,8 +131,7 @@ namespace  mu
         void RunCode_Matrix(const FScheduledOp&, const FParameters*, const FModel* );
 
 		void RunCodeImageDesc(const FScheduledOp&, const FParameters*, const FModel*, uint32 LodMask);
-
-    public:
+	public:
 
 		struct FExternalResourceId
 		{
@@ -259,10 +257,16 @@ namespace  mu
 		class FLoadMeshRomTask : public CodeRunner::FIssuedTask
 		{
 		public:
-			FLoadMeshRomTask( const FScheduledOp& InOp, uint32 InRomIndex)
+			FLoadMeshRomTask( 
+					const FScheduledOp& InOp, 
+					int32 InFirstIndex, 
+					EMeshContentFlags InRomContentFlags, 
+					EMeshContentFlags InExecutionContentFlags)
 				: FIssuedTask(InOp)
+				, FirstIndex(InFirstIndex)
+				, RomContentFlags(InRomContentFlags)
+				, ExecutionContentFlags(InExecutionContentFlags)
 			{
-				RomIndex = InRomIndex;
 			}
 
 			// FIssuedTask interface
@@ -270,7 +274,12 @@ namespace  mu
 			virtual bool Complete(CodeRunner*) override;
 
 		private:
-			uint32 RomIndex = 0;
+
+			int32 FirstIndex = -1;
+			EMeshContentFlags RomContentFlags = EMeshContentFlags::None;
+			EMeshContentFlags ExecutionContentFlags = EMeshContentFlags::None;
+
+			TArray<int32, TInlineAllocator<4>> RomIndices;
 		};
 
 		class FLoadImageRomsTask : public CodeRunner::FIssuedTask
@@ -361,17 +370,17 @@ namespace  mu
 			AddChildren(dep4);
 		}
 
-		template<class RangeType>
-		void AddOp(const FScheduledOp& Op, const RangeType& Deps)
+		void AddOp(const FScheduledOp& Op, TArrayView<const FScheduledOp> Deps)
 		{
 			FTask Task(Op);
 
+			Task.Deps.Reserve(Deps.Num());
 			for (const FScheduledOp& D : Deps)
 			{
 				Task.Deps.Add(D);
 			}
 
-			ClosedTasks.Add(Task);
+			ClosedTasks.Add(MoveTemp(Task));
 
 			ScheduledStagePerOp[Op] = Op.Stage + 1;
 
diff --git a/Engine/Plugins/Mutable/Source/MutableRuntime/Private/MuR/CodeRunnerTasked.cpp b/Engine/Plugins/Mutable/Source/MutableRuntime/Private/MuR/CodeRunnerTasked.cpp
--- a/Engine/Plugins/Mutable/Source/MutableRuntime/Private/MuR/CodeRunnerTasked.cpp
+++ b/Engine/Plugins/Mutable/Source/MutableRuntime/Private/MuR/CodeRunnerTasked.cpp
@@ -1967,10 +1967,7 @@ namespace mu
 		return bSuccess;
 	}
 
-
-	//---------------------------------------------------------------------------------------------
-	//---------------------------------------------------------------------------------------------
-	//---------------------------------------------------------------------------------------------
+	
 	bool CodeRunner::FLoadMeshRomTask::Prepare(CodeRunner* Runner, bool& bOutFailed)
 	{
 		// This runs in the mutable Runner thread
@@ -1984,65 +1981,134 @@ namespace mu
 		bOutFailed = false;
 
 		const FProgram& Program = Runner->Model->GetPrivate()->Program;
-
-		check(int32(RomIndex) < Program.Roms.Num());
-		
 		FWorkingMemoryManager::FModelCacheEntry* ModelCache = Runner->System->WorkingMemoryManager.FindModelCache(Runner->Model.Get());
-		++ModelCache->PendingOpsPerRom[RomIndex];
+		
+		TArray<UE::Tasks::FTask, TInlineAllocator<4>> ReadCompleteEvents;
+		ReadCompleteEvents.Reserve(4); 
+
+		TArray<int32, TInlineAllocator<4>> RomsToLoad;
+
+		// Rom indices are sorted by flag value
+		static_assert(EMeshContentFlags::GeometryData < EMeshContentFlags::PoseData);
+		static_assert(EMeshContentFlags::PoseData     < EMeshContentFlags::PhysicsData);
+		static_assert(EMeshContentFlags::PhysicsData  < EMeshContentFlags::MetaData);
 
-		if (Program.IsRomLoaded(RomIndex))
+		int32 RomContentIndex = 0;
+
+		if (EnumHasAnyFlags(ExecutionContentFlags, EMeshContentFlags::GeometryData) && 
+			EnumHasAnyFlags(RomContentFlags, EMeshContentFlags::GeometryData))
 		{
-			return false;
+			RomsToLoad.Add(RomContentIndex);
 		}
+		RomContentIndex += EnumHasAnyFlags(RomContentFlags, EMeshContentFlags::GeometryData);
 
-		if (const FRomLoadOp* Result = Runner->RomLoadOps.Find(RomIndex))
+		if (EnumHasAnyFlags(ExecutionContentFlags, EMeshContentFlags::PoseData) && 
+			EnumHasAnyFlags(RomContentFlags, EMeshContentFlags::PoseData))
 		{
-			Event = Result->Event;  // Wait for the read operation started by other task
-			return false;
+			RomsToLoad.Add(RomContentIndex);
 		}
-
-		FRomLoadOp& RomLoadOp = Runner->RomLoadOps.Create(RomIndex);
+		RomContentIndex += EnumHasAnyFlags(RomContentFlags, EMeshContentFlags::PoseData);
 		
-		check(Runner->System->StreamInterface);
-
-		const uint32 RomSize = Program.Roms[RomIndex].Size;
-		check(RomSize > 0);
-
-		// Free roms if necessary
+		if (EnumHasAnyFlags(ExecutionContentFlags, EMeshContentFlags::PhysicsData) && 
+			EnumHasAnyFlags(RomContentFlags, EMeshContentFlags::PhysicsData))
 		{
-			MUTABLE_CPUPROFILER_SCOPE(FreeingRoms);
+			RomsToLoad.Add(RomContentIndex);
+		}
+		RomContentIndex += EnumHasAnyFlags(RomContentFlags, EMeshContentFlags::PhysicsData);
 
-			Runner->System->WorkingMemoryManager.MarkRomUsed(RomIndex, Runner->Model);
-			Runner->System->WorkingMemoryManager.EnsureBudgetBelow(RomSize);
+		if (EnumHasAnyFlags(ExecutionContentFlags, EMeshContentFlags::MetaData) && 
+			EnumHasAnyFlags(RomContentFlags, EMeshContentFlags::MetaData))
+		{
+			RomsToLoad.Add(RomContentIndex);
 		}
+		RomContentIndex += EnumHasAnyFlags(RomContentFlags, EMeshContentFlags::MetaData);
 
-		const int32 SizeBefore = RomLoadOp.m_streamBuffer.GetAllocatedSize();
-		RomLoadOp.m_streamBuffer.SetNumUninitialized(RomSize);
-		const int32 SizeAfter = RomLoadOp.m_streamBuffer.GetAllocatedSize();
+		check(RomContentIndex == FMath::CountBits((uint64)RomContentFlags));
 
-	 	UE::Tasks::FTaskEvent ReadCompletionEvent = UE::Tasks::FTaskEvent(TEXT("FLoadMeshRomTask"));	
-		RomLoadOp.Event = ReadCompletionEvent;
-		
-		TFunction<void(bool)> Callback = [ReadCompletionEvent](bool bSuccess) mutable
+		for (const int32 MeshContentRomIndex : RomsToLoad)
 		{
-			ReadCompletionEvent.Trigger();
-		};
+			FConstantResourceIndex CurrentIndex = 
+					Program.ConstantMeshContentIndices[MeshContentRomIndex + FirstIndex];
+
+			if (!CurrentIndex.Streamable)
+			{
+				// This data is always resident.
+				continue;
+			}
+
+			int32 RomIndex = CurrentIndex.Index;
+
+			check(RomIndex < Program.Roms.Num());
+
+			++ModelCache->PendingOpsPerRom[RomIndex];
+
+			if (DebugRom && (DebugRomAll || RomIndex == DebugRomIndex))
+			{
+				UE_LOG(LogMutableCore, Log, TEXT("Preparing rom %d, now peding ops is %d."), RomIndex, ModelCache->PendingOpsPerRom[RomIndex]);
+			}
+
+			if (Program.IsRomLoaded(RomIndex))
+			{
+				continue;
+			}
+
+			RomIndices.Add(RomIndex);
+
+			if (const FRomLoadOp* Result = Runner->RomLoadOps.Find(RomIndex))
+			{
+				ReadCompleteEvents.Add(Result->Event); // Wait for the read operation started by other task
+				continue;
+			}
+
+			FRomLoadOp& RomLoadOp = Runner->RomLoadOps.Create(RomIndex);
+			
+			check(Runner->System->StreamInterface);
+
+			const uint32 RomSize = Program.Roms[RomIndex].Size;
+			check(RomSize > 0);
+
+			// Free roms if necessary
+			{
+				Runner->System->WorkingMemoryManager.MarkRomUsed(RomIndex, Runner->Model);
+				Runner->System->WorkingMemoryManager.EnsureBudgetBelow(RomSize);
+			}
+
+			const int32 SizeBefore = RomLoadOp.m_streamBuffer.GetAllocatedSize();
+			RomLoadOp.m_streamBuffer.SetNumUninitialized(RomSize);
+			const int32 SizeAfter = RomLoadOp.m_streamBuffer.GetAllocatedSize();
+
+			UE::Tasks::FTaskEvent ReadCompletionEvent(TEXT("FLoadMeshRomsTaskRom"));
+			ReadCompleteEvents.Add(ReadCompletionEvent);
+			RomLoadOp.Event = ReadCompletionEvent;
+
+			TFunction<void(bool)> Callback = [ReadCompletionEvent](bool bSuccess) mutable // Mutable due Trigger not being const
+			{
+				ReadCompletionEvent.Trigger();
+			};
+
 		
-		const uint32 RomId = RomIndex;
-		RomLoadOp.m_streamID = Runner->System->StreamInterface->BeginReadBlock(Runner->Model.Get(), RomId, RomLoadOp.m_streamBuffer.GetData(), RomSize, &Callback);
-		if (RomLoadOp.m_streamID < 0)
-		{
-			bOutFailed = true;
-			return false;
+			const uint32 RomId = RomIndex;
+			RomLoadOp.m_streamID = Runner->System->StreamInterface->BeginReadBlock(Runner->Model.Get(), RomId, RomLoadOp.m_streamBuffer.GetData(), RomSize, &Callback);
+			if (RomLoadOp.m_streamID < 0)
+			{
+				bOutFailed = true;
+				return false;
+			}
 		}
 
-		Event = ReadCompletionEvent; // Wait for read operation to end
+		// Wait for all read operations to end
+		UE::Tasks::FTaskEvent GatherReadsCompletionEvent(TEXT("FLoadMeshRomsTask"));
+		GatherReadsCompletionEvent.AddPrerequisites(ReadCompleteEvents);
+		GatherReadsCompletionEvent.Trigger();
+
+		Event = GatherReadsCompletionEvent;
+			
 		return false; // No worker thread work
-	}
-	
+
+	}	
 
 	//---------------------------------------------------------------------------------------------
-	bool CodeRunner::FLoadMeshRomTask::Complete(CodeRunner* Runner)
+		bool CodeRunner::FLoadMeshRomTask::Complete(CodeRunner* Runner)
 	{
 		// This runs in the Runner thread
 		MUTABLE_CPUPROFILER_SCOPE(FLoadMeshRomTask_Complete);
@@ -2053,49 +2119,63 @@ namespace mu
 		}
 
 		FProgram& Program = Runner->Model->GetPrivate()->Program;
-
-		// Since task could be reordered, we need to make sure we end the rom read before continuing
-		if (FRomLoadOp* RomLoadOp = Runner->RomLoadOps.Find(RomIndex))
+		FWorkingMemoryManager::FModelCacheEntry* ModelCache = Runner->System->WorkingMemoryManager.FindModelCache(Runner->Model.Get());
+		
+		bool bSomeMissingData = false;
+		for (const int32 RomIndex : RomIndices)
 		{
-			bool bSuccess = Runner->System->StreamInterface->EndRead(RomLoadOp->m_streamID);
+			// Since task could be reordered, we need to make sure we end the rom read before continuing
+			if (FRomLoadOp* RomLoadOp = Runner->RomLoadOps.Find(RomIndex))
+			{				
+				bool bSuccess = Runner->System->StreamInterface->EndRead(RomLoadOp->m_streamID);
 
-			check(!Program.ConstantMeshesStreamed.Contains(RomIndex));
+				MUTABLE_CPUPROFILER_SCOPE(Unserialise);
 
-			MUTABLE_CPUPROFILER_SCOPE(Unserialise);
+				if (bSuccess)
+				{
+					FInputMemoryStream Stream(RomLoadOp->m_streamBuffer.GetData(), RomLoadOp->m_streamBuffer.Num());
+					FInputArchive Arch(&Stream);
 
-			FInputMemoryStream Stream(RomLoadOp->m_streamBuffer.GetData(), RomLoadOp->m_streamBuffer.Num());
-			FInputArchive Arch(&Stream);
+					TSharedPtr<FMesh> Value = FMesh::StaticUnserialise(Arch);
 
-			TSharedPtr<FMesh> Value;
-			if (bSuccess)
-			{
-				Value = FMesh::StaticUnserialise(Arch);
-			}
-			else
-			{
-				UE_LOG(LogMutableCore, Error, TEXT("Failed to load rom data for rom index [%d]. Missing files?"), RomIndex);
-				Value = MakeShared<FMesh>();
+					Program.SetMeshRomValue(RomIndex, Value);
+				}
+				else
+				{
+					bSomeMissingData = true;
+				}
+				
+				Runner->RomLoadOps.Remove(*RomLoadOp);
 			}
+		}
 
-			Program.SetMeshRomValue(RomIndex, Value);
-
-			check(Program.ConstantMeshesStreamed.Contains(RomIndex));
-
-			Runner->RomLoadOps.Remove(*RomLoadOp);
+		if (bSomeMissingData)
+		{
+			// Some data may be missing. We can try to go on if some requested mesh parts are there. 
+			UE_LOG(LogMutableCore, Verbose, TEXT("FLoadMeshRomsTask::Complete failed: missing data?"));
 		}
 
-		// Process the constant op normally, now that the rom is loaded.
-		Runner->RunCode(Op, Runner->Params, Runner->Model, Runner->LODMask);
+		// Process the constant op normally, now that the rom is loaded.	
+		bool bSuccess = Runner->RunCode_ConstantResource(Op, Runner->Model.Get());
 
-		FWorkingMemoryManager::FModelCacheEntry* ModelCache = Runner->System->WorkingMemoryManager.FindModelCache(Runner->Model.Get());
+		if (bSuccess)
+		{
+			for (int32 RomIndex : RomIndices)
+			{
+				check(RomIndex < Program.Roms.Num());
 
-		Runner->System->WorkingMemoryManager.MarkRomUsed(RomIndex, Runner->Model);
-		--ModelCache->PendingOpsPerRom[RomIndex];
+				Runner->System->WorkingMemoryManager.MarkRomUsed(RomIndex, Runner->Model);
+				--ModelCache->PendingOpsPerRom[RomIndex];
+
+				if (DebugRom && (DebugRomAll || RomIndex == DebugRomIndex))
+				{
+					UE_LOG(LogMutableCore, Log, TEXT("FLoadMeshRomsTask::Complete rom %d, now peding ops is %d."), RomIndex, ModelCache->PendingOpsPerRom[RomIndex]);
+				}
+			}
+		}
 
-		bool bSuccess = true;
 		return bSuccess;
 	}
-
 	
 	//---------------------------------------------------------------------------------------------
 	//---------------------------------------------------------------------------------------------
@@ -2432,16 +2512,13 @@ namespace mu
 		{
 		case EOpType::ME_CONSTANT:
 		{
-			OP::MeshConstantArgs args = Program.GetOpArgs<OP::MeshConstantArgs>(item.At);
-			FConstantResourceIndex ResourceIndex = *reinterpret_cast<const FConstantResourceIndex*>(&args.value);
-			if (ResourceIndex.Streamable && !Program.ConstantMeshesStreamed.Contains(ResourceIndex.Index))
-			{
-				Issued = MakeShared<FLoadMeshRomTask>(item, uint32(ResourceIndex.Index));
-			}
-			else
-			{
-				// If already available, the rest of the constant code will run right away.
-			}
+			OP::MeshConstantArgs Args = Program.GetOpArgs<OP::MeshConstantArgs>(item.At);
+		
+			const FMeshContentRange MeshContentRange = Program.ConstantMeshes[Args.Value];
+			const EMeshContentFlags ContentFilterFlags = static_cast<EMeshContentFlags>(item.ExecutionOptions);
+			
+			Issued = MakeShared<FLoadMeshRomTask>(item, 
+					MeshContentRange.FirstIndex, MeshContentRange.ContentFlags, ContentFilterFlags);
 			break;
 		}
 
diff --git a/Engine/Plugins/Mutable/Source/MutableRuntime/Private/MuR/Mesh.cpp b/Engine/Plugins/Mutable/Source/MutableRuntime/Private/MuR/Mesh.cpp
--- a/Engine/Plugins/Mutable/Source/MutableRuntime/Private/MuR/Mesh.cpp
+++ b/Engine/Plugins/Mutable/Source/MutableRuntime/Private/MuR/Mesh.cpp
@@ -90,44 +90,7 @@ const FString& FMesh::GetReferencedMorph() const
 TSharedPtr<FMesh> FMesh::Clone() const
 {
     //MUTABLE_CPUPROFILER_SCOPE(MeshClone);
-	LLM_SCOPE_BYNAME(TEXT("MutableRuntime"));
-
-    TSharedPtr<FMesh> Result = MakeShared<FMesh>();
-
-    Result->InternalId = InternalId;
-	Result->Flags = Flags;
-	Result->ReferenceID = ReferenceID;
-	Result->Surfaces = Surfaces;
-	Result->Skeleton = Skeleton;
-	Result->PhysicsBody = PhysicsBody;
-	Result->Tags = Tags;
-	Result->StreamedResources = StreamedResources;
-	Result->MeshIDPrefix = MeshIDPrefix;
-
-    // Clone the main buffers
-    Result->VertexBuffers = VertexBuffers;
-    Result->IndexBuffers = IndexBuffers;
-
-	// Clone additional buffers
-	Result->AdditionalBuffers = AdditionalBuffers;
-
-    // Clone the layouts
-	Result->Layouts = Layouts;
-
-    // The skeleton is not cloned because it is not owned by this mesh and it is always assumed
-    // to be shared.
-
-	// physics body doen't need to be deep cloned either as they are also assumed to be shared.
-	
-	// Clone bone poses
-	Result->BonePoses = BonePoses;
-	Result->BoneMap = BoneMap;
-
-	// Clone SkeletonIDs
-	Result->SkeletonIDs = SkeletonIDs;
-	Result->AdditionalPhysicsBodies = AdditionalPhysicsBodies;
-
-    return Result;
+	return Clone(EMeshCopyFlags::AllFlags);
 }
 
 
@@ -138,87 +101,9 @@ TSharedPtr<FMesh> FMesh::Clone(EMeshCopyFlags InFlags) const
 
 	TSharedPtr<FMesh> Result = MakeShared<FMesh>();
 
-    Result->InternalId = InternalId;
-	Result->MeshIDPrefix = MeshIDPrefix;
-	Result->Flags = Flags;
-	Result->ReferenceID = ReferenceID;
-
-	if (EnumHasAnyFlags(InFlags, EMeshCopyFlags::WithSurfaces))
-	{
-		Result->Surfaces = Surfaces;
-	}
-
-	if (EnumHasAnyFlags(InFlags, EMeshCopyFlags::WithSkeleton))
-	{
-		Result->Skeleton = Skeleton;
-	}
-
-	if (EnumHasAnyFlags(InFlags, EMeshCopyFlags::WithPhysicsBody))
-	{
-		Result->PhysicsBody = PhysicsBody;
-	}
-
-	if (EnumHasAnyFlags(InFlags, EMeshCopyFlags::WithTags))
-	{
-		Result->Tags = Tags;
-	}
+	Result->CopyFrom(*this, InFlags);
 
-	if (EnumHasAnyFlags(InFlags, EMeshCopyFlags::WithStreamedResources))
-	{
-		Result->StreamedResources = StreamedResources;
-	}
-
-    // Clone the main buffers
-	if (EnumHasAnyFlags(InFlags, EMeshCopyFlags::WithVertexBuffers))
-    {
-		Result->VertexBuffers = VertexBuffers;
-	}
-
-	if (EnumHasAnyFlags(InFlags, EMeshCopyFlags::WithIndexBuffers))
-    {
-		Result->IndexBuffers = IndexBuffers;
-	}
-
-	// Clone additional buffers
-	if (EnumHasAnyFlags(InFlags, EMeshCopyFlags::WithAdditionalBuffers))
-	{
-		Result->AdditionalBuffers = AdditionalBuffers;
-	}
-
-    // Clone the layout	
-	if (EnumHasAnyFlags(InFlags, EMeshCopyFlags::WithLayouts))
-	{
-		Result->Layouts = Layouts;
-	}
-    // The skeleton is not cloned because it is not owned by this mesh and it is always assumed
-    // to be shared.
-
-	// physics body doen't need to be deep cloned either as they are also assumed to be shared.
-	
-	// Clone bone poses
-	if (EnumHasAnyFlags(InFlags, EMeshCopyFlags::WithPoses))
-	{
-		Result->BonePoses = BonePoses;
-	}
-
-	// Clone BoneMap
-	if (EnumHasAnyFlags(InFlags, EMeshCopyFlags::WithBoneMap))
-	{
-		Result->BoneMap = BoneMap;
-	}
-
-	// Clone SkeletonIDs
-	if (EnumHasAnyFlags(InFlags, EMeshCopyFlags::WithSkeletonIDs))
-	{
-		Result->SkeletonIDs = SkeletonIDs;
-	}
-
-	if (EnumHasAnyFlags(InFlags, EMeshCopyFlags::WithAdditionalPhysics))
-	{
-		Result->AdditionalPhysicsBodies = AdditionalPhysicsBodies;
-	}
-
-    return Result;
+	return Result;
 }
 
 
@@ -1491,7 +1376,6 @@ void FMesh::ResetStaticFormatFlags() const
 	}
 }
 
-
 namespace
 {
     void LogBuffer(FString& Out, const FMeshBufferSet& BufferSet, int32 BufferElementLimit)
diff --git a/Engine/Plugins/Mutable/Source/MutableRuntime/Private/MuR/MeshBufferSet.cpp b/Engine/Plugins/Mutable/Source/MutableRuntime/Private/MuR/MeshBufferSet.cpp
--- a/Engine/Plugins/Mutable/Source/MutableRuntime/Private/MuR/MeshBufferSet.cpp
+++ b/Engine/Plugins/Mutable/Source/MutableRuntime/Private/MuR/MeshBufferSet.cpp
@@ -17,6 +17,7 @@ namespace mu
 	MUTABLE_IMPLEMENT_POD_VECTOR_SERIALISABLE(FMeshBufferChannel);
 	MUTABLE_IMPLEMENT_ENUM_SERIALISABLE(EMeshBufferFormat);
 	MUTABLE_IMPLEMENT_ENUM_SERIALISABLE(EMeshBufferSemantic);
+	MUTABLE_IMPLEMENT_ENUM_SERIALISABLE(EMeshBufferSetFlags);
 
 	
 	static FMeshBufferFormatData s_meshBufferFormatData[] = // EMeshBufferFormat::Count entries
@@ -65,12 +66,14 @@ namespace mu
 	{
 		Arch << ElementCount;
 		Arch << Buffers;
+		Arch << Flags;
 	}
 
 	void FMeshBufferSet::Unserialise(FInputArchive& Arch)
 	{
 		Arch >> ElementCount;
 		Arch >> Buffers;
+		Arch >> Flags;
 	}
 
 
@@ -84,25 +87,29 @@ namespace mu
 		check(Count >= 0);
 		LLM_SCOPE_BYNAME(TEXT("MutableRuntime"));
 
-		// If the new size is 0, allow shrinking.
-		// TODO: Add better shrink policy or let the user decide. Denying shrinking
-		// unconditionally could mean having small meshes that use lots of memory. For now
-		// allow it if no other allocation will be done.  
-		const EAllowShrinking AllowShrinking = (Count == 0) ? EAllowShrinking::Yes : EAllowShrinking::No;
-
-		for (FMeshBuffer& Buffer : Buffers)
+		const bool bAllocateMemory = !IsDescriptor();
+		if (bAllocateMemory)
 		{
-			if (MemoryInitPolicy == EMemoryInitPolicy::Uninitialized)
-			{
-				Buffer.Data.SetNumUninitialized(Buffer.ElementSize * Count, AllowShrinking);
-			}
-			else if (MemoryInitPolicy == EMemoryInitPolicy::Zeroed)
-			{
-				Buffer.Data.SetNumZeroed(Buffer.ElementSize * Count, AllowShrinking);
-			}
-			else
+			// If the new size is 0, allow shrinking.
+			// TODO: Add better shrink policy or let the user decide. Denying shrinking
+			// unconditionally could mean having small meshes that use lots of memory. For now
+			// allow it if no other allocation will be done.  
+			
+			const EAllowShrinking AllowShrinking = (Count == 0) ? EAllowShrinking::Yes : EAllowShrinking::No;
+			for (FMeshBuffer& Buffer : Buffers)
 			{
-				check(false);
+				if (MemoryInitPolicy == EMemoryInitPolicy::Uninitialized)
+				{
+					Buffer.Data.SetNumUninitialized(Buffer.ElementSize * Count, AllowShrinking);
+				}
+				else if (MemoryInitPolicy == EMemoryInitPolicy::Zeroed)
+				{
+					Buffer.Data.SetNumZeroed(Buffer.ElementSize * Count, AllowShrinking);
+				}
+				else
+				{
+					check(false);
+				}
 			}
 		}
 		
@@ -190,6 +197,7 @@ namespace mu
 
 		FMeshBuffer& Buffer = Buffers[BufferIndex];
 
+
 		int32 MinElemSize = 0;
 		Buffer.Channels.SetNum(ChannelCount);
 		for (int32 ChannelIndex = 0; ChannelIndex < ChannelCount; ++ChannelIndex)
@@ -208,18 +216,23 @@ namespace mu
 		// Set the user specified element size, or enlarge it if it was too small
 		Buffer.ElementSize = FMath::Max(ElementSize, MinElemSize);
 
+		bool bAllocateMemory = !IsDescriptor(); 
+
 		// Update the buffer data
-		if (MemoryInitPolicy == EMemoryInitPolicy::Uninitialized)
-		{
-			Buffer.Data.SetNumUninitialized(Buffer.ElementSize * ElementCount, EAllowShrinking::No);
-		}
-		else if (MemoryInitPolicy == EMemoryInitPolicy::Zeroed)
-		{
-			Buffer.Data.SetNumZeroed(Buffer.ElementSize * ElementCount, EAllowShrinking::No);
-		}
-		else
+		if (bAllocateMemory)
 		{
-			check(false);
+			if (MemoryInitPolicy == EMemoryInitPolicy::Uninitialized)
+			{
+				Buffer.Data.SetNumUninitialized(Buffer.ElementSize * ElementCount, EAllowShrinking::No);
+			}
+			else if (MemoryInitPolicy == EMemoryInitPolicy::Zeroed)
+			{
+				Buffer.Data.SetNumZeroed(Buffer.ElementSize * ElementCount, EAllowShrinking::No);
+			}
+			else
+			{
+				check(false);
+			}
 		}
 	}
 
@@ -259,6 +272,8 @@ namespace mu
 
 	uint8* FMeshBufferSet::GetBufferData(int32 BufferIndex)
 	{
+		check(!IsDescriptor());
+
 		check(Buffers.IsValidIndex(BufferIndex));
 		uint8* ResultPtr = Buffers[BufferIndex].Data.GetData();
 		return ResultPtr;
@@ -373,10 +388,17 @@ namespace mu
 	{
 		int32 Result = 0;
 
+		const bool bIsDescriptor = IsDescriptor();
+
 		const int32 NumBuffers = Buffers.Num();
 		for (int32 BufferIndex = 0; BufferIndex < NumBuffers; ++BufferIndex)
 		{
-			Result += Buffers[BufferIndex].ElementSize * ElementCount;
+			Result += sizeof(FMeshBufferChannel) * Buffers[BufferIndex].Channels.Num();
+
+			if (!bIsDescriptor)
+			{
+				Result += Buffers[BufferIndex].ElementSize * ElementCount;
+			}
 		}
 
 		return Result;
@@ -397,6 +419,7 @@ namespace mu
 
 	void FMeshBufferSet::CopyElement(uint32 FromIndex, uint32 ToIndex)
 	{
+		check(!IsDescriptor())
 		check(FromIndex < ElementCount);
 		check(ToIndex < ElementCount);
 
@@ -615,6 +638,8 @@ namespace mu
 
 	void FMeshBufferSet::UpdateOffsets(int32 BufferIndex)
 	{
+		checkf(Buffers[BufferIndex].Data.IsEmpty(), TEXT("UpdateOffsets called on a non empty buffer set. This is not supported."));
+
 		uint32 Offset = 0;
 		for (FMeshBufferChannel& Channel : Buffers[BufferIndex].Channels)
 		{
@@ -656,6 +681,11 @@ namespace mu
 		return false;
 	}
 
+	bool FMeshBufferSet::IsDescriptor() const
+	{
+		return EnumHasAnyFlags(Flags, EMeshBufferSetFlags::IsDescriptor);
+	}
+
 	void FMeshBuffer::Serialise(FOutputArchive& Arch) const
 	{
 		Arch << Channels;
diff --git a/Engine/Plugins/Mutable/Source/MutableRuntime/Private/MuR/Model.cpp b/Engine/Plugins/Mutable/Source/MutableRuntime/Private/MuR/Model.cpp
--- a/Engine/Plugins/Mutable/Source/MutableRuntime/Private/MuR/Model.cpp
+++ b/Engine/Plugins/Mutable/Source/MutableRuntime/Private/MuR/Model.cpp
@@ -25,6 +25,7 @@ namespace mu
 	MUTABLE_IMPLEMENT_POD_SERIALISABLE(FRomDataRuntime);
 	MUTABLE_IMPLEMENT_POD_SERIALISABLE(FRomDataCompile);
 	MUTABLE_IMPLEMENT_POD_SERIALISABLE(FImageLODRange);
+	MUTABLE_IMPLEMENT_POD_SERIALISABLE(FMeshContentRange);
 	MUTABLE_IMPLEMENT_POD_SERIALISABLE(FConstantResourceIndex);
 	MUTABLE_IMPLEMENT_POD_VECTOR_SERIALISABLE(FConstantResourceIndex);
 
diff --git a/Engine/Plugins/Mutable/Source/MutableRuntime/Private/MuR/OpMeshApplyPose.h b/Engine/Plugins/Mutable/Source/MutableRuntime/Private/MuR/OpMeshApplyPose.h
--- a/Engine/Plugins/Mutable/Source/MutableRuntime/Private/MuR/OpMeshApplyPose.h
+++ b/Engine/Plugins/Mutable/Source/MutableRuntime/Private/MuR/OpMeshApplyPose.h
@@ -11,6 +11,39 @@
 
 namespace mu
 {
+	inline void SetPoseAsReference(FMesh* InOutResult, const FMesh* PoseMesh)
+	{
+		// Set the pose as the Result reference pose.
+		// PoseMesh PoseXform can be decomposed to PoseXform = ModelPoseXform * ModelRefXform^-1, 
+		// BaseMesh Poses have ModelRefXform, to get the new ModelRefXform to be ModelPoseXform we need 
+		// to multiply PoseXform * ModelRefXform. 
+		// (Note, transform multiplication application is from right to left for this comment, but TTransform::operator* is reversed).
+		const int32 ResultNumBones = InOutResult->GetBonePoseCount();
+		for (int32 BoneIndex = 0; BoneIndex < ResultNumBones; ++BoneIndex)
+		{
+			const FBoneName BoneName = InOutResult->GetBonePoseId(BoneIndex);
+			const int32 FoundPoseIndex = PoseMesh->FindBonePose(BoneName);
+			if (FoundPoseIndex == INDEX_NONE)
+			{
+				continue;
+			}
+
+			FTransform3f ModelRefTransform;
+			InOutResult->GetBonePoseTransform(BoneIndex, ModelRefTransform);
+
+			FTransform3f PoseTransform;
+			PoseMesh->GetBonePoseTransform(FoundPoseIndex, PoseTransform); 
+
+			// Adding the Reshaped flag will prioritize this bone over other bones without the flag
+			// when merging poses with the same bone.
+			// TODO: Add another flag to indicate this case or generalize the Reshaped flag. 
+			EBoneUsageFlags UsageFlags = InOutResult->GetBoneUsageFlags(BoneIndex);
+			EnumAddFlags(UsageFlags, EBoneUsageFlags::Reshaped);
+			
+			InOutResult->SetBonePose(BoneIndex, BoneName, ModelRefTransform * PoseTransform, UsageFlags);
+		}
+	}
+
 	/**
     * Reference version
     */
@@ -105,34 +138,7 @@ namespace mu
 
 		Result->CopyFrom(*BaseMesh);
 
-		// Set the pose as the Result reference pose.
-		// PoseMesh PoseXform can be decomposed to PoseXform = ModelPoseXform * ModelRefXform^-1, 
-		// BaseMesh Poses have ModelRefXform, to get the new ModelRefXform to be ModelPoseXform we need 
-		// to multiply PoseXform * ModelRefXform. 
-		// (Note, transform multiplication application is from right to left for this comment, but TTransform::operator* is reversed).
-		const int32 ResultNumBones = Result->GetBonePoseCount();
-		for (int32 BoneIndex = 0; BoneIndex < ResultNumBones; ++BoneIndex)
-		{
-			const FBoneName BoneName = Result->GetBonePoseId(BoneIndex);
-			const int32 FoundPoseIndex = PoseMesh->FindBonePose(BoneName);
-			if (FoundPoseIndex == INDEX_NONE)
-			{
-				continue;
-			}
-
-			FTransform3f ModelRefTransform;
-			Result->GetBonePoseTransform(BoneIndex, ModelRefTransform);
-
-			FTransform3f PoseTransform;
-			PoseMesh->GetBonePoseTransform(FoundPoseIndex, PoseTransform); 
-
-			// Reshape flag is added so the bone is prioritized in case of merge conflict.
-			// TODO: Add another flag to indicate this case or generalize the Reshaped flag. 
-			EBoneUsageFlags UsageFlags = BaseMesh->GetBoneUsageFlags(BoneIndex);
-			EnumAddFlags(UsageFlags, EBoneUsageFlags::Reshaped);
-			
-			Result->SetBonePose(BoneIndex, BoneName, ModelRefTransform * PoseTransform, UsageFlags);
-		}
+		SetPoseAsReference(Result, PoseMesh);
 
 		MeshBufferIterator<EMeshBufferFormat::Float32, float, 3> TargetPositionIterBegin(Result->VertexBuffers, EMeshBufferSemantic::Position, 0);
 		
diff --git a/Engine/Plugins/Mutable/Source/MutableRuntime/Private/MuR/OpMeshFormat.cpp b/Engine/Plugins/Mutable/Source/MutableRuntime/Private/MuR/OpMeshFormat.cpp
--- a/Engine/Plugins/Mutable/Source/MutableRuntime/Private/MuR/OpMeshFormat.cpp
+++ b/Engine/Plugins/Mutable/Source/MutableRuntime/Private/MuR/OpMeshFormat.cpp
@@ -476,15 +476,22 @@ namespace mu
 			}
 		}
 
-
-		// For every vertex buffer in result
-		int32 vCount = Source.GetElementCount();
-		Result.SetElementCount(vCount);
-		for (int32 b = 0; b < Result.GetBufferCount(); ++b)
+		if (Source.IsDescriptor())
 		{
-			MeshFormatBuffer(Source, Result.Buffers[b], 0, bIsVertexBuffer, IDPrefix);
+			EnumAddFlags(Result.Flags, EMeshBufferSetFlags::IsDescriptor); 
 		}
 
+		int32 VertexCount = Source.GetElementCount();
+		Result.SetElementCount(VertexCount);
+
+		if (!Source.IsDescriptor())
+		{
+			// For every vertex buffer in result
+			for (int32 BufferIndex = 0; BufferIndex < Result.GetBufferCount(); ++BufferIndex)
+			{
+				MeshFormatBuffer(Source, Result.Buffers[BufferIndex], 0, bIsVertexBuffer, IDPrefix);
+			}
+		}
 
 		// Detect internal system buffers and clone them unmodified.
 		if (bKeepSystemBuffers)
@@ -532,7 +539,94 @@ namespace mu
 
 	}
 
+	int32 FindHighestBoneIndexInUse(const FMesh* InMesh, const FMeshBufferChannel& Channel)
+	{
+		const FMeshBufferSet& VertexBuffers = InMesh->GetVertexBuffers();
+
+		// If InMesh buffers are descriptors, use the bone map to get an upper bound of the highest index
+		// we will find in the mesh. Otherwise, find it from the vertex buffers data.
+		if (VertexBuffers.IsDescriptor())
+		{
+			return FMath::Max(0, InMesh->BoneMap.Num() - 1);
+		}
 
+		const UntypedMeshBufferIteratorConst BoneIndicesBegin(
+				VertexBuffers, EMeshBufferSemantic::BoneIndices, Channel.SemanticIndex);
+
+		if (!BoneIndicesBegin.ptr())
+		{
+			return 0;
+		}
+
+		const int32 NumVertices = VertexBuffers.GetElementCount();
+		const int32 NumInfluences = BoneIndicesBegin.GetComponents();
+		
+		int32 HighestBoneIndex = 0;
+		for (int32 VertexIndex = 0; VertexIndex < NumVertices; ++VertexIndex)
+		{
+			// If MAX_TOTAL_INFLUENCES ever changed, the next line would no longer work or compile and 
+			// GetAsVec12i would need to be changed accordingly
+			int32 VertexInfluences[MAX_TOTAL_INFLUENCES];
+			(BoneIndicesBegin + VertexIndex).GetAsInt32Vec(VertexInfluences, MAX_TOTAL_INFLUENCES);
+			
+			for (int32 InfluenceIndex = 0; InfluenceIndex < NumInfluences; ++InfluenceIndex)
+			{
+				HighestBoneIndex = FMath::Max(HighestBoneIndex, VertexInfluences[InfluenceIndex]);
+			}
+		}
+
+		return HighestBoneIndex;
+	}
+
+	/** Updates the InOutFormat bone index semantic format if it cannot represent the highest bone index in use. */
+	void UpdateBoneIndexFormatIfNeeded(const FMesh* InMesh, FMeshBufferSet& InOutFormatBufferSet)
+	{
+		const FMeshBufferSet& VertexBuffers = InMesh->GetVertexBuffers();
+
+		const int32 BufferCount = VertexBuffers.GetBufferCount();
+		for (int32 BufferIndex = 0; BufferIndex < BufferCount; ++BufferIndex)
+		{
+			const int32 ChannelCount = VertexBuffers.GetBufferChannelCount(BufferIndex);
+			for (int32 ChannelIndex = 0; ChannelIndex < ChannelCount; ++ChannelIndex)
+			{
+				const FMeshBufferChannel& Channel = VertexBuffers.Buffers[BufferIndex].Channels[ChannelIndex];
+
+				if (Channel.Semantic == EMeshBufferSemantic::BoneIndices)
+				{
+					int32 ResultBuf = 0;
+					int32 ResultChan = 0;
+					InOutFormatBufferSet.FindChannel(EMeshBufferSemantic::BoneIndices, Channel.SemanticIndex, &ResultBuf, &ResultChan);
+					if (ResultBuf >= 0)
+					{
+						int32 HighestBoneIndexInUse = FindHighestBoneIndexInUse(InMesh, Channel);
+						check(HighestBoneIndexInUse >= 0);
+
+						EMeshBufferFormat& BoneIndexFormat = InOutFormatBufferSet.Buffers[ResultBuf].Channels[ResultChan].Format;
+						if (HighestBoneIndexInUse > 0xffff && (BoneIndexFormat == EMeshBufferFormat::UInt8 || BoneIndexFormat == EMeshBufferFormat::UInt16))
+						{
+							BoneIndexFormat = EMeshBufferFormat::UInt32;
+							InOutFormatBufferSet.UpdateOffsets(ResultBuf);
+						}
+						else if (HighestBoneIndexInUse > 0x7fff && (BoneIndexFormat == EMeshBufferFormat::UInt8 || BoneIndexFormat == EMeshBufferFormat::UInt16))
+						{
+							BoneIndexFormat = EMeshBufferFormat::UInt32;
+							InOutFormatBufferSet.UpdateOffsets(ResultBuf);
+						}
+						else if (HighestBoneIndexInUse > 0xff && BoneIndexFormat == EMeshBufferFormat::UInt8)
+						{
+							BoneIndexFormat = EMeshBufferFormat::UInt16;
+							InOutFormatBufferSet.UpdateOffsets(ResultBuf);
+						}
+						else if (HighestBoneIndexInUse > 0x7f && BoneIndexFormat == EMeshBufferFormat::UInt8)
+						{
+							BoneIndexFormat = EMeshBufferFormat::UInt16;
+							InOutFormatBufferSet.UpdateOffsets(ResultBuf);
+						}
+					}
+				}
+			}
+		}
+	}
 
 	//-------------------------------------------------------------------------------------------------
 	void MeshFormat
@@ -567,71 +661,10 @@ namespace mu
 		Result->CopyFrom(*Format);
 		Result->MeshIDPrefix = Source->MeshIDPrefix;
 
-		// Make sure that the bone indices will fit in this format, or extend it.
 		if (bFormatVertices)
 		{
-			const FMeshBufferSet& VertexBuffers = Source->GetVertexBuffers();
-
-			const int32 BufferCount = VertexBuffers.GetBufferCount();
-			for (int32 BufferIndex = 0; BufferIndex < BufferCount; ++BufferIndex)
-			{
-				const int32 ChannelCount = VertexBuffers.GetBufferChannelCount(BufferIndex);
-				for (int32 ChannelIndex = 0; ChannelIndex < ChannelCount; ++ChannelIndex)
-				{
-					const FMeshBufferChannel& Channel = VertexBuffers.Buffers[BufferIndex].Channels[ChannelIndex];
-
-					if (Channel.Semantic == EMeshBufferSemantic::BoneIndices)
-					{
-						int32 resultBuf = 0;
-						int32 resultChan = 0;
-						FMeshBufferSet& formBuffs = Result->GetVertexBuffers();
-						formBuffs.FindChannel(EMeshBufferSemantic::BoneIndices, Channel.SemanticIndex, &resultBuf, &resultChan);
-						if (resultBuf >= 0)
-						{
-							UntypedMeshBufferIteratorConst it(VertexBuffers, EMeshBufferSemantic::BoneIndices, Channel.SemanticIndex);
-							int32 maxBoneIndex = 0;
-							for (int32 v = 0; v < VertexBuffers.GetElementCount(); ++v)
-							{
-								// If MAX_TOTAL_INFLUENCES ever changed, the next line would no longer work or compile and 
-								// GetAsVec12i would need to be changed accordingly
-								int32 va[MAX_TOTAL_INFLUENCES];
-								it.GetAsInt32Vec(va, MAX_TOTAL_INFLUENCES);
-								for (int32 c = 0; c < it.GetComponents(); ++c)
-								{
-									maxBoneIndex = FMath::Max(maxBoneIndex, va[c]);
-								}
-								++it;
-							}
-
-							EMeshBufferFormat& format = formBuffs.Buffers[resultBuf].Channels[resultChan].Format;
-							if (maxBoneIndex > 0xffff && (format == EMeshBufferFormat::UInt8 || format == EMeshBufferFormat::UInt16))
-							{
-								format = EMeshBufferFormat::UInt32;
-								formBuffs.UpdateOffsets(resultBuf);
-							}
-							else if (maxBoneIndex > 0x7fff && (format == EMeshBufferFormat::Int8 || format == EMeshBufferFormat::Int16))
-							{
-								format = EMeshBufferFormat::UInt32;
-								formBuffs.UpdateOffsets(resultBuf);
-							}
-							else if (maxBoneIndex > 0xff && format == EMeshBufferFormat::UInt8)
-							{
-								format = EMeshBufferFormat::UInt16;
-								formBuffs.UpdateOffsets(resultBuf);
-							}
-							else if (maxBoneIndex > 0x7f && format == EMeshBufferFormat::Int8)
-							{
-								format = EMeshBufferFormat::Int16;
-								formBuffs.UpdateOffsets(resultBuf);
-							}
-						}
-					}
-				}
-			}
-		}
+			UpdateBoneIndexFormatIfNeeded(Source, Result->GetVertexBuffers());
 
-		if (bFormatVertices)
-		{
 			FormatBufferSet(Source->GetVertexBuffers(), Result->GetVertexBuffers(),
 				bKeepSystemBuffers, bIgnoreMissingChannels, true, Source->MeshIDPrefix);
 		}
@@ -692,6 +725,12 @@ namespace mu
 
 		FMeshBufferSet& VertexBuffers = InMesh->VertexBuffers;
 
+		// Ignore the operation if the MeshBuffersSet is a descriptor.
+		if (VertexBuffers.IsDescriptor())
+		{
+			return;
+		}
+
 		// Reduce the number of influences if possible
 		constexpr int32 SemanticIndex = 0;
 		
diff --git a/Engine/Plugins/Mutable/Source/MutableRuntime/Private/MuR/OpMeshOptimizeSkinning.h b/Engine/Plugins/Mutable/Source/MutableRuntime/Private/MuR/OpMeshOptimizeSkinning.h
--- a/Engine/Plugins/Mutable/Source/MutableRuntime/Private/MuR/OpMeshOptimizeSkinning.h
+++ b/Engine/Plugins/Mutable/Source/MutableRuntime/Private/MuR/OpMeshOptimizeSkinning.h
@@ -75,6 +75,7 @@ namespace mu
 		// TODO: Replace by MeshFormatInPlace once implemented
 		const int32 VertexBuffersCount = InMeshVertexBuffers.GetBufferCount();
 		const int32 ElementCount = InMeshVertexBuffers.GetElementCount();
+		const EMeshBufferSetFlags BufferSetFlags = InMeshVertexBuffers.Flags;
 
 		// Clone mesh without VertexBuffers, they will be copied manually.
 		constexpr EMeshCopyFlags CopyFlags = ~EMeshCopyFlags::WithVertexBuffers;
@@ -82,6 +83,8 @@ namespace mu
 
 		mu::FMeshBufferSet& VertexBuffers = Result->GetVertexBuffers();
 
+
+		VertexBuffers.Flags = InMeshVertexBuffers.Flags;
 		VertexBuffers.Buffers.Reserve(VertexBuffersCount);
 		VertexBuffers.SetElementCount(ElementCount);
 
@@ -135,50 +138,54 @@ namespace mu
 				// This buffer has a BoneIndex channel that needs to be manually fixed. 
 				VertexBuffers.SetBufferCount(BufferIndex + 1);
 				VertexBuffers.SetBuffer(BufferIndex, ElementSize, ChannelsCount, Semantics.GetData(), SemanticIndices.GetData(), Formats.GetData(), Components.GetData(), Offsets.GetData());
-				uint8* Data = VertexBuffers.GetBufferData(BufferIndex);
-
-				const uint8* SourceData = (const uint8*)SourceBuffer.Data.GetData();
-				const int32 SourceBoneIndexSize = GetMeshFormatData(SourceBoneIndexFormat).SizeInBytes;
+				
+				if (!VertexBuffers.IsDescriptor())
+				{
+					uint8* Data = VertexBuffers.GetBufferData(BufferIndex);
 
-				const int32 BoneIndexSize = GetMeshFormatData(mu::EMeshBufferFormat::UInt8).SizeInBytes;
-				const int32 BoneIndexComponentCount = Components[BoneIndexChannelIndex];
-				const int32 BoneIndexChannelOffset = Offsets[BoneIndexChannelIndex];
+					const uint8* SourceData = (const uint8*)SourceBuffer.Data.GetData();
+					const int32 SourceBoneIndexSize = GetMeshFormatData(SourceBoneIndexFormat).SizeInBytes;
 
-				const int32 TailSize = AuxOffset - (BoneIndexChannelOffset + BoneIndexComponentCount * BoneIndexSize);
+					const int32 BoneIndexSize = GetMeshFormatData(mu::EMeshBufferFormat::UInt8).SizeInBytes;
+					const int32 BoneIndexComponentCount = Components[BoneIndexChannelIndex];
+					const int32 BoneIndexChannelOffset = Offsets[BoneIndexChannelIndex];
 
+					const int32 TailSize = AuxOffset - (BoneIndexChannelOffset + BoneIndexComponentCount * BoneIndexSize);
 
-				for (const FMeshSurface& Surface : InMesh->Surfaces)
-				{
-					const uint32 NumBonesInBoneMap = Surface.BoneMapCount;
 
-					for (const FSurfaceSubMesh SubMesh : Surface.SubMeshes)
+					for (const FMeshSurface& Surface : InMesh->Surfaces)
 					{
-						for (int32 VertexIndex = SubMesh.VertexBegin; VertexIndex < SubMesh.VertexEnd; ++VertexIndex)
-						{
-							FMemory::Memcpy(Data, SourceData, BoneIndexChannelOffset);
-							Data += BoneIndexChannelOffset;
-							SourceData += BoneIndexChannelOffset;
+						const uint32 NumBonesInBoneMap = Surface.BoneMapCount;
 
-							for (int32 ComponentIndex = 0; ComponentIndex < BoneIndexComponentCount; ++ComponentIndex)
+						for (const FSurfaceSubMesh SubMesh : Surface.SubMeshes)
+						{
+							for (int32 VertexIndex = SubMesh.VertexBegin; VertexIndex < SubMesh.VertexEnd; ++VertexIndex)
 							{
-								const uint16 SourceIndex = *((const uint16*)SourceData);
-								if (SourceIndex < NumBonesInBoneMap)
-								{
-									*Data = (uint8)SourceIndex;
-								}
-								else
+								FMemory::Memcpy(Data, SourceData, BoneIndexChannelOffset);
+								Data += BoneIndexChannelOffset;
+								SourceData += BoneIndexChannelOffset;
+
+								for (int32 ComponentIndex = 0; ComponentIndex < BoneIndexComponentCount; ++ComponentIndex)
 								{
-									*Data = 0;
+									const uint16 SourceIndex = *((const uint16*)SourceData);
+									if (SourceIndex < NumBonesInBoneMap)
+									{
+										*Data = (uint8)SourceIndex;
+									}
+									else
+									{
+										*Data = 0;
+									}
+									
+									Data += BoneIndexSize;
+									SourceData += SourceBoneIndexSize;
 								}
-								
-								Data += BoneIndexSize;
-								SourceData += SourceBoneIndexSize;
-							}
 
 
-							FMemory::Memcpy(Data, SourceData, TailSize);
-							Data += TailSize;
-							SourceData += TailSize;
+								FMemory::Memcpy(Data, SourceData, TailSize);
+								Data += TailSize;
+								SourceData += TailSize;
+							}
 						}
 					}
 				}
diff --git a/Engine/Plugins/Mutable/Source/MutableRuntime/Private/MuR/System.cpp b/Engine/Plugins/Mutable/Source/MutableRuntime/Private/MuR/System.cpp
--- a/Engine/Plugins/Mutable/Source/MutableRuntime/Private/MuR/System.cpp
+++ b/Engine/Plugins/Mutable/Source/MutableRuntime/Private/MuR/System.cpp
@@ -519,10 +519,9 @@ namespace mu
 				UE::Tasks::ETaskPriority::Inherit,
 				UE::Tasks::EExtendedTaskPriority::Inline);
 
-	}
-
+	}	
 
-    TSharedPtr<const FMesh> FSystem::GetMeshInline(FInstance::FID instanceID, FResourceID MeshId)
+    TSharedPtr<const FMesh> FSystem::GetMeshInline(FInstance::FID instanceID, FResourceID MeshId, EMeshContentFlags MeshContentFilter)
     {
 		LLM_SCOPE_BYNAME(TEXT("MutableRuntime"));
 		MUTABLE_CPUPROFILER_SCOPE(SystemGetMesh);
@@ -535,7 +534,11 @@ namespace mu
 		m_pD->WorkingMemoryManager.CurrentInstanceCache = LiveInstance->Cache;
 
 		OP::ADDRESS RootAddress = GetResourceIDRoot(MeshId);
-		Result = m_pD->BuildMesh(LiveInstance->Model, LiveInstance->OldParameters.Get(), RootAddress);
+		Result = m_pD->BuildMesh(
+				LiveInstance->Model, 
+				LiveInstance->OldParameters.Get(), 
+				RootAddress, 
+				MeshContentFilter);
 
 		// If the mesh is null it means empty, but we still need to return a valid one
 		if (!Result)
@@ -547,7 +550,7 @@ namespace mu
 		return Result;
 	}
 
-	UE::Tasks::TTask<TSharedPtr<const FMesh>> FSystem::GetMesh(FInstance::FID instanceID, FResourceID MeshId)
+	UE::Tasks::TTask<TSharedPtr<const FMesh>> FSystem::GetMesh(FInstance::FID instanceID, FResourceID MeshId, EMeshContentFlags MeshContentFilter)
     {
 		LLM_SCOPE_BYNAME(TEXT("MutableRuntime"));
 		MUTABLE_CPUPROFILER_SCOPE(SystemGetMesh);
@@ -570,22 +573,32 @@ namespace mu
 
 			return UE::Tasks::MakeCompletedTask<TSharedPtr<const FMesh>>(new FMesh());
 		}
-		
+	
+		uint8 ExecutionOptions = static_cast<uint8>(MeshContentFilter);
 		TSharedRef<CodeRunner> Runner = CodeRunner::Create(
-				m_pD->Settings, m_pD, EExecutionStrategy::MinimizeMemory, LiveInstance->Model, LiveInstance->OldParameters.Get(), RootAddress, FSystem::AllLODs, 0, 0, FScheduledOp::EType::Full);
+				m_pD->Settings,
+				m_pD, 
+				EExecutionStrategy::MinimizeMemory, 
+				LiveInstance->Model, 
+				LiveInstance->OldParameters.Get(), 
+				RootAddress, 
+				FSystem::AllLODs, 
+				ExecutionOptions, 
+				0, 
+				FScheduledOp::EType::Full);
 		
 		constexpr bool bForceInlineExecution = false;
 		UE::Tasks::FTask RunnerCompletionEvent = Runner->StartRun(bForceInlineExecution);
 
 		return UE::Tasks::Launch(TEXT("FSystem::GetMeshResultTask"),
-				[SystemPrivate = m_pD, Runner, RootAddress]() -> TSharedPtr<const FMesh>
+				[SystemPrivate = m_pD, Runner, RootAddress, ExecutionOptions]() -> TSharedPtr<const FMesh>
 				{
 					TSharedPtr<const FMesh> Result;
 
 					SystemPrivate->bUnrecoverableError = Runner->bUnrecoverableError;
 					if (!Runner->bUnrecoverableError)
 					{
-						Result = SystemPrivate->WorkingMemoryManager.LoadMesh(FCacheAddress(RootAddress, 0, 0), true);
+						Result = SystemPrivate->WorkingMemoryManager.LoadMesh(FCacheAddress(RootAddress, 0, ExecutionOptions), true);
 					}
 					
 					SystemPrivate->WorkingMemoryManager.EndRunnerThread();
@@ -843,10 +856,10 @@ namespace mu
 
 	//---------------------------------------------------------------------------------------------
 	void FSystem::Private::RunCode(const TSharedPtr<const FModel>& InModel,
-		const FParameters* InParameters, OP::ADDRESS InCodeRoot, uint32 InLODs, uint8 executionOptions, int32 InImageLOD)
+		const FParameters* InParameters, OP::ADDRESS InCodeRoot, uint32 InLODs, uint8 ExecutionOptions, int32 InImageLOD)
 	{
 		TSharedRef<CodeRunner> Runner = CodeRunner::Create(Settings, this, EExecutionStrategy::MinimizeMemory, InModel, InParameters, InCodeRoot, InLODs,
-			executionOptions, InImageLOD, FScheduledOp::EType::Full);
+			ExecutionOptions, InImageLOD, FScheduledOp::EType::Full);
 		
 		constexpr bool bForceInlineExecutution = true;
 		UE::Tasks::FTask RunnerCompletionEvent = Runner->StartRun(bForceInlineExecutution);
@@ -977,19 +990,21 @@ namespace mu
 	}
 
 
-	TSharedPtr<const FMesh> FSystem::Private::BuildMesh(const TSharedPtr<const FModel>& pModel, const FParameters* Params, OP::ADDRESS at)
+	TSharedPtr<const FMesh> FSystem::Private::BuildMesh(const TSharedPtr<const FModel>& Model, const FParameters* Params, OP::ADDRESS RootAddress, EMeshContentFlags MeshContentFilter)
 	{
 		WorkingMemoryManager.BeginRunnerThread();
 
 		TSharedPtr<const FMesh> Result;
 
-		mu::EOpType OpType = pModel->GetPrivate()->Program.GetOpType(at);
+		mu::EOpType OpType = Model->GetPrivate()->Program.GetOpType(RootAddress);
+
 		if (GetOpDataType(OpType) == EDataType::Mesh)
 		{
-			RunCode(pModel, Params, at);
+			uint8 ExecutionOptions = static_cast<uint8>(MeshContentFilter);
+			RunCode(Model, Params, RootAddress, FSystem::AllLODs, ExecutionOptions);
 			if (!bUnrecoverableError)
 			{
-				Result = WorkingMemoryManager.LoadMesh(FCacheAddress(at, 0, 0), true);
+				Result = WorkingMemoryManager.LoadMesh(FCacheAddress(RootAddress, 0, ExecutionOptions), true);
 			}	
 		}
 
diff --git a/Engine/Plugins/Mutable/Source/MutableRuntime/Public/MuR/Mesh.h b/Engine/Plugins/Mutable/Source/MutableRuntime/Public/MuR/Mesh.h
--- a/Engine/Plugins/Mutable/Source/MutableRuntime/Public/MuR/Mesh.h
+++ b/Engine/Plugins/Mutable/Source/MutableRuntime/Public/MuR/Mesh.h
@@ -140,6 +140,18 @@ namespace mu
 	
 	ENUM_CLASS_FLAGS(EMeshCopyFlags);
 
+	enum class EMeshContentFlags : uint8
+	{
+		None          = 0,
+		GeometryData  = 1 << 0,
+		PoseData      = 1 << 1,
+		PhysicsData   = 1 << 2,
+		MetaData      = 1 << 3,
+
+		LastFlag      = MetaData,
+		AllFlags      = (LastFlag << 1) - 1,
+	};
+	ENUM_CLASS_FLAGS(EMeshContentFlags);
 
 	/** Optimised mesh formats that are identified in some operations to chose a faster version. */
 	enum class EMeshFlags : uint32
@@ -358,7 +370,6 @@ namespace mu
 
         //! \}
 
-
         /** 
 		 * Get an internal identifier used to reference this mesh in operations like deferred
          * mesh building, or instance updating.
diff --git a/Engine/Plugins/Mutable/Source/MutableRuntime/Public/MuR/MeshBufferSet.h b/Engine/Plugins/Mutable/Source/MutableRuntime/Public/MuR/MeshBufferSet.h
--- a/Engine/Plugins/Mutable/Source/MutableRuntime/Public/MuR/MeshBufferSet.h
+++ b/Engine/Plugins/Mutable/Source/MutableRuntime/Public/MuR/MeshBufferSet.h
@@ -213,6 +213,13 @@ namespace mu
 		}
 	};
 
+	enum class EMeshBufferSetFlags : uint32
+	{
+		None = 0,
+		IsDescriptor = 1 << 0
+	};
+
+	ENUM_CLASS_FLAGS(EMeshBufferSetFlags);
 
 	/** Set of buffers storing mesh element data. Elements can be vertices, indices or faces. */
 	class MUTABLERUNTIME_API FMeshBufferSet
@@ -220,6 +227,7 @@ namespace mu
 	public:
 
 		uint32 ElementCount = 0;
+		EMeshBufferSetFlags Flags = EMeshBufferSetFlags::None; 
 		TArray<FMeshBuffer> Buffers;
 
 		void Serialise(FOutputArchive& Arch) const;
@@ -227,7 +235,10 @@ namespace mu
 
 		inline bool operator==(const FMeshBufferSet& Other) const
 		{
-			return ElementCount == Other.ElementCount && Buffers == Other.Buffers;
+			return 
+				ElementCount == Other.ElementCount && 
+				Buffers == Other.Buffers && 
+				Flags == Other.Flags;
 		}
 
 	public:
@@ -414,6 +425,9 @@ namespace mu
 		
 		/** Check that all channels of a specific semantic use the provided format. */
 		bool HasAnySemanticWithDifferentFormat(EMeshBufferSemantic Semantic, EMeshBufferFormat ExpectedFormat) const;
+
+		/** Check if this buffer set is only a format descriptor, e.i. reports num elements larger than 0 but does not hold any data*/ 
+		bool IsDescriptor() const;
 	};	
 
 	
@@ -421,5 +435,6 @@ namespace mu
 	MUTABLE_DEFINE_POD_VECTOR_SERIALISABLE(FMeshBufferChannel);
 	MUTABLE_DEFINE_ENUM_SERIALISABLE(EMeshBufferFormat);
 	MUTABLE_DEFINE_ENUM_SERIALISABLE(EMeshBufferSemantic);
+	MUTABLE_DEFINE_ENUM_SERIALISABLE(EMeshBufferSetFlags);
 }
 
diff --git a/Engine/Plugins/Mutable/Source/MutableRuntime/Public/MuR/System.h b/Engine/Plugins/Mutable/Source/MutableRuntime/Public/MuR/System.h
--- a/Engine/Plugins/Mutable/Source/MutableRuntime/Public/MuR/System.h
+++ b/Engine/Plugins/Mutable/Source/MutableRuntime/Public/MuR/System.h
@@ -3,6 +3,7 @@
 #pragma once
 
 #include "MuR/Image.h"
+#include "MuR/Mesh.h"
 #include "MuR/Instance.h"
 #include "MuR/RefCounted.h"
 #include "MuR/Settings.h"
@@ -160,7 +161,7 @@ namespace mu
 		UE::Tasks::TTask<TSharedPtr<const FImage>> GetImage(FInstance::FID InstanceID, FResourceID ImageId, int32 MipsToSkip = 0, int32 LOD = 0);
 
         //! Only valid between BeginUpdate and EndUpdate
-		UE::Tasks::TTask<TSharedPtr<const FMesh>> GetMesh(FInstance::FID InstanceID, FResourceID MeshId);
+		UE::Tasks::TTask<TSharedPtr<const FMesh>> GetMesh(FInstance::FID InstanceID, FResourceID MeshId, EMeshContentFlags MeshContentFilter = EMeshContentFlags::AllFlags);
 
 		//! Only valid between BeginUpdate and EndUpdate
 		//! Calculate the description of an image, without generating it.
@@ -173,7 +174,7 @@ namespace mu
 		TSharedPtr<const FImage> GetImageInline(FInstance::FID InstanceID, FResourceID ImageId, int32 MipsToSkip = 0, int32 LOD = 0);
 
         //! Only valid between BeginUpdate and EndUpdate
-		TSharedPtr<const FMesh> GetMeshInline(FInstance::FID InstanceID, FResourceID MeshId);
+		TSharedPtr<const FMesh> GetMeshInline(FInstance::FID InstanceID, FResourceID MeshId, EMeshContentFlags MeshContentFilter = EMeshContentFlags::AllFlags);
 
         //! Invalidate and free the last Instance data returned by a call to BeginUpdate with
         //! the same instance index. After a call to this method, that Instance cannot be used any
diff --git a/Engine/Plugins/Mutable/Source/MutableTools/Private/MuT/ASTOpConstantResource.cpp b/Engine/Plugins/Mutable/Source/MutableTools/Private/MuT/ASTOpConstantResource.cpp
--- a/Engine/Plugins/Mutable/Source/MutableTools/Private/MuT/ASTOpConstantResource.cpp
+++ b/Engine/Plugins/Mutable/Source/MutableTools/Private/MuT/ASTOpConstantResource.cpp
@@ -277,7 +277,186 @@ namespace mu
 
 	namespace
 	{
-		/** Adds a constant image data to a Program and returns its constant index. */
+ 		/** Adds a constant mesh data to a program and returns its constant index. */
+		int32 AddConstantMesh(FProgram& Program, const TSharedPtr<const FMesh>& MeshData, FLinkerOptions& Options)
+		{
+			auto AddMeshToProgram = [&Options, &Program](const TSharedPtr<FMesh>& Mesh)
+			{
+				// Use a map-based deduplication
+				int32 MeshIndex = -1;
+				TSharedPtr<const FMesh> MeshKey = Mesh;
+				const int32* MeshIndexPtr = Options.MeshConstantMap.Find(MeshKey);
+				if (!MeshIndexPtr)
+				{
+					MeshIndex = Program.ConstantMeshesPermanent.Add(Mesh);
+					Options.MeshConstantMap.Add(Mesh, MeshIndex);
+				}
+				else
+				{
+					MeshIndex = *MeshIndexPtr;
+				}
+
+				check(MeshIndex >= 0)
+				return Program.ConstantMeshContentIndices.Add(FConstantResourceIndex{(uint32)MeshIndex, 0});
+			};
+			
+			// Split generated mesh data in 4 parts. Geometry and Pose and Physiscs and Metadata.
+			// Inidices for a given rom are sorted by content flag value.
+			static_assert(EMeshContentFlags::GeometryData < EMeshContentFlags::PoseData);
+			static_assert(EMeshContentFlags::PoseData     < EMeshContentFlags::PhysicsData);
+			static_assert(EMeshContentFlags::PhysicsData  < EMeshContentFlags::MetaData);
+
+			int32 FirstIndex = Program.ConstantMeshContentIndices.Num();
+			EMeshContentFlags MeshContentFlags = EMeshContentFlags::None;	
+			// GeometryMesh
+			{
+				const EMeshCopyFlags GeometryDataCopyFlags = 
+						EMeshCopyFlags::WithSurfaces      | 
+						EMeshCopyFlags::WithVertexBuffers |
+						EMeshCopyFlags::WithIndexBuffers  |
+						EMeshCopyFlags::WithLayouts;
+
+				TSharedPtr<FMesh> MeshGeometryData = MeshData->Clone(GeometryDataCopyFlags);
+
+				// Copy geometry related additional buffers.
+				for (const TPair<EMeshBufferType, FMeshBufferSet>& AdditionalBuffer : MeshData->AdditionalBuffers)
+				{
+					const bool bIsGeometryBufferType = 
+							AdditionalBuffer.Key == EMeshBufferType::MeshLaplacianData    ||
+							AdditionalBuffer.Key == EMeshBufferType::MeshLaplacianOffsets ||
+							AdditionalBuffer.Key == EMeshBufferType::UniqueVertexMap;
+
+					if (bIsGeometryBufferType)
+					{
+						MeshGeometryData->AdditionalBuffers.Emplace(AdditionalBuffer);
+					}
+				}
+
+				MeshGeometryData->MeshIDPrefix = 0;
+				AddMeshToProgram(MeshGeometryData);	
+				EnumAddFlags(MeshContentFlags, EMeshContentFlags::GeometryData);
+			}
+
+			// Pose Mesh
+			{
+				const EMeshCopyFlags PoseDataCopyFlags = 
+						EMeshCopyFlags::WithPoses  |
+						EMeshCopyFlags::WithBoneMap;
+
+				TSharedPtr<FMesh> MeshPoseData = MeshData->Clone(PoseDataCopyFlags);
+
+				for (const TPair<EMeshBufferType, FMeshBufferSet>& AdditionalBuffer : MeshData->AdditionalBuffers)
+				{
+					const bool bIsPoseBufferType = 	
+							AdditionalBuffer.Key == EMeshBufferType::SkeletonDeformBinding;
+
+					if (bIsPoseBufferType)
+					{
+						MeshPoseData->AdditionalBuffers.Emplace(AdditionalBuffer);
+					}
+				}
+
+				MeshPoseData->MeshIDPrefix = 0;
+				AddMeshToProgram(MeshPoseData);	
+				EnumAddFlags(MeshContentFlags, EMeshContentFlags::PoseData);
+			}
+
+			// PhysicsMeshData
+			{
+				const EMeshCopyFlags PhysicsDataCopyFlags = 
+						EMeshCopyFlags::WithAdditionalPhysics;
+
+				TSharedPtr<FMesh> MeshPhysicsData = MeshData->Clone(PhysicsDataCopyFlags);
+			
+				// Copy components related additional buffers.
+				for (const TPair<EMeshBufferType, FMeshBufferSet>& AdditionalBuffer : MeshData->AdditionalBuffers)
+				{
+					const bool bIsPhysicsBufferType = 
+							AdditionalBuffer.Key == EMeshBufferType::PhysicsBodyDeformBinding   ||
+							AdditionalBuffer.Key == EMeshBufferType::PhysicsBodyDeformSelection ||
+							AdditionalBuffer.Key == EMeshBufferType::PhysicsBodyDeformOffsets;
+
+					if (bIsPhysicsBufferType)
+					{
+						MeshPhysicsData->AdditionalBuffers.Emplace(AdditionalBuffer);
+					}
+				}
+
+				MeshPhysicsData->MeshIDPrefix = 0;
+				AddMeshToProgram(MeshPhysicsData);	
+				EnumAddFlags(MeshContentFlags, EMeshContentFlags::PhysicsData);
+			}
+
+			// MetadataMesh Mesh
+			{
+				const EMeshCopyFlags MetadataDataCopyFlags = 
+						EMeshCopyFlags::WithSurfaces         |
+						EMeshCopyFlags::WithTags             |
+						EMeshCopyFlags::WithSkeletonIDs      |
+						EMeshCopyFlags::WithStreamedResources;
+
+				TSharedPtr<FMesh> MeshMetadataData = MeshData->Clone(MetadataDataCopyFlags);
+			
+				// Add a descriptor MeshBufferSet to the metadata part to have formating info.
+				{
+					FMeshBufferSet VertexMeshFormat;
+					const FMeshBufferSet& VertexBufferSet = MeshData->VertexBuffers;
+					
+					VertexMeshFormat.ElementCount = VertexBufferSet.ElementCount;
+
+					const int32 NumVertexBuffers = VertexBufferSet.Buffers.Num();
+					VertexMeshFormat.Buffers.SetNum(NumVertexBuffers);
+
+					for (int32 BufferIndex = 0; BufferIndex < NumVertexBuffers; ++BufferIndex)
+					{
+						VertexMeshFormat.Buffers[BufferIndex].Channels = VertexBufferSet.Buffers[BufferIndex].Channels;	
+						VertexMeshFormat.Buffers[BufferIndex].ElementSize = VertexBufferSet.Buffers[BufferIndex].ElementSize;	
+					}
+				
+					MeshMetadataData->VertexBuffers = MoveTemp(VertexMeshFormat);
+					EnumAddFlags(MeshMetadataData->VertexBuffers.Flags, EMeshBufferSetFlags::IsDescriptor);
+				}
+
+				{
+					FMeshBufferSet IndexMeshFormat;
+					const FMeshBufferSet& IndexBufferSet = MeshData->IndexBuffers;
+					
+					IndexMeshFormat.ElementCount = IndexBufferSet.ElementCount;
+
+					const int32 NumIndexBuffers = IndexBufferSet.Buffers.Num();
+					IndexMeshFormat.Buffers.SetNum(NumIndexBuffers);
+
+					for (int32 BufferIndex = 0; BufferIndex < NumIndexBuffers; ++BufferIndex)
+					{
+						IndexMeshFormat.Buffers[BufferIndex].Channels = IndexBufferSet.Buffers[BufferIndex].Channels;	
+						IndexMeshFormat.Buffers[BufferIndex].ElementSize = IndexBufferSet.Buffers[BufferIndex].ElementSize;	
+					}
+					
+					MeshMetadataData->IndexBuffers = MoveTemp(IndexMeshFormat);
+					EnumAddFlags(MeshMetadataData->IndexBuffers.Flags, EMeshBufferSetFlags::IsDescriptor);
+				}
+
+				MeshMetadataData->MeshIDPrefix = 0;
+				AddMeshToProgram(MeshMetadataData);	
+				EnumAddFlags(MeshContentFlags, EMeshContentFlags::MetaData);
+			}
+
+
+			// For now empty meshes are not discarded. A mesh rom index will be used even if empty.
+			check(Program.ConstantMeshContentIndices.Num() - FirstIndex == 4);
+
+			FMeshContentRange MeshContentRange;
+
+			FMemory::Memzero(MeshContentRange);
+			MeshContentRange.FirstIndex   = (uint32)FirstIndex;
+			MeshContentRange.ContentFlags = MeshContentFlags;
+			MeshContentRange.MeshIDPrefix = MeshData->MeshIDPrefix;
+
+			return Program.ConstantMeshes.Add(MeshContentRange);
+		}   
+
+
+		/** Adds a constant image data to a program and returns its constant index. */
 		uint32 AddConstantImage(FProgram& Program, const TSharedPtr<const FImage>& pImage, FLinkerOptions& Options)
 		{
 			MUTABLE_CPUPROFILER_SCOPE(AddConstantImage);
@@ -410,39 +589,26 @@ namespace mu
 				OP::MeshConstantArgs Args;
 				FMemory::Memzero(Args);
 
-				TSharedPtr<FMesh> MeshData = static_cast<const FMesh*>(GetValue().Get())->Clone();
-				check(MeshData);
+				TSharedPtr<FMesh> MeshValue = StaticCastSharedPtr<const FMesh>(GetValue())->Clone();
+				check(MeshValue);
 
-				Args.skeleton = -1;
-				if (TSharedPtr<const FSkeleton> pSkeleton = MeshData->GetSkeleton())
+				Args.Skeleton = -1;
+				if (TSharedPtr<const FSkeleton> MeshSkeleton = MeshValue->GetSkeleton())
 				{
-					Args.skeleton = Program.AddConstant(pSkeleton);
-					MeshData->SetSkeleton(nullptr);
+					Args.Skeleton = Program.AddConstant(MeshSkeleton);
+					MeshValue->SetSkeleton(nullptr);
 				}
 
-				Args.physicsBody = -1;
-				if (TSharedPtr<const FPhysicsBody> pPhysicsBody = MeshData->GetPhysicsBody())
+				Args.PhysicsBody = -1;
+				if (TSharedPtr<const FPhysicsBody> MeshPhysicsBody = MeshValue->GetPhysicsBody())
 				{
-					Args.physicsBody = Program.AddConstant(pPhysicsBody);
-					MeshData->SetPhysicsBody(nullptr);
+					Args.PhysicsBody = Program.AddConstant(MeshPhysicsBody);
+					MeshValue->SetPhysicsBody(nullptr);
 				}
 
-				// Use a map-based deduplication
-				TSharedPtr<const FMesh> Key = MeshData;
-				const int32* IndexPtr = Options->MeshConstantMap.Find(Key);
-				if (!IndexPtr)
-				{
-					Args.value = Program.ConstantMeshesPermanent.Add(MeshData);
-
-					int32 DataDescIndex = Options->AdditionalData.SourceMeshPerConstant.Add(SourceDataDescriptor);
-					check(DataDescIndex == Args.value);
-
-					Options->MeshConstantMap.Add(MeshData, int32(Args.value));
-				}
-				else
-				{
-					Args.value = *IndexPtr;
-				}
+				Args.Value = AddConstantMesh(Program, MeshValue, *Options);
+				int32 DataDescIndex = Options->AdditionalData.SourceMeshPerConstant.Add(SourceDataDescriptor);
+				check(DataDescIndex == Args.Value);
 
 				linkedAddress = (OP::ADDRESS)Program.OpAddress.Num();
 				Program.OpAddress.Add((uint32)Program.ByteCode.Num());
diff --git a/Engine/Plugins/Mutable/Source/MutableTools/Private/MuT/CodeOptimiser.cpp b/Engine/Plugins/Mutable/Source/MutableTools/Private/MuT/CodeOptimiser.cpp
--- a/Engine/Plugins/Mutable/Source/MutableTools/Private/MuT/CodeOptimiser.cpp
+++ b/Engine/Plugins/Mutable/Source/MutableTools/Private/MuT/CodeOptimiser.cpp
@@ -518,14 +518,14 @@ namespace mu
 			{
 				MUTABLE_CPUPROFILER_SCOPE(ConstantMesh);
 
-				TSharedPtr<const FMesh> pMesh = System->GetPrivate()->BuildMesh( model, LocalParams.Get(), at );
+				TSharedPtr<const FMesh> MeshBuild = System->GetPrivate()->BuildMesh(model, LocalParams.Get(), at, EMeshContentFlags::AllFlags);
 
-				if (pMesh)
+				if (MeshBuild)
 				{
 					mu::Ptr<ASTOpConstantResource> ConstantOp = new ASTOpConstantResource();
 					ConstantOp->SourceDataDescriptor = SourceDataDescriptor;
 					ConstantOp->Type = EOpType::ME_CONSTANT;
-					ConstantOp->SetValue( pMesh, DiskCacheContext );
+					ConstantOp->SetValue(MeshBuild, DiskCacheContext);
 					Result = ConstantOp;
 				  }
 				break;
diff --git a/Engine/Plugins/Mutable/Source/MutableTools/Private/MuT/Compiler.cpp b/Engine/Plugins/Mutable/Source/MutableTools/Private/MuT/Compiler.cpp
--- a/Engine/Plugins/Mutable/Source/MutableTools/Private/MuT/Compiler.cpp
+++ b/Engine/Plugins/Mutable/Source/MutableTools/Private/MuT/Compiler.cpp
@@ -879,11 +879,30 @@ namespace mu
 		Program.ConstantMeshesPermanent = {};
 		Program.ConstantMeshesPermanent.Reserve(AllMeshes.Num());
 
-		// Meshes indices will be remapped as they are classified between streamed and permanent
-		TArray<uint32> OldToNewMeshIndex;
-		OldToNewMeshIndex.SetNumUninitialized(AllMeshes.Num());
-
 		FMemory::Memzero( RomDatas.GetData(), RomDatas.GetAllocatedSize() );
+		{
+			MUTABLE_CPUPROFILER_SCOPE(GenerateRoms_ImageSourceIds);
+
+			for (int32 MeshIndex = 0; MeshIndex < Program.ConstantMeshes.Num(); ++MeshIndex)
+			{
+				const FMeshContentRange& MeshContentRange = Program.ConstantMeshes[MeshIndex];
+
+				const uint32 SourceId = AdditionalData.SourceMeshPerConstant[MeshIndex].SourceId;
+
+				const int32 NumMeshContentElements = 
+						FMath::CountBits(static_cast<uint64>(MeshContentRange.ContentFlags));
+
+				for (int32 MeshContentIndex = 0; MeshContentIndex < NumMeshContentElements; ++MeshContentIndex)
+				{
+					FConstantResourceIndex ResourceIndex = 
+							Program.ConstantMeshContentIndices[MeshContentRange.FirstIndex + MeshContentIndex];
+
+					check(!ResourceIndex.Streamable); // Not classified yet
+					RomDatasCompile[ResourceIndex.Index].SourceId = SourceId;
+				}
+			}
+		}
+
 		{
 			MUTABLE_CPUPROFILER_SCOPE(GenerateRoms_MeshIds);
 
@@ -911,14 +930,13 @@ namespace mu
 		{
 			MUTABLE_CPUPROFILER_SCOPE(GenerateRoms_MeshSplit);
 
+			TArray<FConstantResourceIndex> IndexPerMeshContent;
+			IndexPerMeshContent.SetNumUninitialized(AllMeshes.Num());
+			
 			for (int32 MeshIndex = 0; MeshIndex < AllMeshes.Num(); ++MeshIndex)
 			{
 				FRomDataRuntime& RomData = RomDatas[MeshIndex];
-
 				FRomDataCompile& RomDataCompile = RomDatasCompile[MeshIndex];
-				RomDataCompile.SourceId = AdditionalData.SourceMeshPerConstant[MeshIndex].SourceId;
-
-				FConstantResourceIndex ResourceIndex;
 
 				if (RomData.Size > EmbeddedDataBytesLimit)
 				{
@@ -928,7 +946,7 @@ namespace mu
 					uint32 RomIndex = Program.Roms.Num();
 					check(RomIndex < 0x7fffffff);
 
-					ResourceIndex = { RomIndex, 1 };
+					IndexPerMeshContent[MeshIndex] = { RomIndex, 1 };
 					Program.ConstantMeshesStreamed.Add(RomIndex, AllMeshes[MeshIndex]);
 
 					Program.Roms.Add(RomData);
@@ -938,26 +956,19 @@ namespace mu
 				{
 					uint32 Index = uint32(Program.ConstantMeshesPermanent.Num());
 					check(Index < 0x7fffffff);
-					ResourceIndex = { Index, 0 };
+
+					IndexPerMeshContent[MeshIndex] = { Index, 0 };
 					Program.ConstantMeshesPermanent.Add(AllMeshes[MeshIndex]);
 
 					NumEmbedded++;
 					NumEmbeddedBytes += RomData.Size;
 				}
-
-				OldToNewMeshIndex[MeshIndex] = *reinterpret_cast<const uint32*>(&ResourceIndex);
 			}
-		}
 
-		// Remap mesh constant references in bytecode
-		for (int32 OpIndex=0; OpIndex<Program.OpAddress.Num(); ++OpIndex)
-		{
-			EOpType Type = Program.GetOpType(OpIndex);
-			if (Type == EOpType::ME_CONSTANT)
+			for (int32 Index = 0; Index < Program.ConstantMeshContentIndices.Num(); ++Index)
 			{
-				OP::ResourceConstantArgs Args = Program.GetOpArgs<OP::ResourceConstantArgs>(OpIndex);
-				Args.value = OldToNewMeshIndex[Args.value];
-				Program.SetOpArgs<OP::ResourceConstantArgs>(OpIndex, Args);
+				int32 RomDataIndex = Program.ConstantMeshContentIndices[Index].Index;
+				Program.ConstantMeshContentIndices[Index] = IndexPerMeshContent[RomDataIndex];
 			}
 		}
 
diff --git a/Engine/Plugins/Mutable/Source/MutableTools/Private/MuT/DataPacker.cpp b/Engine/Plugins/Mutable/Source/MutableTools/Private/MuT/DataPacker.cpp
--- a/Engine/Plugins/Mutable/Source/MutableTools/Private/MuT/DataPacker.cpp
+++ b/Engine/Plugins/Mutable/Source/MutableTools/Private/MuT/DataPacker.cpp
@@ -235,18 +235,16 @@ namespace mu
                 default:
                 {
                     // If we hit the constant we are analysing
-                    if ( OpType == EOpType::ME_CONSTANT
-                         &&
-                         program.GetOpType(at) == OpType )
+                    if (OpType == EOpType::ME_CONSTANT && program.GetOpType(at) == OpType)
                     {
-						OP::MeshConstantArgs args = program.GetOpArgs<OP::MeshConstantArgs>(at);
-                        if ( args.value == Constant )
+						OP::MeshConstantArgs Args = program.GetOpArgs<OP::MeshConstantArgs>(at);
+                        if (Args.Value == Constant)
                         {
                             // Accumulate the currently relevant parameters
-                            const TArray<int>& currentParams = GetCurrentState();
-                            for (int32 i=0; i<currentParams.Num(); ++i )
+                            const TArray<int32>& CurrentParams = GetCurrentState();
+                            for (int32 ParamIndex = 0; ParamIndex < CurrentParams.Num(); ++ParamIndex)
                             {
-                                m_allParams[i] += currentParams[i];
+                                m_allParams[ParamIndex] += CurrentParams[ParamIndex];
                             }
                         }
                     }
