[mutable] Disabled all usages of the "CustomizableObjectCustomSettings.h" prior to its complete removal.
- We have no proof it is currently being used but, just to be sure, disable it before proceeding with the future deletion of the system.

#rnx

diff --git a/Engine/Plugins/Mutable/Source/CustomizableObjectEditor/Private/MuCOE/CustomizableObjectCustomSettings.h b/Engine/Plugins/Mutable/Source/CustomizableObjectEditor/Private/MuCOE/CustomizableObjectCustomSettings.h
--- a/Engine/Plugins/Mutable/Source/CustomizableObjectEditor/Private/MuCOE/CustomizableObjectCustomSettings.h
+++ b/Engine/Plugins/Mutable/Source/CustomizableObjectEditor/Private/MuCOE/CustomizableObjectCustomSettings.h
@@ -22,6 +22,9 @@ struct FGeometry;
 struct FPointerEvent;
 
 
+/**
+ * Custom slate that grants the user the ability of adding, moving and removing lights from the preview scene.
+ */
 class SCustomizableObjectCustomSettings : public SVerticalBox
 {
 public:
diff --git a/Engine/Plugins/Mutable/Source/CustomizableObjectEditor/Private/MuCOE/CustomizableObjectEditor.cpp b/Engine/Plugins/Mutable/Source/CustomizableObjectEditor/Private/MuCOE/CustomizableObjectEditor.cpp
--- a/Engine/Plugins/Mutable/Source/CustomizableObjectEditor/Private/MuCOE/CustomizableObjectEditor.cpp
+++ b/Engine/Plugins/Mutable/Source/CustomizableObjectEditor/Private/MuCOE/CustomizableObjectEditor.cpp
@@ -1501,7 +1501,7 @@ TSharedPtr<SCustomizableObjectEditorAdvancedPreviewSettings> FCustomizableObject
 
 bool FCustomizableObjectEditor::ShowLightingSettings()
 {
-	return true;
+	return false;
 }
 
 
