[mutable] Moved LODSettings from CO to MeshComponent node
- Added compatibility code that handles the copy of LODSettings data from the Parent CO to the Child CO (if required) and later thee copy of said data to the nodes that now will host the LODSettings data.

#rnx

diff --git a/Engine/Plugins/Mutable/Source/CustomizableObject/Internal/MuCO/CustomizableObjectCustomVersion.h b/Engine/Plugins/Mutable/Source/CustomizableObject/Internal/MuCO/CustomizableObjectCustomVersion.h
--- a/Engine/Plugins/Mutable/Source/CustomizableObject/Internal/MuCO/CustomizableObjectCustomVersion.h
+++ b/Engine/Plugins/Mutable/Source/CustomizableObject/Internal/MuCO/CustomizableObjectCustomVersion.h
@@ -178,6 +178,8 @@ struct CUSTOMIZABLEOBJECT_API FCustomizableObjectCustomVersion
 
 		TextureVariationsToVariations,
 
+		MovedLODSettingsToMeshComponentNode,
+
 		// -----<new versions can be added above this line>-------------------------------------------------
 		VersionPlusOne,
 		LatestVersion = VersionPlusOne - 1
diff --git a/Engine/Plugins/Mutable/Source/CustomizableObject/Internal/MuCO/CustomizableObjectPrivate.h b/Engine/Plugins/Mutable/Source/CustomizableObject/Internal/MuCO/CustomizableObjectPrivate.h
--- a/Engine/Plugins/Mutable/Source/CustomizableObject/Internal/MuCO/CustomizableObjectPrivate.h
+++ b/Engine/Plugins/Mutable/Source/CustomizableObject/Internal/MuCO/CustomizableObjectPrivate.h
@@ -1322,11 +1322,11 @@ public:
 
 	/** Minimum LOD to render per Platform. */
 	UPROPERTY()
-	FPerPlatformInt MinLOD;
+	TMap<FName, FPerPlatformInt> MinLODPerComponent;
 
 	/** Minimum LOD to render per Quality level. */
 	UPROPERTY()
-	FPerQualityLevelInt MinQualityLevelLOD;
+	TMap<FName, FPerQualityLevelInt> MinQualityLevelLODPerComponent;
 	
 	UPROPERTY()
 	FString ReleaseVersion;
@@ -1589,6 +1589,8 @@ enum class ECustomizableObjectVersions : int32
 
 	MeshDataRomSplitBackout,
 
+	MovedLODSettingsToMeshComponentNode,
+
 	// -----<new versions can be added above this line>--------
 	LastCustomizableObjectVersion
 };
diff --git a/Engine/Plugins/Mutable/Source/CustomizableObject/Internal/MuCO/ICustomizableObjectEditorModule.h b/Engine/Plugins/Mutable/Source/CustomizableObject/Internal/MuCO/ICustomizableObjectEditorModule.h
--- a/Engine/Plugins/Mutable/Source/CustomizableObject/Internal/MuCO/ICustomizableObjectEditorModule.h
+++ b/Engine/Plugins/Mutable/Source/CustomizableObject/Internal/MuCO/ICustomizableObjectEditorModule.h
@@ -72,6 +72,15 @@ public:
 	  * @return true if the compilation is out of date. */
 	virtual bool IsCompilationOutOfDate(const UCustomizableObject& Object, bool bSkipIndirectReferences, TArray<FName>& OutOfDatePackages, TArray<FName>& AddedPackages, TArray<FName>& RemovedPackages, bool& bReleaseVersionDiff) const = 0;
 
+	/** Method called once all parent COs of this CO have already been loaded and it is safe to get data from the root CO.
+	* @note The compiled data will not be yet available. */
+	virtual void OnUpstreamCOsLoaded(UCustomizableObject* Object) const = 0;
+
+	/** Fixup operations executed once the upstream COs (parent COs) have been loaded.
+	 * @param CustomizableObjectCustomVersion Version index used to determine what change in the data should be performed.
+	 */
+	virtual void OnUpstreamCOsLoadedFixup(UCustomizableObject* Object, int32 CustomizableObjectCustomVersion) const = 0;
+	
 	using IsCompilationOutOfDateCallback = TFunction<void(bool bOutOfDate, bool bVersionDiff, const TArray<FName>& OutOfDatePackages, const TArray<FName>& AddedPackages, const TArray<FName>& RemovedPackages)>;
 
 	/** Async version. See IsCompilationOutOfDate.
diff --git a/Engine/Plugins/Mutable/Source/CustomizableObject/Private/MuCO/CustomizableObject.cpp b/Engine/Plugins/Mutable/Source/CustomizableObject/Private/MuCO/CustomizableObject.cpp
--- a/Engine/Plugins/Mutable/Source/CustomizableObject/Private/MuCO/CustomizableObject.cpp
+++ b/Engine/Plugins/Mutable/Source/CustomizableObject/Private/MuCO/CustomizableObject.cpp
@@ -129,7 +129,9 @@ UCustomizableObject::UCustomizableObject()
 #if WITH_EDITORONLY_DATA
 	const FString CVarName = TEXT("r.SkeletalMesh.MinLodQualityLevel");
 	const FString ScalabilitySectionName = TEXT("ViewDistanceQuality");
+	PRAGMA_DISABLE_DEPRECATION_WARNINGS
 	LODSettings.MinQualityLevelLOD.SetQualityLevelCVarForCooking(*CVarName, *ScalabilitySectionName);
+	PRAGMA_ENABLE_DEPRECATION_WARNINGS
 #endif
 }
 
@@ -2006,12 +2008,21 @@ uint8 UCustomizableObjectPrivate::GetMinLODIndex(const FName& ComponentName) con
 				// Get the quality level for the skeletal meshes for the current scalability setting. The bigger the better.
 				// Value extracted from the Scalability.ini of the current platform (if found).
 				const int32 QualityLevel = UCustomizableObjectSystem::GetInstance()->GetSkeletalMeshMinLODQualityLevel();
-				MinLODIdx = LocalModelResources->MinQualityLevelLOD.GetValue(QualityLevel);
+
+				const FPerQualityLevelInt* QualityLevelInt = LocalModelResources->MinQualityLevelLODPerComponent.Find(ComponentName);
+				if (ensure(QualityLevelInt))
+				{
+					MinLODIdx = QualityLevelInt->GetValue(QualityLevel);
+				}
 			}
 		}
 		else
 		{
-			MinLODIdx = LocalModelResources->MinLOD.GetValue();
+			const FPerPlatformInt* PlatformInt = LocalModelResources->MinLODPerComponent.Find(ComponentName);
+			if (ensure(PlatformInt))
+			{
+				MinLODIdx = PlatformInt->GetValue();
+			}
 		}
 	}
 	
@@ -3424,7 +3435,9 @@ void UCustomizableObjectPrivate::LogMemory()
 				Arch << ModelResources->ClothSharedConfigsData;
 				Arch << ModelResources->ClothingAssetsData;
 				Arch << ModelResources->ComponentNamesPerObjectComponent;
-
+				Arch << ModelResources->MinLODPerComponent;
+				Arch << ModelResources->MinQualityLevelLODPerComponent;
+				
 				ModeleResData += ModelResources->StreamedResourceData.GetAllocatedSize();
 				for (FCustomizableObjectStreamedResourceData& Data : ModelResources->StreamedResourceData)
 				{
diff --git a/Engine/Plugins/Mutable/Source/CustomizableObject/Private/MuCO/CustomizableObjectSystem.cpp b/Engine/Plugins/Mutable/Source/CustomizableObject/Private/MuCO/CustomizableObjectSystem.cpp
--- a/Engine/Plugins/Mutable/Source/CustomizableObject/Private/MuCO/CustomizableObjectSystem.cpp
+++ b/Engine/Plugins/Mutable/Source/CustomizableObject/Private/MuCO/CustomizableObjectSystem.cpp
@@ -4023,6 +4023,11 @@ int32 UCustomizableObjectSystem::TickInternal(const bool bBlocking)
 					break;
 				}
 			}
+
+			if (bReady)
+			{
+				Module->OnUpstreamCOsLoaded(CO);
+			}
 		}
 
 		if (bReady)
diff --git a/Engine/Plugins/Mutable/Source/CustomizableObject/Public/MuCO/CustomizableObject.h b/Engine/Plugins/Mutable/Source/CustomizableObject/Public/MuCO/CustomizableObject.h
--- a/Engine/Plugins/Mutable/Source/CustomizableObject/Public/MuCO/CustomizableObject.h
+++ b/Engine/Plugins/Mutable/Source/CustomizableObject/Public/MuCO/CustomizableObject.h
@@ -255,8 +255,8 @@ public:
 	FInstancedStruct VersionStruct;
 #endif
 
-	// DEPRECATED, Will be per component UE-243403
-	UPROPERTY(EditAnywhere, Category = CustomizableObject, meta = (DisplayName = "LOD Settings"))
+	UE_DEPRECATED(5.6, "Moved to Mesh Component Node")
+	UPROPERTY(meta = (DeprecatedProperty, DeprecationMessage="Moved to Mesh Component Node"))
 	FMutableLODSettings LODSettings;
 	
 	/** true to use the Reference Skeletal Mesh as a placeholder while the Generated Skeletal Mesh is not ready.
diff --git a/Engine/Plugins/Mutable/Source/CustomizableObjectEditor/Private/MuCOE/CustomizableObjectCompiler.cpp b/Engine/Plugins/Mutable/Source/CustomizableObjectEditor/Private/MuCOE/CustomizableObjectCompiler.cpp
--- a/Engine/Plugins/Mutable/Source/CustomizableObjectEditor/Private/MuCOE/CustomizableObjectCompiler.cpp
+++ b/Engine/Plugins/Mutable/Source/CustomizableObjectEditor/Private/MuCOE/CustomizableObjectCompiler.cpp
@@ -924,9 +924,21 @@ void FCustomizableObjectCompiler::CompileInternal(bool bAsync)
 										ImageProp.AddressX, ImageProp.AddressY });
 		}
 
-		// Copy the LODSettings data found in the CO into the ModelResources
-		ModelResources->MinLOD = CurrentObject->LODSettings.MinLOD;
-		ModelResources->MinQualityLevelLOD = CurrentObject->LODSettings.MinQualityLevelLOD;
+		const TArray<FMutableComponentInfo>& ComponentInfos = GenerationContext.CompilationContext.ComponentInfos;
+		for (const FMutableComponentInfo& ComponentInfo : ComponentInfos)
+		{
+			if (!Cast<UCustomizableObjectNodeComponentMesh>(ComponentInfo.Node))
+			{
+				continue;
+			}
+
+			const FName ComponentName = ComponentInfo.ComponentName;
+			const FMutableLODSettings& ComponentLODSettings = ComponentInfo.LODSettings;
+
+			// Copy the LODSettings data found in the Component into the ModelResources
+			ModelResources->MinLODPerComponent.Add(ComponentName, ComponentLODSettings.MinLOD);
+			ModelResources->MinQualityLevelLODPerComponent.Add(ComponentName, ComponentLODSettings.MinQualityLevelLOD);
+		}
 		
 		ModelResources->ParameterUIDataMap = MoveTemp(GenerationContext.ParameterUIDataMap);
 		ModelResources->StateUIDataMap = MoveTemp(GenerationContext.StateUIDataMap);
diff --git a/Engine/Plugins/Mutable/Source/CustomizableObjectEditor/Private/MuCOE/CustomizableObjectEditorModule.cpp b/Engine/Plugins/Mutable/Source/CustomizableObjectEditor/Private/MuCOE/CustomizableObjectEditorModule.cpp
--- a/Engine/Plugins/Mutable/Source/CustomizableObjectEditor/Private/MuCOE/CustomizableObjectEditorModule.cpp
+++ b/Engine/Plugins/Mutable/Source/CustomizableObjectEditor/Private/MuCOE/CustomizableObjectEditorModule.cpp
@@ -92,6 +92,7 @@
 #include "MuCOE/GraphTraversal.h"
 #include "MuCOE/CustomizableObjectInstanceBaker.h"
 #include "Editor.h"
+#include "MuCO/CustomizableObjectCustomVersion.h"
 #include "MuCO/CustomizableObjectInstancePrivate.h"
 #include "MuCO/CustomizableObjectSystemPrivate.h"
 #include "MuCOE/CompileRequest.h"
@@ -736,6 +737,44 @@ void FCustomizableObjectEditorModule::GetRelatedObjects(UCustomizableObject* CO,
 }
 
 
+void FCustomizableObjectEditorModule::OnUpstreamCOsLoaded(UCustomizableObject* Object) const
+{
+	check(Object)
+	const int32 CustomVersion = Object->GetLinkerCustomVersion(FCustomizableObjectCustomVersion::GUID);
+	for (int32 Version = CustomVersion + 1; Version <= FCustomizableObjectCustomVersion::LatestVersion; ++Version)
+	{
+		OnUpstreamCOsLoadedFixup(Object, Version);
+	}
+}
+
+
+void FCustomizableObjectEditorModule::OnUpstreamCOsLoadedFixup(UCustomizableObject* Object, int32 CustomizableObjectCustomVersion) const
+{
+	if (CustomizableObjectCustomVersion == FCustomizableObjectCustomVersion::MovedLODSettingsToMeshComponentNode)
+	{
+		check(Object)
+		const UCustomizableObject* RootObject = GetRootObject(Object);
+		check(RootObject);
+		
+		PRAGMA_DISABLE_DEPRECATION_WARNINGS
+		FMutableLODSettings RootObjectLODSettings = RootObject->LODSettings;
+		Object->LODSettings = RootObjectLODSettings;
+		PRAGMA_ENABLE_DEPRECATION_WARNINGS
+
+		// Override the LOD Settings found in the UCustomizableObjectNodeComponentMesh of this CO using the LODSettings of the parent CO
+		{
+			TArray<UCustomizableObjectNodeComponentMesh*> ComponentMeshNodes;
+			Object->GetPrivate()->GetSource()->GetNodesOfClass<UCustomizableObjectNodeComponentMesh>(ComponentMeshNodes);
+			for (UCustomizableObjectNodeComponentMesh* ComponentNode : ComponentMeshNodes)
+			{
+				ComponentNode->LODSettings = RootObjectLODSettings;
+			}
+		}
+	}
+}
+
+
+
 void FCustomizableObjectEditorModule::BakeCustomizableObjectInstance(UCustomizableObjectInstance* InTargetInstance, const FBakingConfiguration& InBakingConfig)
 {
 	UCustomizableObjectInstanceBaker* InstanceBaker = NewObject<UCustomizableObjectInstanceBaker>();
diff --git a/Engine/Plugins/Mutable/Source/CustomizableObjectEditor/Private/MuCOE/CustomizableObjectEditorModule.h b/Engine/Plugins/Mutable/Source/CustomizableObjectEditor/Private/MuCOE/CustomizableObjectEditorModule.h
--- a/Engine/Plugins/Mutable/Source/CustomizableObjectEditor/Private/MuCOE/CustomizableObjectEditorModule.h
+++ b/Engine/Plugins/Mutable/Source/CustomizableObjectEditor/Private/MuCOE/CustomizableObjectEditorModule.h
@@ -50,6 +50,8 @@ public:
 	virtual UCustomizableObject* GetRootObject(UCustomizableObject* ChildObject) const override;
 	virtual const UCustomizableObject* GetRootObject(const UCustomizableObject* ChildObject) const override;
 	virtual void GetRelatedObjects(UCustomizableObject*, TSet<UCustomizableObject*>& OutRelated) const override;
+	virtual void OnUpstreamCOsLoaded(UCustomizableObject* Object) const override;
+	virtual void OnUpstreamCOsLoadedFixup(UCustomizableObject* Object, int32 CustomizableObjectCustomVersion) const override;
 	virtual void BakeCustomizableObjectInstance(UCustomizableObjectInstance* InTargetInstance, const FBakingConfiguration& InBakingConfig) override;
 	virtual USkeletalMesh* GetReferenceSkeletalMesh(const UCustomizableObject& Object, const FName& ComponentName) const override;
 	virtual TMap<FName, FGuid> GetParticipatingObjects(const UCustomizableObject* Object, const FCompilationOptions* Options = nullptr) const override;
diff --git a/Engine/Plugins/Mutable/Source/CustomizableObjectEditor/Private/MuCOE/GenerateMutableSource/GenerateMutableSource.cpp b/Engine/Plugins/Mutable/Source/CustomizableObjectEditor/Private/MuCOE/GenerateMutableSource/GenerateMutableSource.cpp
--- a/Engine/Plugins/Mutable/Source/CustomizableObjectEditor/Private/MuCOE/GenerateMutableSource/GenerateMutableSource.cpp
+++ b/Engine/Plugins/Mutable/Source/CustomizableObjectEditor/Private/MuCOE/GenerateMutableSource/GenerateMutableSource.cpp
@@ -50,7 +50,6 @@
 #include "MuR/Mesh.h"
 #include "PhysicsEngine/PhysicsAsset.h"
 #include "PlatformInfo.h"
-#include "AssetRegistry/AssetRegistryModule.h"
 #include "Math/NumericLimits.h"
 #include "Hash/CityHash.h"
 #include "MuCO/CustomizableObjectResourceDataTypes.h"
@@ -815,20 +814,28 @@ mu::Ptr<mu::NodeObject> GenerateMutableSource(const UEdGraphPin * Pin, FMutableG
 				ObjectNode->AddStateParam(StateIndex, State.RuntimeParameters[ParamIndex]);
 			}
 
-			const FMutableLODSettings& LODSettings = GenerationContext.CompilationContext.Object->LODSettings;
-			const ITargetPlatform* TargetPlatform = GenerationContext.CompilationContext.Options.TargetPlatform;
-			const FSupportedQualityLevelArray SupportedQualityLevels = LODSettings.MinQualityLevelLOD.GetSupportedQualityLevels(*TargetPlatform->GetPlatformInfo().IniPlatformName.ToString());
-
 			int32 MinQualityLevelLOD = INT32_MAX;
 			int32 MaxQualityLevelLOD = 0;
 
-			for (const int32& QualitySetting : SupportedQualityLevels)
+			for (const FMutableComponentInfo& ComponentInfo : GenerationContext.CompilationContext.ComponentInfos)
 			{
-				if (LODSettings.MinQualityLevelLOD.IsQualityLevelValid(QualitySetting))
+				if (!Cast<UCustomizableObjectNodeComponentMesh>(ComponentInfo.Node))
+				{
+					continue;
+				}
+				
+				const FMutableLODSettings& LODSettings = ComponentInfo.Node->LODSettings;
+				const ITargetPlatform* TargetPlatform = GenerationContext.CompilationContext.Options.TargetPlatform;
+				const FSupportedQualityLevelArray SupportedQualityLevels = LODSettings.MinQualityLevelLOD.GetSupportedQualityLevels(*TargetPlatform->GetPlatformInfo().IniPlatformName.ToString());
+				
+				for (const int32& QualitySetting : SupportedQualityLevels)
 				{
-					int32 QualityLod = LODSettings.MinQualityLevelLOD.GetValueForQualityLevel(QualitySetting);
-					MinQualityLevelLOD = FMath::Min(MinQualityLevelLOD, QualityLod);
-					MaxQualityLevelLOD = FMath::Max(MaxQualityLevelLOD, QualityLod);
+					if (LODSettings.MinQualityLevelLOD.IsQualityLevelValid(QualitySetting))
+					{
+						const int32 QualityLod = LODSettings.MinQualityLevelLOD.GetValueForQualityLevel(QualitySetting);
+						MinQualityLevelLOD = FMath::Min(MinQualityLevelLOD, QualityLod);
+						MaxQualityLevelLOD = FMath::Max(MaxQualityLevelLOD, QualityLod);
+					}
 				}
 			}
 
diff --git a/Engine/Plugins/Mutable/Source/CustomizableObjectEditor/Private/MuCOE/GenerateMutableSource/GenerateMutableSource.h b/Engine/Plugins/Mutable/Source/CustomizableObjectEditor/Private/MuCOE/GenerateMutableSource/GenerateMutableSource.h
--- a/Engine/Plugins/Mutable/Source/CustomizableObjectEditor/Private/MuCOE/GenerateMutableSource/GenerateMutableSource.h
+++ b/Engine/Plugins/Mutable/Source/CustomizableObjectEditor/Private/MuCOE/GenerateMutableSource/GenerateMutableSource.h
@@ -384,6 +384,8 @@ struct FMutableComponentInfo
 
 	FName ComponentName;
 
+	FMutableLODSettings LODSettings;
+	
 	// Each component must have a reference SkeletalMesh with a valid Skeleton
 	USkeletalMesh* RefSkeletalMesh = nullptr;
 	USkeleton* RefSkeleton = nullptr;
diff --git a/Engine/Plugins/Mutable/Source/CustomizableObjectEditor/Private/MuCOE/GenerateMutableSource/GenerateMutableSourceComponent.cpp b/Engine/Plugins/Mutable/Source/CustomizableObjectEditor/Private/MuCOE/GenerateMutableSource/GenerateMutableSourceComponent.cpp
--- a/Engine/Plugins/Mutable/Source/CustomizableObjectEditor/Private/MuCOE/GenerateMutableSource/GenerateMutableSourceComponent.cpp
+++ b/Engine/Plugins/Mutable/Source/CustomizableObjectEditor/Private/MuCOE/GenerateMutableSource/GenerateMutableSourceComponent.cpp
@@ -2,16 +2,13 @@
 
 #include "GenerateMutableSourceComponent.h"
 
-#include "GenerateMutableSourceModifier.h"
 #include "Factories/FbxMeshImportData.h"
 #include "MuCOE/GenerateMutableSource/GenerateMutableSource.h"
 
-#include "MuCOE/CustomizableObjectCompiler.h"
 #include "MuCOE/GenerateMutableSource/GenerateMutableSourceMesh.h"
 #include "MuCOE/GenerateMutableSource/GenerateMutableSourceSurface.h"
 #include "MuCOE/GenerateMutableSource/GenerateMutableSourceFloat.h"
 #include "MuCOE/GraphTraversal.h"
-#include "MuCOE/MutableUtils.h"
 #include "MuCOE/Nodes/CustomizableObjectNodeComponentMesh.h"
 #include "MuCOE/Nodes/CustomizableObjectNodeComponentMeshAddTo.h"
 #include "MuCOE/Nodes/CustomizableObjectNodeComponentPassthroughMesh.h"
@@ -31,11 +28,9 @@
 #include "MuT/NodeScalarConstant.h"
 #include "MuT/NodeSurfaceNew.h"
 #include "MuT/NodeMeshConstant.h"
-#include "MuT/NodeObjectNew.h"
 
 #include "Rendering/SkeletalMeshLODModel.h"
 #include "Rendering/SkeletalMeshModel.h"
-#include "Rendering/SkeletalMeshRenderData.h"
 #include "Interfaces/ITargetPlatform.h"
 
 
@@ -595,7 +590,7 @@ void FirstPass(UCustomizableObjectNodeComponentMesh& Node, FMutableGraphGenerati
 		NumLODs = MaxRefMeshLODs;
 	}
 	
-	const FMutableLODSettings& LODSettings = GenerationContext.CompilationContext.Object->LODSettings;
+	const FMutableLODSettings& LODSettings = Node.LODSettings;
 
 	uint8 FirstLODAvailable = MAX_MESH_LOD_COUNT;
 	
diff --git a/Engine/Plugins/Mutable/Source/CustomizableObjectEditor/Private/MuCOE/Nodes/CustomizableObjectNodeComponentMesh.cpp b/Engine/Plugins/Mutable/Source/CustomizableObjectEditor/Private/MuCOE/Nodes/CustomizableObjectNodeComponentMesh.cpp
--- a/Engine/Plugins/Mutable/Source/CustomizableObjectEditor/Private/MuCOE/Nodes/CustomizableObjectNodeComponentMesh.cpp
+++ b/Engine/Plugins/Mutable/Source/CustomizableObjectEditor/Private/MuCOE/Nodes/CustomizableObjectNodeComponentMesh.cpp
@@ -10,6 +10,14 @@
 #define LOCTEXT_NAMESPACE "CustomizableObjectEditor"
 
 
+UCustomizableObjectNodeComponentMesh::UCustomizableObjectNodeComponentMesh()
+{
+	const FString CVarName = TEXT("r.SkeletalMesh.MinLodQualityLevel");
+	const FString ScalabilitySectionName = TEXT("ViewDistanceQuality");
+	LODSettings.MinQualityLevelLOD.SetQualityLevelCVarForCooking(*CVarName, *ScalabilitySectionName);
+}
+
+
 void UCustomizableObjectNodeComponentMesh::PostEditChangeProperty(FPropertyChangedEvent& PropertyChangedEvent)
 {
 	Super::PostEditChangeProperty(PropertyChangedEvent);
diff --git a/Engine/Plugins/Mutable/Source/CustomizableObjectEditor/Private/MuCOE/Nodes/CustomizableObjectNodeComponentMesh.h b/Engine/Plugins/Mutable/Source/CustomizableObjectEditor/Private/MuCOE/Nodes/CustomizableObjectNodeComponentMesh.h
--- a/Engine/Plugins/Mutable/Source/CustomizableObjectEditor/Private/MuCOE/Nodes/CustomizableObjectNodeComponentMesh.h
+++ b/Engine/Plugins/Mutable/Source/CustomizableObjectEditor/Private/MuCOE/Nodes/CustomizableObjectNodeComponentMesh.h
@@ -51,6 +51,8 @@ class CUSTOMIZABLEOBJECTEDITOR_API UCustomizableObjectNodeComponentMesh : public
 	GENERATED_BODY()
 
 public:
+	UCustomizableObjectNodeComponentMesh();
+	
 	// UObject interface
 	virtual void PostEditChangeProperty(FPropertyChangedEvent& PropertyChangedEvent) override;
 	virtual void BackwardsCompatibleFixup(int32 CustomizableObjectCustomVersion) override;
@@ -83,6 +85,9 @@ public:
 	UPROPERTY(EditAnywhere, Category = ComponentMesh)
 	TSoftObjectPtr<UMaterialInterface> OverlayMaterial;
 
+	UPROPERTY(EditAnywhere, Category = ComponentMesh)
+	FMutableLODSettings LODSettings;
+
 	/** Details selected LOD. */
 	int32 SelectedLOD = 0;
 };
diff --git a/Engine/Plugins/Mutable/Source/MutableTools/Private/MuT/CodeOptimiser_Parameters.cpp b/Engine/Plugins/Mutable/Source/MutableTools/Private/MuT/CodeOptimiser_Parameters.cpp
--- a/Engine/Plugins/Mutable/Source/MutableTools/Private/MuT/CodeOptimiser_Parameters.cpp
+++ b/Engine/Plugins/Mutable/Source/MutableTools/Private/MuT/CodeOptimiser_Parameters.cpp
@@ -14,7 +14,6 @@
 #include "MuT/ASTOpMeshMorph.h"
 #include "MuT/ASTOpInstanceAdd.h"
 #include "MuT/ASTOpLayoutFromMesh.h"
-#include "MuT/ASTOpSwitch.h"
 #include "MuT/ASTOpParameter.h"
 #include "MuT/ASTOpConstantColor.h"
 #include "MuT/CodeOptimiser.h"
@@ -26,9 +25,7 @@
 #include "MuR/MutableTrace.h"
 #include "MuR/Operations.h"
 #include "MuR/ParametersPrivate.h"
-#include "MuR/Platform.h"
 #include "MuR/Ptr.h"
-#include "MuR/RefCounted.h"
 #include "MuR/MutableRuntimeModule.h"
 
 #include "HAL/PlatformCrt.h"
