[mutable] Fixed submit with merge issues.

#rnx

diff --git a/Engine/Plugins/Mutable/Source/CustomizableObject/Private/MuCO/CustomizableObjectInstance.cpp b/Engine/Plugins/Mutable/Source/CustomizableObject/Private/MuCO/CustomizableObjectInstance.cpp
--- a/Engine/Plugins/Mutable/Source/CustomizableObject/Private/MuCO/CustomizableObjectInstance.cpp
+++ b/Engine/Plugins/Mutable/Source/CustomizableObject/Private/MuCO/CustomizableObjectInstance.cpp
@@ -3851,8 +3851,8 @@ void UCustomizableInstancePrivate::InitSkeletalMeshData(const TSharedRef<FUpdate
 
 	// Set Min LOD (Override the Reference Skeletal Mesh LOD Settings)
 	const UModelResources& ModelResources = CustomizableObject.GetPrivate()->GetModelResourcesChecked();
-	SkeletalMesh->SetMinLod(FMath::Max(ModelResources.MinLOD.GetDefault(), static_cast<int32>(OperationData->FirstLODAvailable[ComponentName])));
-	SkeletalMesh->SetQualityLevelMinLod(ModelResources.MinQualityLevelLOD);
+	SkeletalMesh->SetMinLod(FMath::Max(ModelResources.MinLODPerComponent.FindChecked(ComponentName).GetDefault(), static_cast<int32>(OperationData->FirstLODAvailable[ComponentName])));
+	SkeletalMesh->SetQualityLevelMinLod(ModelResources.MinQualityLevelLODPerComponent.FindChecked(ComponentName));
 
 	// Set up unreal's default material, will be replaced when building materials
 	{
