[mutable] Fix missing deduplication of shape constant on a specific operation linkage.
#rnx

diff --git a/Engine/Plugins/Mutable/Source/MutableTools/Private/MuT/ASTOpMeshClipMorphPlane.cpp b/Engine/Plugins/Mutable/Source/MutableTools/Private/MuT/ASTOpMeshClipMorphPlane.cpp
--- a/Engine/Plugins/Mutable/Source/MutableTools/Private/MuT/ASTOpMeshClipMorphPlane.cpp
+++ b/Engine/Plugins/Mutable/Source/MutableTools/Private/MuT/ASTOpMeshClipMorphPlane.cpp
@@ -78,7 +78,7 @@ namespace mu
 	}
 
 
-	void ASTOpMeshClipMorphPlane::Link(FProgram& program, FLinkerOptions*)
+	void ASTOpMeshClipMorphPlane::Link(FProgram& Program, FLinkerOptions*)
 	{
 		// Already linked?
 		if (!linkedAddress)
@@ -88,8 +88,7 @@ namespace mu
 
 			if (Source) Args.Source = Source->linkedAddress;
 
-			Args.MorphShape = (OP::ADDRESS)program.ConstantShapes.Num();
-			program.ConstantShapes.Add(MorphShape);
+			Args.MorphShape = Program.AddConstant(MorphShape);
 
 			Args.FaceCullStrategy = FaceCullStrategy;
 
@@ -101,18 +100,17 @@ namespace mu
 			}
 			else if (VertexSelectionType == EClipVertexSelectionType::Shape)
 			{
-				Args.VertexSelectionShapeOrBone = program.AddConstant(SelectionShape);
+				Args.VertexSelectionShapeOrBone = Program.AddConstant(SelectionShape);
 			}
 
 			Args.Dist = Dist;
 			Args.Factor = Factor;
 			Args.MaxBoneRadius = VertexSelectionBoneMaxRadius;
 
-			linkedAddress = (OP::ADDRESS)program.OpAddress.Num();
-			//program.m_code.push_back(op);
-			program.OpAddress.Add((uint32_t)program.ByteCode.Num());
-			AppendCode(program.ByteCode, EOpType::ME_CLIPMORPHPLANE);
-			AppendCode(program.ByteCode, Args);
+			linkedAddress = (OP::ADDRESS)Program.OpAddress.Num();
+			Program.OpAddress.Add((uint32_t)Program.ByteCode.Num());
+			AppendCode(Program.ByteCode, EOpType::ME_CLIPMORPHPLANE);
+			AppendCode(Program.ByteCode, Args);
 		}
 
 	}
