[Mutable] Remove unused FMutableGraphMeshGenerationData and add a new EMutableMeshConversionFlag to ignore Morphs when converting some meshes to Mutable.

#rnx

diff --git a/Engine/Plugins/Mutable/Source/CustomizableObject/Internal/MuCO/UnrealToMutableTextureConversionUtils.h b/Engine/Plugins/Mutable/Source/CustomizableObject/Internal/MuCO/UnrealToMutableTextureConversionUtils.h
--- a/Engine/Plugins/Mutable/Source/CustomizableObject/Internal/MuCO/UnrealToMutableTextureConversionUtils.h
+++ b/Engine/Plugins/Mutable/Source/CustomizableObject/Internal/MuCO/UnrealToMutableTextureConversionUtils.h
@@ -53,8 +53,11 @@ enum class EMutableMeshConversionFlags : uint32
 	// Ignore Physics assets
 	IgnorePhysics = 1 << 1,
 
+	// Ignore Morphs
+	IgnoreMorphs = 1 << 2,
+
 	// Prevent this mesh generation from adding per mesh metadata. 
-	DoNotCreateMeshMetadata = 1 << 2
+	DoNotCreateMeshMetadata = 1 << 3
 };
 
 ENUM_CLASS_FLAGS(EMutableMeshConversionFlags)
diff --git a/Engine/Plugins/Mutable/Source/CustomizableObjectEditor/Private/MuCOE/GenerateMutableSource/GenerateMutableSource.h b/Engine/Plugins/Mutable/Source/CustomizableObjectEditor/Private/MuCOE/GenerateMutableSource/GenerateMutableSource.h
--- a/Engine/Plugins/Mutable/Source/CustomizableObjectEditor/Private/MuCOE/GenerateMutableSource/GenerateMutableSource.h
+++ b/Engine/Plugins/Mutable/Source/CustomizableObjectEditor/Private/MuCOE/GenerateMutableSource/GenerateMutableSource.h
@@ -222,46 +222,20 @@ struct FMutableGraphMeshGenerationData
 	int32 MinNumTriangles = TNumericLimits<int32>::Max();
 
 	TArray<int32> SkinWeightProfilesSemanticIndices;
-
-	// Combine another generated data looking for the most general case.
-	void Combine(const FMutableGraphMeshGenerationData& Other)
-	{
-		bHasVertexColors = bHasVertexColors || Other.bHasVertexColors;
-		bHasRealTimeMorphs = bHasRealTimeMorphs || Other.bHasRealTimeMorphs;
-		bHasClothing = bHasClothing || Other.bHasClothing;
-		NumTexCoordChannels = FMath::Max(Other.NumTexCoordChannels, NumTexCoordChannels);
-		MaxNumBonesPerVertex = FMath::Max(Other.MaxNumBonesPerVertex, MaxNumBonesPerVertex);
-		MaxBoneIndexTypeSizeBytes = FMath::Max(Other.MaxBoneIndexTypeSizeBytes, MaxBoneIndexTypeSizeBytes);
-		MaxNumTriangles = FMath::Max(Other.MaxNumTriangles, MaxNumTriangles);
-		MinNumTriangles = FMath::Min(Other.MinNumTriangles, MinNumTriangles);
-
-		for (int32 SemanticIndex : Other.SkinWeightProfilesSemanticIndices)
-		{
-			SkinWeightProfilesSemanticIndices.AddUnique(SemanticIndex);
-		}
-	}
 };
 
 
 // Data stored for each processed unreal graph node, stored in the cache.
 struct FGeneratedData
 {
-	FGeneratedData(const UEdGraphNode* InSource, mu::NodePtr InNode,
-		const FMutableGraphMeshGenerationData* InMeshData = nullptr)
+	FGeneratedData(const UEdGraphNode* InSource, mu::NodePtr InNode)
 	{
 		Source = InSource;
 		Node = InNode;
-		if (InMeshData)
-		{
-			meshData = *InMeshData;
-		}
 	}
 
 	const UEdGraphNode* Source;
 	mu::NodePtr Node;
-
-	// Used for mesh nodes only
-	FMutableGraphMeshGenerationData meshData;
 };
 
 
diff --git a/Engine/Plugins/Mutable/Source/CustomizableObjectEditor/Private/MuCOE/GenerateMutableSource/GenerateMutableSourceImage.cpp b/Engine/Plugins/Mutable/Source/CustomizableObjectEditor/Private/MuCOE/GenerateMutableSource/GenerateMutableSourceImage.cpp
--- a/Engine/Plugins/Mutable/Source/CustomizableObjectEditor/Private/MuCOE/GenerateMutableSource/GenerateMutableSourceImage.cpp
+++ b/Engine/Plugins/Mutable/Source/CustomizableObjectEditor/Private/MuCOE/GenerateMutableSource/GenerateMutableSourceImage.cpp
@@ -708,12 +708,10 @@ mu::Ptr<mu::NodeImage> GenerateMutableSourceImage(const UEdGraphPin* Pin, FMutab
 
 		if (const UEdGraphPin* ConnectedPin = FollowInputPin(*TypedNodeProject->MeshPin()))
 		{
-			FMutableGraphMeshGenerationData DummyMeshData;
-
 			FLayoutGenerationFlags LayoutGenerationFlags;
 			LayoutGenerationFlags.TexturePinModes.Init(EPinMode::Mutable, TEXSTREAM_MAX_NUM_UVCHANNELS);
 			GenerationContext.LayoutGenerationFlags.Push(LayoutGenerationFlags);
-			mu::Ptr<mu::NodeMesh> MeshNode = GenerateMutableSourceMesh(ConnectedPin, GenerationContext, DummyMeshData, 0, false, false);
+			mu::Ptr<mu::NodeMesh> MeshNode = GenerateMutableSourceMesh(ConnectedPin, GenerationContext, 0, false, false);
 			GenerationContext.LayoutGenerationFlags.Pop();
 			ImageNode->Mesh = MeshNode;
 		}
diff --git a/Engine/Plugins/Mutable/Source/CustomizableObjectEditor/Private/MuCOE/GenerateMutableSource/GenerateMutableSourceLayout.cpp b/Engine/Plugins/Mutable/Source/CustomizableObjectEditor/Private/MuCOE/GenerateMutableSource/GenerateMutableSourceLayout.cpp
--- a/Engine/Plugins/Mutable/Source/CustomizableObjectEditor/Private/MuCOE/GenerateMutableSource/GenerateMutableSourceLayout.cpp
+++ b/Engine/Plugins/Mutable/Source/CustomizableObjectEditor/Private/MuCOE/GenerateMutableSource/GenerateMutableSourceLayout.cpp
@@ -91,7 +91,8 @@ mu::Ptr<mu::NodeLayout> CreateMutableLayoutNode(FMutableGraphGenerationContext&
 			// We don't need all the data to generate the blocks
 			const EMutableMeshConversionFlags ShapeFlags =
 				EMutableMeshConversionFlags::IgnoreSkinning |
-				EMutableMeshConversionFlags::IgnorePhysics;
+				EMutableMeshConversionFlags::IgnorePhysics |
+				EMutableMeshConversionFlags::IgnoreMorphs;
 
 			GenerationContext.MeshGenerationFlags.Push(ShapeFlags);
 
diff --git a/Engine/Plugins/Mutable/Source/CustomizableObjectEditor/Private/MuCOE/GenerateMutableSource/GenerateMutableSourceMesh.cpp b/Engine/Plugins/Mutable/Source/CustomizableObjectEditor/Private/MuCOE/GenerateMutableSource/GenerateMutableSourceMesh.cpp
--- a/Engine/Plugins/Mutable/Source/CustomizableObjectEditor/Private/MuCOE/GenerateMutableSource/GenerateMutableSourceMesh.cpp
+++ b/Engine/Plugins/Mutable/Source/CustomizableObjectEditor/Private/MuCOE/GenerateMutableSource/GenerateMutableSourceMesh.cpp
@@ -923,8 +923,9 @@ TSharedPtr<mu::FMesh> ConvertSkeletalMeshToMutable(
 
 	// Get the mesh generation flags to use
 	const EMutableMeshConversionFlags CurrentFlags = Source.Flags;
-	const bool bIgnoreSkeleton = EnumHasAnyFlags(CurrentFlags, EMutableMeshConversionFlags::IgnoreSkinning);
+	const bool bIgnoreSkinning = EnumHasAnyFlags(CurrentFlags, EMutableMeshConversionFlags::IgnoreSkinning);
 	const bool bIgnorePhysics = EnumHasAnyFlags(CurrentFlags, EMutableMeshConversionFlags::IgnorePhysics);
+	const bool bIgnoreMorphs = EnumHasAnyFlags(CurrentFlags, EMutableMeshConversionFlags::IgnoreMorphs);
 	const bool bDoNotCreateMeshMetadata = EnumHasAnyFlags(CurrentFlags, EMutableMeshConversionFlags::DoNotCreateMeshMetadata);
 
 	TSharedPtr<mu::FMesh> MutableMesh = MakeShared<mu::FMesh>();
@@ -934,7 +935,7 @@ TSharedPtr<mu::FMesh> ConvertSkeletalMeshToMutable(
 	TArray<FBoneIndexType> RemappedBoneMapIndices;
 
 	// Check if the Skeleton is valid and build the mu::Skeleton
-	if (!bIgnoreSkeleton)
+	if (!bIgnoreSkinning)
 	{
 		const USkeleton* Skeleton = SkeletalMesh->GetSkeleton();
 		if (!Skeleton)
@@ -1213,7 +1214,7 @@ TSharedPtr<mu::FMesh> ConvertSkeletalMeshToMutable(
 		}
 	}
 
-	if (!bIgnoreSkeleton)
+	if (!bIgnoreSkinning)
 	{
 		MutableMesh->GetVertexBuffers().SetBufferCount(3);
 
@@ -1327,7 +1328,7 @@ TSharedPtr<mu::FMesh> ConvertSkeletalMeshToMutable(
 
 	// Apply the morph if necessary
 	TUniquePtr<FSoftSkinVertex[]> MorphedVertexData;
-	if (!MorphName.IsEmpty())
+	if (!bIgnoreMorphs && !MorphName.IsEmpty())
 	{
 		UMorphTarget* MorphTarget = SkeletalMesh->FindMorphTarget(*MorphName);
 		if (MorphTarget && MorphTarget->GetMorphLODModels().IsValidIndex(Source.LODIndex))
@@ -1418,9 +1419,8 @@ TSharedPtr<mu::FMesh> ConvertSkeletalMeshToMutable(
 	CompletionEvent.Trigger();
 	CompletionEvent.Wait();
 
-	// TODO: Add Mesh generation flags to not include RT Morph and clothing if not needed.
 	int32 NextBufferIndex = MutableMesh->VertexBuffers.Buffers.Num();
-	if (Context.Options.bRealTimeMorphTargetsEnabled)
+	if (!bIgnoreMorphs && Context.Options.bRealTimeMorphTargetsEnabled)
 	{
 
 		// This call involves resolving every TObjectPtr<UMorphTarget> to a UMorphTarget*, so
@@ -1732,7 +1732,7 @@ TSharedPtr<mu::FMesh> ConvertSkeletalMeshToMutable(
 	}
 
 	// Clothing vertex info.
-	if (Context.Options.bClothingEnabled)
+	if (!bIgnoreSkinning && Context.Options.bClothingEnabled)
 	{
 		// Create new asset or find an already created one if the section has clothing assets.
 		// clothing assets are shared among all LODs in a section
@@ -2052,7 +2052,7 @@ TSharedPtr<mu::FMesh> ConvertSkeletalMeshToMutable(
 
 
 	// SkinWeightProfiles vertex info.
-	if (Context.Options.bSkinWeightProfilesEnabled)
+	if (!bIgnoreSkinning && Context.Options.bSkinWeightProfilesEnabled)
 	{
 		using namespace mu;
 
@@ -2481,11 +2481,6 @@ TSharedPtr<mu::FMesh> ConvertSkeletalMeshToMutable(
 			MeshData.NumTexCoordChannels = ImportedModel->LODModels[Source.LODIndex].NumTexCoords;
 			MeshData.MaxBoneIndexTypeSizeBytes = ImportedModel->LODModels[Source.LODIndex].RequiredBones.Num() > 256 ? 2 : 1;
 			MeshData.MaxNumBonesPerVertex = ImportedModel->LODModels[Source.LODIndex].GetMaxBoneInfluences();
-
-			// When mesh data is combined we will get an upper and lower bound of the number of triangles.
-			MeshData.MaxNumTriangles = ImportedModel->LODModels[Source.LODIndex].Sections[Source.SectionIndex].NumTriangles;
-			MeshData.MinNumTriangles = MeshData.MaxNumTriangles;
-
 			MeshData.bHasRealTimeMorphs = Context.Options.bRealTimeMorphTargetsEnabled && SkeletalMesh->GetMorphTargets().Num();
 			MeshData.bHasClothing = Context.Options.bClothingEnabled && ImportedModel->LODModels[Source.LODIndex].HasClothData();
 		}
@@ -3636,7 +3631,6 @@ mu::NodeMeshPtr GenerateMorphMesh(const UEdGraphPin* Pin,
 
 mu::Ptr<mu::NodeMesh> GenerateMutableSourceMesh(const UEdGraphPin* Pin,
 	FMutableGraphGenerationContext& GenerationContext,
-	FMutableGraphMeshGenerationData& MeshData,
 	uint32 SurfaceMetadataId,
 	const bool bLinkedToExtendMaterial,
 	const bool bOnlyConnectedLOD)
@@ -3653,7 +3647,6 @@ mu::Ptr<mu::NodeMesh> GenerateMutableSourceMesh(const UEdGraphPin* Pin,
 	const FGeneratedKey Key(reinterpret_cast<void*>(&GenerateMutableSourceMesh), *Pin , *Node, GenerationContext, true, bOnlyConnectedLOD);
 	if (const FGeneratedData* Generated = GenerationContext.Generated.Find(Key))
 	{
-		MeshData = Generated->meshData;
 		return static_cast<mu::NodeMesh*>(Generated->Node.get());
 	}
 	
@@ -3760,18 +3753,6 @@ mu::Ptr<mu::NodeMesh> GenerateMutableSourceMesh(const UEdGraphPin* Pin,
 			{
 				MeshNode->Value = MutableMesh;
 
-				if (GenerationContext.CompilationContext.Options.bSkinWeightProfilesEnabled && LODIndex>=0)
-				{
-					if (ImportedModel && ImportedModel->LODModels.IsValidIndex(LODIndex))
-					{
-						for (const auto& SkinWeightProfile : ImportedModel->LODModels[LODIndex].SkinWeightProfiles)
-						{
-							const int32 ProfileId = static_cast<int32>(GenerationContext.CompilationContext.GetSkinWeightProfileIdUnique(SkinWeightProfile.Key));
-							MeshData.SkinWeightProfilesSemanticIndices.AddUnique(ProfileId);
-						}
-					}
-				}
-
 				if (!TypedNodeSkel->AnimInstance.IsNull())
 				{
 					AddTagToMutableMeshUnique(*MutableMesh, AnimBPAssetTag);
@@ -3796,22 +3777,6 @@ mu::Ptr<mu::NodeMesh> GenerateMutableSourceMesh(const UEdGraphPin* Pin,
 					FString MeshTag = FString("__MeshPath:") + MeshPath;
 					AddTagToMutableMeshUnique(*MutableMesh, MeshTag);
 				}
-
-				if (ImportedModel->LODModels.IsValidIndex(LODIndex) &&
-					ImportedModel->LODModels[LODIndex].Sections.IsValidIndex(SectionIndex))
-				{
-					MeshData.bHasVertexColors = SkeletalMesh->GetHasVertexColors();	
-					MeshData.NumTexCoordChannels = ImportedModel->LODModels[LODIndex].NumTexCoords;
-					MeshData.MaxBoneIndexTypeSizeBytes = ImportedModel->LODModels[LODIndex].RequiredBones.Num() > 256 ? 2 : 1;
-					MeshData.MaxNumBonesPerVertex = ImportedModel->LODModels[LODIndex].GetMaxBoneInfluences();
-
-					// When mesh data is combined we will get an upper and lower bound of the number of triangles.
-					MeshData.MaxNumTriangles = ImportedModel->LODModels[LODIndex].Sections[SectionIndex].NumTriangles;
-					MeshData.MinNumTriangles = MeshData.MaxNumTriangles;
-
-					MeshData.bHasRealTimeMorphs = GenerationContext.CompilationContext.Options.bRealTimeMorphTargetsEnabled && SkeletalMesh->GetMorphTargets().Num();
-					MeshData.bHasClothing = GenerationContext.CompilationContext.Options.bClothingEnabled && ImportedModel->LODModels[LODIndex].HasClothData();
-				}
 			}
 
 			// Layouts
@@ -3997,7 +3962,7 @@ mu::Ptr<mu::NodeMesh> GenerateMutableSourceMesh(const UEdGraphPin* Pin,
 			// Mesh Morph Stack Management
 			FMorphNodeData NewMorphData = { TypedNodeMorph, MorphTargetName, FactorNode, TypedNodeMorph->MeshPin() };
 			GenerationContext.MeshMorphStack.Push(NewMorphData);
-			Result = GenerateMutableSourceMesh(ConnectedPin, GenerationContext, MeshData, SurfaceMetadataId, false, bOnlyConnectedLOD);
+			Result = GenerateMutableSourceMesh(ConnectedPin, GenerationContext, SurfaceMetadataId, false, bOnlyConnectedLOD);
 			GenerationContext.MeshMorphStack.Pop(EAllowShrinking::Yes);
 		}
 		else
@@ -4062,7 +4027,7 @@ mu::Ptr<mu::NodeMesh> GenerateMutableSourceMesh(const UEdGraphPin* Pin,
 
 				if (const UEdGraphPin* MeshConnectedPin = FollowInputPin(*TypedNodeMeshMorphStackApp->GetMeshPin()))
 				{
-					Result = GenerateMutableSourceMesh(MeshConnectedPin, GenerationContext, MeshData, SurfaceMetadataId, false, bOnlyConnectedLOD);
+					Result = GenerateMutableSourceMesh(MeshConnectedPin, GenerationContext, SurfaceMetadataId, false, bOnlyConnectedLOD);
 				}
 
 				for (int32 MorphIndex = 0; MorphIndex < AddedMorphs; ++MorphIndex)
@@ -4133,10 +4098,8 @@ mu::Ptr<mu::NodeMesh> GenerateMutableSourceMesh(const UEdGraphPin* Pin,
 				{
 					if (const UEdGraphPin* ConnectedPin = FollowInputPin(*TypedNodeMeshSwitch->GetElementPin(SelectorIndex)))
 					{
-						FMutableGraphMeshGenerationData ChildMeshData;
-						Result = GenerateMutableSourceMesh(ConnectedPin, GenerationContext, ChildMeshData, SurfaceMetadataId, false, bOnlyConnectedLOD);
+						Result = GenerateMutableSourceMesh(ConnectedPin, GenerationContext, SurfaceMetadataId, false, bOnlyConnectedLOD);
 						SwitchNode->Options[SelectorIndex] = Result;
-						MeshData.Combine(ChildMeshData);
 					}
 				}
 
@@ -4158,12 +4121,10 @@ mu::Ptr<mu::NodeMesh> GenerateMutableSourceMesh(const UEdGraphPin* Pin,
 
 		if (const UEdGraphPin* ConnectedPin = FollowInputPin(*TypedNodeMeshVar->DefaultPin()))
 		{
-			FMutableGraphMeshGenerationData ChildMeshData;
-			mu::NodeMeshPtr ChildNode = GenerateMutableSourceMesh(ConnectedPin, GenerationContext, ChildMeshData, SurfaceMetadataId, false, bOnlyConnectedLOD);
+			mu::NodeMeshPtr ChildNode = GenerateMutableSourceMesh(ConnectedPin, GenerationContext, SurfaceMetadataId, false, bOnlyConnectedLOD);
 			if (ChildNode)
 			{
 				MeshNode->DefaultMesh = ChildNode;
-				MeshData.Combine(ChildMeshData);
 			}
 			else
 			{
@@ -4185,10 +4146,8 @@ mu::Ptr<mu::NodeMesh> GenerateMutableSourceMesh(const UEdGraphPin* Pin,
 
 			if (const UEdGraphPin* ConnectedPin = FollowInputPin(*VariationPin))
 			{
-				FMutableGraphMeshGenerationData VariationMeshData;
-				mu::NodeMeshPtr ChildNode = GenerateMutableSourceMesh(ConnectedPin, GenerationContext, VariationMeshData, SurfaceMetadataId, false, bOnlyConnectedLOD);
+				mu::NodeMeshPtr ChildNode = GenerateMutableSourceMesh(ConnectedPin, GenerationContext, SurfaceMetadataId, false, bOnlyConnectedLOD);
 				MeshNode->Variations[VariationIndex].Mesh = ChildNode;
-				MeshData.Combine(VariationMeshData);
 			}
 		}
 	}
@@ -4200,12 +4159,10 @@ mu::Ptr<mu::NodeMesh> GenerateMutableSourceMesh(const UEdGraphPin* Pin,
 
 		if (const UEdGraphPin* ConnectedPin = FollowInputPin(*TypedNodeReshape->BaseMeshPin()))
 		{
-			FMutableGraphMeshGenerationData ChildMeshData;
-			mu::Ptr<mu::NodeMesh> ChildNode = GenerateMutableSourceMesh(ConnectedPin, GenerationContext, ChildMeshData, SurfaceMetadataId, false, bOnlyConnectedLOD);
+			mu::Ptr<mu::NodeMesh> ChildNode = GenerateMutableSourceMesh(ConnectedPin, GenerationContext, SurfaceMetadataId, false, bOnlyConnectedLOD);
 			if (ChildNode)
 			{
 				MeshNode->BaseMesh = ChildNode;
-				MeshData.Combine(ChildMeshData);
 			}
 			else
 			{
@@ -4323,16 +4280,12 @@ mu::Ptr<mu::NodeMesh> GenerateMutableSourceMesh(const UEdGraphPin* Pin,
 
 		GenerationContext.MeshGenerationFlags.Push( ShapeFlags );
 			
-		constexpr int32 PinNotSetValue = TNumericLimits<int32>::Max();
-		int32 BaseShapeTriangleCount = PinNotSetValue;
 		if (const UEdGraphPin* ConnectedPin = FollowInputPin(*TypedNodeReshape->BaseShapePin()))
 		{
-			FMutableGraphMeshGenerationData ChildMeshData;
-			mu::Ptr<mu::NodeMesh> ChildNode = GenerateMutableSourceMesh(ConnectedPin, GenerationContext, ChildMeshData, 0, false, true);
+			mu::Ptr<mu::NodeMesh> ChildNode = GenerateMutableSourceMesh(ConnectedPin, GenerationContext, 0, false, true);
 	
 			if (ChildNode)
 			{
-				BaseShapeTriangleCount = ChildMeshData.MaxNumTriangles == ChildMeshData.MinNumTriangles ? ChildMeshData.MaxNumTriangles : -1;
 				MeshNode->BaseShape = ChildNode;
 			}
 			else
@@ -4341,15 +4294,12 @@ mu::Ptr<mu::NodeMesh> GenerateMutableSourceMesh(const UEdGraphPin* Pin,
 			}
 		}
 
-		int32 TargetShapeTriangleCount = PinNotSetValue;
 		if (const UEdGraphPin* ConnectedPin = FollowInputPin(*TypedNodeReshape->TargetShapePin()))
 		{
-			FMutableGraphMeshGenerationData ChildMeshData;
-			mu::Ptr<mu::NodeMesh> ChildNode = GenerateMutableSourceMesh(ConnectedPin, GenerationContext, ChildMeshData, 0, false, true);
+			mu::Ptr<mu::NodeMesh> ChildNode = GenerateMutableSourceMesh(ConnectedPin, GenerationContext, 0, false, true);
 			
 			if (ChildNode)
 			{
-				TargetShapeTriangleCount = ChildMeshData.MaxNumTriangles == ChildMeshData.MinNumTriangles ? ChildMeshData.MaxNumTriangles : -1;
 				MeshNode->TargetShape = ChildNode;
 			}
 			else
@@ -4358,26 +4308,6 @@ mu::Ptr<mu::NodeMesh> GenerateMutableSourceMesh(const UEdGraphPin* Pin,
 			}
 		}
 
-
-		// There is cases where it is not possible to determine if the test passes or not, e.g., mesh variations or switches.
-		// Until now if there were the possibility of two meshes not being compatible the warning was raised. This is not ideal
-		// as there are legitimate cases were the meshes will match but we cannot be sure they will. For now disable the warning.
-		
-		constexpr bool bDissableMeshReshapeWarning = true;
-
-		if (!bDissableMeshReshapeWarning)
-		{
-			// If any of the shape pins is not set, don't warn about it.
-			if (BaseShapeTriangleCount != PinNotSetValue && TargetShapeTriangleCount != PinNotSetValue)
-			{
-				if (BaseShapeTriangleCount != TargetShapeTriangleCount || BaseShapeTriangleCount == -1 || TargetShapeTriangleCount == -1)
-				{
-					GenerationContext.Log(LOCTEXT("ReshapeMeshShapeIncompatible",
-						"Base and Target Shapes might not be compatible. Don't have the same number of triangles."), Node, EMessageSeverity::Warning);
-				}
-			}
-		}
-
 		GenerationContext.MeshGenerationFlags.Pop();
 	}
 
@@ -4385,7 +4315,7 @@ mu::Ptr<mu::NodeMesh> GenerateMutableSourceMesh(const UEdGraphPin* Pin,
 	{
 		if (const UEdGraphPin* InputMeshPin = FollowInputPin(*TypedNode->GetInputMeshPin()))
 		{
-			mu::Ptr<mu::NodeMesh> InputMeshNode = GenerateMutableSourceMesh(InputMeshPin, GenerationContext, MeshData, SurfaceMetadataId, false, bOnlyConnectedLOD);
+			mu::Ptr<mu::NodeMesh> InputMeshNode = GenerateMutableSourceMesh(InputMeshPin, GenerationContext, SurfaceMetadataId, false, bOnlyConnectedLOD);
 
 			if (GenerationContext.GetCurrentComponentInfo()->RefSkeletalMesh)
 			{
@@ -4413,7 +4343,7 @@ mu::Ptr<mu::NodeMesh> GenerateMutableSourceMesh(const UEdGraphPin* Pin,
 					if (const UCustomizableObjectNodeTable* TypedNodeTable = Cast<UCustomizableObjectNodeTable>(TablePosePin->GetOwningNode()))
 					{
 						mu::Ptr<mu::NodeMeshApplyPose> NodeMeshApplyPose = new mu::NodeMeshApplyPose();
-						mu::Ptr<mu::NodeMesh> MeshTableNode = GenerateMutableSourceMesh(TablePosePin, GenerationContext, MeshData, SurfaceMetadataId, false, bOnlyConnectedLOD);
+						mu::Ptr<mu::NodeMesh> MeshTableNode = GenerateMutableSourceMesh(TablePosePin, GenerationContext, SurfaceMetadataId, false, bOnlyConnectedLOD);
 
 						NodeMeshApplyPose->Base = InputMeshNode;
 						NodeMeshApplyPose->Pose = MeshTableNode;
@@ -4536,30 +4466,6 @@ mu::Ptr<mu::NodeMesh> GenerateMutableSourceMesh(const UEdGraphPin* Pin,
 						// Pose Assets do not need this part of the code
 						if (Pin->PinType.PinCategory == Schema->PC_Mesh)
 						{
-							if (DefaultSkeletalMesh)
-							{
-								FSkeletalMeshModel* ImportedModel = DefaultSkeletalMesh->GetImportedModel();
-
-								if (ImportedModel->LODModels.IsValidIndex(LODIndex) &&
-									ImportedModel->LODModels[LODIndex].Sections.IsValidIndex(SectionIndex))
-								{
-									// TODO: this should be made for all the meshes of the Column to support meshes with different values
-									// Filling Mesh Data
-									MeshData.bHasVertexColors = DefaultSkeletalMesh->GetHasVertexColors();
-									MeshData.NumTexCoordChannels = ImportedModel->LODModels[LODIndex].NumTexCoords;
-									MeshData.MaxBoneIndexTypeSizeBytes = ImportedModel->LODModels[LODIndex].RequiredBones.Num() > 256 ? 2 : 1;
-									MeshData.MaxNumBonesPerVertex = ImportedModel->LODModels[LODIndex].GetMaxBoneInfluences();
-
-									// When mesh data is combined we will get an upper and lower bound of the number of triangles.
-									MeshData.MaxNumTriangles = ImportedModel->LODModels[LODIndex].Sections[SectionIndex].NumTriangles;
-									MeshData.MinNumTriangles = ImportedModel->LODModels[LODIndex].Sections[SectionIndex].NumTriangles;
-
-									// With tables we have to use the default skeletal mesh to tell if the surface mesh needs morphs or clothing.
-									MeshData.bHasRealTimeMorphs = GenerationContext.CompilationContext.Options.bRealTimeMorphTargetsEnabled && DefaultSkeletalMesh->GetMorphTargets().Num();
-									MeshData.bHasClothing = GenerationContext.CompilationContext.Options.bClothingEnabled && ImportedModel->LODModels[LODIndex].HasClothData();
-								}
-							}
-
 							TArray<UCustomizableObjectLayout*> Layouts = TypedNodeTable->GetLayouts(Pin);
 							MeshTableNode->Layouts.SetNum(Layouts.Num());
 
@@ -4658,7 +4564,7 @@ mu::Ptr<mu::NodeMesh> GenerateMutableSourceMesh(const UEdGraphPin* Pin,
 			if (const UEdGraphPin* FollowPin = FollowInputPin(*OutputPin))
 			{
 				GenerationContext.MacroNodesStack.Push(TypedNodeMacro);
-				Result = GenerateMutableSourceMesh(FollowPin, GenerationContext, MeshData, SurfaceMetadataId, bLinkedToExtendMaterial, bOnlyConnectedLOD);
+				Result = GenerateMutableSourceMesh(FollowPin, GenerationContext, SurfaceMetadataId, bLinkedToExtendMaterial, bOnlyConnectedLOD);
 				GenerationContext.MacroNodesStack.Pop();
 			}
 			else
@@ -4688,7 +4594,7 @@ mu::Ptr<mu::NodeMesh> GenerateMutableSourceMesh(const UEdGraphPin* Pin,
 		{
 			if (const UEdGraphPin* FollowPin = FollowInputPin(*InputPin))
 			{
-				Result = GenerateMutableSourceMesh(FollowPin, GenerationContext, MeshData, SurfaceMetadataId, bLinkedToExtendMaterial, bOnlyConnectedLOD);
+				Result = GenerateMutableSourceMesh(FollowPin, GenerationContext, SurfaceMetadataId, bLinkedToExtendMaterial, bOnlyConnectedLOD);
 			}
 			else
 			{
@@ -4718,7 +4624,7 @@ mu::Ptr<mu::NodeMesh> GenerateMutableSourceMesh(const UEdGraphPin* Pin,
 	
 	if (bCacheNode)
 	{
-		GenerationContext.Generated.Add(Key, FGeneratedData(Node, Result, &MeshData));
+		GenerationContext.Generated.Add(Key, FGeneratedData(Node, Result));
 		GenerationContext.GeneratedNodes.Add(Node);
 	}
 
diff --git a/Engine/Plugins/Mutable/Source/CustomizableObjectEditor/Private/MuCOE/GenerateMutableSource/GenerateMutableSourceMesh.h b/Engine/Plugins/Mutable/Source/CustomizableObjectEditor/Private/MuCOE/GenerateMutableSource/GenerateMutableSourceMesh.h
--- a/Engine/Plugins/Mutable/Source/CustomizableObjectEditor/Private/MuCOE/GenerateMutableSource/GenerateMutableSourceMesh.h
+++ b/Engine/Plugins/Mutable/Source/CustomizableObjectEditor/Private/MuCOE/GenerateMutableSource/GenerateMutableSourceMesh.h
@@ -109,7 +109,7 @@ TSharedPtr<mu::FMesh> BuildMorphedMutableMesh(const UEdGraphPin* BaseSourcePin,
  * @param bLinkedToExtendMaterial 
  * @param bOnlyConnectedLOD Corrected LOD and Section will unconditionally always be the connected ones.
  * @return  Mutable Mesh Node. */
-mu::Ptr<mu::NodeMesh> GenerateMutableSourceMesh(const UEdGraphPin* Pin, FMutableGraphGenerationContext& GenerationContext, FMutableGraphMeshGenerationData& MeshData, uint32 SurfaceMetadataId, bool bLinkedToExtendMaterial, bool bOnlyConnectedLOD);
+mu::Ptr<mu::NodeMesh> GenerateMutableSourceMesh(const UEdGraphPin* Pin, FMutableGraphGenerationContext& GenerationContext, uint32 SurfaceMetadataId, bool bLinkedToExtendMaterial, bool bOnlyConnectedLOD);
 
 
 /** */
diff --git a/Engine/Plugins/Mutable/Source/CustomizableObjectEditor/Private/MuCOE/GenerateMutableSource/GenerateMutableSourceModifier.cpp b/Engine/Plugins/Mutable/Source/CustomizableObjectEditor/Private/MuCOE/GenerateMutableSource/GenerateMutableSourceModifier.cpp
--- a/Engine/Plugins/Mutable/Source/CustomizableObjectEditor/Private/MuCOE/GenerateMutableSource/GenerateMutableSourceModifier.cpp
+++ b/Engine/Plugins/Mutable/Source/CustomizableObjectEditor/Private/MuCOE/GenerateMutableSource/GenerateMutableSourceModifier.cpp
@@ -104,7 +104,8 @@ mu::Ptr<mu::NodeModifier> GenerateMutableSourceModifier(const UEdGraphPin * Pin,
 	{
 		const EMutableMeshConversionFlags ModifiersMeshFlags =
 			EMutableMeshConversionFlags::IgnoreSkinning |
-			EMutableMeshConversionFlags::IgnorePhysics;
+			EMutableMeshConversionFlags::IgnorePhysics |
+			EMutableMeshConversionFlags::IgnoreMorphs;
 		GenerationContext.MeshGenerationFlags.Push(ModifiersMeshFlags);
 
 		// This modifier can be connected to multiple nodes at the same time and, when that happens and if the cache is being used, only the first node to be processed does work. 
@@ -135,7 +136,8 @@ mu::Ptr<mu::NodeModifier> GenerateMutableSourceModifier(const UEdGraphPin * Pin,
 	{
 		const EMutableMeshConversionFlags ModifiersMeshFlags =
 			EMutableMeshConversionFlags::IgnoreSkinning |
-			EMutableMeshConversionFlags::IgnorePhysics;
+			EMutableMeshConversionFlags::IgnorePhysics |
+			EMutableMeshConversionFlags::IgnoreMorphs;
 		GenerationContext.MeshGenerationFlags.Push(ModifiersMeshFlags);
 
 		mu::Ptr<mu::NodeModifierMeshClipDeform> ClipNode = new mu::NodeModifierMeshClipDeform();
@@ -145,8 +147,7 @@ mu::Ptr<mu::NodeModifier> GenerateMutableSourceModifier(const UEdGraphPin * Pin,
 
 		if (const UEdGraphPin* ConnectedPin = FollowInputPin(*TypedNodeClipDeform->ClipShapePin()))
 		{
-			FMutableGraphMeshGenerationData DummyMeshData;
-			mu::Ptr<mu::NodeMesh> ClipMesh = GenerateMutableSourceMesh(ConnectedPin, GenerationContext, DummyMeshData, 0, false, true);
+			mu::Ptr<mu::NodeMesh> ClipMesh = GenerateMutableSourceMesh(ConnectedPin, GenerationContext, 0, false, true);
 
 			ClipNode->ClipMesh = ClipMesh;
 
@@ -186,7 +187,8 @@ mu::Ptr<mu::NodeModifier> GenerateMutableSourceModifier(const UEdGraphPin * Pin,
 	{
 		const EMutableMeshConversionFlags ModifiersMeshFlags =
 			EMutableMeshConversionFlags::IgnoreSkinning |
-			EMutableMeshConversionFlags::IgnorePhysics;
+			EMutableMeshConversionFlags::IgnorePhysics |
+			EMutableMeshConversionFlags::IgnoreMorphs;
 		GenerationContext.MeshGenerationFlags.Push(ModifiersMeshFlags);
 
 		// MeshClipWithMesh can be connected to multiple objects, so the compiled NodeModifierMeshClipWithMesh
@@ -200,9 +202,7 @@ mu::Ptr<mu::NodeModifier> GenerateMutableSourceModifier(const UEdGraphPin * Pin,
 
 		if (const UEdGraphPin* ConnectedPin = FollowInputPin(*TypedNodeClipMesh->ClipMeshPin()))
 		{
-			FMutableGraphMeshGenerationData DummyMeshData;
-
-			mu::NodeMeshPtr ClipMesh = GenerateMutableSourceMesh(ConnectedPin, GenerationContext, DummyMeshData, 0, false, true);
+			mu::NodeMeshPtr ClipMesh = GenerateMutableSourceMesh(ConnectedPin, GenerationContext, 0, false, true);
 
 			if (FMatrix Matrix = TypedNodeClipMesh->Transform.ToMatrixWithScale(); Matrix != FMatrix::Identity)
 			{
@@ -232,7 +232,8 @@ mu::Ptr<mu::NodeModifier> GenerateMutableSourceModifier(const UEdGraphPin * Pin,
 	{
 		const EMutableMeshConversionFlags ModifiersMeshFlags =
 			EMutableMeshConversionFlags::IgnoreSkinning |
-			EMutableMeshConversionFlags::IgnorePhysics;
+			EMutableMeshConversionFlags::IgnorePhysics |
+			EMutableMeshConversionFlags::IgnoreMorphs;
 		GenerationContext.MeshGenerationFlags.Push(ModifiersMeshFlags);
 
 		// This modifier can be connected to multiple objects, so the compiled node
@@ -246,8 +247,6 @@ mu::Ptr<mu::NodeModifier> GenerateMutableSourceModifier(const UEdGraphPin * Pin,
 
 		if (const UEdGraphPin* ConnectedPin = FollowInputPin(*TypedNodeClipUVMask->ClipMaskPin()))
 		{
-			FMutableGraphMeshGenerationData DummyMeshData;
-
 			mu::Ptr<mu::NodeImage> ClipMask = GenerateMutableSourceImage(ConnectedPin, GenerationContext, 0);
 
 			ClipNode->ClipMask = ClipMask;
@@ -291,7 +290,6 @@ mu::Ptr<mu::NodeModifier> GenerateMutableSourceModifier(const UEdGraphPin * Pin,
 			GenerationContext.CurrentLOD = LODIndex;
 
 			mu::Ptr<mu::NodeMesh> AddMeshNode;
-			FMutableGraphMeshGenerationData MeshData;
 			if (const UEdGraphPin* ConnectedPin = FollowInputPin(*TypedNodeExt->AddMeshPin()))
 			{
 				// Flags to know which UV channels need layout
@@ -326,7 +324,7 @@ mu::Ptr<mu::NodeModifier> GenerateMutableSourceModifier(const UEdGraphPin * Pin,
 					}
 				}	
 
-				AddMeshNode = GenerateMutableSourceMesh(ConnectedPin, GenerationContext, MeshData, SurfaceMetadataUniqueHash, true, false);
+				AddMeshNode = GenerateMutableSourceMesh(ConnectedPin, GenerationContext, SurfaceMetadataUniqueHash, true, false);
 
 				GenerationContext.LayoutGenerationFlags.Pop();
 			}
@@ -434,7 +432,8 @@ mu::Ptr<mu::NodeModifier> GenerateMutableSourceModifier(const UEdGraphPin * Pin,
 	{
 		const EMutableMeshConversionFlags ModifiersMeshFlags =
 			EMutableMeshConversionFlags::IgnoreSkinning |
-			EMutableMeshConversionFlags::IgnorePhysics;
+			EMutableMeshConversionFlags::IgnorePhysics |
+			EMutableMeshConversionFlags::IgnoreMorphs;
 		GenerationContext.MeshGenerationFlags.Push(ModifiersMeshFlags);
 
 		mu::Ptr<mu::NodeModifierSurfaceEdit> SurfNode = new mu::NodeModifierSurfaceEdit();
@@ -456,8 +455,7 @@ mu::Ptr<mu::NodeModifier> GenerateMutableSourceModifier(const UEdGraphPin * Pin,
 				GenerationContext.FromLOD = 0;
 				GenerationContext.CurrentLOD = LODIndex;
 
-				FMutableGraphMeshGenerationData DummyMeshData;
-				mu::Ptr<mu::NodeMesh> RemoveMeshNode = GenerateMutableSourceMesh(ConnectedPin, GenerationContext, DummyMeshData, 0, false, true);
+				mu::Ptr<mu::NodeMesh> RemoveMeshNode = GenerateMutableSourceMesh(ConnectedPin, GenerationContext, 0, false, true);
 				SurfNode->LODs[LODIndex].MeshRemove = RemoveMeshNode;
 			}
 		}
@@ -471,7 +469,8 @@ mu::Ptr<mu::NodeModifier> GenerateMutableSourceModifier(const UEdGraphPin * Pin,
 	{
 		const EMutableMeshConversionFlags ModifiersMeshFlags =
 			EMutableMeshConversionFlags::IgnoreSkinning |
-			EMutableMeshConversionFlags::IgnorePhysics;
+			EMutableMeshConversionFlags::IgnorePhysics |
+			EMutableMeshConversionFlags::IgnoreMorphs;
 		GenerationContext.MeshGenerationFlags.Push(ModifiersMeshFlags);
 
 		mu::Ptr<mu::NodeModifierMeshClipWithUVMask> ClipNode = new mu::NodeModifierMeshClipWithUVMask();
@@ -493,8 +492,9 @@ mu::Ptr<mu::NodeModifier> GenerateMutableSourceModifier(const UEdGraphPin * Pin,
 	else if (UCustomizableObjectNodeModifierEditMeshSection* TypedNodeEdit = Cast<UCustomizableObjectNodeModifierEditMeshSection>(Node))
 	{
 		const EMutableMeshConversionFlags ModifiersMeshFlags =
-				EMutableMeshConversionFlags::IgnoreSkinning |
-				EMutableMeshConversionFlags::IgnorePhysics;
+			EMutableMeshConversionFlags::IgnoreSkinning |
+			EMutableMeshConversionFlags::IgnorePhysics |
+			EMutableMeshConversionFlags::IgnoreMorphs;
 		GenerationContext.MeshGenerationFlags.Push(ModifiersMeshFlags);
 
 		mu::Ptr<mu::NodeModifierSurfaceEdit> SurfNode = new mu::NodeModifierSurfaceEdit();
@@ -571,7 +571,8 @@ mu::Ptr<mu::NodeModifier> GenerateMutableSourceModifier(const UEdGraphPin * Pin,
 	{
 		const EMutableMeshConversionFlags ModifiersMeshFlags =
 			EMutableMeshConversionFlags::IgnoreSkinning |
-			EMutableMeshConversionFlags::IgnorePhysics;
+			EMutableMeshConversionFlags::IgnorePhysics |
+			EMutableMeshConversionFlags::IgnoreMorphs;
 		GenerationContext.MeshGenerationFlags.Push(ModifiersMeshFlags);
 
 		mu::Ptr<mu::NodeModifierSurfaceEdit> SurfNode = new mu::NodeModifierSurfaceEdit();
@@ -654,7 +655,8 @@ mu::Ptr<mu::NodeModifier> GenerateMutableSourceModifier(const UEdGraphPin * Pin,
 	{
 		const EMutableMeshConversionFlags ModifiersMeshFlags =
 			EMutableMeshConversionFlags::IgnoreSkinning |
-			EMutableMeshConversionFlags::IgnorePhysics;
+			EMutableMeshConversionFlags::IgnorePhysics |
+			EMutableMeshConversionFlags::IgnoreMorphs;
 		GenerationContext.MeshGenerationFlags.Push(ModifiersMeshFlags);
 
 		// MeshTransformInMesh can be connected to multiple objects, so the compiled NodeModifierMeshTransformInMesh
@@ -672,9 +674,7 @@ mu::Ptr<mu::NodeModifier> GenerateMutableSourceModifier(const UEdGraphPin * Pin,
 		// If no bounding mesh is provided, we transform the entire mesh.
 		if (const UEdGraphPin* ConnectedPin = FollowInputPin(*TypedNodeTransformMesh->BoundingMeshPin()))
 		{
-			FMutableGraphMeshGenerationData DummyMeshData;
-
-			mu::Ptr<mu::NodeMesh> BoundingMesh = GenerateMutableSourceMesh(ConnectedPin, GenerationContext, DummyMeshData, 0, false, true);
+			mu::Ptr<mu::NodeMesh> BoundingMesh = GenerateMutableSourceMesh(ConnectedPin, GenerationContext, 0, false, true);
 
 			if (FMatrix Matrix = TypedNodeTransformMesh->BoundingMeshTransform.ToMatrixWithScale(); Matrix != FMatrix::Identity)
 			{
diff --git a/Engine/Plugins/Mutable/Source/CustomizableObjectEditor/Private/MuCOE/GenerateMutableSource/GenerateMutableSourceSurface.cpp b/Engine/Plugins/Mutable/Source/CustomizableObjectEditor/Private/MuCOE/GenerateMutableSource/GenerateMutableSourceSurface.cpp
--- a/Engine/Plugins/Mutable/Source/CustomizableObjectEditor/Private/MuCOE/GenerateMutableSource/GenerateMutableSourceSurface.cpp
+++ b/Engine/Plugins/Mutable/Source/CustomizableObjectEditor/Private/MuCOE/GenerateMutableSource/GenerateMutableSourceSurface.cpp
@@ -193,8 +193,7 @@ mu::Ptr<mu::NodeSurface> GenerateMutableSourceSurface(const UEdGraphPin * Pin, F
 
 				GenerationContext.LayoutGenerationFlags.Push(LayoutGenerationFlags);
 
-				FMutableGraphMeshGenerationData MeshData;
-				MeshNode = GenerateMutableSourceMesh(ConnectedPin, GenerationContext, MeshData, SurfaceMetadataUniqueHash, false, false);
+				MeshNode = GenerateMutableSourceMesh(ConnectedPin, GenerationContext, SurfaceMetadataUniqueHash, false, false);
 
 				GenerationContext.LayoutGenerationFlags.Pop();
 
diff --git a/Engine/Plugins/Mutable/Source/CustomizableObjectEditor/Private/MuCOE/GenerateMutableSource/GenerateMutableSourceSurface.h b/Engine/Plugins/Mutable/Source/CustomizableObjectEditor/Private/MuCOE/GenerateMutableSource/GenerateMutableSourceSurface.h
--- a/Engine/Plugins/Mutable/Source/CustomizableObjectEditor/Private/MuCOE/GenerateMutableSource/GenerateMutableSourceSurface.h
+++ b/Engine/Plugins/Mutable/Source/CustomizableObjectEditor/Private/MuCOE/GenerateMutableSource/GenerateMutableSourceSurface.h
@@ -7,7 +7,7 @@
 
 class UEdGraphPin;
 struct FMutableGraphGenerationContext;
-struct FMutableGraphMeshGenerationData;
+
 enum class ECustomizableObjectNumBoneInfluences:uint8;
 namespace mu
 {
