[mutable] Fix binary asset indeterminism when using some operations using shapes.
- fix bug where bone radius may be lost for mesh clip morph operations.
#rnx

diff --git a/Engine/Plugins/Mutable/Source/MutableRuntime/Internal/MuR/ParametersPrivate.h b/Engine/Plugins/Mutable/Source/MutableRuntime/Internal/MuR/ParametersPrivate.h
--- a/Engine/Plugins/Mutable/Source/MutableRuntime/Internal/MuR/ParametersPrivate.h
+++ b/Engine/Plugins/Mutable/Source/MutableRuntime/Internal/MuR/ParametersPrivate.h
@@ -198,11 +198,11 @@ namespace Private
 	struct FShape
 	{
 		// Transform
-		FVector3f position;
-		FVector3f up;
-		FVector3f side;
+		FVector3f position = FVector3f(0,0,0);
+		FVector3f up = FVector3f(0, 0, 0);
+		FVector3f side = FVector3f(0, 0, 0);
 		
-		FVector3f size;
+		FVector3f size = FVector3f(0, 0, 0);
 
 		// 
 		enum class Type : uint8
@@ -212,7 +212,6 @@ namespace Private
 			AABox
 		};
         uint8 type = 0;
-		
 
 		//!
 		void Serialise(FOutputArchive& arch) const
diff --git a/Engine/Plugins/Mutable/Source/MutableTools/Private/MuT/ASTOpConditional.h b/Engine/Plugins/Mutable/Source/MutableTools/Private/MuT/ASTOpConditional.h
--- a/Engine/Plugins/Mutable/Source/MutableTools/Private/MuT/ASTOpConditional.h
+++ b/Engine/Plugins/Mutable/Source/MutableTools/Private/MuT/ASTOpConditional.h
@@ -21,7 +21,7 @@ struct FProgram;
 	public:
 
 		//! Type of switch
-		EOpType type;
+		EOpType type = EOpType::NONE;
 
 		//! Boolean expression
 		ASTChild condition;
diff --git a/Engine/Plugins/Mutable/Source/MutableTools/Private/MuT/ASTOpMeshClipMorphPlane.cpp b/Engine/Plugins/Mutable/Source/MutableTools/Private/MuT/ASTOpMeshClipMorphPlane.cpp
--- a/Engine/Plugins/Mutable/Source/MutableTools/Private/MuT/ASTOpMeshClipMorphPlane.cpp
+++ b/Engine/Plugins/Mutable/Source/MutableTools/Private/MuT/ASTOpMeshClipMorphPlane.cpp
@@ -42,7 +42,8 @@ namespace mu
 				VertexSelectionType == other->VertexSelectionType &&
 				FaceCullStrategy == other->FaceCullStrategy &&
 				Dist == other->Dist &&
-				Factor == other->Factor;
+				Factor == other->Factor &&
+				VertexSelectionBoneMaxRadius == other->VertexSelectionBoneMaxRadius;
 		}
 		return false;
 	}
@@ -68,6 +69,7 @@ namespace mu
 		n->FaceCullStrategy = FaceCullStrategy;
 		n->Dist = Dist;
 		n->Factor = Factor;
+		n->VertexSelectionBoneMaxRadius = VertexSelectionBoneMaxRadius;
 		return n;
 	}
 
@@ -108,8 +110,8 @@ namespace mu
 			Args.MaxBoneRadius = VertexSelectionBoneMaxRadius;
 
 			linkedAddress = (OP::ADDRESS)Program.OpAddress.Num();
-			Program.OpAddress.Add((uint32_t)Program.ByteCode.Num());
-			AppendCode(Program.ByteCode, EOpType::ME_CLIPMORPHPLANE);
+			Program.OpAddress.Add(Program.ByteCode.Num());
+			AppendCode(Program.ByteCode, GetOpType());
 			AppendCode(Program.ByteCode, Args);
 		}
 
diff --git a/Engine/Plugins/Mutable/Source/MutableTools/Private/MuT/CodeOptimiser.cpp b/Engine/Plugins/Mutable/Source/MutableTools/Private/MuT/CodeOptimiser.cpp
--- a/Engine/Plugins/Mutable/Source/MutableTools/Private/MuT/CodeOptimiser.cpp
+++ b/Engine/Plugins/Mutable/Source/MutableTools/Private/MuT/CodeOptimiser.cpp
@@ -605,8 +605,7 @@ namespace mu
 	};
 
 
-	//---------------------------------------------------------------------------------------------
-	bool ConstantGeneratorAST( const CompilerOptions::Private* InOptions, Ptr<ASTOp>& Root, int32 Pass )
+	bool ConstantGenerator( const CompilerOptions::Private* InOptions, Ptr<ASTOp>& Root, int32 Pass )
 	{
 		MUTABLE_CPUPROFILER_SCOPE(ConstantGenerator);
 
@@ -621,11 +620,6 @@ namespace mu
 		{
 			Ptr<ASTOp> Root;
 			UE::Tasks::FTaskEvent CompletedEvent;
-
-			// This is only necessary for non-concurrent execution
-			TArray< UE::Tasks::FTask, TInlineAllocator<8> > Requisites;
-			TFunction<void()> NonConcurrentTask;
-			TUniquePtr<FConstantTask> TaskData;
 		};
 		TArray< FConstantSubgraph > ConstantSubgraphs;
 		ConstantSubgraphs.Reserve(256);
@@ -793,7 +787,7 @@ namespace mu
 
 
 		// Launch the tasks.
-		UE::Tasks::FTask LaunchTask = UE::Tasks::Launch(TEXT("ConstantGeneratorLaunchTasks"), 
+		UE::Tasks::FTask LaunchTask = UE::Tasks::Launch(TEXT("ConstantGeneratorLaunchTasks"),
 			[&ConstantSubgraphs, &GetRequisites, Pass, InOptions]()
 			{
 				MUTABLE_CPUPROFILER_SCOPE(ConstantGenerator_LaunchTasks);
@@ -817,16 +811,16 @@ namespace mu
 						SubgraphType == EOpType::IM_CONSTANT)
 					{
 						uint32 ImageID = 0;
-						if (SubgraphType == EOpType::IM_REFERENCE)							
+						if (SubgraphType == EOpType::IM_REFERENCE)
 						{
 							bIsReference = true;
 							const ASTOpReferenceResource* Typed = static_cast<const ASTOpReferenceResource*>(SubgraphRoot.get());
 							ImageID = Typed->ID;
 						}
-						else if(SubgraphType == EOpType::IM_CONSTANT)
+						else if (SubgraphType == EOpType::IM_CONSTANT)
 						{
 							const ASTOpConstantResource* Typed = static_cast<const ASTOpConstantResource*>(SubgraphRoot.get());
-							const FImage* Value = static_cast<const FImage*>( Typed->GetValue().Get() );
+							const FImage* Value = static_cast<const FImage*>(Typed->GetValue().Get());
 							bIsReference = Value->IsReference();
 							if (bIsReference)
 							{
@@ -957,7 +951,7 @@ namespace mu
 
 					ConstantSubgraphs[Index].Root = nullptr;
 					SubgraphCompletionEvent.Trigger();
-						
+
 					UE::Tasks::AddNested(SubgraphCompletionEvent);
 				}
 
@@ -969,6 +963,7 @@ namespace mu
 			LaunchTask.Wait();
 		}
 
+
 		bool bSomethingModified = ConstantSubgraphs.Num() > 0;
 		return bSomethingModified;
 	}
@@ -1048,7 +1043,7 @@ namespace mu
 			UE_LOG(LogMutableCore, Verbose, TEXT(" - constant generator"));
 
 			// Constant subtree generation
-			bModified = ConstantGeneratorAST( Options->GetPrivate(), Root, Pass );
+			bModified = ConstantGenerator( Options->GetPrivate(), Root, Pass );
 		}
 
 		//UE_LOG(LogMutableCore, Verbose, TEXT("(int) %s : %ld"), TEXT("ast size"), int64(ASTOp::CountNodes(roots)));
@@ -1296,9 +1291,9 @@ namespace mu
 
 
 	//---------------------------------------------------------------------------------------------
-	void CodeOptimiser::OptimiseAST()
+	void CodeOptimiser::Optimise()
 	{
-		MUTABLE_CPUPROFILER_SCOPE(OptimiseAST);
+		MUTABLE_CPUPROFILER_SCOPE(Optimise);
 
 		// Gather all the roots (one for each state)
 		TArray<Ptr<ASTOp>> roots;
@@ -1388,7 +1383,7 @@ namespace mu
 				for (Ptr<ASTOp>& Root : roots)
 				{
 					// Constant subtree generation
-					bModified = ConstantGeneratorAST(Options->GetPrivate(), Root, Pass);
+					bModified = ConstantGenerator(Options->GetPrivate(), Root, Pass);
 				}
 
 				DuplicatedDataRemoverAST(roots);
@@ -1465,7 +1460,7 @@ namespace mu
 				constexpr int32 Pass = 1;
 
 				UE_LOG(LogMutableCore, Verbose, TEXT(" - constant generator"));
-				ConstantGeneratorAST( Options->GetPrivate(), roots[StateIndex], Pass);
+				ConstantGenerator( Options->GetPrivate(), roots[StateIndex], Pass);
 				//AXE_INT_VALUE("Mutable", Verbose, "ast size", (int64_t)ASTOp::CountNodes(roots));
 			}
 
diff --git a/Engine/Plugins/Mutable/Source/MutableTools/Private/MuT/CodeOptimiser.h b/Engine/Plugins/Mutable/Source/MutableTools/Private/MuT/CodeOptimiser.h
--- a/Engine/Plugins/Mutable/Source/MutableTools/Private/MuT/CodeOptimiser.h
+++ b/Engine/Plugins/Mutable/Source/MutableTools/Private/MuT/CodeOptimiser.h
@@ -16,7 +16,7 @@ namespace mu
     public:
         CodeOptimiser( Ptr<CompilerOptions>, TArray<FStateCompilationData>& States );
 
-        void OptimiseAST();
+        void Optimise();
 
     private:
 
@@ -79,7 +79,7 @@ namespace mu
 
 
     /** ConstantGenerator replaces constant subtrees of operations with an equivalent single */
-    extern bool ConstantGeneratorAST( const CompilerOptions::Private*, Ptr<ASTOp>& Root, int32 Pass );
+    extern bool ConstantGenerator( const CompilerOptions::Private*, Ptr<ASTOp>& Root, int32 Pass );
 
     /** \TODO: shapes, projectors, others? but not switches (they must be unique) */
     extern void DuplicatedDataRemoverAST( ASTOpList& Roots );
diff --git a/Engine/Plugins/Mutable/Source/MutableTools/Private/MuT/CodeOptimiser_Parameters.cpp b/Engine/Plugins/Mutable/Source/MutableTools/Private/MuT/CodeOptimiser_Parameters.cpp
--- a/Engine/Plugins/Mutable/Source/MutableTools/Private/MuT/CodeOptimiser_Parameters.cpp
+++ b/Engine/Plugins/Mutable/Source/MutableTools/Private/MuT/CodeOptimiser_Parameters.cpp
@@ -2154,7 +2154,7 @@ namespace mu
             for(Ptr<ASTOp>& Root : roots)
             {
                 UE_LOG(LogMutableCore, Verbose, TEXT(" - constant optimiser"));
-				modified = ConstantGeneratorAST( Options->GetPrivate(), Root, Pass );
+				modified = ConstantGenerator( Options->GetPrivate(), Root, Pass );
 			}
 
 			//UE_LOG(LogMutableCore, Verbose, TEXT("(int) %s : %ld"), TEXT("ast size"), int64(ASTOp::CountNodes(roots)));
diff --git a/Engine/Plugins/Mutable/Source/MutableTools/Private/MuT/Compiler.cpp b/Engine/Plugins/Mutable/Source/MutableTools/Private/MuT/Compiler.cpp
--- a/Engine/Plugins/Mutable/Source/MutableTools/Private/MuT/Compiler.cpp
+++ b/Engine/Plugins/Mutable/Source/MutableTools/Private/MuT/Compiler.cpp
@@ -547,7 +547,7 @@ namespace mu
         // Optimize the generated code
         {
             CodeOptimiser Optimiser( m_pD->Options, States );
-            Optimiser.OptimiseAST( );
+            Optimiser.Optimise();
         }
 
         // Link the Program and generate state data.
diff --git a/Engine/Plugins/Mutable/Source/MutableTools/Public/MuT/NodeProjector.h b/Engine/Plugins/Mutable/Source/MutableTools/Public/MuT/NodeProjector.h
--- a/Engine/Plugins/Mutable/Source/MutableTools/Public/MuT/NodeProjector.h
+++ b/Engine/Plugins/Mutable/Source/MutableTools/Public/MuT/NodeProjector.h
@@ -40,10 +40,10 @@ namespace mu
 	public:
 
 		EProjectorType Type = EProjectorType::Planar;
-		FVector3f Position;
-		FVector3f Direction;
-		FVector3f Up;
-		FVector3f Scale;
+		FVector3f Position = FVector3f::ZeroVector;
+		FVector3f Direction = FVector3f::ZeroVector;
+		FVector3f Up = FVector3f::ZeroVector;
+		FVector3f Scale = FVector3f::ZeroVector;
 		float ProjectionAngle = 0.0f;
 
 	public:
@@ -83,10 +83,10 @@ namespace mu
 	public:
 
 		EProjectorType Type = EProjectorType::Planar;
-		FVector3f Position;
-		FVector3f Direction;
-		FVector3f Up;
-		FVector3f Scale;
+		FVector3f Position = FVector3f::ZeroVector;
+		FVector3f Direction = FVector3f::ZeroVector;
+		FVector3f Up = FVector3f::ZeroVector;
+		FVector3f Scale = FVector3f::ZeroVector;
 		float ProjectionAngle = 0.0f;
 
 		FString Name;
