[Mutable] Fix non-deterministic output when cooking multiple platforms at once.

#rnx

diff --git a/Engine/Plugins/Mutable/Source/CustomizableObject/Internal/MuCO/CustomizableObjectPrivate.h b/Engine/Plugins/Mutable/Source/CustomizableObject/Internal/MuCO/CustomizableObjectPrivate.h
--- a/Engine/Plugins/Mutable/Source/CustomizableObject/Internal/MuCO/CustomizableObjectPrivate.h
+++ b/Engine/Plugins/Mutable/Source/CustomizableObject/Internal/MuCO/CustomizableObjectPrivate.h
@@ -1341,7 +1341,7 @@ public:
 	UPROPERTY()
 	bool bCompiledWithHDTextureCompression = false;
 	
-	void InitCookData();
+	void InitCookData(UObject& InCustomizableObject);
 #endif
 };
 
@@ -1714,6 +1714,10 @@ public:
 	void LoadModelStreamableBulk(FArchive& Ar, bool bIsCooking);
 	void LoadModel(FArchive& Ar);
 
+	TObjectPtr<UModelResources> LoadModelResources_Internal(FArchive& Ar, const ITargetPlatform* InTargetPlatform, bool bSkipEditorOnlyData = false) const;
+	const TSharedPtr<FModelStreamableBulkData> LoadModelStreamableBulk_Internal(FArchive& Ar) const;
+	const TSharedPtr<mu::FModel, ESPMode::ThreadSafe> LoadModel_Internal(FArchive& Ar) const;
+
 	/** Load compiled data for the running platform from disk, this is used to load Editor Compilations. */
 	void LoadCompiledDataFromDisk();
 	
diff --git a/Engine/Plugins/Mutable/Source/CustomizableObject/Internal/MuCO/CustomizableObjectStreamedResourceData.h b/Engine/Plugins/Mutable/Source/CustomizableObject/Internal/MuCO/CustomizableObjectStreamedResourceData.h
--- a/Engine/Plugins/Mutable/Source/CustomizableObject/Internal/MuCO/CustomizableObjectStreamedResourceData.h
+++ b/Engine/Plugins/Mutable/Source/CustomizableObject/Internal/MuCO/CustomizableObjectStreamedResourceData.h
@@ -34,7 +34,7 @@ public:
 	 * Note that the container may not have been moved to its new path yet, so don't try to load
 	 * the container from the given path here.
 	 */
-	void ConvertToSoftReferenceForCooking(const TSoftObjectPtr<UCustomizableObjectResourceDataContainer>& NewContainerPath);
+	void ConvertToSoftReferenceForCooking();
 #endif // WITH_EDITOR
 
 	/**
diff --git a/Engine/Plugins/Mutable/Source/CustomizableObject/Private/MuCO/CustomizableObject.cpp b/Engine/Plugins/Mutable/Source/CustomizableObject/Private/MuCO/CustomizableObject.cpp
--- a/Engine/Plugins/Mutable/Source/CustomizableObject/Private/MuCO/CustomizableObject.cpp
+++ b/Engine/Plugins/Mutable/Source/CustomizableObject/Private/MuCO/CustomizableObject.cpp
@@ -193,6 +193,11 @@ void UCustomizableObject::PreSave(FObjectPreSaveContext ObjectSaveContext)
 		// Load cached data before saving
 		if (GetPrivate()->TryLoadCompiledCookDataForPlatform(TargetPlatform))
 		{
+			// Change current platform ModelResources' outer from TransientPkg to CO before save. 
+			// PostSaveRoot will set the outer to TransientPkg again. This is done to avoid serializing multiple ModelResources when cooking more than one platform at once.
+			TObjectPtr<UModelResources> ModelResources = GetPrivate()->GetModelResources(true);
+			ModelResources->Rename(nullptr, this, REN_DontCreateRedirectors);
+
 			const bool bUseBulkData = CVarMutableUseBulkData.GetValueOnAnyThread();
 			if (bUseBulkData)
 			{
@@ -252,6 +257,24 @@ void UCustomizableObject::PreSave(FObjectPreSaveContext ObjectSaveContext)
 }
 
 
+void UCustomizableObject::PostSaveRoot(FObjectPostSaveRootContext ObjectSaveContext)
+{
+	Super::PostSaveRoot(ObjectSaveContext);
+
+	if (ObjectSaveContext.IsCooking() && !bIsChildObject)
+	{
+		const ITargetPlatform* TargetPlatform = ObjectSaveContext.GetTargetPlatform();
+
+		const MutablePrivate::FMutableCachedPlatformData* PlatformData = GetPrivate()->CachedPlatformsData.Find(TargetPlatform->PlatformName());
+		if (PlatformData && PlatformData->ModelResources)
+		{
+			// Set the outer to TransientPkg again. This is done to avoid serializing multiple ModelResources when cooking more than one platform at once.
+			PlatformData->ModelResources->Rename(nullptr, GetTransientPackage(), REN_DontCreateRedirectors);
+		}
+	}
+}
+
+
 bool UCustomizableObjectPrivate::TryUpdateIsChildObject()
 {
 	if (const ICustomizableObjectEditorModule* Module = ICustomizableObjectEditorModule::Get())
@@ -486,53 +509,64 @@ FString GetModelResourcesNameForPlatform(const UCustomizableObject& Customizable
 
 
 bool UCustomizableObjectPrivate::LoadModelResources(FArchive& MemoryReader, const ITargetPlatform* InTargetPlatform, bool bIsCooking)
+{
+	TObjectPtr<UModelResources> LocalModelResources = LoadModelResources_Internal(MemoryReader, InTargetPlatform, bIsCooking);
+	SetModelResources(LocalModelResources, bIsCooking);
+
+	return LocalModelResources != nullptr;
+}
+
+
+void UCustomizableObjectPrivate::LoadModelStreamableBulk(FArchive& MemoryReader, bool bIsCooking)
+{
+	SetModelStreamableBulkData(LoadModelStreamableBulk_Internal(MemoryReader), bIsCooking);
+}
+
+
+void UCustomizableObjectPrivate::LoadModel(FArchive& MemoryReader)
+{
+	SetModel(LoadModel_Internal(MemoryReader), GenerateIdentifier(*GetPublic()));
+}
+
+
+TObjectPtr<UModelResources> UCustomizableObjectPrivate::LoadModelResources_Internal(FArchive& MemoryReader, const ITargetPlatform* InTargetPlatform, bool bIsCooking) const
 {
 	// Make sure mutable has been initialised.
 	UCustomizableObjectSystem::GetInstance();
-	
+
 	const FString ModelResourcesName = GetModelResourcesNameForPlatform(*GetPublic(), *InTargetPlatform);
-	UObject* Outer = bIsCooking ? GetPublic() : (UObject*)GetTransientPackage();
-	TObjectPtr<UModelResources> LocalModelResources = NewObject<UModelResources>(Outer, FName(*ModelResourcesName), RF_Public);
-	
+	TObjectPtr<UModelResources> LocalModelResources = NewObject<UModelResources>((UObject*)GetTransientPackage(), FName(*ModelResourcesName), RF_Public);
+
 	FObjectAndNameAsStringProxyArchive ObjectReader(MemoryReader, true);
 	LocalModelResources->Serialize(ObjectReader);
 
 	const bool bLoadedSuccessfully = LocalModelResources->CodeVersion == UCustomizableObjectPrivate::CurrentSupportedVersion;
-	SetModelResources(bLoadedSuccessfully ? LocalModelResources : nullptr, bIsCooking);
-
-	return bLoadedSuccessfully;
+	return bLoadedSuccessfully ? LocalModelResources : nullptr;
 }
 
 
-void UCustomizableObjectPrivate::LoadModelStreamableBulk(FArchive& MemoryReader, bool bIsCooking)
+const TSharedPtr<FModelStreamableBulkData> UCustomizableObjectPrivate::LoadModelStreamableBulk_Internal(FArchive& MemoryReader) const
 {
 	TSharedPtr<FModelStreamableBulkData> LocalModelStreamablesPtr = MakeShared<FModelStreamableBulkData>();
 	FModelStreamableBulkData& LocalModelStreamables = *LocalModelStreamablesPtr.Get();
 	MemoryReader << LocalModelStreamables;
 
-	SetModelStreamableBulkData(LocalModelStreamablesPtr, bIsCooking);
+	return LocalModelStreamablesPtr;
 }
 
 
-void UCustomizableObjectPrivate::LoadModel(FArchive& MemoryReader)
+const TSharedPtr<mu::FModel, ESPMode::ThreadSafe> UCustomizableObjectPrivate::LoadModel_Internal(FArchive& MemoryReader) const
 {
-	TSharedPtr<mu::FModel, ESPMode::ThreadSafe> LoadedModel;
-
 	FUnrealMutableInputStream Stream(MemoryReader);
 	mu::FInputArchive Arch(&Stream);
-	LoadedModel = mu::FModel::StaticUnserialise(Arch);
-
-	SetModel(LoadedModel, GenerateIdentifier(*GetPublic()));
+	return mu::FModel::StaticUnserialise(Arch);
 }
 
 
-void UModelResources::InitCookData()
+void UModelResources::InitCookData(UObject& InCustomizableObject)
 {
-	UCustomizableObject* CustomizableObject = Cast<UCustomizableObject>(GetOuter());
-	check(CustomizableObject);
+	const FString ObjectName = InCustomizableObject.GetName();
 
-	const FString ObjectName = GetNameSafe(CustomizableObject);
-	
 	const int32 NumStreamedResources = StreamedResourceDataEditor.Num();
 	for (int32 Index = 0; Index < NumStreamedResources; ++Index)
 	{
@@ -540,32 +574,32 @@ void UModelResources::InitCookData()
 
 		const FString& ContainerName = ObjectName + FString::Printf(TEXT("_SR_%d"), Index);
 
-		UCustomizableObjectResourceDataContainer* Container = FindObject<UCustomizableObjectResourceDataContainer>(CustomizableObject, *ContainerName);
+		UCustomizableObjectResourceDataContainer* Container = FindObject<UCustomizableObjectResourceDataContainer>(&InCustomizableObject, *ContainerName);
 		if (!Container)
 		{
 			Container = NewObject<UCustomizableObjectResourceDataContainer>(
-				CustomizableObject,
+				&InCustomizableObject,
 				FName(*ContainerName),
 				RF_Public);
-		}
 
-		Container->Data = ResourceData;
-
-		if (FCustomizableObjectAssetUserData* AUDResource = Container->Data.Data.GetMutablePtr<FCustomizableObjectAssetUserData>())
-		{
-			// Find or duplicate the AUD replacing the outer
-			UAssetUserData* SourceAssetUserData = MutablePrivate::LoadObject(AUDResource->AssetUserDataEditor);
-			const FString AssetName = GetNameSafe(SourceAssetUserData);
-			check(SourceAssetUserData != nullptr);
+			Container->Data = ResourceData;
 
-			UAssetUserData* AssetUserData = FindObject<UAssetUserData>(Container, *AssetName);
-			if (!AssetUserData)
+			if (FCustomizableObjectAssetUserData* AUDResource = Container->Data.Data.GetMutablePtr<FCustomizableObjectAssetUserData>())
 			{
-				// AUD may be private objects within meshes. Duplicate changing the outer to avoid including meshes into the builds.
-				AssetUserData = DuplicateObject<UAssetUserData>(SourceAssetUserData, Container, FName(*AssetName));
+				// Find or duplicate the AUD replacing the outer
+				UAssetUserData* SourceAssetUserData = MutablePrivate::LoadObject(AUDResource->AssetUserDataEditor);
+				const FString AssetName = GetNameSafe(SourceAssetUserData);
+				check(SourceAssetUserData != nullptr);
+
+				UAssetUserData* AssetUserData = FindObject<UAssetUserData>(Container, *AssetName);
+				if (!AssetUserData)
+				{
+					// AUD may be private objects within meshes. Duplicate changing the outer to avoid including meshes into the builds.
+					AssetUserData = DuplicateObject<UAssetUserData>(SourceAssetUserData, Container, FName(*AssetName));
+				}
+
+				AUDResource->AssetUserData = AssetUserData;
 			}
-			
-			AUDResource->AssetUserData = AssetUserData;
 		}
 
 		StreamedResourceData.Emplace(Container);
@@ -579,21 +613,22 @@ void UModelResources::InitCookData()
 		const FString& ContainerName = ObjectName + FString::Printf(TEXT("_SE_%d"), Index);
 		FCustomizableObjectResourceData& ResourceData = StreamedExtensionDataEditor[Index];
 
-		UCustomizableObjectResourceDataContainer* Container = FindObject<UCustomizableObjectResourceDataContainer>(CustomizableObject, *ContainerName);
+		UCustomizableObjectResourceDataContainer* Container = FindObject<UCustomizableObjectResourceDataContainer>(&InCustomizableObject, *ContainerName);
 		if (!Container)
 		{
 			Container = NewObject<UCustomizableObjectResourceDataContainer>(
-				CustomizableObject,
+				&InCustomizableObject,
 				FName(*ContainerName),
 				RF_Public);
-		}
 
-		for (const UCustomizableObjectExtension* Extension : AllExtensions)
-		{
-			Extension->MovePrivateReferencesToContainer(ResourceData.Data, Container);
+			for (const UCustomizableObjectExtension* Extension : AllExtensions)
+			{
+				Extension->MovePrivateReferencesToContainer(ResourceData.Data, Container);
+			}
+
+			Container->Data = MoveTemp(ResourceData);
 		}
 
-		Container->Data = MoveTemp(ResourceData);
 		StreamedExtensionData.Emplace(Container);
 	}
 }
@@ -3160,30 +3195,34 @@ void UCustomizableObjectPrivate::LoadCompiledDataFromDDC(FCompilationOptions Opt
 
 		// Load the compiled data to validate it.
 		FMemoryReaderView ModelResourcesReader(ModelResourcesBytesDDC.GetView());
-		if (LoadModelResources(ModelResourcesReader, Options.TargetPlatform, Options.bIsCooking))
+
+		if (TObjectPtr<UModelResources> LocalModelResources = LoadModelResources_Internal(ModelResourcesReader, Options.TargetPlatform, Options.bIsCooking))
 		{
-			CachedData.ModelResources = TStrongObjectPtr<UModelResources>(GetModelResources(Options.bIsCooking));
+			CachedData.ModelResources = TStrongObjectPtr<UModelResources>(LocalModelResources);
 			CachedData.ModelResources->bIsStoredInDDC = true;
 			CachedData.ModelResources->DDCKey = CacheKey;
 			CachedData.ModelResources->DDCDefaultPolicy = ECachePolicy::Default;
 
 			FMemoryReaderView ModelStreamablesReader(ModelStreamablesBytesDDC.GetView());
-			LoadModelStreamableBulk(ModelStreamablesReader, Options.bIsCooking);
-			CachedData.ModelStreamableBulkData = GetModelStreamableBulkData(Options.bIsCooking);
+			CachedData.ModelStreamableBulkData = LoadModelStreamableBulk_Internal(ModelStreamablesReader);
 
 			FMemoryReaderView ModelReader(ModelBytesDDC.GetView());
-			LoadModel(ModelReader);
-			CachedData.Model = GetModel();
+			CachedData.Model = LoadModel_Internal(ModelReader);
+		}
+
+		if (!CachedData.Model || !CachedData.ModelResources || !CachedData.ModelStreamableBulkData)
+		{
+			return;
 		}
 
 		// Cache CookedPlatfomData. 
-		if (Options.bIsCooking && CachedData.Model && CachedData.ModelResources && CachedData.ModelStreamableBulkData)
+		if (Options.bIsCooking)
 		{
 			// Sync cache cooked platform data
 			// TODO UE-220138: Sync -> Async
 			MUTABLE_CPUPROFILER_SCOPE(CacheCookedPlatformData);
 
-			CachedData.ModelResources->InitCookData();
+			CachedData.ModelResources->InitCookData(*GetPublic());
 
 			// Value Id to file mapping to reconstruct the cached data
 			TMap<FValueId, MutablePrivate::FFile> ValueIdToFile;
@@ -3300,6 +3339,12 @@ void UCustomizableObjectPrivate::LoadCompiledDataFromDDC(FCompilationOptions Opt
 			MutablePrivate::FMutableCachedPlatformData& CachedPlatformData = CachedPlatformsData.Add(Options.TargetPlatform->PlatformName(), {});
 			CachedPlatformData = MoveTemp(CachedData);
 		}
+		else // Editor DDC
+		{
+			SetModel(CachedData.Model, GenerateIdentifier(*GetPublic()));
+			SetModelResources(CachedData.ModelResources.Get(), Options.bIsCooking);
+			SetModelStreamableBulkData(CachedData.ModelStreamableBulkData, Options.bIsCooking);
+		}
 	}
 
 	return;
diff --git a/Engine/Plugins/Mutable/Source/CustomizableObject/Private/MuCO/CustomizableObjectStreamedResourceData.cpp b/Engine/Plugins/Mutable/Source/CustomizableObject/Private/MuCO/CustomizableObjectStreamedResourceData.cpp
--- a/Engine/Plugins/Mutable/Source/CustomizableObject/Private/MuCO/CustomizableObjectStreamedResourceData.cpp
+++ b/Engine/Plugins/Mutable/Source/CustomizableObject/Private/MuCO/CustomizableObjectStreamedResourceData.cpp
@@ -19,15 +19,16 @@ FCustomizableObjectStreamedResourceData::FCustomizableObjectStreamedResourceData
 	ContainerPath = InContainer;
 }
 
-void FCustomizableObjectStreamedResourceData::ConvertToSoftReferenceForCooking(
-	const TSoftObjectPtr<UCustomizableObjectResourceDataContainer>& NewContainerPath)
+void FCustomizableObjectStreamedResourceData::ConvertToSoftReferenceForCooking()
 {
 	check(IsInGameThread());
+	check(Container);
 
+	// Update ContainerPath;
+	ContainerPath = TSoftObjectPtr<UCustomizableObjectResourceDataContainer>(const_cast<UCustomizableObjectResourceDataContainer*>(Container.Get()));
+	
 	// Remove the hard reference to the container, so that it can be unloaded
 	Container = nullptr;
-
-	ContainerPath = NewContainerPath;
 }
 #endif // WITH_EDITOR
 
diff --git a/Engine/Plugins/Mutable/Source/CustomizableObject/Public/MuCO/CustomizableObject.h b/Engine/Plugins/Mutable/Source/CustomizableObject/Public/MuCO/CustomizableObject.h
--- a/Engine/Plugins/Mutable/Source/CustomizableObject/Public/MuCO/CustomizableObject.h
+++ b/Engine/Plugins/Mutable/Source/CustomizableObject/Public/MuCO/CustomizableObject.h
@@ -495,6 +495,7 @@ public:
 	virtual FString GetDesc() override;
 	virtual bool IsEditorOnly() const override;
 	virtual void PreSave(FObjectPreSaveContext ObjectSaveContext) override;
+	virtual void PostSaveRoot(FObjectPostSaveRootContext ObjectSaveContext) override;
 	virtual void PostRename(UObject* OldOuter, const FName OldName) override;	
 	virtual void BeginCacheForCookedPlatformData(const ITargetPlatform* TargetPlatform) override;
 	virtual bool IsCachedCookedPlatformDataLoaded(const ITargetPlatform* TargetPlatform) override;
diff --git a/Engine/Plugins/Mutable/Source/CustomizableObjectEditor/Private/MuCOE/CustomizableObjectCompiler.cpp b/Engine/Plugins/Mutable/Source/CustomizableObjectEditor/Private/MuCOE/CustomizableObjectCompiler.cpp
--- a/Engine/Plugins/Mutable/Source/CustomizableObjectEditor/Private/MuCOE/CustomizableObjectCompiler.cpp
+++ b/Engine/Plugins/Mutable/Source/CustomizableObjectEditor/Private/MuCOE/CustomizableObjectCompiler.cpp
@@ -1227,9 +1227,11 @@ bool FCustomizableObjectCompiler::TryLoadCompiledDataFromDDC(UCustomizableObject
 	FCacheKey CacheKey = CurrentRequest->GetDerivedDataCacheKey();
 	CustomizableObject.GetPrivate()->LoadCompiledDataFromDDC(CurrentOptions, DefaultPolicy, &CacheKey);
 	
-	PlatformData->Model = CustomizableObject.GetPrivate()->GetModel();
-	PlatformData->ModelResources = TStrongObjectPtr<UModelResources>(CustomizableObject.GetPrivate()->GetModelResources(CurrentOptions.bIsCooking));
-	PlatformData->ModelStreamableBulkData = CustomizableObject.GetPrivate()->GetModelStreamableBulkData(CurrentOptions.bIsCooking);
+	if (CurrentOptions.bIsCooking)
+	{
+		return CustomizableObject.GetPrivate()->CachedPlatformsData.Find(CurrentOptions.TargetPlatform->PlatformName()) != nullptr;
+	}
+
 	return CustomizableObject.IsCompiled();
 }
 
@@ -1501,7 +1503,7 @@ void FCustomizableObjectCompiler::FinishSavingDerivedDataTask()
 		MutablePrivate::FMutableCachedPlatformData& Data = CurrentObject->GetPrivate()->CachedPlatformsData.Add(PlatformName);
 		Data = MoveTemp(*PlatformData.Get());
 
-		Data.ModelResources->InitCookData();
+		Data.ModelResources->InitCookData(*CurrentObject);
 	}
 
 	// Order matters
diff --git a/Engine/Plugins/Mutable/Source/CustomizableObjectEditor/Private/MuCOE/CustomizableObjectCookPackageSplitter.cpp b/Engine/Plugins/Mutable/Source/CustomizableObjectEditor/Private/MuCOE/CustomizableObjectCookPackageSplitter.cpp
--- a/Engine/Plugins/Mutable/Source/CustomizableObjectEditor/Private/MuCOE/CustomizableObjectCookPackageSplitter.cpp
+++ b/Engine/Plugins/Mutable/Source/CustomizableObjectEditor/Private/MuCOE/CustomizableObjectCookPackageSplitter.cpp
@@ -208,8 +208,6 @@ bool FCustomizableObjectCookPackageSplitter::PreSaveGeneratorPackage(FPopulateCo
 	// 
 	// 2.	Remove hard references to the streamed data, so that it doesn't get loaded as soon as
 	//		the CO is loaded
-	//
-	// 3.	Fix up the soft references to point to the streamed data's new location
 	TConstArrayView<ICookPackageSplitter::FGeneratedPackageForPopulate>& PlaceholderPackages = PopulateContext.GetGeneratedPackages();
 	const auto& PreSavePackage = [] (const ICookPackageSplitter::FGeneratedPackageForPopulate& GeneratedPackage,
 		TArray<FCustomizableObjectStreamedResourceData>& StreamedResources
@@ -233,9 +231,6 @@ bool FCustomizableObjectCookPackageSplitter::PreSaveGeneratorPackage(FPopulateCo
 			return false;
 		}
 
-		// Remove the hard reference and set the soft reference to the streamed data's new location
-		FoundData->ConvertToSoftReferenceForCooking(TSoftObjectPtr<UCustomizableObjectResourceDataContainer>(Container));
-
 		return true;
 	};
 
@@ -283,6 +278,26 @@ bool FCustomizableObjectCookPackageSplitter::PreSaveGeneratorPackage(FPopulateCo
 		SavedExtensionContainerNames.Add(GeneratedPackage.RelativePath);
 	}
 
+	// All platforms should have the same resources
+	const TMap<FString, MutablePrivate::FMutableCachedPlatformData>& CachePlatforms = Object->GetPrivate()->CachedPlatformsData;
+	for (const TPair<FString, MutablePrivate::FMutableCachedPlatformData>& PlatformData : CachePlatforms)
+	{
+		if (PlatformData.Value.ModelResources)
+		{
+			for (FCustomizableObjectStreamedResourceData& StreamedResourceData : PlatformData.Value.ModelResources->StreamedResourceData)
+			{
+				// Remove the hard reference and set the soft reference to the streamed data's new location
+				StreamedResourceData.ConvertToSoftReferenceForCooking();
+			}
+
+			for (FCustomizableObjectStreamedResourceData& StreamedExtensionData : PlatformData.Value.ModelResources->StreamedExtensionData)
+			{
+				// Remove the hard reference and set the soft reference to the streamed data's new location
+				StreamedExtensionData.ConvertToSoftReferenceForCooking();
+			}
+		}
+	}
+
 	return true;
 }
 
