[Mutable] Allow different layout grid sizes on the X/Y axes.

#rnx

diff --git a/Engine/Plugins/Mutable/Source/CustomizableObjectEditor/Private/MuCOE/SCustomizableObjectLayoutEditor.cpp b/Engine/Plugins/Mutable/Source/CustomizableObjectEditor/Private/MuCOE/SCustomizableObjectLayoutEditor.cpp
--- a/Engine/Plugins/Mutable/Source/CustomizableObjectEditor/Private/MuCOE/SCustomizableObjectLayoutEditor.cpp
+++ b/Engine/Plugins/Mutable/Source/CustomizableObjectEditor/Private/MuCOE/SCustomizableObjectLayoutEditor.cpp
@@ -258,11 +258,27 @@ void FCustomizableObjectLayoutEditorDetailsBuilder::CustomizeDetails(IDetailLayo
 			]
 			.ValueContent()
 			[
-				SAssignNew(LayoutEditor->GridSizeComboBox, STextComboBox)
-					.InitiallySelectedItem(LayoutEditor->LayoutGridSizes[0])
-					.OptionsSource(&LayoutEditor->LayoutGridSizes)
-					.OnSelectionChanged(LayoutEditor.Get(), &SCustomizableObjectLayoutEditor::OnGridSizeChanged)
-					.Font(IDetailLayoutBuilder::GetDetailFont())
+				SNew(SHorizontalBox)
+					+ SHorizontalBox::Slot()
+					.FillWidth(0.5f)
+					[
+						SAssignNew(LayoutEditor->GridSizeXComboBox, STextComboBox)
+							.InitiallySelectedItem(LayoutEditor->LayoutGridSizes[0])
+							.OptionsSource(&LayoutEditor->LayoutGridSizes)
+							.OnComboBoxOpening(LayoutEditor.Get(), &SCustomizableObjectLayoutEditor::OnOpenGridSizeComboBox)
+							.OnSelectionChanged(LayoutEditor.Get(), &SCustomizableObjectLayoutEditor::OnGridSizeChanged, true)
+							.Font(IDetailLayoutBuilder::GetDetailFont())
+					]
+					+ SHorizontalBox::Slot()
+					.FillWidth(0.5f)
+					[
+						SAssignNew(LayoutEditor->GridSizeYComboBox, STextComboBox)
+							.InitiallySelectedItem(LayoutEditor->LayoutGridSizes[0])
+							.OptionsSource(&LayoutEditor->LayoutGridSizes)
+							.OnComboBoxOpening(LayoutEditor.Get(), &SCustomizableObjectLayoutEditor::OnOpenGridSizeComboBox)
+							.OnSelectionChanged(LayoutEditor.Get(), &SCustomizableObjectLayoutEditor::OnGridSizeChanged, false)
+							.Font(IDetailLayoutBuilder::GetDetailFont())
+					]
 			];
 	}
 
@@ -279,8 +295,8 @@ void FCustomizableObjectLayoutEditorDetailsBuilder::CustomizeDetails(IDetailLayo
 			.ValueContent()
 			[
 				SAssignNew(LayoutEditor->MaxGridSizeComboBox, STextComboBox)
-					.InitiallySelectedItem(LayoutEditor->LayoutGridSizes[0])
-					.OptionsSource(&LayoutEditor->LayoutGridSizes)
+					.InitiallySelectedItem(LayoutEditor->MaxLayoutGridSizes[0])
+					.OptionsSource(&LayoutEditor->MaxLayoutGridSizes)
 					.OnSelectionChanged(LayoutEditor.Get(), &SCustomizableObjectLayoutEditor::OnMaxGridSizeChanged)
 					.Font(IDetailLayoutBuilder::GetDetailFont())
 			];
@@ -494,10 +510,12 @@ void SCustomizableObjectLayoutEditor::Construct(const FArguments& InArgs)
 	// Array of available grid size options
 	const int32 MaxGridSize = 128;
 	LayoutGridSizes.Empty();
+	MaxLayoutGridSizes.Empty();
 
 	for (int32 Size = 1; Size <= MaxGridSize; Size *= 2)
 	{
-		LayoutGridSizes.Add(MakeShareable(new FString(FString::Printf(TEXT("%d x %d"), Size, Size))));
+		LayoutGridSizes.Add(MakeShareable(new FString(FString::FromInt(Size))));
+		MaxLayoutGridSizes.Add(MakeShareable(new FString(FString::Printf(TEXT("%d x %d"), Size, Size))));
 	}
 
 	{
@@ -750,7 +768,32 @@ void SCustomizableObjectLayoutEditor::OnUVChannelChanged(TSharedPtr<FString> UVC
 }
 
 
-void SCustomizableObjectLayoutEditor::OnGridSizeChanged(TSharedPtr<FString> NewSelection, ESelectInfo::Type SelectInfo)
+void SCustomizableObjectLayoutEditor::OnOpenGridSizeComboBox()
+{
+	if (!CurrentLayout)
+	{
+		return;
+	}
+
+	const int32 MaxGridSize = FixedStrategyOptionsVisibility() == EVisibility::Visible ? CurrentLayout->GetMaxGridSize().X : 128;
+	const int32 NumOptions = FMath::CountTrailingZeros(MaxGridSize) + 1;
+
+	LayoutGridSizes.SetNum(NumOptions);
+
+	int32 AuxIndex = 0;
+	for (TSharedPtr<FString>& GridSizeString : LayoutGridSizes)
+	{
+		if (!GridSizeString)
+		{
+			GridSizeString = MakeShareable(new FString(FString::FromInt(1 << AuxIndex)));
+		}
+
+		AuxIndex++;
+	}
+}
+
+
+void SCustomizableObjectLayoutEditor::OnGridSizeChanged(TSharedPtr<FString> NewSelection, ESelectInfo::Type SelectInfo, bool bIsGridSizeX)
 {
 	if (!Node.IsValid() || !CurrentLayout)
 	{
@@ -759,9 +802,13 @@ void SCustomizableObjectLayoutEditor::OnGridSizeChanged(TSharedPtr<FString> NewS
 
 	const int32 Size = 1 << LayoutGridSizes.Find(NewSelection);
 
-	if (CurrentLayout->GetGridSize().X != Size || CurrentLayout->GetGridSize().Y != Size)
+	FIntPoint GridSize = CurrentLayout->GetGridSize();
+	GridSize.X = bIsGridSizeX ? Size : GridSize.X;
+	GridSize.Y = !bIsGridSizeX ? Size : GridSize.Y;
+
+	if (GridSize != CurrentLayout->GetGridSize())
 	{
-		CurrentLayout->SetGridSize(FIntPoint(Size));
+		CurrentLayout->SetGridSize(GridSize);
 
 		Node->Modify();
 
@@ -779,12 +826,18 @@ void SCustomizableObjectLayoutEditor::OnMaxGridSizeChanged(TSharedPtr<FString> N
 		return;
 	}
 
-	const int32 Size = 1 << LayoutGridSizes.Find(NewSelection);
+	const int32 Size = 1 << MaxLayoutGridSizes.Find(NewSelection);
 
 	if (CurrentLayout->GetMaxGridSize().X != Size || CurrentLayout->GetMaxGridSize().Y != Size)
 	{
 		CurrentLayout->SetMaxGridSize(FIntPoint(Size));
 
+		// GridSize must be equal or smaller than MaxGridSize.
+		FIntPoint GridSize = CurrentLayout->GetGridSize();
+		GridSize.X = FMath::Clamp(GridSize.X, 1, Size);
+		GridSize.Y = FMath::Clamp(GridSize.Y, 1, Size);
+		CurrentLayout->SetGridSize(GridSize);
+		
 		Node->Modify();
 
 		// Reset to update the UI.
@@ -868,6 +921,18 @@ void SCustomizableObjectLayoutEditor::OnPackagingStrategyChanged(TSharedPtr<FStr
 	{
 		CurrentLayout->PackingStrategy = LayoutPackingStrategiesOptions[Index].Value;
 		
+		// Update max size when changing strategies
+		if (CurrentLayout->PackingStrategy == ECustomizableObjectTextureLayoutPackingStrategy::Fixed)
+		{
+			// GridSize must be equal or smaller than MaxGridSize.
+			const int32 MaxGridSize = CurrentLayout->GetMaxGridSize().X;
+			
+			FIntPoint GridSize = CurrentLayout->GetGridSize();
+			GridSize.X = FMath::Clamp(GridSize.X, 1, MaxGridSize);
+			GridSize.Y = FMath::Clamp(GridSize.Y, 1, MaxGridSize);
+			CurrentLayout->SetGridSize(GridSize);
+		}
+
 		Node->Modify();
 	
 		// Reset to update the UI.
@@ -1443,16 +1508,22 @@ void SCustomizableObjectLayoutEditor::FillLayoutComboBoxOptions()
 			AutoBlocksMergeComboBox->SetSelectedItem(AutoBlocksMergeStrategies[AutoBlockMergeIndex]);
 		}
 
-		if (GridSizeComboBox)
+		if (GridSizeXComboBox)
 		{
 			const int32 Index = FMath::CountTrailingZeros(CurrentLayout->GetGridSize().X);
-			GridSizeComboBox->SetSelectedItem(LayoutGridSizes[Index]);
+			GridSizeXComboBox->SetSelectedItem(LayoutGridSizes[Index]);
+		}
+
+		if (GridSizeYComboBox)
+		{
+			const int32 Index = FMath::CountTrailingZeros(CurrentLayout->GetGridSize().Y);
+			GridSizeYComboBox->SetSelectedItem(LayoutGridSizes[Index]);
 		}
 
 		if (MaxGridSizeComboBox)
 		{
 			const int32 Index = FMath::CountTrailingZeros(CurrentLayout->GetMaxGridSize().X);
-			MaxGridSizeComboBox->SetSelectedItem(LayoutGridSizes[Index]);
+			MaxGridSizeComboBox->SetSelectedItem(MaxLayoutGridSizes[Index]);
 		}
 
 		if (ReductionMethodComboBox)
diff --git a/Engine/Plugins/Mutable/Source/CustomizableObjectEditor/Private/MuCOE/SCustomizableObjectLayoutEditor.h b/Engine/Plugins/Mutable/Source/CustomizableObjectEditor/Private/MuCOE/SCustomizableObjectLayoutEditor.h
--- a/Engine/Plugins/Mutable/Source/CustomizableObjectEditor/Private/MuCOE/SCustomizableObjectLayoutEditor.h
+++ b/Engine/Plugins/Mutable/Source/CustomizableObjectEditor/Private/MuCOE/SCustomizableObjectLayoutEditor.h
@@ -111,6 +111,7 @@ private:
 	
 	/** List of available layout grid sizes. */
 	TArray<TSharedPtr<FString>> LayoutGridSizes;
+	TArray<TSharedPtr<FString>> MaxLayoutGridSizes;
 
 	TSharedPtr<FString> NothingSelectedString;
 
@@ -159,7 +160,8 @@ private:
 	TSharedPtr<SSearchableComboBox> AutoBlocksMergeComboBox;
 
 	// ComboBox widget to select a Grid Size from the Selected Layout
-	TSharedPtr<STextComboBox> GridSizeComboBox;
+	TSharedPtr<STextComboBox> GridSizeXComboBox;
+	TSharedPtr<STextComboBox> GridSizeYComboBox;
 
 	// ComboBox widget to select a Max Grid Size from the Selected Layout
 	TSharedPtr<STextComboBox> MaxGridSizeComboBox;
@@ -191,7 +193,8 @@ private:
 	/** */
 	void OnMeshSectionChanged(TSharedPtr<FString> NewSelection, ESelectInfo::Type SelectInfo);
 	void OnUVChannelChanged(TSharedPtr<FString> NewSelection, ESelectInfo::Type SelectInfo);
-	void OnGridSizeChanged(TSharedPtr<FString> NewSelection, ESelectInfo::Type SelectInfo);
+	void OnOpenGridSizeComboBox();
+	void OnGridSizeChanged(TSharedPtr<FString> NewSelection, ESelectInfo::Type SelectInfo, bool bIsGridSizeX);
 	void OnMaxGridSizeChanged(TSharedPtr<FString> NewSelection, ESelectInfo::Type SelectInfo);
 
 	void OnAutoBlocksChanged(TSharedPtr<FString> NewSelection, ESelectInfo::Type SelectInfo);
