[mutable] Add support for layouts creation at runtime, for mesh parameters.
- Support for multiple output sections in mesh parameter node. Fix mesh conversion for sections other than 0.
- Added support for color buffers to mesh parameters.
- Added support for 16-bit float tex coords.
#rnx

diff --git a/Engine/Plugins/Mutable/Source/CustomizableObject/Internal/MuCO/CustomizableObjectPrivate.h b/Engine/Plugins/Mutable/Source/CustomizableObject/Internal/MuCO/CustomizableObjectPrivate.h
--- a/Engine/Plugins/Mutable/Source/CustomizableObject/Internal/MuCO/CustomizableObjectPrivate.h
+++ b/Engine/Plugins/Mutable/Source/CustomizableObject/Internal/MuCO/CustomizableObjectPrivate.h
@@ -1591,6 +1591,10 @@ enum class ECustomizableObjectVersions : int32
 
 	MovedLODSettingsToMeshComponentNode,
 
+	AddedMeshPrepareLayoutOp,
+
+	AddedMeshIDToMeshParamOp,
+
 	// -----<new versions can be added above this line>--------
 	LastCustomizableObjectVersion
 };
diff --git a/Engine/Plugins/Mutable/Source/CustomizableObject/Internal/MuCO/MutableMeshBufferUtils.h b/Engine/Plugins/Mutable/Source/CustomizableObject/Internal/MuCO/MutableMeshBufferUtils.h
--- a/Engine/Plugins/Mutable/Source/CustomizableObject/Internal/MuCO/MutableMeshBufferUtils.h
+++ b/Engine/Plugins/Mutable/Source/CustomizableObject/Internal/MuCO/MutableMeshBufferUtils.h
@@ -43,7 +43,7 @@ namespace  MutableMeshBufferUtils
 	 * @param OutTargetVertexBuffers Mesh buffer to be set up. 
 	 * @param InTextureSemanticIndicesOverride (Optional) If you are certain of the texture semantic indices that should be used populate this argument, if not, leave it as null
 	 */
-	CUSTOMIZABLEOBJECT_API void SetupTexCoordinatesBuffer(const int32& InCurrentVertexBuffer,const int32& InChannelCount, mu::FMeshBufferSet& OutTargetVertexBuffers, const int32* InTextureSemanticIndicesOverride = nullptr);
+	CUSTOMIZABLEOBJECT_API void SetupTexCoordinatesBuffer(const int32& InCurrentVertexBuffer,const int32& InChannelCount, bool bHighPrecision, mu::FMeshBufferSet& OutTargetVertexBuffers, const int32* InTextureSemanticIndicesOverride = nullptr);
 
 	/**
 	 * Sets up all the channels that ought to be part of the Vertex Skin Buffer of the Mutable mesh owning the provided MeshBufferSet
diff --git a/Engine/Plugins/Mutable/Source/CustomizableObject/Private/MuCO/CustomizableObjectInstance.cpp b/Engine/Plugins/Mutable/Source/CustomizableObject/Private/MuCO/CustomizableObjectInstance.cpp
--- a/Engine/Plugins/Mutable/Source/CustomizableObject/Private/MuCO/CustomizableObjectInstance.cpp
+++ b/Engine/Plugins/Mutable/Source/CustomizableObject/Private/MuCO/CustomizableObjectInstance.cpp
@@ -7281,7 +7281,6 @@ void UCustomizableInstancePrivate::RegenerateImportedModels()
 			for (int32 SectionIndex = 0; SectionIndex < LODRenderData.RenderSections.Num(); ++SectionIndex)
 			{
 				check(!LODRenderData.StaticVertexBuffers.StaticMeshVertexBuffer.GetUseHighPrecisionTangentBasis());
-				check(LODRenderData.StaticVertexBuffers.StaticMeshVertexBuffer.GetUseFullPrecisionUVs());
 				
 				const FSkelMeshRenderSection& RenderSection = LODRenderData.RenderSections[SectionIndex];
 				FSkelMeshSection& ImportedSection = ImportedModel->LODModels[LODIndex].Sections[SectionIndex];
@@ -7443,8 +7442,9 @@ void UCustomizableInstancePrivate::RegenerateImportedModels()
 							static_cast<const FPackedNormal*>(Job.StaticVertexBuffers->StaticMeshVertexBuffer.GetTangentData()) + Job.VerticesOffset*2;
 
 					const int32 NumTexCoords = Job.StaticVertexBuffers->StaticMeshVertexBuffer.GetNumTexCoords();
-					const UVsVectorType* TexCoordBuffer = 
-							static_cast<const UVsVectorType*>(Job.StaticVertexBuffers->StaticMeshVertexBuffer.GetTexCoordData()) + Job.VerticesOffset*NumTexCoords;
+					const int32 UVSize = Job.StaticVertexBuffers->StaticMeshVertexBuffer.GetUseFullPrecisionUVs() ? 2 * sizeof(float) : 2 * sizeof(FFloat16);
+					const uint8* TexCoordBuffer = 
+							static_cast<const uint8*>(Job.StaticVertexBuffers->StaticMeshVertexBuffer.GetTexCoordData()) + Job.VerticesOffset*NumTexCoords*UVSize;
 					
 					const FColor* ColorBuffer = 
 							static_cast<const FColor*>(Job.StaticVertexBuffers->ColorVertexBuffer.GetVertexData());
@@ -7470,16 +7470,32 @@ void UCustomizableInstancePrivate::RegenerateImportedModels()
 						float TangentSign = Tangent[1].Vector.W < 0 ? -1.0f : 1.0f;
 						Vertex->TangentY = FVector3f::CrossProduct(Vertex->TangentZ, Vertex->TangentX) * TangentSign;
 
-						const UVsVectorType* TexCoord = TexCoordBuffer + JobVertexIndex*NumTexCoords;
+						const void* TexCoord = TexCoordBuffer + JobVertexIndex*NumTexCoords*UVSize;
 
 						// Switch based jumptable.
-						switch (NumTexCoords)
+						if (UVSize == 4)
 						{
-							case 4: Vertex->UVs[3] = TexCoord[3]; // Fall through
-							case 3: Vertex->UVs[2] = TexCoord[2]; // Fall through
-							case 2: Vertex->UVs[1] = TexCoord[1]; // Fall through
-							case 1: Vertex->UVs[0] = TexCoord[0]; // Fall through
+							const FFloat16* TypedSource = reinterpret_cast<const FFloat16*>(TexCoord);
+							switch (NumTexCoords)
+							{
+							case 4: Vertex->UVs[3] = { TypedSource[6], TypedSource[7] }; // Fall through
+							case 3: Vertex->UVs[2] = { TypedSource[4], TypedSource[5] }; // Fall through
+							case 2: Vertex->UVs[1] = { TypedSource[2], TypedSource[3] }; // Fall through
+							case 1: Vertex->UVs[0] = { TypedSource[0], TypedSource[1] }; // Fall through
+							default: break;
+							}
+						}
+						else
+						{
+							const FVector2f* TypedSource = reinterpret_cast<const FVector2f*>(TexCoord);
+							switch (NumTexCoords)
+							{
+							case 4: Vertex->UVs[3] = TypedSource[3]; // Fall through
+							case 3: Vertex->UVs[2] = TypedSource[2]; // Fall through
+							case 2: Vertex->UVs[1] = TypedSource[1]; // Fall through
+							case 1: Vertex->UVs[0] = TypedSource[0]; // Fall through
 							default: break;
+							}
 						}
 
 						Vertex->Color = bHasColor ? ColorBuffer[JobVertexIndex] : FColor::White;
diff --git a/Engine/Plugins/Mutable/Source/CustomizableObject/Private/MuCO/MutableMeshBufferUtils.cpp b/Engine/Plugins/Mutable/Source/CustomizableObject/Private/MuCO/MutableMeshBufferUtils.cpp
--- a/Engine/Plugins/Mutable/Source/CustomizableObject/Private/MuCO/MutableMeshBufferUtils.cpp
+++ b/Engine/Plugins/Mutable/Source/CustomizableObject/Private/MuCO/MutableMeshBufferUtils.cpp
@@ -51,40 +51,28 @@ void MutableMeshBufferUtils::SetupTangentBuffer(const int32& InCurrentVertexBuff
 }
 
 
-void MutableMeshBufferUtils::SetupTexCoordinatesBuffer(const int32& InCurrentVertexBuffer, const int32& InChannelCount,
+void MutableMeshBufferUtils::SetupTexCoordinatesBuffer(const int32& InCurrentVertexBuffer, const int32& InChannelCount, bool bHighPrecision,
                                                        mu::FMeshBufferSet& OutTargetVertexBuffers,const int32* InTextureSemanticIndicesOverride /* = nullptr */)
-{
-	// \todo: support for half precision?
-	typedef TStaticMeshVertexUVsDatum<TStaticMeshVertexUVsTypeSelector<
-		EStaticMeshVertexUVType::HighPrecision>::UVsTypeT> UVType;
+{	
+	int32 UVSize = bHighPrecision 
+		? sizeof(TStaticMeshVertexUVsDatum<TStaticMeshVertexUVsTypeSelector<EStaticMeshVertexUVType::HighPrecision>::UVsTypeT>)
+		: sizeof(TStaticMeshVertexUVsDatum<TStaticMeshVertexUVsTypeSelector<EStaticMeshVertexUVType::Default>::UVsTypeT>);
 
 	using namespace mu;
-	const int32 ElementSize = sizeof(UVType) * InChannelCount;
+	const int32 ElementSize = UVSize * InChannelCount;
 	constexpr int32 MaxChannelCount = MaxTexCordChannelCount;
-	const EMeshBufferSemantic Semantics[MaxChannelCount] = {
-		EMeshBufferSemantic::TexCoords, EMeshBufferSemantic::TexCoords, EMeshBufferSemantic::TexCoords, EMeshBufferSemantic::TexCoords
-	};
-	const EMeshBufferFormat Formats[MaxChannelCount] = {EMeshBufferFormat::Float32, EMeshBufferFormat::Float32, EMeshBufferFormat::Float32, EMeshBufferFormat::Float32};
+	const EMeshBufferSemantic Semantics[MaxChannelCount] = { EMeshBufferSemantic::TexCoords, EMeshBufferSemantic::TexCoords, EMeshBufferSemantic::TexCoords, EMeshBufferSemantic::TexCoords };
 	const int32 Components[MaxChannelCount] = {2, 2, 2, 2};
-	const int32 Offsets[MaxChannelCount] =
-	{
-		STRUCT_OFFSET(UVType, UVs) + 0 * sizeof(UVType),
-		STRUCT_OFFSET(UVType, UVs) + 1 * sizeof(UVType),
-		STRUCT_OFFSET(UVType, UVs) + 2 * sizeof(UVType),
-		STRUCT_OFFSET(UVType, UVs) + 3 * sizeof(UVType),
-	};
+	const int32 Offsets[MaxChannelCount] = { 0 * UVSize, 1 * UVSize, 2 * UVSize, 3 * UVSize, };
 
-	if (InTextureSemanticIndicesOverride)
-	{
-		OutTargetVertexBuffers.SetBuffer(InCurrentVertexBuffer, ElementSize, InChannelCount, Semantics,
-		                                  InTextureSemanticIndicesOverride, Formats, Components, Offsets);
-	}
-	else
-	{
-		const int32 SemanticIndices[MaxChannelCount] = {0, 1, 2, 3};
-		OutTargetVertexBuffers.SetBuffer(InCurrentVertexBuffer, ElementSize, InChannelCount, Semantics,
-		                                  SemanticIndices, Formats, Components, Offsets);
-	}
+	const int32 StandardSemanticIndices[MaxChannelCount] = { 0, 1, 2, 3 };
+	const int32* SemanticIndices = InTextureSemanticIndicesOverride ? InTextureSemanticIndicesOverride : StandardSemanticIndices;
+
+	const EMeshBufferFormat HighFormats[MaxChannelCount] = { EMeshBufferFormat::Float32, EMeshBufferFormat::Float32, EMeshBufferFormat::Float32, EMeshBufferFormat::Float32 };
+	const EMeshBufferFormat DefaultFormats[MaxChannelCount] = { EMeshBufferFormat::Float16, EMeshBufferFormat::Float16, EMeshBufferFormat::Float16, EMeshBufferFormat::Float16 };
+	const EMeshBufferFormat* Formats = bHighPrecision ? HighFormats : DefaultFormats;
+
+	OutTargetVertexBuffers.SetBuffer(InCurrentVertexBuffer, ElementSize, InChannelCount, Semantics, SemanticIndices, Formats, Components, Offsets);
 }
 
 
diff --git a/Engine/Plugins/Mutable/Source/CustomizableObject/Private/MuCO/UnrealConversionUtils.cpp b/Engine/Plugins/Mutable/Source/CustomizableObject/Private/MuCO/UnrealConversionUtils.cpp
--- a/Engine/Plugins/Mutable/Source/CustomizableObject/Private/MuCO/UnrealConversionUtils.cpp
+++ b/Engine/Plugins/Mutable/Source/CustomizableObject/Private/MuCO/UnrealConversionUtils.cpp
@@ -348,7 +348,16 @@ namespace UnrealConversionUtils
 
 		//const ESkeletalMeshVertexFlags BuildFlags = SkeletalMesh->GetVertexBufferFlags();
 		//const bool bUseFullPrecisionUVs = EnumHasAllFlags(BuildFlags, ESkeletalMeshVertexFlags::UseFullPrecisionUVs);
-		const bool bUseFullPrecisionUVs = true;
+		bool bUseFullPrecisionUVs = true;
+		{
+			int32 TexCoordsBufferIndex = -1;
+			int32 TexCoordsChannelIndex = -1;
+			MutableMesh->VertexBuffers.FindChannel(mu::EMeshBufferSemantic::TexCoords, 0, &TexCoordsBufferIndex, &TexCoordsChannelIndex);
+			if (TexCoordsBufferIndex >= 0 && TexCoordsChannelIndex >= 0)
+			{
+				bUseFullPrecisionUVs = MutableMesh->VertexBuffers.Buffers[TexCoordsBufferIndex].Channels[TexCoordsChannelIndex].Format == mu::EMeshBufferFormat::Float32;
+			}
+		}
 
 		const int NumTexCoords = MutableMeshVertexBuffers.GetBufferChannelCount(MUTABLE_VERTEXBUFFER_TEXCOORDS);
 
@@ -406,7 +415,7 @@ namespace UnrealConversionUtils
 		);
 
 		// Optional buffers
-		for (int Buffer = MUTABLE_VERTEXBUFFER_TEXCOORDS + 1; Buffer < MutableMeshVertexBuffers.GetBufferCount(); ++Buffer)
+		for (int32 Buffer = MUTABLE_VERTEXBUFFER_TEXCOORDS + 1; Buffer < MutableMeshVertexBuffers.GetBufferCount(); ++Buffer)
 		{
 			if (MutableMeshVertexBuffers.GetBufferChannelCount(Buffer) > 0)
 			{
@@ -991,6 +1000,7 @@ namespace UnrealConversionUtils
 
 
 		// Mesh data
+		int32 FirstVertexIndex = LOD.RenderSections[SectionIndex].BaseVertexIndex;
 		int32 VertexCount = Section.GetNumVertices();
 
 		mu::FMeshBufferSet& Vertices = Result.GetVertexBuffers();
@@ -999,10 +1009,17 @@ namespace UnrealConversionUtils
 
 		// Position buffer
 		int32 CurrentVertexBuffer = 0;
-		MutableMeshBufferUtils::SetupVertexPositionsBuffer(CurrentVertexBuffer, Vertices);
-		check(Vertices.Buffers[CurrentVertexBuffer].ElementSize * VertexCount <= LOD.StaticVertexBuffers.PositionVertexBuffer.GetAllocatedSize());
-		FMemory::Memcpy(Vertices.GetBufferData(CurrentVertexBuffer), LOD.StaticVertexBuffers.PositionVertexBuffer.GetVertexData(), Vertices.Buffers[CurrentVertexBuffer].ElementSize * VertexCount);
-		++CurrentVertexBuffer;
+		{
+			MutableMeshBufferUtils::SetupVertexPositionsBuffer(CurrentVertexBuffer, Vertices);
+
+			int32 ElementSize = Vertices.Buffers[CurrentVertexBuffer].ElementSize;
+			const uint8* SourceVertexData = ((const uint8*)LOD.StaticVertexBuffers.PositionVertexBuffer.GetVertexData()) + FirstVertexIndex * ElementSize;
+
+			check(ElementSize * (FirstVertexIndex+VertexCount) <= LOD.StaticVertexBuffers.PositionVertexBuffer.GetAllocatedSize());
+			FMemory::Memcpy(Vertices.GetBufferData(CurrentVertexBuffer), SourceVertexData, ElementSize * VertexCount);
+
+			++CurrentVertexBuffer;
+		}
 
 		// Tangent buffer
 		{
@@ -1026,8 +1043,10 @@ namespace UnrealConversionUtils
 			}
 			Vertices.SetBuffer(CurrentVertexBuffer, ElementSize, 2, Semantics, SemanticIndices, Formats, Components, Offsets);
 
-			check((int32)Vertices.Buffers[CurrentVertexBuffer].ElementSize * VertexCount <= LOD.StaticVertexBuffers.StaticMeshVertexBuffer.GetTangentSize());
-			FMemory::Memcpy(Vertices.GetBufferData(CurrentVertexBuffer), LOD.StaticVertexBuffers.StaticMeshVertexBuffer.GetTangentData(), ElementSize * VertexCount);
+			const uint8* SourceVertexData = ((const uint8*)LOD.StaticVertexBuffers.StaticMeshVertexBuffer.GetTangentData()) + FirstVertexIndex * ElementSize;
+
+			check((int32)ElementSize * (FirstVertexIndex + VertexCount) <= LOD.StaticVertexBuffers.StaticMeshVertexBuffer.GetTangentSize());
+			FMemory::Memcpy(Vertices.GetBufferData(CurrentVertexBuffer), SourceVertexData, ElementSize * VertexCount);
 
 			++CurrentVertexBuffer;
 		}
@@ -1044,33 +1063,23 @@ namespace UnrealConversionUtils
 			int32 Offsets[MaxChannelCount];
 
 			int32 UVSize = LOD.StaticVertexBuffers.StaticMeshVertexBuffer.GetUseFullPrecisionUVs() ? 8 : 4;
-			const int32 ElementSize = 8 * NumTexCoords; // UVSize * NumTexCoords;
+			const int32 ElementSize = UVSize * NumTexCoords;
 
 			for (int32 UV = 0; UV < NumTexCoords; ++UV)
 			{
 				Semantics[UV] = mu::EMeshBufferSemantic::TexCoords;
 				SemanticIndices[UV] = UV;
-				Formats[UV] = mu::EMeshBufferFormat::Float32; //  (UVSize == 8) ? mu::EMeshBufferFormat::Float32 : mu::EMeshBufferFormat::Float16;
+				Formats[UV] = (UVSize == 8) ? mu::EMeshBufferFormat::Float32 : mu::EMeshBufferFormat::Float16;
 				Components[UV] = 2;
 				Offsets[UV] = UVSize * UV;
 			}
 			Vertices.SetBuffer(CurrentVertexBuffer, ElementSize, NumTexCoords, Semantics, SemanticIndices, Formats, Components, Offsets);
 
-			if (UVSize == 8)
-			{
-				FMemory::Memcpy(Vertices.GetBufferData(CurrentVertexBuffer), LOD.StaticVertexBuffers.StaticMeshVertexBuffer.GetTexCoordData(), Vertices.Buffers[CurrentVertexBuffer].ElementSize * VertexCount);
-			}
-			else
-			{
-				float* Target = (float*)(Vertices.GetBufferData(CurrentVertexBuffer));
-				const FFloat16* Source = (const FFloat16*)LOD.StaticVertexBuffers.StaticMeshVertexBuffer.GetTexCoordData();
-				for (int32 v = 0; v < VertexCount * 2 * NumTexCoords; ++v)
-				{
-					*Target = *Source;
-					++Target;
-					++Source;
-				}
-			}
+			const uint8* SourceVertexData = ((const uint8*)LOD.StaticVertexBuffers.StaticMeshVertexBuffer.GetTexCoordData()) + FirstVertexIndex * UVSize * NumTexCoords;
+			check(UVSize * NumTexCoords * (FirstVertexIndex + VertexCount) <= LOD.StaticVertexBuffers.StaticMeshVertexBuffer.GetTexCoordSize());
+
+			FMemory::Memcpy(Vertices.GetBufferData(CurrentVertexBuffer), SourceVertexData, Vertices.Buffers[CurrentVertexBuffer].ElementSize * VertexCount);
+
 			++CurrentVertexBuffer;
 		}
 
@@ -1082,30 +1091,44 @@ namespace UnrealConversionUtils
 			const int32 MaxBoneWeightTypeSizeBytes = SkinBuffer->GetBoneWeightByteSize();
 			const int32 MaxBonesPerVertex = SkinBuffer->GetMaxBoneInfluences();
 			MutableMeshBufferUtils::SetupSkinBuffer(CurrentVertexBuffer, MaxBoneIndexTypeSizeBytes, MaxBoneWeightTypeSizeBytes, MaxBonesPerVertex, Vertices);
-			check(Vertices.Buffers[CurrentVertexBuffer].ElementSize <= SkinBuffer->GetVertexDataSize());
-			FMemory::Memcpy(Vertices.GetBufferData(CurrentVertexBuffer),
-				SkinBuffer->GetDataVertexBuffer()->GetWeightData(),
-				Vertices.Buffers[CurrentVertexBuffer].ElementSize * VertexCount);
+
+			int32 ElementSize = Vertices.Buffers[CurrentVertexBuffer].ElementSize;
+			const uint8* SourceVertexData = ((const uint8*)SkinBuffer->GetDataVertexBuffer()->GetWeightData()) + FirstVertexIndex * ElementSize;
+
+			check(ElementSize * (FirstVertexIndex + VertexCount) <= (int32)SkinBuffer->GetVertexDataSize());
+
+			FMemory::Memcpy(Vertices.GetBufferData(CurrentVertexBuffer), SourceVertexData, ElementSize * VertexCount);
 
 			++CurrentVertexBuffer;
 		}
 
 		// Colour buffer
-		// TODO
-		//if (MeshData.bHasVertexColors)
-		//{
-		//	MutableMeshBufferUtils::SetupVertexColorBuffer(CurrentVertexBuffer, Vertices);
-		//	++CurrentVertexBuffer;
-		//}
+		if (LOD.StaticVertexBuffers.ColorVertexBuffer.GetNumVertices() > 0)
+		{			
+			const FColorVertexBuffer& ColorBuffer = LOD.StaticVertexBuffers.ColorVertexBuffer;
+
+			MutableMeshBufferUtils::SetupVertexColorBuffer(CurrentVertexBuffer, Vertices);
+
+			int32 ElementSize = Vertices.Buffers[CurrentVertexBuffer].ElementSize;
+			const uint8* SourceVertexData = ((const uint8*)ColorBuffer.GetVertexData()) + FirstVertexIndex * ElementSize;
+
+			check(ElementSize * (FirstVertexIndex + VertexCount) <= (int32)ColorBuffer.GetAllocatedSize());
+			check(ElementSize == ColorBuffer.GetStride());
+
+			FMemory::Memcpy(Vertices.GetBufferData(CurrentVertexBuffer), SourceVertexData, ElementSize * VertexCount);
+
+			++CurrentVertexBuffer;
+		}
 
 		// Indices
 		{
-			int32 FaceCount = Section.NumTriangles;
+			int32 FirstIndexIndex = LOD.RenderSections[SectionIndex].BaseIndex;
+			int32 IndexCount = Section.NumTriangles*3;
 			int32 ElementSize = LOD.MultiSizeIndexContainer.GetDataTypeSize();
 
 			mu::FMeshBufferSet& Indices = Result.GetIndexBuffers();
 			Indices.SetBufferCount(1);
-			Indices.SetElementCount(FaceCount * 3);
+			Indices.SetElementCount(IndexCount);
 			constexpr int32 ChannelCount = 1;
 			const mu::EMeshBufferSemantic Semantics[ChannelCount] = { mu::EMeshBufferSemantic::VertexIndex };
 			const int32 SemanticIndices[ChannelCount] = { 0 };
@@ -1114,19 +1137,46 @@ namespace UnrealConversionUtils
 			const int32 Offsets[ChannelCount] = { 0 };
 			Indices.SetBuffer(0, ElementSize, ChannelCount, Semantics, SemanticIndices, Formats, Components, Offsets);
 
-			FMemory::Memcpy(Indices.GetBufferData(0), LOD.MultiSizeIndexContainer.GetIndexBuffer()->GetPointerTo(0), Indices.GetElementCount() * ElementSize);
-		}
+			const uint8* SourceData = (const uint8*)LOD.MultiSizeIndexContainer.GetIndexBuffer()->GetPointerTo(FirstIndexIndex);
+			uint8* TargetData = Indices.GetBufferData(0);
 
-		// Assign an ID prefix to the mesh. 
-		// TODO: Hash all data?
-		// TODO: Ensure unique?
-		uint32 MeshID = HashCombine(GetTypeHash(LODIndex), GetTypeHash(SkeletalMesh.GetFullName()));
-		MeshID = HashCombine(MeshID, GetTypeHash(SectionIndex));
-		Result.MeshIDPrefix = MeshID;
+			if (FirstVertexIndex == 0)
+			{
+				FMemory::Memcpy(TargetData, SourceData, IndexCount* ElementSize);
+			}
+			else
+			{
+				// Apply vertex offset
+				switch (ElementSize)
+				{
+				case 2:
+					for (int32 Index=0;Index<IndexCount; ++Index)
+					{
+						const uint16* Source = ((const uint16*)SourceData) + Index;
+						uint16* Target = ((uint16*)TargetData) + Index;
+						check(*Source>=FirstVertexIndex);
+						*Target = *Source - FirstVertexIndex;
+					}
+					break;
+				case 4:
+					for (int32 Index = 0; Index < IndexCount; ++Index)
+					{
+						const uint32* Source = ((const uint32*)SourceData) + Index;
+						uint32* Target = ((uint32*)TargetData) + Index;
+						check(*Source >= uint32(FirstVertexIndex));
+						*Target = *Source - FirstVertexIndex;
+					}
+					break;
+				default:
+					// Index size not implemented
+					break;
+				}
+			}
+		}
 
-		FString DebugLog;
-		Result.Log(DebugLog, 8);
-		UE_LOG(LogMutable, Warning, TEXT("[%s]"), *DebugLog);
+		//FString DebugLog;
+		//Result.Log(DebugLog, 8);
+		//UE_LOG(LogMutable, Warning, TEXT("[%s]"), *DebugLog);
 
 		Result.EnsureSurfaceData();
 	}
diff --git a/Engine/Plugins/Mutable/Source/CustomizableObject/Private/MuCO/UnrealMutableImageProvider.cpp b/Engine/Plugins/Mutable/Source/CustomizableObject/Private/MuCO/UnrealMutableImageProvider.cpp
--- a/Engine/Plugins/Mutable/Source/CustomizableObject/Private/MuCO/UnrealMutableImageProvider.cpp
+++ b/Engine/Plugins/Mutable/Source/CustomizableObject/Private/MuCO/UnrealMutableImageProvider.cpp
@@ -505,18 +505,18 @@ TTuple<UE::Tasks::FTask, TFunction<void()>> FUnrealMutableResourceProvider::GetM
 			}
 			else
 			{
-				TStrongObjectPtr<UCustomizableObject> CO = CustomizableObject.Pin();
-				check(CO);
+				UModelResources* ModelResources = nullptr;
 
+				TStrongObjectPtr<UCustomizableObject> CO = CustomizableObject.Pin();
 				if (CO)
 				{
-					UModelResources* ModelResources = CO->GetPrivate()->GetModelResources();
-					check(ModelResources);
-					if (ModelResources)
-					{
-						UnrealConversionUtils::ConvertSkeletalMeshFromRuntimeData(*SkeletalMesh, LODIndex, SectionIndex, ModelResources, *Result.Get());
-					}
+					ModelResources = CO->GetPrivate()->GetModelResources();
 				}
+
+				// It is valid not to have a CO or MutableStreamedResources. It is only used for skeleton data. 
+				// This may happen when updating mips, and for those operations we don't need skeleton data.
+				//ensure(ModelResources);
+				UnrealConversionUtils::ConvertSkeletalMeshFromRuntimeData(*SkeletalMesh, LODIndex, SectionIndex, ModelResources, *Result.Get());
 			}
 
 			// Signal
diff --git a/Engine/Plugins/Mutable/Source/CustomizableObjectEditor/Private/MuCOE/CustomizableObjectEditorModule.cpp b/Engine/Plugins/Mutable/Source/CustomizableObjectEditor/Private/MuCOE/CustomizableObjectEditorModule.cpp
--- a/Engine/Plugins/Mutable/Source/CustomizableObjectEditor/Private/MuCOE/CustomizableObjectEditorModule.cpp
+++ b/Engine/Plugins/Mutable/Source/CustomizableObjectEditor/Private/MuCOE/CustomizableObjectEditorModule.cpp
@@ -69,6 +69,8 @@
 #include "MuCOE/Nodes/CustomizableObjectNodeSkeletalMesh.h"
 #include "MuCOE/Nodes/CustomizableObjectNodeSkeletalMeshDetails.h"
 #include "MuCOE/Nodes/CustomizableObjectNodeStaticMesh.h"
+#include "MuCOE/Nodes/CustomizableObjectNodeMeshParameter.h"
+#include "MuCOE/Nodes/CustomizableObjectNodeMeshParameterDetails.h"
 #include "MuCOE/Nodes/CustomizableObjectNodeTable.h"
 #include "MuCOE/Nodes/CustomizableObjectNodeTableDetails.h"
 #include "MuCOE/Nodes/CustomizableObjectNodeTunnel.h"
@@ -234,6 +236,7 @@ void FCustomizableObjectEditorModule::StartupModule()
 	RegisterCustomDetails(PropertyModule, UCustomizableObjectNodeMaterial::StaticClass(), FOnGetDetailCustomizationInstance::CreateStatic(&FCustomizableObjectNodeMeshSectionDetails::MakeInstance));
 	RegisterCustomDetails(PropertyModule, UCustomizableObjectNodeSkeletalMesh::StaticClass(), FOnGetDetailCustomizationInstance::CreateStatic(&FCustomizableObjectNodeSkeletalMeshDetails::MakeInstance));
 	RegisterCustomDetails(PropertyModule, UCustomizableObjectNodeStaticMesh::StaticClass(), FOnGetDetailCustomizationInstance::CreateStatic(&FCustomizableObjectNodeDetails::MakeInstance));
+	RegisterCustomDetails(PropertyModule, UCustomizableObjectNodeMeshParameter::StaticClass(), FOnGetDetailCustomizationInstance::CreateStatic(&FCustomizableObjectNodeMeshParameterDetails::MakeInstance));
 	RegisterCustomDetails(PropertyModule, UCustomizableObjectNodeTable::StaticClass(), FOnGetDetailCustomizationInstance::CreateStatic(&FCustomizableObjectNodeTableDetails::MakeInstance));
 	RegisterCustomDetails(PropertyModule, UCustomizableObjectNodeComponentMesh::StaticClass(), FOnGetDetailCustomizationInstance::CreateStatic(&FCustomizableObjectNodeComponentMeshDetails::MakeInstance));
 	RegisterCustomDetails(PropertyModule, UCustomizableObjectNodeMacroInstance::StaticClass(), FOnGetDetailCustomizationInstance::CreateStatic(&FCustomizableObjectNodeMacroInstanceDetails::MakeInstance));
diff --git a/Engine/Plugins/Mutable/Source/CustomizableObjectEditor/Private/MuCOE/CustomizableObjectLayout.cpp b/Engine/Plugins/Mutable/Source/CustomizableObjectEditor/Private/MuCOE/CustomizableObjectLayout.cpp
--- a/Engine/Plugins/Mutable/Source/CustomizableObjectEditor/Private/MuCOE/CustomizableObjectLayout.cpp
+++ b/Engine/Plugins/Mutable/Source/CustomizableObjectEditor/Private/MuCOE/CustomizableObjectLayout.cpp
@@ -223,7 +223,7 @@ TSoftObjectPtr<UObject> UCustomizableObjectLayout::GetMesh() const
 {
 	if (UObject* Node = GetOuter())
 	{
-		if (const UCustomizableObjectNodeSkeletalMesh* TypedNodeSkeletalMesh = Cast<UCustomizableObjectNodeSkeletalMesh>(Node))
+		if (const UCustomizableObjectNodeMesh* TypedNodeSkeletalMesh = Cast<UCustomizableObjectNodeMesh>(Node))
 		{
 			return TypedNodeSkeletalMesh->GetMesh();
 		}
diff --git a/Engine/Plugins/Mutable/Source/CustomizableObjectEditor/Private/MuCOE/GenerateMutableSource/GenerateMutableSourceMesh.cpp b/Engine/Plugins/Mutable/Source/CustomizableObjectEditor/Private/MuCOE/GenerateMutableSource/GenerateMutableSourceMesh.cpp
--- a/Engine/Plugins/Mutable/Source/CustomizableObjectEditor/Private/MuCOE/GenerateMutableSource/GenerateMutableSourceMesh.cpp
+++ b/Engine/Plugins/Mutable/Source/CustomizableObjectEditor/Private/MuCOE/GenerateMutableSource/GenerateMutableSourceMesh.cpp
@@ -116,7 +116,8 @@ void SetSurfaceFormat(FMutableCompilationContext& Context,
 	++CurrentVertexBuffer;
 
 	// Texture coords buffer
-	MutableMeshBufferUtils::SetupTexCoordinatesBuffer(CurrentVertexBuffer, MeshData.NumTexCoordChannels, OutVertexBufferFormat);
+	bool bHighPrecision = true;
+	MutableMeshBufferUtils::SetupTexCoordinatesBuffer(CurrentVertexBuffer, MeshData.NumTexCoordChannels, bHighPrecision, OutVertexBufferFormat);
 	++CurrentVertexBuffer;
 
 	// Skin buffer
@@ -953,7 +954,7 @@ TSharedPtr<mu::FMesh> ConvertSkeletalMeshToMutable(
 		check(ComponentRefSkeletalMesh);
 		check(ComponentRefSkeleton);
 
-		// Compatibility check
+		// Compatibility check and add skeleton references
 		{
 			FString ErrorMessage;
 			const bool bCompatible = IsSkeletalMeshCompatibleWithRefSkeleton(*MutComponentInfo, SkeletalMesh, ErrorMessage);
@@ -4538,20 +4539,60 @@ mu::Ptr<mu::NodeMesh> GenerateMutableSourceMesh(const UEdGraphPin* Pin,
 			ParameterNode->DefaultValue = FName(TypedNodeParam->DefaultValue->GetPathName());
 		}
 
+		const int32 LODIndex = 0;
+		int32 SectionIndex = 0;
+		TypedNodeParam->GetPinSection(*Pin, SectionIndex);
+		ParameterNode->SectionIndex = SectionIndex;
+
 		GenerationContext.ParameterUIDataMap.Add(TypedNodeParam->ParameterName, FMutableParameterData(
 			TypedNodeParam->ParamUIMetadata,
 			EMutableParameterType::Texture));
 
+		// Layouts
+		{
+			const int32 LODIndexLayout = LODIndex;
+			const int32 SectionIndexLayout = SectionIndex;
+
+
+			const FLayoutGenerationFlags& LayoutFlags = GenerationContext.LayoutGenerationFlags.Last();
+
+			const TArray<UCustomizableObjectLayout*>& Layouts = TypedNodeParam->GetLayouts(*Pin);
+			int32 NumLayouts = Layouts.Num();
+			ParameterNode->Layouts.SetNum(NumLayouts);
+			for (int32 LayoutIndex = 0; LayoutIndex < NumLayouts; ++LayoutIndex)
+			{
+				if (!LayoutFlags.TexturePinModes.IsValidIndex(LayoutIndex) ||
+					LayoutFlags.TexturePinModes[LayoutIndex] != EPinMode::Mutable)
+				{
+					ParameterNode->Layouts[LayoutIndex] = CreateDefaultLayout();
+
+					// Keep packing strategy if possible, Overlay can be valid with EPinMode == Passthrough	
+					if (Layouts.IsValidIndex(LayoutIndex) && Layouts[LayoutIndex])
+					{
+						ParameterNode->Layouts[LayoutIndex]->Strategy = ConvertLayoutStrategy(Layouts[LayoutIndex]->PackingStrategy);
+					}
+
+					// Ignore layout
+					continue;
+				}
+
+				const UCustomizableObjectLayout* Layout = Layouts.IsValidIndex(LayoutIndex) ? Layouts[LayoutIndex] : nullptr;
+				if (ensure(Layout))
+				{
+					bool bWasEmpty = false;
+					mu::Ptr<mu::NodeLayout> LayoutNode = CreateMutableLayoutNode(GenerationContext, Layout, bLinkedToExtendMaterial, bWasEmpty);
+					LayoutNode->SetMessageContext(Node);
+
+					ParameterNode->Layouts[LayoutIndex] = LayoutNode;
+				}
+			}
+		}
+
 		// For constants we also have:
 		// animbp instance?
 		// gameplay tags?
-		// layouts?
 		// morphs?
 
-		// Force the format?
-		//mu::NodeMeshFormatPtr FormatNode = new mu::NodeMeshFormat();
-		//FormatNode->SetSource(ParameterNode);
-
 		Result = ParameterNode;
 	}
 
diff --git a/Engine/Plugins/Mutable/Source/CustomizableObjectEditor/Private/MuCOE/MutableMeshPreviewUtils.cpp b/Engine/Plugins/Mutable/Source/CustomizableObjectEditor/Private/MuCOE/MutableMeshPreviewUtils.cpp
--- a/Engine/Plugins/Mutable/Source/CustomizableObjectEditor/Private/MuCOE/MutableMeshPreviewUtils.cpp
+++ b/Engine/Plugins/Mutable/Source/CustomizableObjectEditor/Private/MuCOE/MutableMeshPreviewUtils.cpp
@@ -151,7 +151,8 @@ namespace MutableMeshPreviewUtils
 					CurrentVertexBuffer++;
 				
 					// Texture coords buffer
-					MutableMeshBufferUtils::SetupTexCoordinatesBuffer(CurrentVertexBuffer,NumbOfTextCoordChannels,FormattedVertexBuffers, TextureSemanticIndices);
+					const bool bHighPrecision = true;
+					MutableMeshBufferUtils::SetupTexCoordinatesBuffer(CurrentVertexBuffer,NumbOfTextCoordChannels, bHighPrecision, FormattedVertexBuffers, TextureSemanticIndices);
 					CurrentVertexBuffer++;
 				
 					// Skin buffer
diff --git a/Engine/Plugins/Mutable/Source/CustomizableObjectEditor/Private/MuCOE/Nodes/CustomizableObjectNodeMeshParameter.cpp b/Engine/Plugins/Mutable/Source/CustomizableObjectEditor/Private/MuCOE/Nodes/CustomizableObjectNodeMeshParameter.cpp
--- a/Engine/Plugins/Mutable/Source/CustomizableObjectEditor/Private/MuCOE/Nodes/CustomizableObjectNodeMeshParameter.cpp
+++ b/Engine/Plugins/Mutable/Source/CustomizableObjectEditor/Private/MuCOE/Nodes/CustomizableObjectNodeMeshParameter.cpp
@@ -2,8 +2,11 @@
 
 #include "MuCOE/Nodes/CustomizableObjectNodeMeshParameter.h"
 
-#include "MuCO/CustomizableObjectCustomVersion.h"
 #include "MuCOE/EdGraphSchema_CustomizableObject.h"
+#include "MuCOE/CustomizableObjectLayout.h"
+#include "MuCO/CustomizableObjectCustomVersion.h"
+#include "Rendering/SkeletalMeshModel.h"
+#include "Engine/SkeletalMesh.h"
 
 class UCustomizableObjectNodeRemapPins;
 
@@ -12,13 +15,62 @@ class UCustomizableObjectNodeRemapPins;
 
 void UCustomizableObjectNodeMeshParameter::AllocateDefaultPins(UCustomizableObjectNodeRemapPins* RemapPins)
 {
-	const FName Type = UEdGraphSchema_CustomizableObject::PC_Mesh;
-	const FName PinName = UEdGraphSchema_CustomizableObject::GetPinCategoryName(Type);
-	const FText PinFriendlyName = UEdGraphSchema_CustomizableObject::GetPinCategoryFriendlyName(Type);
+	USkeletalMesh* SkeletalMesh = DefaultValue.LoadSynchronous();
+
+	if (!SkeletalMesh)
+	{
+		return;
+	}
+
+	const UEdGraphSchema_CustomizableObject* Schema = GetDefault<UEdGraphSchema_CustomizableObject>();
+
+	if (const FSkeletalMeshModel* ImportedModel = SkeletalMesh->GetImportedModel())
+	{
+		int32 LODIndex = 0;
+		{
+			const int32 NumSections = ImportedModel->LODModels[LODIndex].Sections.Num();
+			for (int32 SectionIndex = 0; SectionIndex < NumSections; ++SectionIndex)
+			{
+				// Ignore disabled sections.
+				if (ImportedModel->LODModels[LODIndex].Sections[SectionIndex].bDisabled)
+				{
+					continue;
+				}
+
+				UMaterialInterface* MaterialInterface = GetMaterialInterfaceFor( SectionIndex);
+
+				FString Section = FString::Printf(TEXT("Section %i"), SectionIndex);
+				if (MaterialInterface)
+				{
+					Section.Append(FString::Printf(TEXT(" : %s"), *MaterialInterface->GetName()));
+				}
+
+				// Mesh
+				{
+					UCustomizableObjectNodeMeshParameterPinDataSection* PinData = NewObject<UCustomizableObjectNodeMeshParameterPinDataSection>(this);
+					PinData->Init( SectionIndex, ImportedModel->LODModels[LODIndex].NumTexCoords );
+
+					const FString MeshPinName = FString::Printf(TEXT("Section %i - Mesh"), SectionIndex);
+
+					UEdGraphPin* Pin = CustomCreatePin(EGPD_Output, Schema->PC_Mesh, FName(*MeshPinName), PinData);
+					Pin->PinFriendlyName = FText::FromString(FString::Printf(TEXT("%s"), *Section));
+				}
+			}
+		}
+	}
+
+}
+
+
+void UCustomizableObjectNodeMeshParameter::PostEditChangeProperty(FPropertyChangedEvent& PropertyChangedEvent)
+{
+	Super::PostEditChangeProperty(PropertyChangedEvent);
 	
-	UEdGraphPin* ValuePin = CustomCreatePin(EGPD_Output, Type, PinName);
-	ValuePin->bDefaultValueIsIgnored = true;
-	ValuePin->PinFriendlyName = PinFriendlyName;
+	FProperty* PropertyThatChanged = PropertyChangedEvent.Property;
+	if (PropertyThatChanged && PropertyThatChanged->GetName() == GET_MEMBER_NAME_CHECKED(UCustomizableObjectNodeMeshParameter, DefaultValue))
+	{
+		ReconstructNode();
+	}
 }
 
 
@@ -66,5 +118,187 @@ void UCustomizableObjectNodeMeshParameter::OnRenameNode(const FString& NewName)
 	}
 }
 
+
+UTexture2D* UCustomizableObjectNodeMeshParameter::FindTextureForPin(const UEdGraphPin* Pin) const
+{
+	return nullptr;
+}
+
+
+TArray<UCustomizableObjectLayout*> UCustomizableObjectNodeMeshParameter::GetLayouts(const UEdGraphPin& MeshPin) const
+{
+	const UCustomizableObjectNodeMeshParameterPinDataSection& MeshPinData = GetPinData<UCustomizableObjectNodeMeshParameterPinDataSection>(MeshPin);
+	return MeshPinData.Layouts;
+}
+
+
+TSoftObjectPtr<UObject> UCustomizableObjectNodeMeshParameter::GetMesh() const
+{
+	return DefaultValue;
+}
+
+
+UEdGraphPin* UCustomizableObjectNodeMeshParameter::GetMeshPin(const int32 LODIndex, const int32 SectionIndex) const
+{
+	for (UEdGraphPin* Pin : GetAllNonOrphanPins())
+	{
+		if (const UCustomizableObjectNodeMeshParameterPinDataSection* PinData = Cast<UCustomizableObjectNodeMeshParameterPinDataSection>(GetPinData(*Pin)))
+		{
+			if (PinData->GetSectionIndex() == SectionIndex)
+			{
+				return Pin;
+			}
+		}
+	}
+
+	return nullptr;
+}
+
+
+void UCustomizableObjectNodeMeshParameter::GetPinSection(const UEdGraphPin& Pin, int32& OutLODIndex, int32& OutSectionIndex, int32& OutLayoutIndex) const
+{
+	OutLODIndex = -1;
+	OutSectionIndex = -1;
+	OutLayoutIndex = -1;
+
+	if (const UCustomizableObjectNodeMeshParameterPinDataSection* PinData = Cast<UCustomizableObjectNodeMeshParameterPinDataSection>(GetPinData(Pin)))
+	{
+		OutLODIndex = 0;
+		OutSectionIndex = PinData->GetSectionIndex();
+	}
+}
+
+
+void UCustomizableObjectNodeMeshParameter::GetPinSection(const UEdGraphPin& Pin, int32& OutSectionIndex) const
+{
+	if (const UCustomizableObjectNodeMeshParameterPinDataSection* PinData = Cast<UCustomizableObjectNodeMeshParameterPinDataSection>(GetPinData(Pin)))
+	{
+		OutSectionIndex = PinData->GetSectionIndex();
+		return;
+	}
+
+	OutSectionIndex = -1;
+}
+
+
+UMaterialInterface* UCustomizableObjectNodeMeshParameter::GetMaterialInterfaceFor(const int32 SectionIndex) const
+{
+	if (FSkeletalMaterial* SkeletalMaterial = GetSkeletalMaterialFor(SectionIndex))
+	{
+		return SkeletalMaterial->MaterialInterface;
+	}
+
+	return nullptr;
+}
+
+
+FSkeletalMaterial* UCustomizableObjectNodeMeshParameter::GetSkeletalMaterialFor(const int32 SectionIndex) const
+{
+	USkeletalMesh* SkeletalMesh = DefaultValue.LoadSynchronous();
+
+	if (!SkeletalMesh)
+	{
+		return nullptr;
+	}
+
+	const int32 SkeletalMeshMaterialIndex = GetSkeletalMaterialIndexFor(SectionIndex);
+	if (SkeletalMesh->GetMaterials().IsValidIndex(SkeletalMeshMaterialIndex))
+	{
+		return &SkeletalMesh->GetMaterials()[SkeletalMeshMaterialIndex];
+	}
+
+	return nullptr;
+}
+
+int32 UCustomizableObjectNodeMeshParameter::GetSkeletalMaterialIndexFor(const int32 SectionIndex) const
+{
+	USkeletalMesh* SkeletalMesh = DefaultValue.LoadSynchronous();
+
+	if (!SkeletalMesh)
+	{
+		return INDEX_NONE;
+	}
+
+	int32 LODIndex = 0;
+
+	// We assume that LODIndex and MaterialIndex are valid for the imported model
+	int32 SkeletalMeshMaterialIndex = INDEX_NONE;
+
+	// Check if we have lod info map to get the correct material index
+	if (const FSkeletalMeshLODInfo* LodInfo = SkeletalMesh->GetLODInfo(LODIndex))
+	{
+		if (LodInfo->LODMaterialMap.IsValidIndex(SectionIndex))
+		{
+			SkeletalMeshMaterialIndex = LodInfo->LODMaterialMap[SectionIndex];
+		}
+	}
+
+	// Only deduce index when the explicit mapping is not found or there is no remap
+	if (SkeletalMeshMaterialIndex == INDEX_NONE)
+	{
+		FSkeletalMeshModel* ImportedModel = SkeletalMesh->GetImportedModel();
+		if (ImportedModel && ImportedModel->LODModels.IsValidIndex(LODIndex) && ImportedModel->LODModels[LODIndex].Sections.IsValidIndex(SectionIndex))
+		{
+			SkeletalMeshMaterialIndex = ImportedModel->LODModels[LODIndex].Sections[SectionIndex].MaterialIndex;
+		}
+	}
+
+	return SkeletalMeshMaterialIndex;
+}
+
+
+int32 UCustomizableObjectNodeMeshParameter::GetSkeletalMaterialIndexFor(const UEdGraphPin& Pin) const
+{
+	int32 SectionIndex;
+	GetPinSection(Pin, SectionIndex);
+
+	return GetSkeletalMaterialIndexFor(SectionIndex);
+}
+
+
+void UCustomizableObjectNodeMeshParameterPinDataSection::Init(int32 InSectionIndex, int32 NumTexCoords)
+{
+	SectionIndex = InSectionIndex;
+
+	if (NumTexCoords > 0)
+	{
+		UObject* Outer = GetOuter();
+
+		Layouts.SetNum(NumTexCoords);
+
+		const int LODIndex = 0;
+
+		for (int32 Index = 0; Index < NumTexCoords; ++Index)
+		{
+			Layouts[Index] = NewObject<UCustomizableObjectLayout>(Outer);
+			Layouts[Index]->SetLayout(LODIndex, InSectionIndex, Index);
+		}
+	}
+}
+
+
+int32 UCustomizableObjectNodeMeshParameterPinDataSection::GetSectionIndex() const
+{
+	return SectionIndex;
+}
+
+
+bool UCustomizableObjectNodeMeshParameterPinDataSection::Equals(const UCustomizableObjectNodePinData& Other) const
+{
+	if (GetClass() != Other.GetClass())
+	{
+		return false;
+	}
+
+	const UCustomizableObjectNodeMeshParameterPinDataSection& OtherTyped = static_cast<const UCustomizableObjectNodeMeshParameterPinDataSection&>(Other);
+	if (SectionIndex != OtherTyped.SectionIndex)
+	{
+		return false;
+	}
+
+	return Super::Equals(Other);
+}
+
+
 #undef LOCTEXT_NAMESPACE
 
diff --git a/Engine/Plugins/Mutable/Source/CustomizableObjectEditor/Private/MuCOE/Nodes/CustomizableObjectNodeMeshParameter.h b/Engine/Plugins/Mutable/Source/CustomizableObjectEditor/Private/MuCOE/Nodes/CustomizableObjectNodeMeshParameter.h
--- a/Engine/Plugins/Mutable/Source/CustomizableObjectEditor/Private/MuCOE/Nodes/CustomizableObjectNodeMeshParameter.h
+++ b/Engine/Plugins/Mutable/Source/CustomizableObjectEditor/Private/MuCOE/Nodes/CustomizableObjectNodeMeshParameter.h
@@ -2,19 +2,46 @@
 
 #pragma once
 
-#include "MuCOE/Nodes/CustomizableObjectNode.h"
+#include "MuCOE/Nodes/CustomizableObjectNodeMesh.h"
 
 #include "CustomizableObjectNodeMeshParameter.generated.h"
 
 namespace ENodeTitleType { enum Type : int; }
 
-class UCustomizableObjectNodeRemapPins;
-class UObject;
-class UMesh2D;
+class UCustomizableObjectLayout;
+class UMaterialInterface;
+struct FSkeletalMaterial;
 
 
+/** PinData of a pin that belongs to a Skeletal Mesh Section. */
 UCLASS()
-class CUSTOMIZABLEOBJECTEDITOR_API UCustomizableObjectNodeMeshParameter : public UCustomizableObjectNode
+class UCustomizableObjectNodeMeshParameterPinDataSection : public UCustomizableObjectNodePinData
+{
+	GENERATED_BODY()
+
+public:
+	void Init(int32 InSectionIndex, int32 NumTexCoords);
+
+	int32 GetSectionIndex() const;
+
+	/** Layouts related to this Mesh pin */
+	UPROPERTY()
+	TArray<TObjectPtr<UCustomizableObjectLayout>> Layouts;
+
+protected:
+
+	virtual bool Equals(const UCustomizableObjectNodePinData& Other) const override;
+
+private:
+
+	UPROPERTY()
+	int32 SectionIndex = -1;
+
+};
+
+
+UCLASS()
+class CUSTOMIZABLEOBJECTEDITOR_API UCustomizableObjectNodeMeshParameter : public UCustomizableObjectNodeMesh
 {
 public:
 	GENERATED_BODY()
@@ -23,15 +50,15 @@ public:
 	UPROPERTY(EditAnywhere, Category=CustomizableObject)
 	TSoftObjectPtr<USkeletalMesh> DefaultValue;
 
-	UPROPERTY(EditAnywhere, Category = CustomizableObject)
-	int32 DefaultValueSection = 0;
-
 	UPROPERTY(EditAnywhere, Category=CustomizableObject)
 	FString ParameterName = "Mesh Param";
 
 	UPROPERTY(EditAnywhere, Category = UI, meta = (DisplayName = "Parameter UI Metadata"))
 	FMutableParamUIMetadata ParamUIMetadata;
 
+	// UObject interface.
+	virtual void PostEditChangeProperty(FPropertyChangedEvent& PropertyChangedEvent) override;
+
 	// EdGraphNode interface
 	virtual FText GetNodeTitle(ENodeTitleType::Type TitleType) const override;
 	virtual FLinearColor GetNodeTitleColor() const override;
@@ -43,5 +70,20 @@ public:
 	virtual void AllocateDefaultPins(UCustomizableObjectNodeRemapPins* RemapPins) override;
 	virtual bool IsAffectedByLOD() const override { return false; }
 	virtual bool IsExperimental() const override;
+
+	// UCustomizableObjectNodeMesh interface
+	virtual UTexture2D* FindTextureForPin(const UEdGraphPin* Pin) const override;
+	virtual TArray<UCustomizableObjectLayout*> GetLayouts(const UEdGraphPin& MeshPin) const override;
+	virtual TSoftObjectPtr<UObject> GetMesh() const override;
+	virtual UEdGraphPin* GetMeshPin(int32 LOD, int32 SectionIndex) const override;
+	virtual void GetPinSection(const UEdGraphPin& Pin, int32& OutLODIndex, int32& OutSectionIndex, int32& OutLayoutIndex) const override;
+
+	// Own interface
+	void GetPinSection(const UEdGraphPin& Pin, int32& OutSectionIndex) const;
+	UMaterialInterface* GetMaterialInterfaceFor(const int32 SectionIndex) const;
+	FSkeletalMaterial* GetSkeletalMaterialFor(const int32 SectionIndex) const;
+	int32 GetSkeletalMaterialIndexFor(const UEdGraphPin& Pin) const;
+	int32 GetSkeletalMaterialIndexFor(const int32 SectionIndex) const;
+
 };
 
diff --git a/Engine/Plugins/Mutable/Source/CustomizableObjectEditor/Private/MuCOE/Nodes/CustomizableObjectNodeMeshParameterDetails.h b/Engine/Plugins/Mutable/Source/CustomizableObjectEditor/Private/MuCOE/Nodes/CustomizableObjectNodeMeshParameterDetails.h
new file mode 100644
--- /dev/null
+++ b/Engine/Plugins/Mutable/Source/CustomizableObjectEditor/Private/MuCOE/Nodes/CustomizableObjectNodeMeshParameterDetails.h
@@ -0,0 +1,27 @@
+// Copyright Epic Games, Inc. All Rights Reserved.
+
+#pragma once
+
+#include "CustomizableObjectNodeDetails.h"
+#include "IDetailCustomization.h"
+
+class IDetailLayoutBuilder;
+struct FLayoutEditorMeshSection;
+
+
+class FCustomizableObjectNodeMeshParameterDetails : public FCustomizableObjectNodeDetails
+{
+public:
+	/** Makes a new instance of this detail layout class for a specific detail view requesting it */
+	static TSharedRef<IDetailCustomization> MakeInstance();
+
+	/** ILayoutDetails interface */
+	virtual void CustomizeDetails(IDetailLayoutBuilder& DetailBuilder) override;
+
+private:
+
+	void GenerateMeshSectionOptions(TArray<FLayoutEditorMeshSection>& OutMeshSections);
+
+	// Pointer to the node represented in this details
+	class UCustomizableObjectNodeMeshParameter* Node;
+};
diff --git a/Engine/Plugins/Mutable/Source/CustomizableObjectEditor/Private/MuCOE/Nodes/CustomizableObjectNodeMeshParametersDetails.cpp b/Engine/Plugins/Mutable/Source/CustomizableObjectEditor/Private/MuCOE/Nodes/CustomizableObjectNodeMeshParametersDetails.cpp
new file mode 100644
--- /dev/null
+++ b/Engine/Plugins/Mutable/Source/CustomizableObjectEditor/Private/MuCOE/Nodes/CustomizableObjectNodeMeshParametersDetails.cpp
@@ -0,0 +1,89 @@
+// Copyright Epic Games, Inc. All Rights Reserved.
+
+#include "MuCOE/Nodes/CustomizableObjectNodeMeshParameterDetails.h"
+
+#include "DetailLayoutBuilder.h"
+#include "DetailWidgetRow.h"
+#include "IDetailsView.h"
+#include "MuCOE/Nodes/CustomizableObjectNodeMeshParameter.h"
+#include "MuCOE/PinViewer/SPinViewer.h"
+#include "MuCOE/SCustomizableObjectLayoutEditor.h"
+
+#define LOCTEXT_NAMESPACE "CustomizableObjectNodeMaterialDetails"
+
+
+TSharedRef<IDetailCustomization> FCustomizableObjectNodeMeshParameterDetails::MakeInstance()
+{
+	return MakeShareable(new FCustomizableObjectNodeMeshParameterDetails);
+}
+
+
+void FCustomizableObjectNodeMeshParameterDetails::CustomizeDetails(IDetailLayoutBuilder& DetailBuilder)
+{
+	FCustomizableObjectNodeDetails::CustomizeDetails(DetailBuilder);
+
+	TSharedPtr<const IDetailsView> DetailsView = DetailBuilder.GetDetailsViewSharedPtr();
+
+    if (DetailsView && DetailsView->GetSelectedObjects().Num())
+	{
+		Node = Cast<UCustomizableObjectNodeMeshParameter>(DetailsView->GetSelectedObjects()[0].Get());
+	}
+
+    if (!Node)
+    {
+        return;
+    }
+
+    IDetailCategoryBuilder& CustomizableObject = DetailBuilder.EditCategory("CustomizableObject");
+
+	TSharedRef<IPropertyHandle> DefaultValueProperty = DetailBuilder.GetProperty(GET_MEMBER_NAME_CHECKED(UCustomizableObjectNodeMeshParameter, DefaultValue));
+
+	TArray<FLayoutEditorMeshSection> MeshSectionsAndLayouts;
+	GenerateMeshSectionOptions(MeshSectionsAndLayouts);
+
+	TSharedPtr<SCustomizableObjectLayoutEditor> LayoutBlocksEditor = SNew(SCustomizableObjectLayoutEditor)
+		.Node(Node)
+		.MeshSections(MeshSectionsAndLayouts);
+
+	FCustomizableObjectLayoutEditorDetailsBuilder LayoutEditorBuilder;
+	LayoutEditorBuilder.LayoutEditor = LayoutBlocksEditor;
+	LayoutEditorBuilder.bShowLayoutSelector = true;
+	LayoutEditorBuilder.bShowPackagingStrategy = true;
+	LayoutEditorBuilder.bShowAutomaticGenerationSettings = true;
+	LayoutEditorBuilder.bShowGridSize = true;
+	LayoutEditorBuilder.bShowMaxGridSize = true;
+	LayoutEditorBuilder.bShowReductionMethods = true;
+	LayoutEditorBuilder.bShowWarningSettings = true;
+
+	LayoutEditorBuilder.CustomizeDetails(DetailBuilder);
+
+	LayoutBlocksEditor->UpdateLayout(nullptr);
+}
+
+
+void FCustomizableObjectNodeMeshParameterDetails::GenerateMeshSectionOptions(TArray<FLayoutEditorMeshSection>& OutMeshSections)
+{
+	OutMeshSections.Empty();
+
+	if (!Node)
+	{
+		return;
+	}
+
+	for (const UEdGraphPin* Pin : Node->GetAllNonOrphanPins())
+	{
+		if (const UCustomizableObjectNodeMeshParameterPinDataSection* PinData = Cast<UCustomizableObjectNodeMeshParameterPinDataSection>(Node->GetPinData(*Pin)))
+		{
+			FLayoutEditorMeshSection& MeshSection = OutMeshSections.AddDefaulted_GetRef();
+			MeshSection.MeshName = MakeShareable(new FString(Pin->PinFriendlyName.ToString()));
+
+			for (UCustomizableObjectLayout* Layout : PinData->Layouts)
+			{
+				MeshSection.Layouts.Add(Layout);
+			}
+		}
+	}
+}
+
+
+#undef LOCTEXT_NAMESPACE
diff --git a/Engine/Plugins/Mutable/Source/CustomizableObjectEditor/Private/MuCOE/SMutableCodeViewer.cpp b/Engine/Plugins/Mutable/Source/CustomizableObjectEditor/Private/MuCOE/SMutableCodeViewer.cpp
--- a/Engine/Plugins/Mutable/Source/CustomizableObjectEditor/Private/MuCOE/SMutableCodeViewer.cpp
+++ b/Engine/Plugins/Mutable/Source/CustomizableObjectEditor/Private/MuCOE/SMutableCodeViewer.cpp
@@ -1527,7 +1527,7 @@ void SMutableCodeViewer::GenerateElementRecursive(const int32& InStateIndex, mu:
 
 		break;
 	}
-		
+
 	case mu::EOpType::ME_APPLYLAYOUT:
 	{
 		mu::OP::MeshApplyLayoutArgs Args = InProgram.GetOpArgs<mu::OP::MeshApplyLayoutArgs>(InParentAddress);
@@ -1535,7 +1535,15 @@ void SMutableCodeViewer::GenerateElementRecursive(const int32& InStateIndex, mu:
 		AddOpFunc(Args.Mesh, TEXT("mesh "));
 		break;
 	}
-		
+
+	case mu::EOpType::ME_PREPARELAYOUT:
+	{
+		mu::OP::MeshPrepareLayoutArgs Args = InProgram.GetOpArgs<mu::OP::MeshPrepareLayoutArgs>(InParentAddress);
+		AddOpFunc(Args.Layout, TEXT("layout "));
+		AddOpFunc(Args.Mesh, TEXT("mesh "));
+		break;
+	}
+
 	case mu::EOpType::ME_DIFFERENCE:
 	{
 		const uint8* data = InProgram.GetOpArgsPointer(InParentAddress);
diff --git a/Engine/Plugins/Mutable/Source/CustomizableObjectEditor/Private/MuCOE/SMutableMeshViewer.cpp b/Engine/Plugins/Mutable/Source/CustomizableObjectEditor/Private/MuCOE/SMutableMeshViewer.cpp
--- a/Engine/Plugins/Mutable/Source/CustomizableObjectEditor/Private/MuCOE/SMutableMeshViewer.cpp
+++ b/Engine/Plugins/Mutable/Source/CustomizableObjectEditor/Private/MuCOE/SMutableMeshViewer.cpp
@@ -6,6 +6,7 @@
 #include "Framework/Views/TableViewMetadata.h"
 #include "MuCOE/SMutableMeshViewport.h"
 #include "MuCOE/SMutableSkeletonViewer.h"
+#include "MuCO/CustomizableObject.h"
 #include "MuT/TypeInfo.h"
 #include "PropertyCustomizationHelpers.h"
 #include "Widgets/Layout/SScrollBox.h"
@@ -227,6 +228,10 @@ void SMutableMeshViewer::SetMesh(const TSharedPtr<const mu::FMesh>& InMesh)
 			{
 				MutableSkeletonViewer->SetVisibility(EVisibility::Hidden);
 			}
+
+			FString DebugLog;
+			MutableMesh->Log(DebugLog, 8);
+			UE_LOG(LogMutable, Warning, TEXT("[%s]"), *DebugLog);
 		}
 
 		MeshViewport->SetMesh(MutableMesh);
diff --git a/Engine/Plugins/Mutable/Source/MutableRuntime/Internal/MuR/OpMeshPrepareLayout.h b/Engine/Plugins/Mutable/Source/MutableRuntime/Internal/MuR/OpMeshPrepareLayout.h
new file mode 100644
--- /dev/null
+++ b/Engine/Plugins/Mutable/Source/MutableRuntime/Internal/MuR/OpMeshPrepareLayout.h
@@ -0,0 +1,23 @@
+// Copyright Epic Games, Inc. All Rights Reserved.
+
+#pragma once
+
+#include "MuR/Ptr.h"
+#include "HAL/Platform.h"
+
+namespace mu
+{
+	class FMesh;
+	class FLayout;
+
+	MUTABLERUNTIME_API void MeshPrepareLayout(
+		FMesh& Mesh,
+		const FLayout& InLayout,
+		int32 LayoutChannel,
+		bool bNormalizeUVs,
+		bool bClampUVIslands,
+		bool bEnsureAllVerticesHaveLayoutBlock,
+		bool bUseAbsoluteBlockIds
+	);
+
+}
diff --git a/Engine/Plugins/Mutable/Source/MutableRuntime/Internal/MuR/Operations.h b/Engine/Plugins/Mutable/Source/MutableRuntime/Internal/MuR/Operations.h
--- a/Engine/Plugins/Mutable/Source/MutableRuntime/Internal/MuR/Operations.h
+++ b/Engine/Plugins/Mutable/Source/MutableRuntime/Internal/MuR/Operations.h
@@ -206,6 +206,9 @@ namespace mu
         //! Apply a layout to a mesh texture coordinates channel
         ME_APPLYLAYOUT,
 
+		/** */
+		ME_PREPARELAYOUT,
+
         //! Compare two meshes and extract a morph from the first to the second
         //! The meshes must have the same topology, etc.
         ME_DIFFERENCE,
@@ -426,6 +429,7 @@ namespace mu
 			ADDRESS variable;
 			uint8 LOD = 0;
 			uint8 Section = 0;
+			uint32 MeshID = 0;
 		};
 
         struct ConditionalArgs
@@ -794,6 +798,17 @@ namespace mu
             uint16 Channel;
         };
 
+		struct MeshPrepareLayoutArgs
+		{
+			ADDRESS Mesh;
+			ADDRESS Layout;
+			uint8 LayoutChannel = 0;
+			uint8 bUseAbsoluteBlockIds : 1;
+			uint8 bNormalizeUVs : 1;
+			uint8 bClampUVIslands : 1;
+			uint8 bEnsureAllVerticesHaveLayoutBlock : 1;
+		};
+		
         struct MeshMergeArgs
         {
             ADDRESS Base;
diff --git a/Engine/Plugins/Mutable/Source/MutableRuntime/Private/MuR/CodeRunner.cpp b/Engine/Plugins/Mutable/Source/MutableRuntime/Private/MuR/CodeRunner.cpp
--- a/Engine/Plugins/Mutable/Source/MutableRuntime/Private/MuR/CodeRunner.cpp
+++ b/Engine/Plugins/Mutable/Source/MutableRuntime/Private/MuR/CodeRunner.cpp
@@ -35,6 +35,7 @@
 #include "MuR/OpLayoutPack.h"
 #include "MuR/OpLayoutRemoveBlocks.h"
 #include "MuR/OpMeshApplyLayout.h"
+#include "MuR/OpMeshPrepareLayout.h"
 #include "MuR/OpMeshApplyPose.h"
 #include "MuR/OpMeshBind.h"
 #include "MuR/OpMeshClipDeform.h"
@@ -1128,12 +1129,12 @@ namespace mu
            		MUTABLE_CPUPROFILER_SCOPE(ME_APPLYLAYOUT)
             		
                 TSharedPtr<const FMesh> Base = LoadMesh(FCacheAddress(Args.Mesh, Item));
+                TSharedPtr<const FLayout> Layout = LoadLayout(FCacheAddress(Args.Layout, Item));
 
                 if (Base)
                 {
 					TSharedPtr<FMesh> Result = CloneOrTakeOver(Base);
-
-                    TSharedPtr<const FLayout> Layout = LoadLayout(FCacheAddress(Args.Layout, Item));
+						
                     int32 texCoordsSet = Args.Channel;
 
                     MeshApplyLayout(Result.Get(), Layout.Get(), texCoordsSet);
@@ -1156,6 +1157,50 @@ namespace mu
             break;
         }
 
+
+		case EOpType::ME_PREPARELAYOUT:
+		{
+			OP::MeshPrepareLayoutArgs Args = Program.GetOpArgs<OP::MeshPrepareLayoutArgs>(Item.At);
+			switch (Item.Stage)
+			{
+			case 0:
+			{
+				AddOp(FScheduledOp(Item.At, Item, 1),
+					FScheduledOp(Args.Mesh, Item),
+					FScheduledOp(Args.Layout, Item));
+				break;
+			}
+			case 1:
+			{
+				MUTABLE_CPUPROFILER_SCOPE(ME_PREPARELAYOUT)
+
+				TSharedPtr<const FMesh> Base = LoadMesh(FCacheAddress(Args.Mesh, Item));
+				TSharedPtr<const FLayout> Layout = LoadLayout(FCacheAddress(Args.Layout, Item));
+
+				if (Base && Layout)
+				{
+					TSharedPtr<FMesh> Result = CloneOrTakeOver(Base);
+
+					MeshPrepareLayout(*Result, *Layout, Args.LayoutChannel, Args.bNormalizeUVs, Args.bClampUVIslands, Args.bEnsureAllVerticesHaveLayoutBlock, Args.bUseAbsoluteBlockIds );
+
+					StoreMesh(Item, Result);
+				}
+				else
+				{
+					StoreMesh(Item, Base);
+				}
+
+
+				break;
+			}
+
+			default:
+				check(false);
+			}
+
+			break;
+		}
+
         case EOpType::ME_DIFFERENCE:
         {
 			const uint8* data = Program.GetOpArgsPointer(Item.At);
diff --git a/Engine/Plugins/Mutable/Source/MutableRuntime/Private/MuR/CodeRunnerTasked.cpp b/Engine/Plugins/Mutable/Source/MutableRuntime/Private/MuR/CodeRunnerTasked.cpp
--- a/Engine/Plugins/Mutable/Source/MutableRuntime/Private/MuR/CodeRunnerTasked.cpp
+++ b/Engine/Plugins/Mutable/Source/MutableRuntime/Private/MuR/CodeRunnerTasked.cpp
@@ -2352,7 +2352,7 @@ namespace mu
 	{
 	public:
 
-		FMeshExternalLoadTask(const FScheduledOp&, CodeRunner::FExternalResourceId, int32 InLODIndex, int32 InSectionIndex);
+		FMeshExternalLoadTask(const FScheduledOp&, CodeRunner::FExternalResourceId, int32 InLODIndex, int32 InSectionIndex, uint32 MeshID);
 
 		// FIssuedTask interface
 		virtual bool Prepare(CodeRunner*, bool& bOutFailed) override;
@@ -2363,6 +2363,7 @@ namespace mu
 		CodeRunner::FExternalResourceId Id;
 		int32 LODIndex = 0;
 		int32 SectionIndex = 0;
+		uint32 MeshID = 0;
 
 		TSharedPtr<FMesh> Result;
 
@@ -2370,11 +2371,12 @@ namespace mu
 	};
 
 
-	FMeshExternalLoadTask::FMeshExternalLoadTask(const FScheduledOp& InOp, CodeRunner::FExternalResourceId InId, int32 InLODIndex, int32 InSectionIndex)
+	FMeshExternalLoadTask::FMeshExternalLoadTask(const FScheduledOp& InOp, CodeRunner::FExternalResourceId InId, int32 InLODIndex, int32 InSectionIndex, uint32 InMeshID)
 		: FIssuedTask(InOp)
 		, Id(InId)
 		, LODIndex(InLODIndex)
 		, SectionIndex(InSectionIndex)
+		, MeshID(InMeshID)
 	{
 	}
 
@@ -2410,6 +2412,7 @@ namespace mu
 			Invoke(ExternalCleanUpFunc);
 		}
 
+		Result->MeshIDPrefix = MeshID;
 		Runner->StoreMesh(Op, Result);
 
 		bool bSuccess = true;
@@ -2535,7 +2538,7 @@ namespace mu
 
 			CodeRunner::FExternalResourceId FullId;
 			FullId.ParameterId = Id;
-			Issued = MakeShared<FMeshExternalLoadTask>(item, FullId, LODIndex, SectionIndex);
+			Issued = MakeShared<FMeshExternalLoadTask>(item, FullId, LODIndex, SectionIndex, Args.MeshID);
 
 			break;
 		}
@@ -2555,8 +2558,9 @@ namespace mu
 				// TODO
 				int32 LODIndex = 0;
 				int32 SectionIndex = 0;
+				int32 MeshID = 0;
 
-				Issued = MakeShared<FMeshExternalLoadTask>(item, FullId, LODIndex, SectionIndex);
+				Issued = MakeShared<FMeshExternalLoadTask>(item, FullId, LODIndex, SectionIndex, MeshID);
 			}
 
 			break;
diff --git a/Engine/Plugins/Mutable/Source/MutableRuntime/Private/MuR/Mesh.cpp b/Engine/Plugins/Mutable/Source/MutableRuntime/Private/MuR/Mesh.cpp
--- a/Engine/Plugins/Mutable/Source/MutableRuntime/Private/MuR/Mesh.cpp
+++ b/Engine/Plugins/Mutable/Source/MutableRuntime/Private/MuR/Mesh.cpp
@@ -10,6 +10,7 @@
 #include "MuR/MeshPrivate.h"
 #include "MuR/MutableTrace.h"
 #include "MuR/OpMeshClipWithMesh.h"
+#include "Spatial/PointHashGrid3.h"
 
 namespace mu
 {
@@ -1568,7 +1569,7 @@ namespace
 }
 
 
-void FMesh::Log( FString& out, int32 BufferElementLimit)
+void FMesh::Log( FString& out, int32 BufferElementLimit) const
 {
     out += "Mesh:\n";
 
@@ -1579,5 +1580,167 @@ void FMesh::Log( FString& out, int32 BufferElementLimit)
     LogBuffer( out, VertexBuffers, BufferElementLimit);
 }
 
+
+void GetUVIsland(TArray<FTriangleInfo>& InTriangles,
+	const uint32 InFirstTriangle,
+	TArray<uint32>& OutTriangleIndices,
+	const TArray<FVector2f>& InUVs,
+	const TMultiMap<int32, uint32>& InVertexToTriangleMap)
+{
+	MUTABLE_CPUPROFILER_SCOPE(GetUVIsland);
+
+	const uint32 NumTriangles = (uint32)InTriangles.Num();
+
+	OutTriangleIndices.Reserve(NumTriangles);
+	OutTriangleIndices.Add(InFirstTriangle);
+
+	TArray<bool> SkipTriangles;
+	SkipTriangles.Init(false, NumTriangles);
+
+	TArray<uint32> PendingTriangles;
+	PendingTriangles.Reserve(NumTriangles / 64);
+	PendingTriangles.Add(InFirstTriangle);
+
+	while (!PendingTriangles.IsEmpty())
+	{
+		const uint32 TriangleIndex = PendingTriangles.Pop();
+
+		// Triangle about to be proccessed, mark as skip;
+		SkipTriangles[TriangleIndex] = true;
+
+		bool ConnectedEdges[3] = { false, false, false };
+
+		const FTriangleInfo& Triangle = InTriangles[TriangleIndex];
+
+		// Find Triangles connected to edges 0 and 2
+		int32 CollapsedVertex1 = Triangle.CollapsedIndices[1];
+		int32 CollapsedVertex2 = Triangle.CollapsedIndices[2];
+
+		TArray<uint32> FoundTriangleIndices;
+		InVertexToTriangleMap.MultiFind(Triangle.CollapsedIndices[0], FoundTriangleIndices);
+
+		for (uint32 OtherTriangleIndex : FoundTriangleIndices)
+		{
+			const FTriangleInfo& OtherTriangle = InTriangles[OtherTriangleIndex];
+
+			for (int32 OtherIndex = 0; OtherIndex < 3; ++OtherIndex)
+			{
+				const int32 OtherCollapsedIndex = OtherTriangle.CollapsedIndices[OtherIndex];
+				if (OtherCollapsedIndex == CollapsedVertex1)
+				{
+					// Check if the vertex is in the same UV Island 
+					if (!SkipTriangles[OtherTriangleIndex]
+						&& InUVs[Triangle.Indices[1]].Equals(InUVs[OtherTriangle.Indices[OtherIndex]], 0.00001f))
+					{
+						OutTriangleIndices.Add(OtherTriangleIndex);
+						PendingTriangles.Add(OtherTriangleIndex);
+						SkipTriangles[OtherTriangleIndex] = true;
+					}
+
+					// Connected but already processed or in another island
+					break;
+				}
+
+				if (OtherCollapsedIndex == CollapsedVertex2)
+				{
+					// Check if the vertex is in the same UV Island 
+					if (!SkipTriangles[OtherTriangleIndex]
+						&& InUVs[Triangle.Indices[2]].Equals(InUVs[OtherTriangle.Indices[OtherIndex]], 0.00001f))
+					{
+						OutTriangleIndices.Add(OtherTriangleIndex);
+						PendingTriangles.Add(OtherTriangleIndex);
+						SkipTriangles[OtherTriangleIndex] = true;
+					}
+
+					// Connected but already processed or in another UV Island
+					break;
+				}
+			}
+
+		}
+
+		// Find the triangle connected to edge 1
+		FoundTriangleIndices.Reset();
+		InVertexToTriangleMap.MultiFind(CollapsedVertex1, FoundTriangleIndices);
+
+		for (uint32 OtherTriangleIndex : FoundTriangleIndices)
+		{
+			const FTriangleInfo& OtherTriangle = InTriangles[OtherTriangleIndex];
+
+			for (int32 OtherIndex = 0; OtherIndex < 3; ++OtherIndex)
+			{
+				const int32 OtherCollapsedIndex = OtherTriangle.CollapsedIndices[OtherIndex];
+				if (OtherCollapsedIndex == CollapsedVertex2)
+				{
+					// Check if the vertex belong to the same UV island
+					if (!SkipTriangles[OtherTriangleIndex]
+						&& InUVs[Triangle.Indices[2]].Equals(InUVs[OtherTriangle.Indices[OtherIndex]], 0.00001f))
+					{
+						OutTriangleIndices.Add(OtherTriangleIndex);
+						PendingTriangles.Add(OtherTriangleIndex);
+						SkipTriangles[OtherTriangleIndex] = true;
+					}
+
+					// Connected but already processed or in another island
+					break;
+				}
+			}
+		}
+	}
+}
+
+
+void MeshCreateCollapsedVertexMap(const mu::FMesh* Mesh, TArray<int32>& CollapsedVertices)
+{
+	MUTABLE_CPUPROFILER_SCOPE(LayoutUV_CreateCollapsedVertexMap);
+
+	const int32 NumVertices = Mesh->GetVertexCount();
+	CollapsedVertices.Reserve(NumVertices);
+
+	UE::Geometry::TPointHashGrid3f<int32> VertHash(0.01f, INDEX_NONE);
+	VertHash.Reserve(NumVertices);
+
+	TArray<FVector3f> Vertices;
+	Vertices.SetNumUninitialized(NumVertices);
+
+	mu::UntypedMeshBufferIteratorConst ItPosition = mu::UntypedMeshBufferIteratorConst(Mesh->GetVertexBuffers(), mu::EMeshBufferSemantic::Position);
+
+	FVector3f* VertexData = Vertices.GetData();
+
+	for (int32 VertexIndex = 0; VertexIndex < NumVertices; ++VertexIndex)
+	{
+		*VertexData = ItPosition.GetAsVec3f();
+		VertHash.InsertPointUnsafe(VertexIndex, *VertexData);
+
+		++ItPosition;
+		++VertexData;
+	}
+
+	// Find unique vertices
+	CollapsedVertices.Init(INDEX_NONE, NumVertices);
+
+	TArray<int32> NearbyVertices;
+	for (int32 VertexIndex = 0; VertexIndex < NumVertices; ++VertexIndex)
+	{
+		if (CollapsedVertices[VertexIndex] != INDEX_NONE)
+		{
+			continue;
+		}
+
+		const FVector3f& Vertex = Vertices[VertexIndex];
+
+		NearbyVertices.Reset();
+		VertHash.FindPointsInBall(Vertex, 0.00001,
+			[&Vertex, &Vertices](const int32& Other) -> float {return FVector3f::DistSquared(Vertices[Other], Vertex); },
+			NearbyVertices);
+
+		// Find equals
+		for (int32 NearbyVertexIndex : NearbyVertices)
+		{
+			CollapsedVertices[NearbyVertexIndex] = VertexIndex;
+		}
+	}
+}
+
 	
 }
diff --git a/Engine/Plugins/Mutable/Source/MutableRuntime/Private/MuR/OpMeshPrepareLayout.cpp b/Engine/Plugins/Mutable/Source/MutableRuntime/Private/MuR/OpMeshPrepareLayout.cpp
new file mode 100644
--- /dev/null
+++ b/Engine/Plugins/Mutable/Source/MutableRuntime/Private/MuR/OpMeshPrepareLayout.cpp
@@ -0,0 +1,521 @@
+// Copyright Epic Games, Inc. All Rights Reserved.
+
+#include "MuR/OpMeshPrepareLayout.h"
+
+#include "MuR/MeshPrivate.h"
+#include "MuR/MutableTrace.h"
+#include "MuR/Platform.h"
+#include "MuR/Mesh.h"
+#include "MuR/Layout.h"
+#include "MuR/Image.h"
+#include "MuR/MutableRuntimeModule.h"
+
+
+namespace mu
+{
+
+	template<typename T>
+	struct TArray2D
+	{
+		int32 SizeX = 0;
+		int32 SizeY = 0;
+		TArray<T> Data;
+
+		void Init(const T& Value, int32 InSizeX, int32 InSizeY)
+		{
+			SizeX = InSizeX;
+			SizeY = InSizeY;
+			Data.Init(Value, SizeX * SizeY);
+		}
+
+		inline const T& Get(int32 X, int32 Y) const
+		{
+			check(X >= 0 && X < SizeX);
+			check(Y >= 0 && Y < SizeY);
+			return Data[SizeX * Y + X];
+		}
+
+		inline void Set(int32 X, int32 Y, const T& Value)
+		{
+			check(X >= 0 && X < SizeX);
+			check(Y >= 0 && Y < SizeY);
+			Data[SizeX * Y + X] = Value;
+		}
+
+	};
+
+
+	void MeshPrepareLayout(
+		FMesh& Mesh,
+		const FLayout& InLayout,
+		int32 LayoutChannel,
+		bool bNormalizeUVs,
+		bool bClampUVIslands,
+		bool bEnsureAllVerticesHaveLayoutBlock,
+		bool bUseAbsoluteBlockIds
+	)
+	{
+		MUTABLE_CPUPROFILER_SCOPE(MeshPrepareLayout);
+
+		if (Mesh.GetVertexCount() == 0)
+		{
+			return;
+		}
+
+		TSharedPtr<FLayout> Layout = InLayout.Clone();
+
+		// The layout must have block ids.
+		check(Layout->Blocks.IsEmpty() || Layout->Blocks[0].Id != FLayoutBlock::InvalidBlockId);
+
+		Mesh.AddLayout(Layout);
+
+		const int32 NumVertices = Mesh.GetVertexCount();
+		const int32 NumBlocks = Layout->GetBlockCount();
+
+		// Find block ids for each block in the grid. Calculate a grid size that contains all blocks
+		FIntPoint LayoutGrid = Layout->GetGridSize();
+		FIntPoint WorkingGrid = LayoutGrid;
+		for (const FLayoutBlock& Block : Layout->Blocks)
+		{
+			WorkingGrid.X = FMath::Max(WorkingGrid.X, Block.Min.X + Block.Size.X);
+			WorkingGrid.Y = FMath::Max(WorkingGrid.Y, Block.Min.Y + Block.Size.Y);
+		}
+
+
+		TArray2D<int32> GridBlockBlockId;
+		GridBlockBlockId.Init(MAX_uint16, WorkingGrid.X, WorkingGrid.Y);
+
+		// 
+		TArray<box<FVector2f>> BlockRects;
+		BlockRects.SetNumUninitialized(NumBlocks);
+
+		// Create an array of block index per cell
+		TArray<int32> OverlappingBlocks;
+		for (int32 BlockIndex = 0; BlockIndex < NumBlocks; ++BlockIndex)
+		{
+			// TODO
+			bool bBlockHasMask = false; // GeneratedLayout.Source->Blocks[BlockIndex].Mask != nullptr;
+
+			// Fill the block rect
+			FIntVector2 Min = Layout->Blocks[BlockIndex].Min;
+			FIntVector2 Size = Layout->Blocks[BlockIndex].Size;
+
+			box<FVector2f>& BlockRect = BlockRects[BlockIndex];
+			BlockRect.min[0] = float(Min.X) / float(LayoutGrid.X);
+			BlockRect.min[1] = float(Min.Y) / float(LayoutGrid.Y);
+			BlockRect.size[0] = float(Size.X) / float(LayoutGrid.X);
+			BlockRect.size[1] = float(Size.Y) / float(LayoutGrid.Y);
+
+			// Fill the block index per cell array
+			// Ignore the block in this stage if it has a mask, because blocks with masks will very likely overlap other blocks
+			if (!bBlockHasMask)
+			{
+				for (uint16 Y = Min.Y; Y < Min.Y + Size.Y; ++Y)
+				{
+					for (uint16 X = Min.X; X < Min.X + Size.X; ++X)
+					{
+						if (GridBlockBlockId.Get(X, Y) == MAX_uint16)
+						{
+							GridBlockBlockId.Set(X, Y, BlockIndex);
+						}
+						else
+						{
+							OverlappingBlocks.AddUnique(BlockIndex);
+						}
+					}
+				}
+			}
+		}
+
+
+		// Notify Overlapping layout blocks
+		if (!OverlappingBlocks.IsEmpty())
+		{
+			UE_LOG(LogMutableCore, Warning, TEXT("Mesh has %d layout block overlapping in channel %d."), OverlappingBlocks.Num() + 1, LayoutChannel);
+		}
+
+		// Get the information about the texture coordinates channel
+		int32 TexCoordsBufferIndex = -1;
+		int32 TexCoordsChannelIndex = -1;
+		Mesh.GetVertexBuffers().FindChannel(EMeshBufferSemantic::TexCoords, LayoutChannel, &TexCoordsBufferIndex, &TexCoordsChannelIndex);
+
+
+		if (TexCoordsBufferIndex < 0 || TexCoordsChannelIndex < 0)
+		{
+			// This is actually fine when using shared materials across LODs
+			UE_LOG(LogMutableCore, Warning, TEXT("Trying to generate layout for missing UV channel %d. No layout blocks generated."), LayoutChannel);
+			return;
+		}
+
+		check(TexCoordsBufferIndex >= 0);
+		check(TexCoordsChannelIndex >= 0);
+
+		const FMeshBufferChannel& TexCoordsChannel = Mesh.VertexBuffers.Buffers[TexCoordsBufferIndex].Channels[TexCoordsChannelIndex];
+		check(TexCoordsChannel.Semantic == EMeshBufferSemantic::TexCoords);
+
+		uint8* TexCoordData = Mesh.GetVertexBuffers().GetBufferData(TexCoordsBufferIndex);
+		int32 elemSize = Mesh.GetVertexBuffers().GetElementSize(TexCoordsBufferIndex);
+		int32 channelOffset = TexCoordsChannel.Offset;
+		TexCoordData += channelOffset;
+
+		// Get a copy of the UVs as FVector2f to work with them. 
+		TArray<FVector2f> TexCoords;
+		{
+			TexCoords.SetNumUninitialized(NumVertices);
+
+			bool bNonNormalizedUVs = false;
+			const bool bIsOverlayLayout = Layout->GetLayoutPackingStrategy() == mu::EPackStrategy::Overlay;
+
+			const uint8* pVertices = TexCoordData;
+			for (int32 VertexIndex = 0; VertexIndex < NumVertices; ++VertexIndex)
+			{
+				FVector2f& UV = TexCoords[VertexIndex];
+				if (TexCoordsChannel.Format == EMeshBufferFormat::Float32)
+				{
+					UV = *((FVector2f*)pVertices);
+				}
+				else if (TexCoordsChannel.Format == EMeshBufferFormat::Float16)
+				{
+					const FFloat16* pUV = reinterpret_cast<const FFloat16*>(pVertices);
+					UV = FVector2f(float(pUV[0]), float(pUV[1]));
+				}
+
+				// Check that UVs are normalized. If not, clamp the values and throw a warning.
+				if (bNormalizeUVs && !bIsOverlayLayout
+					&& (UV[0] < 0.f || UV[0] > 1.f || UV[1] < 0.f || UV[1] > 1.f))
+				{
+					UV[0] = FMath::Clamp(UV[0], 0.f, 1.f);
+					UV[1] = FMath::Clamp(UV[1], 0.f, 1.f);
+					bNonNormalizedUVs = true;
+				}
+
+				pVertices += elemSize;
+			}
+
+			// Mutable does not support non-normalized UVs
+			if (bNonNormalizedUVs && !bIsOverlayLayout)
+			{
+				UE_LOG(LogMutableCore, Warning, TEXT("Source mesh has non-normalized UVs index %d"), LayoutChannel);
+			}
+		}
+
+
+		const int32 NumTriangles = Mesh.GetIndexCount() / 3;
+		TArray<FTriangleInfo> Triangles;
+
+		// Vertices mapped to unique vertex index
+		TArray<int32> CollapsedVertices;
+
+		// Vertex to face map used to speed up connectivity building
+		TMultiMap<int32, uint32> VertexToFaceMap;
+
+		// Find Unique Vertices
+		if (bClampUVIslands)
+		{
+			VertexToFaceMap.Reserve(NumVertices);
+			Triangles.SetNumUninitialized(NumTriangles);
+
+			MeshCreateCollapsedVertexMap(&Mesh, CollapsedVertices);
+		}
+
+		TArray<int32> ConflictiveTriangles;
+
+		const uint32 MaxGridX = bNormalizeUVs ? MAX_uint32 : WorkingGrid.X - 1;
+		const uint32 MaxGridY = bNormalizeUVs ? MAX_uint32 : WorkingGrid.Y - 1;
+
+		// Allocate the per-vertex layout block data
+		TArray<uint16> LayoutData;
+		constexpr uint16 NullBlockId = MAX_uint16 - 1;
+		LayoutData.Init(NullBlockId, NumVertices);
+
+		UntypedMeshBufferIteratorConst ItIndices(Mesh.GetIndexBuffers(), EMeshBufferSemantic::VertexIndex);
+		for (int32 TriangleIndex = 0; TriangleIndex < NumTriangles; ++TriangleIndex)
+		{
+			uint32 Index0 = ItIndices.GetAsUINT32();
+			++ItIndices;
+			uint32 Index1 = ItIndices.GetAsUINT32();
+			++ItIndices;
+			uint32 Index2 = ItIndices.GetAsUINT32();
+			++ItIndices;
+
+			auto AssignOneVertex =
+				[LayoutGrid, MaxGridX, MaxGridY, NumBlocks, Layout, &GridBlockBlockId, &TexCoords, &LayoutData]
+				(int32 VertexIndex)
+				{
+					uint16& BlockIndex = LayoutData[VertexIndex];
+
+					// Was it previously assigned?
+					if (BlockIndex != NullBlockId)
+					{
+						return BlockIndex;
+					}
+
+					FVector2f UV = TexCoords[VertexIndex];
+
+					int32 VertexWorkingGridX = FMath::Clamp(LayoutGrid.X * UV[0], 0, LayoutGrid.X - 1);
+					int32 VertexWorkingGridY = FMath::Clamp(LayoutGrid.Y * UV[1], 0, LayoutGrid.Y - 1);
+
+					// First: Assign the vertices to masked blocks in order
+					for (int32 CandidateBlockIndex = 0; CandidateBlockIndex < NumBlocks; ++CandidateBlockIndex)
+					{
+						// TODO
+						const mu::FImage* Mask = nullptr; // GeneratedLayout.Source->Blocks[CandidateBlockIndex].Mask.Get();
+						if (Mask)
+						{
+							// First discard with block limits.
+							FIntVector2 Min = Layout->Blocks[CandidateBlockIndex].Min;
+							FIntVector2 Size = Layout->Blocks[CandidateBlockIndex].Size;
+
+							bool bInBlock =
+								(VertexWorkingGridX >= Min.X && VertexWorkingGridX < Min.X + Size.X)
+								&&
+								(VertexWorkingGridY >= Min.Y && VertexWorkingGridY < Min.Y + Size.Y);
+
+							if (bInBlock)
+							{
+								// TODO: This always clamps the UVs
+								FVector2f SampleUV;
+								SampleUV.X = FMath::Fmod(UV.X, 1.0);
+								SampleUV.Y = FMath::Fmod(UV.Y, 1.0);
+
+								FVector4f MaskValue = Mask->Sample(SampleUV);
+								if (MaskValue.X > 0.5f)
+								{
+									BlockIndex = CandidateBlockIndex;
+									break;
+								}
+							}
+						}
+					}
+
+					// Second: Assign to non-masked blocks if not assigned yet
+					if (BlockIndex == NullBlockId)
+					{
+						uint32 ClampedX = FMath::Min<uint32>(MaxGridX, FMath::Max<uint32>(0, VertexWorkingGridX));
+						uint32 ClampedY = FMath::Min<uint32>(MaxGridY, FMath::Max<uint32>(0, VertexWorkingGridY));
+						BlockIndex = GridBlockBlockId.Get(ClampedX, ClampedY);
+					}
+					return BlockIndex;
+				};
+
+			const uint16 BlockIndexV0 = AssignOneVertex(Index0);
+			const uint16 BlockIndexV1 = AssignOneVertex(Index1);
+			const uint16 BlockIndexV2 = AssignOneVertex(Index2);
+
+			if (bClampUVIslands)
+			{
+				if (BlockIndexV0 != BlockIndexV1 || BlockIndexV0 != BlockIndexV2)
+				{
+					ConflictiveTriangles.Add(TriangleIndex);
+				}
+
+				FTriangleInfo& Triangle = Triangles[TriangleIndex];
+
+				Triangle.Indices[0] = Index0;
+				Triangle.Indices[1] = Index1;
+				Triangle.Indices[2] = Index2;
+				Triangle.CollapsedIndices[0] = CollapsedVertices[Index0];
+				Triangle.CollapsedIndices[1] = CollapsedVertices[Index1];
+				Triangle.CollapsedIndices[2] = CollapsedVertices[Index2];
+
+				Triangle.BlockIndices[0] = BlockIndexV0;
+				Triangle.BlockIndices[1] = BlockIndexV1;
+				Triangle.BlockIndices[2] = BlockIndexV2;
+				Triangle.bUVsFixed = false;
+
+				VertexToFaceMap.Add(Triangle.CollapsedIndices[0], TriangleIndex);
+				VertexToFaceMap.Add(Triangle.CollapsedIndices[1], TriangleIndex);
+				VertexToFaceMap.Add(Triangle.CollapsedIndices[2], TriangleIndex);
+			}
+		}
+
+		// Clamp UV islands to the predominant block of each island. Will only happen if bClampUVIslands is true.
+		for (int32 ConflictiveTriangleIndex : ConflictiveTriangles)
+		{
+			FTriangleInfo& Triangle = Triangles[ConflictiveTriangleIndex];
+
+			// Skip the ones that have been fixed already
+			if (Triangle.bUVsFixed)
+			{
+				continue;
+			}
+
+			// Find triangles from the same UV Island
+			TArray<uint32> TriangleIndices;
+			GetUVIsland(Triangles, ConflictiveTriangleIndex, TriangleIndices, TexCoords, VertexToFaceMap);
+
+			// Get predominant BlockId != MAX_uint16
+			TArray<uint32> NumVerticesPerBlock;
+			NumVerticesPerBlock.SetNumZeroed(NumBlocks);
+
+			for (int32 TriangleIndex : TriangleIndices)
+			{
+				FTriangleInfo& OtherTriangle = Triangles[TriangleIndex];
+				for (int32 VertexIndex = 0; VertexIndex < 3; ++VertexIndex)
+				{
+					const uint16& BlockIndex = OtherTriangle.BlockIndices[VertexIndex];
+					if (BlockIndex != MAX_uint16)
+					{
+						NumVerticesPerBlock[BlockIndex]++;
+					}
+				}
+			}
+
+			uint16 BlockIndex = 0;
+			uint32 CurrentMaxVertices = 0;
+			for (int32 Index = 0; Index < NumBlocks; ++Index)
+			{
+				if (NumVerticesPerBlock[Index] > CurrentMaxVertices)
+				{
+					BlockIndex = Index;
+					CurrentMaxVertices = NumVerticesPerBlock[Index];
+				}
+			}
+
+			// Get the limits of the predominant block rect
+			const FLayoutBlock& LayoutBlock = Layout->Blocks[BlockIndex];
+
+			const float SmallNumber = 0.000001;
+			const float MinX = ((float)LayoutBlock.Min.X) / (float)LayoutGrid.X + SmallNumber;
+			const float MinY = ((float)LayoutBlock.Min.Y) / (float)LayoutGrid.Y + SmallNumber;
+			const float MaxX = (((float)LayoutBlock.Size.X + LayoutBlock.Min.X) / (float)LayoutGrid.X) - 2 * SmallNumber;
+			const float MaxY = (((float)LayoutBlock.Size.Y + LayoutBlock.Min.Y) / (float)LayoutGrid.Y) - 2 * SmallNumber;
+
+			// Iterate triangles and clamp the UVs
+			for (int32 TriangleIndex : TriangleIndices)
+			{
+				FTriangleInfo& OtherTriangle = Triangles[TriangleIndex];
+
+				for (int8 VertexIndex = 0; VertexIndex < 3; ++VertexIndex)
+				{
+					if (OtherTriangle.BlockIndices[VertexIndex] == BlockIndex)
+					{
+						continue;
+					}
+
+					OtherTriangle.BlockIndices[VertexIndex] = BlockIndex;
+
+					// Clamp UVs to the block they are assigned to
+					const int32 UVIndex = OtherTriangle.Indices[VertexIndex];
+					FVector2f& UV = TexCoords[UVIndex];
+					UV[0] = FMath::Clamp(UV[0], MinX, MaxX);
+					UV[1] = FMath::Clamp(UV[1], MinY, MaxY);
+					LayoutData[UVIndex] = BlockIndex;
+				}
+
+				OtherTriangle.bUVsFixed = true;
+			}
+		}
+
+		// Warn about vertices without a block id
+		//int32 FirstLODToIgnoreWarnings = GeneratedLayout.Source->FirstLODToIgnoreWarnings;
+		//if (FirstLODToIgnoreWarnings == -1 || StaticMeshOptions.LODIndex < FirstLODToIgnoreWarnings)
+		//{
+		//	TArray<float> UnassignedUVs;
+		//	UnassignedUVs.Reserve(NumVertices / 100);
+
+		//	const FVector2f* UVs = TexCoords.GetData();
+		//	for (int32 VertexIndex = 0; VertexIndex < NumVertices; ++VertexIndex)
+		//	{
+		//		if (LayoutData[VertexIndex] == MAX_uint16)
+		//		{
+		//			UnassignedUVs.Add((*(UVs + VertexIndex))[0]);
+		//			UnassignedUVs.Add((*(UVs + VertexIndex))[1]);
+		//		}
+		//	}
+
+		//	if (!UnassignedUVs.IsEmpty())
+		//	{
+		//		// TODO: How do we translate this to editor-time info?
+		//		UE_LOG(LogMutableCore, Warning, TEXT("Source mesh has %d vertices not assigned to any layout block in LOD %d in UVs index %d"), UnassignedUVs.Num(), LayoutChannel);
+		//	}
+		//}
+
+		// Create the layout block vertex buffer
+		uint8* LayoutBufferPtr = nullptr;
+		{
+			const int32 LayoutBufferIndex = Mesh.GetVertexBuffers().GetBufferCount();
+			Mesh.GetVertexBuffers().SetBufferCount(LayoutBufferIndex + 1);
+
+			// TODO
+			check(Layout->GetBlockCount() < MAX_uint16);
+			const EMeshBufferSemantic LayoutSemantic = EMeshBufferSemantic::LayoutBlock;
+			const int32 LayoutSemanticIndex = int32(LayoutChannel);
+			const EMeshBufferFormat LayoutFormat = bUseAbsoluteBlockIds ? EMeshBufferFormat::UInt64 : EMeshBufferFormat::UInt16;
+			const int32 LayoutComponents = 1;
+			const int32 LayoutOffset = 0;
+			int32 ElementSize = bUseAbsoluteBlockIds ? sizeof(uint64) : sizeof(uint16);
+			Mesh.GetVertexBuffers().SetBuffer
+			(
+				LayoutBufferIndex,
+				ElementSize,
+				1,
+				&LayoutSemantic, &LayoutSemanticIndex,
+				&LayoutFormat, &LayoutComponents,
+				&LayoutOffset
+			);
+			LayoutBufferPtr = Mesh.GetVertexBuffers().GetBufferData(LayoutBufferIndex);
+		}
+
+		// Copy UVs back to the mesh
+		{
+			uint8* pVertices = TexCoordData;
+			FVector2f* UVs = TexCoords.GetData();
+
+			for (int32 VertexIndex = 0; VertexIndex < NumVertices; ++VertexIndex)
+			{
+				FVector2f& UV = TexCoords[VertexIndex];
+
+				uint16 LayoutBlockIndex = LayoutData[VertexIndex];
+				if (Layout->Blocks.IsValidIndex(LayoutBlockIndex))
+				{
+					uint64 LayoutBlockId = Layout->Blocks[LayoutBlockIndex].Id;
+
+					UV = BlockRects[LayoutBlockIndex].Homogenize(UV);
+
+					// Replace block index by the actual id of the block
+					if (bUseAbsoluteBlockIds)
+					{
+						uint64* Ptr = reinterpret_cast<uint64*>(LayoutBufferPtr) + VertexIndex;
+						*Ptr = LayoutBlockId;
+					}
+					else
+					{
+						uint16* Ptr = reinterpret_cast<uint16*>(LayoutBufferPtr) + VertexIndex;
+						*Ptr = uint16(LayoutBlockId & 0xffff);
+					}
+				}
+				else
+				{
+					// Map vertices without block
+					if (bUseAbsoluteBlockIds)
+					{
+						uint64* Ptr = reinterpret_cast<uint64*>(LayoutBufferPtr) + VertexIndex;
+						*Ptr = bEnsureAllVerticesHaveLayoutBlock ? 0 : std::numeric_limits<uint64>::max();
+					}
+					else
+					{
+						uint16* Ptr = reinterpret_cast<uint16*>(LayoutBufferPtr) + VertexIndex;
+						*Ptr = bEnsureAllVerticesHaveLayoutBlock ? 0 : std::numeric_limits<uint16>::max();
+					}
+				}
+
+				// Copy UVs
+				if (TexCoordsChannel.Format == EMeshBufferFormat::Float32)
+				{
+					FVector2f* pUV = reinterpret_cast<FVector2f*>(pVertices);
+					*pUV = UV;
+				}
+				else if (TexCoordsChannel.Format == EMeshBufferFormat::Float16)
+				{
+					FFloat16* pUV = reinterpret_cast<FFloat16*>(pVertices);
+					pUV[0] = FFloat16(UV[0]);
+					pUV[1] = FFloat16(UV[1]);
+				}
+
+				pVertices += elemSize;
+			}
+		}
+	}
+
+}
diff --git a/Engine/Plugins/Mutable/Source/MutableRuntime/Private/MuR/Operations.cpp b/Engine/Plugins/Mutable/Source/MutableRuntime/Private/MuR/Operations.cpp
--- a/Engine/Plugins/Mutable/Source/MutableRuntime/Private/MuR/Operations.cpp
+++ b/Engine/Plugins/Mutable/Source/MutableRuntime/Private/MuR/Operations.cpp
@@ -100,6 +100,7 @@ namespace mu
 		{ EDataType::Image,			},	// IM_TRANSFORM
 
 		{ EDataType::Mesh,			},	// ME_APPLYLAYOUT
+		{ EDataType::Mesh,			},	// ME_PREPARELAYOUT
 		{ EDataType::Mesh,			},	// ME_DIFFERENCE
 		{ EDataType::Mesh,			},	// ME_MORPH
 		{ EDataType::Mesh,			},	// ME_MERGE
@@ -536,6 +537,14 @@ namespace mu
             break;
         }
 
+		case EOpType::ME_PREPARELAYOUT:
+		{
+			OP::MeshPrepareLayoutArgs Args = program.GetOpArgs<OP::MeshPrepareLayoutArgs>(at);
+			f(Args.Layout);
+			f(Args.Mesh);
+			break;
+		}
+
         case EOpType::ME_DIFFERENCE:
         {
 			const uint8_t* data = program.GetOpArgsPointer(at);
diff --git a/Engine/Plugins/Mutable/Source/MutableRuntime/Public/MuR/Mesh.h b/Engine/Plugins/Mutable/Source/MutableRuntime/Public/MuR/Mesh.h
--- a/Engine/Plugins/Mutable/Source/MutableRuntime/Public/MuR/Mesh.h
+++ b/Engine/Plugins/Mutable/Source/MutableRuntime/Public/MuR/Mesh.h
@@ -26,8 +26,9 @@ namespace mu
 {
 
 	// Forward references
+	class FMesh;
 	class FLayout;
-    class FSkeleton;
+	class FSkeleton;
     class FPhysicsBody;
 
 	struct FSurfaceSubMesh
@@ -70,6 +71,34 @@ namespace mu
 		inline void Unserialise(FInputArchive& Arch);
 	};
 
+	/** Helper structs for mesh utility methods below. */
+	struct FTriangleInfo
+	{
+		/** Vertex indices in the original mesh. */
+		uint32 Indices[3];
+
+		/** Vertex indices in the collapsed vertex list of the mesh. */
+		uint32 CollapsedIndices[3];
+
+		/** Optional data with layout block indices. */
+		uint16 BlockIndices[3];
+
+		/** Optional data with a flag indicating the UVs have changed during layout for this trioangle. */
+		bool bUVsFixed = false;
+	};
+
+
+	/** Fill an array with the indices of all triangles belonging to the same UV island as InFirstTriangle.
+	*/
+	MUTABLERUNTIME_API void GetUVIsland(TArray<FTriangleInfo>& InTriangles,
+		const uint32 InFirstTriangle,
+		TArray<uint32>& OutTriangleIndices,
+		const TArray<FVector2f>& InUVs,
+		const TMultiMap<int32, uint32>& InVertexToTriangleMap);
+
+	/** Create a map from vertices into vertices, collapsing vertices that have the same position. */
+	MUTABLERUNTIME_API void MeshCreateCollapsedVertexMap(const FMesh* Mesh, TArray<int32>& CollapsedVertices);
+
 
 	enum class EBoneUsageFlags : uint32
 	{
@@ -563,7 +592,7 @@ namespace mu
 		void ResetBufferIndices();
 
 		/** Debug: get a text representation of the mesh */
-		void Log(FString& Out, int32 VertexLimit);
+		void Log(FString& Out, int32 VertexLimit) const;
     };
 
 
diff --git a/Engine/Plugins/Mutable/Source/MutableTools/Private/MuT/ASTOpMeshPrepareLayout.cpp b/Engine/Plugins/Mutable/Source/MutableTools/Private/MuT/ASTOpMeshPrepareLayout.cpp
new file mode 100644
--- /dev/null
+++ b/Engine/Plugins/Mutable/Source/MutableTools/Private/MuT/ASTOpMeshPrepareLayout.cpp
@@ -0,0 +1,106 @@
+// Copyright Epic Games, Inc. All Rights Reserved.
+
+#include "MuT/ASTOpMeshPrepareLayout.h"
+#include "MuT/ASTOpMeshAddTags.h"
+
+#include "HAL/PlatformMath.h"
+#include "MuR/ModelPrivate.h"
+#include "MuR/RefCounted.h"
+#include "MuR/Types.h"
+
+namespace mu
+{
+
+	ASTOpMeshPrepareLayout::ASTOpMeshPrepareLayout()
+		: Mesh(this)
+		, Layout(this)
+	{
+	}
+
+
+	ASTOpMeshPrepareLayout::~ASTOpMeshPrepareLayout()
+	{
+		// Explicit call needed to avoid recursive destruction
+		ASTOp::RemoveChildren();
+	}
+
+
+	bool ASTOpMeshPrepareLayout::IsEqual(const ASTOp& OtherUntyped) const
+	{
+		if (OtherUntyped.GetOpType()==GetOpType())
+		{
+			const ASTOpMeshPrepareLayout* Other = static_cast<const ASTOpMeshPrepareLayout*>(&OtherUntyped);
+			return Mesh == Other->Mesh &&
+				Layout == Other->Layout &&
+				bUseAbsoluteBlockIds == Other->bUseAbsoluteBlockIds &&
+				bNormalizeUVs == Other->bNormalizeUVs &&
+				bClampUVIslands == Other->bClampUVIslands &&
+				bEnsureAllVerticesHaveLayoutBlock == Other->bEnsureAllVerticesHaveLayoutBlock;
+		}
+		return false;
+	}
+
+
+	uint64 ASTOpMeshPrepareLayout::Hash() const
+	{
+		uint64 Result = std::hash<void*>()(Mesh.child().get());
+		hash_combine(Result, Layout.child().get());
+		hash_combine(Result, bUseAbsoluteBlockIds);
+		return Result;
+	}
+
+
+	Ptr<ASTOp> ASTOpMeshPrepareLayout::Clone(MapChildFuncRef MapChild) const
+	{
+		Ptr<ASTOpMeshPrepareLayout> New = new ASTOpMeshPrepareLayout();
+		New->Mesh = MapChild(Mesh.child());
+		New->Layout = MapChild(Layout.child());
+		New->bUseAbsoluteBlockIds = bUseAbsoluteBlockIds;
+		New->bNormalizeUVs = bNormalizeUVs;
+		New->bClampUVIslands = bClampUVIslands;
+		New->bEnsureAllVerticesHaveLayoutBlock = bEnsureAllVerticesHaveLayoutBlock;
+		return New;
+	}
+
+
+	void ASTOpMeshPrepareLayout::ForEachChild(const TFunctionRef<void(ASTChild&)> Func)
+	{
+		Func(Mesh);
+		Func(Layout);
+	}
+
+
+	void ASTOpMeshPrepareLayout::Link(FProgram& Program, FLinkerOptions*)
+	{
+		// Already linked?
+		if (!linkedAddress)
+		{
+			OP::MeshPrepareLayoutArgs Args;
+			FMemory::Memzero(Args);
+
+			if (Mesh) Args.Mesh = Mesh->linkedAddress;
+			if (Layout) Args.Layout = Layout->linkedAddress;
+			Args.LayoutChannel = LayoutChannel;
+			Args.bUseAbsoluteBlockIds = bUseAbsoluteBlockIds;
+			Args.bNormalizeUVs = bNormalizeUVs;
+			Args.bClampUVIslands = bClampUVIslands;
+			Args.bEnsureAllVerticesHaveLayoutBlock = bEnsureAllVerticesHaveLayoutBlock;
+
+			linkedAddress = (OP::ADDRESS)Program.OpAddress.Num();
+			Program.OpAddress.Add(Program.ByteCode.Num());
+			AppendCode(Program.ByteCode, GetOpType());
+			AppendCode(Program.ByteCode, Args);
+		}
+	}
+
+
+	FSourceDataDescriptor ASTOpMeshPrepareLayout::GetSourceDataDescriptor(FGetSourceDataDescriptorContext* Context) const
+	{
+		if (Mesh)
+		{
+			return Mesh->GetSourceDataDescriptor(Context);
+		}
+
+		return {};
+	}
+}
diff --git a/Engine/Plugins/Mutable/Source/MutableTools/Private/MuT/ASTOpMeshPrepareLayout.h b/Engine/Plugins/Mutable/Source/MutableTools/Private/MuT/ASTOpMeshPrepareLayout.h
new file mode 100644
--- /dev/null
+++ b/Engine/Plugins/Mutable/Source/MutableTools/Private/MuT/ASTOpMeshPrepareLayout.h
@@ -0,0 +1,46 @@
+// Copyright Epic Games, Inc. All Rights Reserved.
+
+#pragma once
+
+#include "HAL/Platform.h"
+#include "MuR/Operations.h"
+#include "MuR/Ptr.h"
+#include "MuT/AST.h"
+
+
+namespace mu
+{
+	struct FProgram;
+
+	/** */
+	class ASTOpMeshPrepareLayout final : public ASTOp
+	{
+	public:
+
+		ASTChild Mesh;
+		ASTChild Layout;
+		uint8 LayoutChannel = 0;
+		bool bUseAbsoluteBlockIds = false;
+		bool bNormalizeUVs = false;
+		bool bClampUVIslands = false;
+		bool bEnsureAllVerticesHaveLayoutBlock = false;
+
+	public:
+
+		ASTOpMeshPrepareLayout();
+		ASTOpMeshPrepareLayout(const ASTOpMeshPrepareLayout&) = delete;
+		~ASTOpMeshPrepareLayout();
+
+		// ASTOp interface
+		virtual EOpType GetOpType() const override { return EOpType::ME_PREPARELAYOUT; }
+		virtual uint64 Hash() const override;
+		virtual bool IsEqual(const ASTOp&) const override;
+		virtual Ptr<ASTOp> Clone(MapChildFuncRef) const override;
+		virtual void ForEachChild(const TFunctionRef<void(ASTChild&)>) override;
+		virtual void Link(FProgram&, FLinkerOptions*) override;
+		virtual FSourceDataDescriptor GetSourceDataDescriptor(FGetSourceDataDescriptorContext*) const override;
+	};
+
+
+}
+
diff --git a/Engine/Plugins/Mutable/Source/MutableTools/Private/MuT/ASTOpParameter.cpp b/Engine/Plugins/Mutable/Source/MutableTools/Private/MuT/ASTOpParameter.cpp
--- a/Engine/Plugins/Mutable/Source/MutableTools/Private/MuT/ASTOpParameter.cpp
+++ b/Engine/Plugins/Mutable/Source/MutableTools/Private/MuT/ASTOpParameter.cpp
@@ -40,6 +40,7 @@ namespace mu
 				Parameter == Other->Parameter &&
 				LODIndex == Other->LODIndex &&
 				SectionIndex == Other->SectionIndex &&
+				MeshID == Other->MeshID &&
 				Ranges == Other->Ranges;
 		}
 		return false;
@@ -53,7 +54,8 @@ namespace mu
 		n->Parameter = Parameter;
 		n->LODIndex = LODIndex;
 		n->SectionIndex = SectionIndex;
-		for (const auto& c : Ranges)
+		n->MeshID = MeshID;
+		for (const FRangeData& c : Ranges)
 		{
 			n->Ranges.Emplace(n.get(), MapChild(c.rangeSize.child()), c.rangeName, c.rangeUID);
 		}
@@ -107,6 +109,7 @@ namespace mu
 				Args.variable = (OP::ADDRESS)LinkedParameterIndex;
 				Args.LOD = LODIndex;
 				Args.Section = SectionIndex;
+				Args.MeshID = MeshID;
 
 				for (const auto& d : Ranges)
 				{
diff --git a/Engine/Plugins/Mutable/Source/MutableTools/Private/MuT/ASTOpParameter.h b/Engine/Plugins/Mutable/Source/MutableTools/Private/MuT/ASTOpParameter.h
--- a/Engine/Plugins/Mutable/Source/MutableTools/Private/MuT/ASTOpParameter.h
+++ b/Engine/Plugins/Mutable/Source/MutableTools/Private/MuT/ASTOpParameter.h
@@ -29,6 +29,7 @@ namespace mu
 		/** Used by some parameter types (Mesh) to specify which subdata from the actual parameter value should be used in the operation. */
 		int32 LODIndex = 0;
 		int32 SectionIndex = 0;
+		uint32 MeshID = 0;
 
 		//** Ranges adding dimensions to this parameter. */
 		TArray<FRangeData> Ranges;
diff --git a/Engine/Plugins/Mutable/Source/MutableTools/Private/MuT/CodeGenerator.h b/Engine/Plugins/Mutable/Source/MutableTools/Private/MuT/CodeGenerator.h
--- a/Engine/Plugins/Mutable/Source/MutableTools/Private/MuT/CodeGenerator.h
+++ b/Engine/Plugins/Mutable/Source/MutableTools/Private/MuT/CodeGenerator.h
@@ -318,14 +318,23 @@ namespace mu
 			{
 				TSharedPtr<FMesh> Mesh;
 				Ptr<ASTOp> LastMeshOp;
+				TArray<FGeneratedLayout> GeneratedLayouts;
 			};
 			TMap<uint64, TArray<FGeneratedConstantMesh>> GeneratedConstantMeshes;
-
-			/** List of already used vertex ID groups that must be unique. */
-			TSet<uint32> UniqueVertexIDGroups;
 		};
 		FGenerateMeshConstantState GenerateMeshConstantState;
 
+
+		/** List of already used vertex ID for meshes that must be unique. */
+		struct FUniqueMeshIds
+		{
+			uint32 EnsureUnique(uint32 Id);
+		private:
+			UE::FMutex Mutex;
+			TSet<uint32> Map;
+		};
+		FUniqueMeshIds UniqueMeshIds;
+
 		/** Any task accessing GenerateMeshConstantState should run in this pipe. */
 		UE::Tasks::FPipe GenerateMeshConstantPipe;
 
@@ -643,30 +652,26 @@ namespace mu
 		};
 		FGenerateLayoutConstantState GenerateLayoutConstantState;
 
-		FMeshTask GenerateMesh(const FMeshGenerationStaticOptions& StaticOptions, FMeshOptionsTask, const Ptr<const NodeMesh>&);
-		FMeshTask GenerateMesh_Constant(const FMeshGenerationStaticOptions& StaticOptions, FMeshOptionsTask, const NodeMeshConstant* );
-		FMeshTask GenerateMesh_Format(const FMeshGenerationStaticOptions& StaticOptions, FMeshOptionsTask, const NodeMeshFormat* );
-		FMeshTask GenerateMesh_Morph(const FMeshGenerationStaticOptions& StaticOptions, FMeshOptionsTask, const NodeMeshMorph* );
-		FMeshTask GenerateMesh_MakeMorph(const FMeshGenerationStaticOptions& StaticOptions, FMeshOptionsTask, const NodeMeshMakeMorph* );
-        FMeshTask GenerateMesh_Fragment(const FMeshGenerationStaticOptions& StaticOptions, FMeshOptionsTask, const NodeMeshFragment* );
-        FMeshTask GenerateMesh_Switch(const FMeshGenerationStaticOptions& StaticOptions, FMeshOptionsTask, const NodeMeshSwitch* );
-        FMeshTask GenerateMesh_Transform(const FMeshGenerationStaticOptions& StaticOptions, FMeshOptionsTask, const NodeMeshTransform* );
-        FMeshTask GenerateMesh_ClipMorphPlane(const FMeshGenerationStaticOptions& StaticOptions, FMeshOptionsTask, const NodeMeshClipMorphPlane* );
-        FMeshTask GenerateMesh_ClipWithMesh(const FMeshGenerationStaticOptions& StaticOptions, FMeshOptionsTask, const NodeMeshClipWithMesh* );
-        FMeshTask GenerateMesh_ApplyPose(const FMeshGenerationStaticOptions& StaticOptions, FMeshOptionsTask, const NodeMeshApplyPose* );
-        FMeshTask GenerateMesh_Variation(const FMeshGenerationStaticOptions& StaticOptions, FMeshOptionsTask, const NodeMeshVariation* );
-		FMeshTask GenerateMesh_Table(const FMeshGenerationStaticOptions& StaticOptions, FMeshOptionsTask, const NodeMeshTable*);
-		FMeshTask GenerateMesh_Reshape(const FMeshGenerationStaticOptions& StaticOptions, FMeshOptionsTask, const NodeMeshReshape*);
-		FMeshTask GenerateMesh_ClipDeform(const FMeshGenerationStaticOptions& StaticOptions, FMeshOptionsTask, const NodeMeshClipDeform*);
-		FMeshTask GenerateMesh_Parameter(const FMeshGenerationStaticOptions& StaticOptions, FMeshOptionsTask, const NodeMeshParameter*);
+		FMeshTask GenerateMesh(const FMeshGenerationStaticOptions&, FMeshOptionsTask, const Ptr<const NodeMesh>&);
+		FMeshTask GenerateMesh_Constant(const FMeshGenerationStaticOptions&, FMeshOptionsTask, const NodeMeshConstant* );
+		FMeshTask GenerateMesh_Format(const FMeshGenerationStaticOptions&, FMeshOptionsTask, const NodeMeshFormat* );
+		FMeshTask GenerateMesh_Morph(const FMeshGenerationStaticOptions&, FMeshOptionsTask, const NodeMeshMorph* );
+		FMeshTask GenerateMesh_MakeMorph(const FMeshGenerationStaticOptions&, FMeshOptionsTask, const NodeMeshMakeMorph* );
+        FMeshTask GenerateMesh_Fragment(const FMeshGenerationStaticOptions&, FMeshOptionsTask, const NodeMeshFragment* );
+        FMeshTask GenerateMesh_Switch(const FMeshGenerationStaticOptions&, FMeshOptionsTask, const NodeMeshSwitch* );
+        FMeshTask GenerateMesh_Transform(const FMeshGenerationStaticOptions&, FMeshOptionsTask, const NodeMeshTransform* );
+        FMeshTask GenerateMesh_ClipMorphPlane(const FMeshGenerationStaticOptions&, FMeshOptionsTask, const NodeMeshClipMorphPlane* );
+        FMeshTask GenerateMesh_ClipWithMesh(const FMeshGenerationStaticOptions&, FMeshOptionsTask, const NodeMeshClipWithMesh* );
+        FMeshTask GenerateMesh_ApplyPose(const FMeshGenerationStaticOptions&, FMeshOptionsTask, const NodeMeshApplyPose* );
+        FMeshTask GenerateMesh_Variation(const FMeshGenerationStaticOptions&, FMeshOptionsTask, const NodeMeshVariation* );
+		FMeshTask GenerateMesh_Table(const FMeshGenerationStaticOptions&, FMeshOptionsTask, const NodeMeshTable*);
+		FMeshTask GenerateMesh_Reshape(const FMeshGenerationStaticOptions&, FMeshOptionsTask, const NodeMeshReshape*);
+		FMeshTask GenerateMesh_ClipDeform(const FMeshGenerationStaticOptions&, FMeshOptionsTask, const NodeMeshClipDeform*);
+		FMeshTask GenerateMesh_Parameter(const FMeshGenerationStaticOptions&, FMeshOptionsTask, const NodeMeshParameter*);
+
+		Ptr<ASTOp> GenerateLayoutOpsAndResult(const FMeshGenerationDynamicOptions&, Ptr<ASTOp> LastMeshOp, const TArray<Ptr<NodeLayout>>& OriginalLayouts, uint32 MeshId, FMeshGenerationResult& OutResult);
 
 		//-----------------------------------------------------------------------------------------
-		void PrepareMeshForLayout(const FGeneratedLayout&, TSharedPtr<FMesh>,
-			int32 currentLayoutChannel,
-			const void* errorContext,
-			const FMeshGenerationStaticOptions&,
-			const FMeshGenerationDynamicOptions&,
-			bool bUseAbsoluteBlockIds) const;
 
 		//!
 		TSharedPtr<const FLayout> GenerateLayout(Ptr<const NodeLayout> SourceLayout, uint32 MeshIDPrefix);
diff --git a/Engine/Plugins/Mutable/Source/MutableTools/Private/MuT/CodeGenerator_FirstPass.h b/Engine/Plugins/Mutable/Source/MutableTools/Private/MuT/CodeGenerator_FirstPass.h
--- a/Engine/Plugins/Mutable/Source/MutableTools/Private/MuT/CodeGenerator_FirstPass.h
+++ b/Engine/Plugins/Mutable/Source/MutableTools/Private/MuT/CodeGenerator_FirstPass.h
@@ -58,6 +58,7 @@ namespace mu
 		/** Generated node layouts with their own block ids. */
 		TArray<FGeneratedLayout> GeneratedLayouts;
 
+		/** TODO: The following members seem related to surface-sharing data and not actual mesh generation result. Maybe they should be moved to a different struct.  */
 		TArray<Ptr<ASTOp>> LayoutOps;
 
 		struct FExtraLayouts
@@ -254,8 +255,8 @@ namespace mu
 			/** Cache for generic parameters. */
 			TMap< Ptr<const Node>, Ptr<ASTOpParameter> > GenericParametersCache;
 
-			/** For mesh parameters we generate a different ASTOp for each LOD. */
-			TMap< Ptr<const NodeMeshParameter>, TArray<Ptr<ASTOpParameter>> > MeshParametersCache;
+			/** For mesh parameters we generate a different result for each LOD. */
+			TMap< Ptr<const NodeMeshParameter>, TArray<TPair<Ptr<ASTOpParameter>, FMeshGenerationResult>> > MeshParametersCache;
 		};
 		FSafeParameterNodes ParameterNodes;
 
diff --git a/Engine/Plugins/Mutable/Source/MutableTools/Private/MuT/CodeGenerator_Mesh.cpp b/Engine/Plugins/Mutable/Source/MutableTools/Private/MuT/CodeGenerator_Mesh.cpp
--- a/Engine/Plugins/Mutable/Source/MutableTools/Private/MuT/CodeGenerator_Mesh.cpp
+++ b/Engine/Plugins/Mutable/Source/MutableTools/Private/MuT/CodeGenerator_Mesh.cpp
@@ -8,6 +8,7 @@
 #include "MuR/Mesh.h"
 #include "MuR/MeshPrivate.h"
 #include "MuR/MeshBufferSet.h"
+#include "MuR/OpMeshPrepareLayout.h"
 #include "MuR/MutableMath.h"
 #include "MuR/Operations.h"
 #include "MuR/ParametersPrivate.h"
@@ -31,6 +32,7 @@
 #include "MuT/ASTOpMeshDifference.h"
 #include "MuT/ASTOpMeshMorph.h"
 #include "MuT/ASTOpMeshAddTags.h"
+#include "MuT/ASTOpMeshPrepareLayout.h"
 #include "MuT/ASTOpSwitch.h"
 #include "MuT/ASTOpReferenceResource.h"
 #include "MuT/CodeGenerator.h"
@@ -64,520 +66,7 @@ namespace mu
 {
 	class Node;
 
-	template<typename T>
-	struct TArray2D
-	{
-		int32 SizeX = 0;
-		int32 SizeY = 0;
-		TArray<T> Data;
-
-		void Init(const T& Value, int32 InSizeX, int32 InSizeY)
-		{
-			SizeX = InSizeX;
-			SizeY = InSizeY;
-			Data.Init(Value, SizeX*SizeY );
-		}
-
-		inline const T& Get(int32 X, int32 Y) const
-		{
-			check(X >= 0 && X < SizeX);
-			check(Y >= 0 && Y < SizeY);
-			return Data[SizeX * Y + X];
-		}
-
-		inline void Set(int32 X, int32 Y, const T& Value)
-		{
-			check(X >= 0 && X < SizeX);
-			check(Y >= 0 && Y < SizeY);
-			Data[SizeX * Y + X] = Value;
-		}
-
-	};
-
-    //---------------------------------------------------------------------------------------------
-	void CodeGenerator::PrepareMeshForLayout( const FGeneratedLayout& GeneratedLayout,
-		TSharedPtr<FMesh> Mesh,
-		int32 LayoutChannel,
-		const void* ErrorContext,
-		const FMeshGenerationStaticOptions& StaticMeshOptions,
-		const FMeshGenerationDynamicOptions& DynamicMeshOptions,
-		bool bUseAbsoluteBlockIds
-		) const
-	{
-		// This can run in any thread.
-
-		MUTABLE_CPUPROFILER_SCOPE(PrepareMeshForLayout);
-
-		if (Mesh->GetVertexCount() == 0)
-		{
-			return;
-		}
-
-		// The layout must have block ids.
-		check(GeneratedLayout.Layout->Blocks.IsEmpty() || GeneratedLayout.Layout->Blocks[0].Id != FLayoutBlock::InvalidBlockId);
-
-		// 
-		TSharedPtr<const FLayout> Layout = GeneratedLayout.Layout;
-		Mesh->AddLayout(Layout);
-
-		const int32 NumVertices = Mesh->GetVertexCount();
-		const int32 NumBlocks = Layout->GetBlockCount();
-
-		// Find block ids for each block in the grid. Calculate a grid size that contains all blocks
-		FIntPoint LayoutGrid = Layout->GetGridSize();
-		FIntPoint WorkingGrid = LayoutGrid;
-		for (const FSourceLayoutBlock& Block: GeneratedLayout.Source->Blocks)
-		{
-			WorkingGrid.X = FMath::Max(WorkingGrid.X, Block.Min.X + Block.Size.X );
-			WorkingGrid.Y = FMath::Max(WorkingGrid.Y, Block.Min.Y + Block.Size.Y );
-		}
-
-
-		TArray2D<int32> GridBlockBlockId;
-		GridBlockBlockId.Init(MAX_uint16, WorkingGrid.X, WorkingGrid.Y);
-
-		// 
-		TArray<box<FVector2f>> BlockRects;
-		BlockRects.SetNumUninitialized(NumBlocks);
-
-		// Create an array of block index per cell
-		TArray<int32> OverlappingBlocks;
-		for (int32 BlockIndex = 0; BlockIndex < NumBlocks; ++BlockIndex)
-		{
-			bool bBlockHasMask = GeneratedLayout.Source->Blocks[BlockIndex].Mask != nullptr;
-
-			// Fill the block rect
-			FIntVector2 Min = Layout->Blocks[BlockIndex].Min;
-			FIntVector2 Size = Layout->Blocks[BlockIndex].Size;
-
-			box<FVector2f>& BlockRect = BlockRects[BlockIndex];
-			BlockRect.min[0] = float(Min.X) / float(LayoutGrid.X);
-			BlockRect.min[1] = float(Min.Y) / float(LayoutGrid.Y);
-			BlockRect.size[0] = float(Size.X) / float(LayoutGrid.X);
-			BlockRect.size[1] = float(Size.Y) / float(LayoutGrid.Y);
-
-			// Fill the block index per cell array
-			// Ignore the block in this stage if it has a mask, because blocks with masks will very likely overlap other blocks
-			if (!bBlockHasMask)
-			{
-				for (uint16 Y = Min.Y; Y < Min.Y + Size.Y; ++Y)
-				{
-					for (uint16 X = Min.X; X < Min.X + Size.X; ++X)
-					{
-						if (GridBlockBlockId.Get(X,Y) == MAX_uint16)
-						{
-							GridBlockBlockId.Set(X,Y, BlockIndex);
-						}
-						else
-						{
-							OverlappingBlocks.AddUnique(BlockIndex);
-						}
-					}
-				}
-			}
-		}
-
-
-		// Notify Overlapping layout blocks
-		if (!OverlappingBlocks.IsEmpty())
-		{
-			FString Msg = FString::Printf(TEXT("Source mesh has %d layout block overlapping in LOD %d"), OverlappingBlocks.Num() + 1, StaticMeshOptions.LODIndex );
-			ErrorLog->Add(Msg, ELMT_WARNING, ErrorContext);
-		}
-
-		// Get the information about the texture coordinates channel
-		int32 TexCoordsBufferIndex = -1;
-		int32 TexCoordsChannelIndex = -1;
-		Mesh->GetVertexBuffers().FindChannel(EMeshBufferSemantic::TexCoords, LayoutChannel, &TexCoordsBufferIndex, &TexCoordsChannelIndex);
-
-
-		if (TexCoordsBufferIndex<0 || TexCoordsChannelIndex<0)
-		{
-			// This is actually fine when using shared materials across LODs
-			FString Msg = FString::Printf(TEXT("LOD %d : Trying to generate layout for missing UV channel %d. No layout blocks generated."), StaticMeshOptions.LODIndex, LayoutChannel);
-			ErrorLog->Add(Msg, ELMT_INFO, ErrorContext);
-			return;
-		}
-
-		check(TexCoordsBufferIndex >= 0);
-		check(TexCoordsChannelIndex >= 0);
-
-		const FMeshBufferChannel& TexCoordsChannel = Mesh->VertexBuffers.Buffers[TexCoordsBufferIndex].Channels[TexCoordsChannelIndex];
-		check(TexCoordsChannel.Semantic == EMeshBufferSemantic::TexCoords);
-
-		uint8* TexCoordData = Mesh->GetVertexBuffers().GetBufferData(TexCoordsBufferIndex);
-		int32 elemSize = Mesh->GetVertexBuffers().GetElementSize(TexCoordsBufferIndex);
-		int32 channelOffset = TexCoordsChannel.Offset;
-		TexCoordData += channelOffset;
-
-		// Get a copy of the UVs as FVector2f to work with them. 
-		TArray<FVector2f> TexCoords;
-		{
-			TexCoords.SetNumUninitialized(NumVertices);
-
-			bool bNonNormalizedUVs = false;
-			const bool bIsOverlayLayout = GeneratedLayout.Layout->GetLayoutPackingStrategy() == mu::EPackStrategy::Overlay;
-
-			const uint8* pVertices = TexCoordData;
-			for (int32 VertexIndex = 0; VertexIndex < NumVertices; ++VertexIndex)
-			{
-				FVector2f& UV = TexCoords[VertexIndex];
-				if (TexCoordsChannel.Format == EMeshBufferFormat::Float32)
-				{
-					UV = *((FVector2f*)pVertices);
-				}
-				else if (TexCoordsChannel.Format == EMeshBufferFormat::Float16)
-				{
-					const FFloat16* pUV = reinterpret_cast<const FFloat16*>(pVertices);
-					UV = FVector2f(float(pUV[0]), float(pUV[1]));
-				}
-
-				// Check that UVs are normalized. If not, clamp the values and throw a warning.
-				if (DynamicMeshOptions.bNormalizeUVs && !bIsOverlayLayout
-					&& (UV[0] < 0.f || UV[0] > 1.f || UV[1] < 0.f || UV[1] > 1.f))
-				{
-					UV[0] = FMath::Clamp(UV[0], 0.f, 1.f);
-					UV[1] = FMath::Clamp(UV[1], 0.f, 1.f);
-					bNonNormalizedUVs = true;
-				}
-
-				pVertices += elemSize;
-			}
-
-			// Mutable does not support non-normalized UVs
-			if (bNonNormalizedUVs && !bIsOverlayLayout)
-			{
-				FString Msg = FString::Printf(TEXT("Source mesh has non-normalized UVs in LOD %d"), StaticMeshOptions.LODIndex);
-				ErrorLog->Add(Msg, ELMT_WARNING, ErrorContext);
-			}
-		}
-
-
-		const int32 NumTriangles = Mesh->GetIndexCount() / 3;
-		TArray<FTriangleInfo> Triangles;
-
-		// Vertices mapped to unique vertex index
-		TArray<int32> CollapsedVertices;
-
-		// Vertex to face map used to speed up connectivity building
-		TMultiMap<int32, uint32> VertexToFaceMap;
-
-		// Find Unique Vertices
-		if (DynamicMeshOptions.bClampUVIslands)
-		{
-			VertexToFaceMap.Reserve(NumVertices);
-			Triangles.SetNumUninitialized(NumTriangles);
-
-			MeshCreateCollapsedVertexMap(Mesh.Get(), CollapsedVertices);
-		}
-
-		TArray<int32> ConflictiveTriangles;
-
-		const uint32 MaxGridX = DynamicMeshOptions.bNormalizeUVs ? MAX_uint32 : WorkingGrid.X - 1;
-		const uint32 MaxGridY = DynamicMeshOptions.bNormalizeUVs ? MAX_uint32 : WorkingGrid.Y - 1;
-
-		// Allocate the per-vertex layout block data
-		TArray<uint16> LayoutData;
-		constexpr uint16 NullBlockId = MAX_uint16 - 1;
-		LayoutData.Init(NullBlockId, NumVertices);
-
-		UntypedMeshBufferIteratorConst ItIndices(Mesh->GetIndexBuffers(), EMeshBufferSemantic::VertexIndex);
-		for (int32 TriangleIndex = 0;  TriangleIndex < NumTriangles; ++TriangleIndex)
-		{
-			uint32 Index0 = ItIndices.GetAsUINT32();
-			++ItIndices;
-			uint32 Index1 = ItIndices.GetAsUINT32();
-			++ItIndices;
-			uint32 Index2 = ItIndices.GetAsUINT32();
-			++ItIndices;
-
-			auto AssignOneVertex = 
-				[LayoutGrid, MaxGridX, MaxGridY, NumBlocks, Layout, &GeneratedLayout, &GridBlockBlockId, &TexCoords, &LayoutData]
-				(int32 VertexIndex)
-				{
-					uint16& BlockIndex = LayoutData[VertexIndex];
-
-					// Was it previously assigned?
-					if (BlockIndex != NullBlockId)
-					{
-						return BlockIndex;
-					}
-
-					FVector2f UV = TexCoords[VertexIndex];
-
-					int32 VertexWorkingGridX = FMath::Clamp(LayoutGrid.X * UV[0], 0, LayoutGrid.X - 1);
-					int32 VertexWorkingGridY = FMath::Clamp(LayoutGrid.Y * UV[1], 0, LayoutGrid.Y - 1);
-
-					// First: Assign the vertices to masked blocks in order
-					for (int32 CandidateBlockIndex = 0; CandidateBlockIndex < NumBlocks; ++CandidateBlockIndex)
-					{
-						const mu::FImage* Mask = GeneratedLayout.Source->Blocks[CandidateBlockIndex].Mask.Get();
-						if (Mask)
-						{
-							// First discard with block limits.
-							FIntVector2 Min = Layout->Blocks[CandidateBlockIndex].Min;
-							FIntVector2 Size = Layout->Blocks[CandidateBlockIndex].Size;
-
-							bool bInBlock = 
-								(VertexWorkingGridX >= Min.X && VertexWorkingGridX < Min.X + Size.X)
-								&&
-								(VertexWorkingGridY >= Min.Y && VertexWorkingGridY < Min.Y + Size.Y);
-
-							if (bInBlock)
-							{
-								// TODO: This always clamps the UVs
-								FVector2f SampleUV;
-								SampleUV.X = FMath::Fmod(UV.X, 1.0);
-								SampleUV.Y = FMath::Fmod(UV.Y, 1.0);
 
-								FVector4f MaskValue = Mask->Sample(SampleUV);
-								if (MaskValue.X > 0.5f)
-								{
-									BlockIndex = CandidateBlockIndex;
-									break;
-								}
-							}
-						}
-					}
-
-					// Second: Assign to non-masked blocks if not assigned yet
-					if (BlockIndex == NullBlockId)
-					{
-						uint32 ClampedX = FMath::Min<uint32>(MaxGridX, FMath::Max<uint32>(0, VertexWorkingGridX));
-						uint32 ClampedY = FMath::Min<uint32>(MaxGridY, FMath::Max<uint32>(0, VertexWorkingGridY));
-						BlockIndex = GridBlockBlockId.Get(ClampedX, ClampedY);
-					}
-					return BlockIndex;
-				};
-
-			const uint16 BlockIndexV0 = AssignOneVertex(Index0);
-			const uint16 BlockIndexV1 = AssignOneVertex(Index1);
-			const uint16 BlockIndexV2 = AssignOneVertex(Index2);
-
-			if (DynamicMeshOptions.bClampUVIslands)
-			{
-				if (BlockIndexV0 != BlockIndexV1 || BlockIndexV0 != BlockIndexV2)
-				{
-					ConflictiveTriangles.Add(TriangleIndex);
-				}
-
-				FTriangleInfo& Triangle = Triangles[TriangleIndex];
-
-				Triangle.Indices[0] = Index0;
-				Triangle.Indices[1] = Index1;
-				Triangle.Indices[2] = Index2;
-				Triangle.CollapsedIndices[0] = CollapsedVertices[Index0];
-				Triangle.CollapsedIndices[1] = CollapsedVertices[Index1];
-				Triangle.CollapsedIndices[2] = CollapsedVertices[Index2];
-
-				Triangle.BlockIndices[0] = BlockIndexV0;
-				Triangle.BlockIndices[1] = BlockIndexV1;
-				Triangle.BlockIndices[2] = BlockIndexV2;
-				Triangle.bUVsFixed = false;
-
-				VertexToFaceMap.Add(Triangle.CollapsedIndices[0], TriangleIndex);
-				VertexToFaceMap.Add(Triangle.CollapsedIndices[1], TriangleIndex);
-				VertexToFaceMap.Add(Triangle.CollapsedIndices[2], TriangleIndex);
-			}
-		}
-
-		// Clamp UV islands to the predominant block of each island. Will only happen if bClampUVIslands is true.
-		for (int32 ConflictiveTriangleIndex : ConflictiveTriangles)
-		{
-			FTriangleInfo& Triangle = Triangles[ConflictiveTriangleIndex];
-
-			// Skip the ones that have been fixed already
-			if (Triangle.bUVsFixed)
-			{
-				continue;
-			}
-
-			// Find triangles from the same UV Island
-			TArray<uint32> TriangleIndices;
-			GetUVIsland(Triangles, ConflictiveTriangleIndex, TriangleIndices, TexCoords, VertexToFaceMap);
-
-			// Get predominant BlockId != MAX_uint16
-			TArray<uint32> NumVerticesPerBlock;
-			NumVerticesPerBlock.SetNumZeroed(NumBlocks);
-
-			for (int32 TriangleIndex : TriangleIndices)
-			{
-				FTriangleInfo& OtherTriangle = Triangles[TriangleIndex];
-				for (int32 VertexIndex = 0; VertexIndex < 3; ++VertexIndex)
-				{
-					const uint16& BlockIndex = OtherTriangle.BlockIndices[VertexIndex];
-					if (BlockIndex != MAX_uint16)
-					{
-						NumVerticesPerBlock[BlockIndex]++;
-					}
-				}
-			}
-
-			uint16 BlockIndex = 0;
-			uint32 CurrentMaxVertices = 0;
-			for (int32 Index = 0; Index < NumBlocks; ++Index)
-			{
-				if (NumVerticesPerBlock[Index] > CurrentMaxVertices)
-				{
-					BlockIndex = Index;
-					CurrentMaxVertices = NumVerticesPerBlock[Index];
-				}
-			}
-
-			// Get the limits of the predominant block rect
-			const FLayoutBlock& LayoutBlock = Layout->Blocks[BlockIndex];
-
-			const float SmallNumber = 0.000001;
-			const float MinX = ((float)LayoutBlock.Min.X) / (float)LayoutGrid.X + SmallNumber;
-			const float MinY = ((float)LayoutBlock.Min.Y) / (float)LayoutGrid.Y + SmallNumber;
-			const float MaxX = (((float)LayoutBlock.Size.X + LayoutBlock.Min.X) / (float)LayoutGrid.X) - 2 * SmallNumber;
-			const float MaxY = (((float)LayoutBlock.Size.Y + LayoutBlock.Min.Y) / (float)LayoutGrid.Y) - 2 * SmallNumber;
-
-			// Iterate triangles and clamp the UVs
-			for (int32 TriangleIndex : TriangleIndices)
-			{
-				FTriangleInfo& OtherTriangle = Triangles[TriangleIndex];
-
-				for (int8 VertexIndex = 0; VertexIndex < 3; ++VertexIndex)
-				{
-					if (OtherTriangle.BlockIndices[VertexIndex] == BlockIndex)
-					{
-						continue;
-					}
-
-					OtherTriangle.BlockIndices[VertexIndex] = BlockIndex;
-
-					// Clamp UVs to the block they are assigned to
-					const int32 UVIndex = OtherTriangle.Indices[VertexIndex];
-					FVector2f& UV = TexCoords[UVIndex];
-					UV[0] = FMath::Clamp(UV[0], MinX, MaxX);
-					UV[1] = FMath::Clamp(UV[1], MinY, MaxY);
-					LayoutData[UVIndex] = BlockIndex;
-				}
-
-				OtherTriangle.bUVsFixed = true;
-			}
-		}
-
-		// Warn about vertices without a block id
-		check(StaticMeshOptions.LODIndex>=0);
-		int32 FirstLODToIgnoreWarnings = GeneratedLayout.Source->FirstLODToIgnoreWarnings;
-		if (FirstLODToIgnoreWarnings == -1 || StaticMeshOptions.LODIndex < FirstLODToIgnoreWarnings)
-		{
-			TArray<float> UnassignedUVs;
-			UnassignedUVs.Reserve(NumVertices / 100);
-
-			const FVector2f* UVs = TexCoords.GetData();
-			for (int32 VertexIndex = 0; VertexIndex < NumVertices; ++VertexIndex)
-			{
-				if (LayoutData[VertexIndex] == MAX_uint16)
-				{
-					UnassignedUVs.Add((*(UVs + VertexIndex))[0]);
-					UnassignedUVs.Add((*(UVs + VertexIndex))[1]);
-				}
-			}
-
-			if (!UnassignedUVs.IsEmpty())
-			{
-				FString Msg = FString::Printf(TEXT("Source mesh has %d vertices not assigned to any layout block in LOD %d"), UnassignedUVs.Num(), StaticMeshOptions.LODIndex);
-
-				ErrorLogMessageAttachedDataView AttachedDataView;
-				AttachedDataView.UnassignedUVs = UnassignedUVs.GetData();
-				AttachedDataView.UnassignedUVsSize = UnassignedUVs.Num();
-
-				ErrorLog->Add(Msg, AttachedDataView, ELMT_WARNING, ErrorContext);
-			}
-		}
-
-		// Create the layout block vertex buffer
-		uint8* LayoutBufferPtr = nullptr;
-		{
-			const int32 LayoutBufferIndex = Mesh->GetVertexBuffers().GetBufferCount();
-			Mesh->GetVertexBuffers().SetBufferCount(LayoutBufferIndex + 1);
-
-			// TODO
-			check(Layout->GetBlockCount() < MAX_uint16);
-			const EMeshBufferSemantic LayoutSemantic = EMeshBufferSemantic::LayoutBlock;
-			const int32 LayoutSemanticIndex = int32(LayoutChannel);
-			const EMeshBufferFormat LayoutFormat = bUseAbsoluteBlockIds ? EMeshBufferFormat::UInt64 : EMeshBufferFormat::UInt16;
-			const int32 LayoutComponents = 1;
-			const int32 LayoutOffset = 0;
-			int32 ElementSize = bUseAbsoluteBlockIds ? sizeof(uint64) : sizeof(uint16);
-			Mesh->GetVertexBuffers().SetBuffer
-			(
-				LayoutBufferIndex,
-				ElementSize,
-				1,
-				&LayoutSemantic, &LayoutSemanticIndex,
-				&LayoutFormat, &LayoutComponents,
-				&LayoutOffset
-			);
-			LayoutBufferPtr = Mesh->GetVertexBuffers().GetBufferData(LayoutBufferIndex);
-		}
-
-		// Copy UVs back to the mesh
-		{
-			uint8* pVertices = TexCoordData;
-			FVector2f* UVs = TexCoords.GetData();
-
-			for (int32 VertexIndex = 0; VertexIndex < NumVertices; ++VertexIndex)
-			{
-				FVector2f& UV = TexCoords[VertexIndex];
-
-				uint16 LayoutBlockIndex = LayoutData[VertexIndex];
-				if (Layout->Blocks.IsValidIndex(LayoutBlockIndex))
-				{
-					uint64 LayoutBlockId = Layout->Blocks[LayoutBlockIndex].Id;
-
-					UV = BlockRects[LayoutBlockIndex].Homogenize(UV);
-
-					// Replace block index by the actual id of the block
-					if (bUseAbsoluteBlockIds)
-					{
-						uint64* Ptr = reinterpret_cast<uint64*>(LayoutBufferPtr) + VertexIndex;
-						*Ptr = LayoutBlockId;
-					}
-					else
-					{
-						uint16* Ptr = reinterpret_cast<uint16*>(LayoutBufferPtr) + VertexIndex;
-						*Ptr = uint16(LayoutBlockId & 0xffff);
-					}
-				}
-				else
-				{
-					// Map vertices without block
-					if (bUseAbsoluteBlockIds)
-					{
-						uint64* Ptr = reinterpret_cast<uint64*>(LayoutBufferPtr) + VertexIndex;
-						*Ptr = DynamicMeshOptions.bEnsureAllVerticesHaveLayoutBlock ? 0 : std::numeric_limits<uint64>::max();
-					}
-					else
-					{
-						uint16* Ptr = reinterpret_cast<uint16*>(LayoutBufferPtr) + VertexIndex;
-						*Ptr = DynamicMeshOptions.bEnsureAllVerticesHaveLayoutBlock ? 0 : std::numeric_limits<uint16>::max();
-					}
-				}
-
-				// Copy UVs
-				if (TexCoordsChannel.Format == EMeshBufferFormat::Float32)
-				{
-					FVector2f* pUV = reinterpret_cast<FVector2f*>(pVertices);
-					*pUV = UV;
-				}
-				else if (TexCoordsChannel.Format == EMeshBufferFormat::Float16)
-				{
-					FFloat16* pUV = reinterpret_cast<FFloat16*>(pVertices);
-					pUV[0] = FFloat16(UV[0]);
-					pUV[1] = FFloat16(UV[1]);
-				}
-
-				pVertices += elemSize;
-			}
-		}
-	}
-	
-
-    //---------------------------------------------------------------------------------------------
 	FMeshTask CodeGenerator::GenerateMesh(const FMeshGenerationStaticOptions& StaticOptions, FMeshOptionsTask InOptions, const Ptr<const NodeMesh>& InUntypedNode)
     {
         if (!InUntypedNode)
@@ -1132,7 +621,76 @@ namespace mu
 			Requisites
 		);
 
-    }
+	}
+
+
+	Ptr<ASTOp> CodeGenerator::GenerateLayoutOpsAndResult(
+		const FMeshGenerationDynamicOptions& Options, 
+		Ptr<ASTOp> LastMeshOp, 
+		const TArray<Ptr<NodeLayout>>& OriginalLayouts,
+		uint32 MeshPrefix,
+		FMeshGenerationResult& OutResult)
+	{
+		if (!Options.bLayouts)
+		{
+			return LastMeshOp;
+		}
+
+		bool bUseAbsoluteBlockIds = false;
+		bool bIsOverridingLayouts = !Options.OverrideLayouts.IsEmpty();
+		if (!bIsOverridingLayouts)
+		{
+			OutResult.GeneratedLayouts.Reserve(OriginalLayouts.Num());
+			for (int32 LayoutIndex = 0; LayoutIndex < OriginalLayouts.Num(); ++LayoutIndex)
+			{
+				FGeneratedLayout GeneratedData;
+				Ptr<NodeLayout> LayoutNode = OriginalLayouts[LayoutIndex];
+				if (LayoutNode)
+				{
+					GeneratedData.Source = LayoutNode;
+					GeneratedData.Layout = this->GenerateLayout(LayoutNode, MeshPrefix);
+				}
+				OutResult.GeneratedLayouts.Add(GeneratedData);
+			}
+		}
+		else
+		{
+			// In this case we need the layout block ids to use the ids in the parent layout, and not be prefixed with
+			// the current mesh id prefix. For this reason we need them to be absolute.
+			bUseAbsoluteBlockIds = true;
+
+			// We need to apply the transform of the layouts used to override
+			OutResult.GeneratedLayouts.Reserve(Options.OverrideLayouts.Num());
+			for (int32 LayoutIndex = 0; LayoutIndex < Options.OverrideLayouts.Num(); ++LayoutIndex)
+			{
+				const FGeneratedLayout& OverrideData = Options.OverrideLayouts[LayoutIndex];
+				OutResult.GeneratedLayouts.Add(OverrideData);
+			}
+		}
+
+		// Generate the chain of ops preparing the mesh for the layouts
+		for (int32 LayoutIndex = 0; LayoutIndex < OutResult.GeneratedLayouts.Num(); ++LayoutIndex)
+		{
+			TSharedPtr<const mu::FLayout> Layout = OutResult.GeneratedLayouts[LayoutIndex].Layout;
+
+			Ptr<ASTOpConstantResource> LayoutOp = new ASTOpConstantResource;
+			LayoutOp->Type = EOpType::LA_CONSTANT;
+			LayoutOp->SetValue(Layout, CompilerOptions->OptimisationOptions.DiskCacheContext);
+
+			Ptr<ASTOpMeshPrepareLayout> PrepareOp = new ASTOpMeshPrepareLayout;
+			PrepareOp->Mesh = LastMeshOp;
+			PrepareOp->Layout = LayoutOp;
+			PrepareOp->LayoutChannel = LayoutIndex;
+			PrepareOp->bUseAbsoluteBlockIds = bUseAbsoluteBlockIds;
+			PrepareOp->bNormalizeUVs = Options.bNormalizeUVs;
+			PrepareOp->bClampUVIslands = Options.bClampUVIslands;
+			PrepareOp->bEnsureAllVerticesHaveLayoutBlock = Options.bEnsureAllVerticesHaveLayoutBlock;
+
+			LastMeshOp = PrepareOp;
+		}
+
+		return LastMeshOp;
+	}
 
 
 	FMeshTask CodeGenerator::GenerateMesh_Constant(const FMeshGenerationStaticOptions& StaticOptions, FMeshOptionsTask InOptions, const NodeMeshConstant* Node )
@@ -1158,20 +716,20 @@ namespace mu
 			return UE::Tasks::MakeCompletedTask<FMeshGenerationResult>(Result);
 		}
 
-        Ptr<ASTOpConstantResource> ConstantOp = new ASTOpConstantResource();
-		ConstantOp->Type = EOpType::ME_CONSTANT;
-		ConstantOp->SourceDataDescriptor = Node->SourceDataDescriptor;
-
-		FMeshGenerationResult Result;
-		Result.BaseMeshOp = ConstantOp;
-		Result.MeshOp = ConstantOp;
-
 		if (!Mesh)
 		{
+			Ptr<ASTOpConstantResource> ConstantOp = new ASTOpConstantResource();
+			ConstantOp->Type = EOpType::ME_CONSTANT;
+			ConstantOp->SourceDataDescriptor = Node->SourceDataDescriptor;
+
+			FMeshGenerationResult Result;
+			Result.BaseMeshOp = ConstantOp;
+			Result.MeshOp = ConstantOp;
+
 			// This data is required
 			TSharedPtr<FMesh> EmptyMesh = MakeShared<FMesh>();
 			ConstantOp->SetValue(EmptyMesh, CompilerOptions->OptimisationOptions.DiskCacheContext);
-			EmptyMesh->MeshIDPrefix = ConstantOp->GetValueHash();
+			EmptyMesh->MeshIDPrefix = UniqueMeshIds.EnsureUnique(ConstantOp->GetValueHash());
 
 			// Log an error message
 			ErrorLog->Add("Constant mesh not set.", ELMT_WARNING, Node->GetMessageContext());
@@ -1202,7 +760,7 @@ namespace mu
 		// This task does the necessary processing of the mesh when it is available
 		Ptr<const NodeMeshConstant> NodeCopy = Node;
 		FMeshTask MeshProcessingTask = GenerateMeshConstantPipe.Launch( TEXT("MutableGenerateMeshConstant"),
-			[ReferenceMesh=Mesh,ResolveMesh,Result,ConstantOp, StaticOptions, InOptions, NodeCopy,
+			[ReferenceMesh=Mesh,ResolveMesh,StaticOptions, InOptions, NodeCopy,
 			// This is read-only
 			//CompilerOptions=this->CompilerOptions,
 			this
@@ -1217,6 +775,10 @@ namespace mu
 					return FMeshGenerationResult();
 				}
 
+				Ptr<ASTOpConstantResource> ConstantOp = new ASTOpConstantResource();
+				ConstantOp->Type = EOpType::ME_CONSTANT;
+				ConstantOp->SourceDataDescriptor = NodeCopy->SourceDataDescriptor;
+
 				// Separate the tags from the mesh
 				TArray<FString> Tags = ReferenceMesh->Tags;
 				if (!Mesh->Tags.IsEmpty())
@@ -1238,81 +800,71 @@ namespace mu
 				// Find out if we can (or have to) reuse a mesh that we have already generated.
 				FGenerateMeshConstantState::FGeneratedConstantMesh DuplicateOf;
 				uint32 ThisMeshHash = HashCombineFast(GetTypeHash(Mesh->GetVertexCount()), GetTypeHash(Mesh->GetIndexCount()));
-				TSharedPtr<FMesh> ClonedMesh;
+
+				TArray<FGenerateMeshConstantState::FGeneratedConstantMesh>& CachedCandidates = GenerateMeshConstantState.GeneratedConstantMeshes.FindOrAdd(ThisMeshHash, {});
+				for (const FGenerateMeshConstantState::FGeneratedConstantMesh& Candidate : CachedCandidates)
 				{
-					TArray<FGenerateMeshConstantState::FGeneratedConstantMesh>& CachedCandidates = GenerateMeshConstantState.GeneratedConstantMeshes.FindOrAdd(ThisMeshHash, {});
-					for (const FGenerateMeshConstantState::FGeneratedConstantMesh& Candidate : CachedCandidates)
-					{
-						bool bCompareLayouts = Options.bLayouts && !bIsOverridingLayouts;
+					bool bCompareLayouts = Options.bLayouts && !bIsOverridingLayouts;
 
-						if (Candidate.Mesh->IsSimilar(*Mesh, bCompareLayouts))
+					if (Candidate.Mesh->IsSimilar(*Mesh, bCompareLayouts))
+					{
+						// If it is similar and we are overriding the layouts, we must compare the layouts of the candidate with the ones
+						// we are using to override.
+						if (bIsOverridingLayouts)
 						{
-							// If it is similar and we are overriding the layouts, we must compare the layouts of the candidate with the ones
-							// we are using to override.
-							if (bIsOverridingLayouts)
+							if (Candidate.Mesh->GetLayoutCount() != Options.OverrideLayouts.Num())
 							{
-								if (Candidate.Mesh->GetLayoutCount() != Options.OverrideLayouts.Num())
-								{
-									continue;
-								}
-
-								bool bLayoutsAreEqual = true;
-								for (int32 l = 0; l < Candidate.Mesh->GetLayoutCount(); ++l)
-								{
-									bLayoutsAreEqual = (*Candidate.Mesh->GetLayout(l) == *Options.OverrideLayouts[l].Layout);
-									if (!bLayoutsAreEqual)
-									{
-										break;
-									}
-								}
+								continue;
+							}
 
+							bool bLayoutsAreEqual = true;
+							for (int32 l = 0; l < Candidate.Mesh->GetLayoutCount(); ++l)
+							{
+								bLayoutsAreEqual = (*Candidate.Mesh->GetLayout(l) == *Options.OverrideLayouts[l].Layout);
 								if (!bLayoutsAreEqual)
 								{
-									continue;
+									break;
 								}
 							}
 
-							DuplicateOf = Candidate;
-							break;
+							if (!bLayoutsAreEqual)
+							{
+								continue;
+							}
 						}
+
+						DuplicateOf = Candidate;
+						break;
 					}
+				}
 
-					if (!DuplicateOf.Mesh)
-					{
-						// We need to clone the mesh in the node because we will modify it.
-						ClonedMesh = Mesh->Clone();
-						ClonedMesh->EnsureSurfaceData();
+				FMeshGenerationResult Result;
 
-						ConstantOp->SetValue(ClonedMesh, CompilerOptions->OptimisationOptions.DiskCacheContext);
+				Ptr<ASTOp> LastMeshOp = ConstantOp;
 
-						// Add the unique vertex ID prefix in all cases, since it is free memory-wise
-						uint32 MeshIDPrefix = uint32(ConstantOp->GetValueHash());
-						{
-							// Ensure the ID group is unique
-							bool bValid = false;
-							do
-							{
-								bool bAlreadyPresent = false;
-								GenerateMeshConstantState.UniqueVertexIDGroups.FindOrAdd(MeshIDPrefix, &bAlreadyPresent);
-								bValid = !bAlreadyPresent && MeshIDPrefix != 0;
-								if (!bValid)
-								{
-									++MeshIDPrefix;
-								}
-							} while (bValid);
+				TSharedPtr<FMesh> ClonedMesh;
+				if (!DuplicateOf.Mesh)
+				{
+					// We need to clone the mesh in the node because we will modify it.
+					ClonedMesh = Mesh->Clone();
+					ClonedMesh->EnsureSurfaceData();
 
-							ClonedMesh->MeshIDPrefix = MeshIDPrefix;
-						}
+					ConstantOp->SetValue(ClonedMesh, CompilerOptions->OptimisationOptions.DiskCacheContext);
 
-						// Add the constant data
-						FGenerateMeshConstantState::FGeneratedConstantMesh MeshEntry;
-						MeshEntry.Mesh = ClonedMesh;
-						MeshEntry.LastMeshOp = ConstantOp;
-						CachedCandidates.Add(MeshEntry);
-					}
-				}
+					// Add the unique vertex ID prefix in all cases, since it is free memory-wise
+					uint32 MeshIDPrefix = uint32(ConstantOp->GetValueHash());
+					MeshIDPrefix = UniqueMeshIds.EnsureUnique(MeshIDPrefix);
+					ClonedMesh->MeshIDPrefix = MeshIDPrefix;
 
-				Ptr<ASTOp> LastMeshOp = ConstantOp;
+					LastMeshOp = GenerateLayoutOpsAndResult(Options, ConstantOp, NodeCopy->Layouts, ClonedMesh->MeshIDPrefix, Result);
+
+					// Add the constant data
+					FGenerateMeshConstantState::FGeneratedConstantMesh MeshEntry;
+					MeshEntry.Mesh = ClonedMesh;
+					MeshEntry.LastMeshOp = LastMeshOp;
+					MeshEntry.GeneratedLayouts = Result.GeneratedLayouts;
+					CachedCandidates.Add(MeshEntry);
+				}
 
 				if (DuplicateOf.Mesh)
 				{
@@ -1325,62 +877,13 @@ namespace mu
 						}
 						else
 						{
-							for (int32 l = 0; l < DuplicateOf.Mesh->GetLayoutCount(); ++l)
-							{
-								TSharedPtr<const FLayout> DuplicatedLayout = DuplicateOf.Mesh->GetLayout(l);
-								Result.GeneratedLayouts.Add({ DuplicatedLayout,0 });
-							}
+							Result.GeneratedLayouts = DuplicateOf.GeneratedLayouts;
 						}
 					}
 
 					LastMeshOp = DuplicateOf.LastMeshOp;
 					ConstantOp = nullptr;
 				}
-				else
-				{
-					if (Options.bLayouts)
-					{
-						if (!bIsOverridingLayouts)
-						{
-							// Apply whatever transform is necessary for every layout
-							for (int32 LayoutIndex = 0; LayoutIndex < NodeCopy->Layouts.Num(); ++LayoutIndex)
-							{
-								Ptr<NodeLayout> LayoutNode = NodeCopy->Layouts[LayoutIndex];
-								if (!LayoutNode)
-								{
-									continue;
-								}
-
-								FGeneratedLayout GeneratedData;
-								GeneratedData.Source = LayoutNode;
-								GeneratedData.Layout = this->GenerateLayout(LayoutNode, ClonedMesh->MeshIDPrefix);
-								const void* Context = Options.OverrideContext.Get(NodeCopy->GetMessageContext());
-
-								bool bUseAbsoluteBlockIds = false;
-								this->PrepareMeshForLayout(GeneratedData, ClonedMesh, LayoutIndex, Context, StaticOptions, Options, bUseAbsoluteBlockIds);
-
-								Result.GeneratedLayouts.Add(GeneratedData);
-							}
-						}
-						else
-						{
-							// We need to apply the transform of the layouts used to override
-							for (int32 LayoutIndex = 0; LayoutIndex < Options.OverrideLayouts.Num(); ++LayoutIndex)
-							{
-								const FGeneratedLayout& OverrideData = Options.OverrideLayouts[LayoutIndex];
-								TSharedPtr<const FLayout> GeneratedLayout = OverrideData.Layout;
-								const void* Context = Options.OverrideContext.Get(NodeCopy->GetMessageContext());
-
-								// In this case we need the layout block ids to use the ids in the parent layout, and not be prefixed with
-								// the current mesh id prefix. For this reason we need them to be absolute.
-								bool bUseAbsoluteBlockIds = true;
-								this->PrepareMeshForLayout(OverrideData, ClonedMesh, LayoutIndex, Context, StaticOptions, Options, bUseAbsoluteBlockIds);
-
-								Result.GeneratedLayouts.Add(OverrideData);
-							}
-						}
-					}
-				}
 
 				Result.BaseMeshOp = LastMeshOp;
 
@@ -1813,55 +1316,72 @@ namespace mu
 	}
 
 
+	uint32 CodeGenerator::FUniqueMeshIds::EnsureUnique(uint32 Id)
+	{
+		UE::TUniqueLock Lock(Mutex);
+		bool bValid = false;
+		do
+		{
+			bool bAlreadyPresent = false;
+			Map.FindOrAdd(Id, &bAlreadyPresent);
+			bValid = !bAlreadyPresent && Id != 0;
+			if (!bValid)
+			{
+				++Id;
+			}
+		} while (bValid);
+
+		return Id;
+	}
+
+
 	FMeshTask CodeGenerator::GenerateMesh_Parameter(const FMeshGenerationStaticOptions& StaticOptions, FMeshOptionsTask InOptions, const NodeMeshParameter* Node)
 	{
 		// Local pipe because we call GenerateRange in there.
 		FMeshTask MeshParamTask = this->LocalPipe.Launch(TEXT("MutableMeshParameter"),
 			[InOptions, Node, StaticOptions, this]() mutable
 			{
+				FMeshGenerationDynamicOptions Options = InOptions.GetResult();
+
+				FMeshGenerationResult Result;
 				Ptr<ASTOpParameter> Op;
+
+				FParameterDesc param;
+				param.Name = Node->Name;
+				bool bParseOk = FGuid::Parse(Node->UID, param.UID);
+				check(bParseOk);
+				param.Type = EParameterType::Mesh;
+				param.DefaultValue.Set<FParamAssetType>(Node->DefaultValue);
+
+				Op = new ASTOpParameter();
+				Op->Type = EOpType::ME_PARAMETER;
+				Op->Parameter = param;
+				Op->LODIndex = StaticOptions.LODIndex;
+				Op->SectionIndex = Node->SectionIndex;
+
+				// Assign an ID prefix to the mesh. 
+				uint32 MeshID = HashCombine(GetTypeHash(Op->LODIndex), GetTypeHash(param.UID));
+				MeshID = HashCombine(MeshID, GetTypeHash(Op->SectionIndex));
+				MeshID = UniqueMeshIds.EnsureUnique(MeshID);
+				Op->MeshID = MeshID;
+
 				{
 					UE::TUniqueLock Lock(FirstPass.ParameterNodes.Mutex);
 
-					Ptr<ASTOpParameter> Found;					
-					TArray<Ptr<ASTOpParameter>>& ArrayFound = FirstPass.ParameterNodes.MeshParametersCache.FindOrAdd(Node);
+					TArray<TPair<Ptr<ASTOpParameter>, FMeshGenerationResult>>& ArrayFound = FirstPass.ParameterNodes.MeshParametersCache.FindOrAdd(Node);
 					if (ArrayFound.IsValidIndex(StaticOptions.LODIndex))
 					{
-						Found = ArrayFound[StaticOptions.LODIndex];
+						return ArrayFound[StaticOptions.LODIndex].Value;
 					}
 
-					if (Found)
-					{
-						check(Found->Type == EOpType::ME_PARAMETER);
-
-						FMeshGenerationResult Result;
-						Result.BaseMeshOp = Found;
-						Result.MeshOp = Found;
-						// Rest of the parameters?
-						//Result.LayoutOps ...
-						// ... maybe a special cache for these parameters?
-						return Result;
-					}
-
-					FParameterDesc param;
-					param.Name = Node->Name;
-					bool bParseOk = FGuid::Parse(Node->UID, param.UID);
-					check(bParseOk);
-					param.Type = EParameterType::Mesh;
-					param.DefaultValue.Set<FParamAssetType>(Node->DefaultValue);
+					ArrayFound.SetNum(StaticOptions.LODIndex+1);
+					ArrayFound[StaticOptions.LODIndex].Key = Op;
 
-					Op = new ASTOpParameter();
-					Op->Type = EOpType::ME_PARAMETER;
-					Op->Parameter = param;
-					Op->LODIndex = StaticOptions.LODIndex;
-					Op->SectionIndex = Node->SectionIndex;
-
-					if (!ArrayFound.IsValidIndex(StaticOptions.LODIndex))
-					{
-						ArrayFound.SetNumZeroed(StaticOptions.LODIndex+1);
-					}
-
-					ArrayFound[StaticOptions.LODIndex] = Op;
+					// Fill the result structure. Some ops will be completed outside the lock.
+					Ptr<ASTOp> LastMeshOp = GenerateLayoutOpsAndResult(Options, Op, Node->Layouts, Op->MeshID, Result);
+					Result.MeshOp = LastMeshOp;
+					Result.BaseMeshOp = LastMeshOp;
+					ArrayFound[StaticOptions.LODIndex].Value = Result;
 				}
 
 				int32 LODIndex = StaticOptions.LODIndex;
@@ -1876,11 +1396,6 @@ namespace mu
 					Op->Ranges.Emplace(Op.get(), RangeResult.sizeOp, RangeResult.rangeName, RangeResult.rangeUID);
 				}
 
-				FMeshGenerationResult Result;
-				Result.MeshOp = Op;
-				// Rest of the parameters?
-				//Result.LayoutOps ...
-
 				return Result;
 			},
 			InOptions
diff --git a/Engine/Plugins/Mutable/Source/MutableTools/Private/MuT/Compiler.cpp b/Engine/Plugins/Mutable/Source/MutableTools/Private/MuT/Compiler.cpp
--- a/Engine/Plugins/Mutable/Source/MutableTools/Private/MuT/Compiler.cpp
+++ b/Engine/Plugins/Mutable/Source/MutableTools/Private/MuT/Compiler.cpp
@@ -138,6 +138,7 @@ namespace mu
 		{ true,		{ 0, 1, 1, 1, 0, 0, 0, 0, 0, 0, 0, 0, 0, 0, 0, 0, 0, 0, 0, 0, 0, 0, 0, 0, 0, 0, 0, 0, 0 }		},	// IM_TRANSFORM
 
 		{ true,		{ 0, 0, 0, 0, 0, 0, 0, 0, 0, 0, 0, 0, 0, 0, 0, 0, 0, 0, 0, 0, 0, 0, 0, 0, 0, 0, 0, 0, 0 }		},	// ME_APPLYLAYOUT
+		{ true,		{ 0, 0, 0, 0, 0, 0, 0, 0, 0, 0, 0, 0, 0, 0, 0, 0, 0, 0, 0, 0, 0, 0, 0, 0, 0, 0, 0, 0, 0 }		},	// ME_PREPARELAYOUT
 		{ true,		{ 0, 0, 0, 0, 0, 0, 0, 0, 0, 0, 0, 0, 0, 0, 0, 0, 0, 0, 0, 0, 0, 0, 0, 0, 0, 0, 0, 0, 0 }		},	// ME_DIFFERENCE
 		{ true,		{ 0, 0, 0, 0, 0, 0, 0, 0, 0, 0, 0, 0, 0, 0, 0, 0, 0, 0, 0, 0, 0, 0, 0, 0, 0, 0, 0, 0, 0 }		},	// ME_MORPH
 		{ true,		{ 0, 0, 0, 0, 0, 0, 0, 0, 0, 0, 0, 0, 0, 0, 0, 0, 0, 0, 0, 0, 0, 0, 0, 0, 0, 0, 0, 0, 0 }		},	// ME_MERGE
@@ -465,13 +466,13 @@ namespace mu
 					{
 						Parameters.Add(Entry.Value->Parameter);
 					}
-					for (const TPair<Ptr<const NodeMeshParameter>, TArray<Ptr<ASTOpParameter>>>& Entry : Generator.FirstPass.ParameterNodes.MeshParametersCache)
+					for (const TPair<Ptr<const NodeMeshParameter>, TArray<TPair<Ptr<ASTOpParameter>,FMeshGenerationResult>>>& Entry : Generator.FirstPass.ParameterNodes.MeshParametersCache)
 					{
-						for (const Ptr<ASTOpParameter>& ArrayEntry : Entry.Value)
+						for (const TPair<Ptr<ASTOpParameter>, FMeshGenerationResult>& ArrayEntry : Entry.Value)
 						{
-							if (ArrayEntry)
+							if (ArrayEntry.Key)
 							{
-								Parameters.AddUnique(ArrayEntry->Parameter);
+								Parameters.AddUnique(ArrayEntry.Key->Parameter);
 							}
 						}
 					}
@@ -513,11 +514,11 @@ namespace mu
 										MessageContexts.AddUnique(Entry.Key->GetMessageContext());
 									}
 								}
-								for (const TPair<Ptr<const NodeMeshParameter>, TArray<Ptr<ASTOpParameter>>>& Entry : Generator.FirstPass.ParameterNodes.MeshParametersCache)
+								for (const TPair<Ptr<const NodeMeshParameter>, TArray<TPair<Ptr<ASTOpParameter>, FMeshGenerationResult>>>& Entry : Generator.FirstPass.ParameterNodes.MeshParametersCache)
 								{
-									for (const Ptr<ASTOpParameter>& ArrayEntry : Entry.Value)
+									for (const TPair<Ptr<ASTOpParameter>, FMeshGenerationResult>& ArrayEntry : Entry.Value)
 									{
-										if (ArrayEntry && ArrayEntry->Parameter.Name.Compare(CurrentParameterName))
+										if (ArrayEntry.Key && ArrayEntry.Key->Parameter.Name.Compare(CurrentParameterName))
 										{
 											MessageContexts.AddUnique(Entry.Key->GetMessageContext());
 										}
diff --git a/Engine/Plugins/Mutable/Source/MutableTools/Private/MuT/ErrorLog.cpp b/Engine/Plugins/Mutable/Source/MutableTools/Private/MuT/ErrorLog.cpp
--- a/Engine/Plugins/Mutable/Source/MutableTools/Private/MuT/ErrorLog.cpp
+++ b/Engine/Plugins/Mutable/Source/MutableTools/Private/MuT/ErrorLog.cpp
@@ -284,6 +284,7 @@ namespace mu
 		TEXT("IM_TRANSFORM     "),
 
 		TEXT("ME_APPLYLAYOUT   "),
+		TEXT("ME_PREPARELAYOUT "),
 		TEXT("ME_DIFFERENCE    "),
 		TEXT("ME_MORPH         "),
 		TEXT("ME_MERGE         "),
diff --git a/Engine/Plugins/Mutable/Source/MutableTools/Private/MuT/NodeMesh.cpp b/Engine/Plugins/Mutable/Source/MutableTools/Private/MuT/NodeMesh.cpp
deleted file mode 100644
diff --git a/Engine/Plugins/Mutable/Source/MutableTools/Public/MuT/NodeMesh.h b/Engine/Plugins/Mutable/Source/MutableTools/Public/MuT/NodeMesh.h
--- a/Engine/Plugins/Mutable/Source/MutableTools/Public/MuT/NodeMesh.h
+++ b/Engine/Plugins/Mutable/Source/MutableTools/Public/MuT/NodeMesh.h
@@ -40,35 +40,5 @@ namespace mu
 	};
 
 
-	/** Helper structs for mesh utility methods below. */
-	struct FTriangleInfo
-	{
-		/** Vertex indices in the original mesh. */
-		uint32 Indices[3];
-
-		/** Vertex indices in the collapsed vertex list of the mesh. */
-		uint32 CollapsedIndices[3];
-
-		/** Optional data with layout block indices. */
-		uint16 BlockIndices[3];
-
-		/** Optional data with a flag indicating the UVs have changed during layout for this trioangle. */
-		bool bUVsFixed = false;
-	};
-
-
-	/** Fill an array with the indices of all triangles belonging to the same UV island as InFirstTriangle. 
-	*/
-	extern void GetUVIsland(TArray<FTriangleInfo>& InTriangles,
-		const uint32 InFirstTriangle,
-		TArray<uint32>& OutTriangleIndices,
-		const TArray<FVector2f>& InUVs,
-		const TMultiMap<int32, uint32>& InVertexToTriangleMap);
-
-
-	/** Create a map from vertices into vertices, collapsing vertices that have the same position.
-	*/
-	extern void MeshCreateCollapsedVertexMap(const mu::FMesh* Mesh, TArray<int32>& CollapsedVertices);
-
 }
 
diff --git a/Engine/Plugins/Mutable/Source/MutableTools/Public/MuT/NodeMeshParameter.h b/Engine/Plugins/Mutable/Source/MutableTools/Public/MuT/NodeMeshParameter.h
--- a/Engine/Plugins/Mutable/Source/MutableTools/Public/MuT/NodeMeshParameter.h
+++ b/Engine/Plugins/Mutable/Source/MutableTools/Public/MuT/NodeMeshParameter.h
@@ -7,6 +7,7 @@
 #include "MuT/Node.h"
 #include "MuT/NodeMesh.h"
 #include "MuT/NodeRange.h"
+#include "MuT/NodeLayout.h"
 
 namespace mu
 {
@@ -19,7 +20,10 @@ namespace mu
 		/** Full asset name for the default value of the parameter. */
 		FName DefaultValue;
 
-		/** TODO */
+		/** */
+		TArray<Ptr<NodeLayout>> Layouts;
+
+		/** */
 		uint8 SectionIndex = 0;
 
 		/** Name of the parameter */
