[Mutable] Compile Mesh Metadata.

#rnx

diff --git a/Engine/Plugins/Mutable/Source/CustomizableObject/Internal/MuCO/CustomizableObjectPrivate.h b/Engine/Plugins/Mutable/Source/CustomizableObject/Internal/MuCO/CustomizableObjectPrivate.h
--- a/Engine/Plugins/Mutable/Source/CustomizableObject/Internal/MuCO/CustomizableObjectPrivate.h
+++ b/Engine/Plugins/Mutable/Source/CustomizableObject/Internal/MuCO/CustomizableObjectPrivate.h
@@ -1595,6 +1595,8 @@ enum class ECustomizableObjectVersions : int32
 
 	AddedMeshIDToMeshParamOp,
 
+	ClothMorphMeshMetadata,
+
 	// -----<new versions can be added above this line>--------
 	LastCustomizableObjectVersion
 };
diff --git a/Engine/Plugins/Mutable/Source/CustomizableObjectEditor/Private/MuCOE/GenerateMutableSource/GenerateMutableSourceMesh.cpp b/Engine/Plugins/Mutable/Source/CustomizableObjectEditor/Private/MuCOE/GenerateMutableSource/GenerateMutableSourceMesh.cpp
--- a/Engine/Plugins/Mutable/Source/CustomizableObjectEditor/Private/MuCOE/GenerateMutableSource/GenerateMutableSourceMesh.cpp
+++ b/Engine/Plugins/Mutable/Source/CustomizableObjectEditor/Private/MuCOE/GenerateMutableSource/GenerateMutableSourceMesh.cpp
@@ -1420,6 +1420,8 @@ TSharedPtr<mu::FMesh> ConvertSkeletalMeshToMutable(
 	CompletionEvent.Trigger();
 	CompletionEvent.Wait();
 
+	uint32 MorphUniqueHash = 0;
+	
 	int32 NextBufferIndex = MutableMesh->VertexBuffers.Buffers.Num();
 	if (!bIgnoreMorphs && Context.Options.bRealTimeMorphTargetsEnabled)
 	{
@@ -1643,9 +1645,9 @@ TSharedPtr<mu::FMesh> ConvertSkeletalMeshToMutable(
 				const uint32 DataHash = CityHash32(
 					reinterpret_cast<const char*>(MorphsMeshData.GetData()), MorphsMeshData.Num() * sizeof(FMorphTargetVertexData));
 
-				uint32 UniqueHash = DataHash == InvalidResourceId ? InvalidResourceId + 1 : DataHash;
+				MorphUniqueHash = DataHash == InvalidResourceId ? InvalidResourceId + 1 : DataHash;
 
-				const FRealTimeMorphMeshData* FoundHash = Context.RealTimeMorphTargetPerMeshData.Find(UniqueHash);
+				const FRealTimeMorphMeshData* FoundHash = Context.RealTimeMorphTargetPerMeshData.Find(MorphUniqueHash);
 
 				bool bIsDataAlreadyCollected = false;
 				if (FoundHash)
@@ -1665,7 +1667,7 @@ TSharedPtr<mu::FMesh> ConvertSkeletalMeshToMutable(
 					uint32 NumTries = 0;
 					for (; NumTries < TNumericLimits<uint32>::Max(); ++NumTries)
 					{
-						FoundHash = Context.RealTimeMorphTargetPerMeshData.Find(UniqueHash);
+						FoundHash = Context.RealTimeMorphTargetPerMeshData.Find(MorphUniqueHash);
 
 						if (!FoundHash)
 						{
@@ -1683,17 +1685,17 @@ TSharedPtr<mu::FMesh> ConvertSkeletalMeshToMutable(
 							break;
 						}
 
-						UniqueHash = UniqueHash + 1 == InvalidResourceId ? InvalidResourceId + 1 : UniqueHash + 1;
+						MorphUniqueHash = MorphUniqueHash + 1 == InvalidResourceId ? InvalidResourceId + 1 : MorphUniqueHash + 1;
 					}
 
 					if (NumTries == TNumericLimits<uint32>::Max())
 					{
 						UE_LOG(LogMutable, Warning, TEXT("Maximum number of meshes with real time morphs reached, some morphs may not work as expected."));
-						UniqueHash = InvalidResourceId;
+						MorphUniqueHash = InvalidResourceId;
 					}
 				}
 
-				if (UniqueHash != InvalidResourceId)
+				if (MorphUniqueHash != InvalidResourceId)
 				{
 					// Add the data block reference to the mesh vertices.
 					for (int32 VertexIdx = VertexStart; VertexIdx < VertexStart + VertexCount && VertexIdx < Vertices.Num(); ++VertexIdx)
@@ -1704,19 +1706,19 @@ TSharedPtr<mu::FMesh> ConvertSkeletalMeshToMutable(
 						// mesh generation, but we need to know if a vertex has morphs in both phases. 
 						if (VertexMorphsOffsetAndCountBufferView[VertexIdx - VertexStart] != 0)
 						{
-							VertexMorphsResourceIdBufferView[VertexIdx - VertexStart] = UniqueHash;
+							VertexMorphsResourceIdBufferView[VertexIdx - VertexStart] = MorphUniqueHash;
 						}
 					}
 
 					FCustomizableObjectStreameableResourceId StreamedMorphResource;
-					StreamedMorphResource.Id = UniqueHash;
+					StreamedMorphResource.Id = MorphUniqueHash;
 					StreamedMorphResource.Type = static_cast<uint8>(FCustomizableObjectStreameableResourceId::EType::RealTimeMorphTarget);
 
 					MutableMesh->AddStreamedResource(BitCast<uint64>(StreamedMorphResource));
 
 					if (!bIsDataAlreadyCollected)
 					{
-						FRealTimeMorphMeshData& RealTimeMorphMeshData = Context.RealTimeMorphTargetPerMeshData.FindOrAdd(UniqueHash);
+						FRealTimeMorphMeshData& RealTimeMorphMeshData = Context.RealTimeMorphTargetPerMeshData.FindOrAdd(MorphUniqueHash);
 
 						check(RealTimeMorphMeshData.NameResolutionMap.IsEmpty());
 						check(RealTimeMorphMeshData.Data.IsEmpty());
@@ -1732,6 +1734,9 @@ TSharedPtr<mu::FMesh> ConvertSkeletalMeshToMutable(
 		}
 	}
 
+
+	uint32 ClothingUniqueHash = 0;
+	
 	// Clothing vertex info.
 	if (!bIgnoreSkinning && Context.Options.bClothingEnabled)
 	{
@@ -1967,9 +1972,9 @@ TSharedPtr<mu::FMesh> ConvertSkeletalMeshToMutable(
 				reinterpret_cast<const char*>(ClothingMeshData.Data.GetData()),
 				ClothingMeshData.Data.Num() * sizeof(FCustomizableObjectMeshToMeshVertData));
 
-			uint32 UniqueHash = DataHash == InvalidResourceId ? DataHash + 1 : DataHash;
+			ClothingUniqueHash = DataHash == InvalidResourceId ? DataHash + 1 : DataHash;
 
-			const FClothingMeshData* FoundHash = Context.ClothingPerMeshData.Find(UniqueHash);
+			const FClothingMeshData* FoundHash = Context.ClothingPerMeshData.Find(ClothingUniqueHash);
 
 			bool bIsDataAlreadyCollected = false;
 
@@ -1995,7 +2000,7 @@ TSharedPtr<mu::FMesh> ConvertSkeletalMeshToMutable(
 				uint32 NumTries = 0;
 				for (; NumTries < TNumericLimits<uint32>::Max(); ++NumTries)
 				{
-					FoundHash = Context.ClothingPerMeshData.Find(UniqueHash);
+					FoundHash = Context.ClothingPerMeshData.Find(ClothingUniqueHash);
 
 					if (!FoundHash)
 					{
@@ -2009,31 +2014,31 @@ TSharedPtr<mu::FMesh> ConvertSkeletalMeshToMutable(
 						break;
 					}
 
-					UniqueHash = UniqueHash + 1 == InvalidResourceId ? InvalidResourceId + 1 : UniqueHash + 1;
+					ClothingUniqueHash = ClothingUniqueHash + 1 == InvalidResourceId ? InvalidResourceId + 1 : ClothingUniqueHash + 1;
 				}
 
 				if (NumTries == TNumericLimits<uint32>::Max())
 				{
 					UE_LOG(LogMutable, Warning, TEXT("Maximum number of meshes with clothing reached, some cloth meshes may not work as expected."));
-					UniqueHash = InvalidResourceId;
+					ClothingUniqueHash = InvalidResourceId;
 				}
 			}
 
 			// fill the resource buffer with the generated unique id.
 			for (uint32& Elem : ClothSectionResourceIdView)
 			{
-				Elem = UniqueHash;
+				Elem = ClothingUniqueHash;
 			}
 
 			FCustomizableObjectStreameableResourceId StreamedClothResource;
-			StreamedClothResource.Id = UniqueHash;
+			StreamedClothResource.Id = ClothingUniqueHash;
 			StreamedClothResource.Type = static_cast<uint8>(FCustomizableObjectStreameableResourceId::EType::Clothing);
 
 			MutableMesh->AddStreamedResource(BitCast<uint64>(StreamedClothResource));
 
 			if (!bIsDataAlreadyCollected)
 			{
-				FClothingMeshData& NewClothingMeshData = Context.ClothingPerMeshData.FindOrAdd(UniqueHash);
+				FClothingMeshData& NewClothingMeshData = Context.ClothingPerMeshData.FindOrAdd(ClothingUniqueHash);
 
 				check(NewClothingMeshData.ClothingAssetIndex == INDEX_NONE);
 				check(NewClothingMeshData.Data.IsEmpty());
@@ -2415,10 +2420,8 @@ TSharedPtr<mu::FMesh> ConvertSkeletalMeshToMutable(
 		FMutableMeshMetadata MeshMetadata;
 		FMemory::Memzero(MeshMetadata);
 
-		// TODO: Store morph and cloth metadata in this struct, currenlty its added per vertex. With per mesh metadata
-		// this info can be stored in the CO per mesh.
-		MeshMetadata.MorphMetadataId = 0;
-		MeshMetadata.ClothingMetadataId = 0;
+		MeshMetadata.MorphMetadataId = MorphUniqueHash;
+		MeshMetadata.ClothingMetadataId = ClothingUniqueHash;
 		MeshMetadata.SurfaceMetadataId = Source.SurfaceMetadataId;
 
 		auto HashMeshMetadataFunc = [](const FMutableMeshMetadata& Data) -> uint32
