[Mutable] Added new filter to compile specific data table rows.
+ Improved table nodes cache to support different compilation restrictions in different table nodes.

#rnx

diff --git a/Engine/Plugins/Mutable/Source/CustomizableObject/Internal/MuCO/CustomizableObjectCustomVersion.h b/Engine/Plugins/Mutable/Source/CustomizableObject/Internal/MuCO/CustomizableObjectCustomVersion.h
--- a/Engine/Plugins/Mutable/Source/CustomizableObject/Internal/MuCO/CustomizableObjectCustomVersion.h
+++ b/Engine/Plugins/Mutable/Source/CustomizableObject/Internal/MuCO/CustomizableObjectCustomVersion.h
@@ -180,6 +180,8 @@ struct CUSTOMIZABLEOBJECT_API FCustomizableObjectCustomVersion
 
 		MovedLODSettingsToMeshComponentNode,
 
+		AddedTableNodeCompilationFilters,
+
 		// -----<new versions can be added above this line>-------------------------------------------------
 		VersionPlusOne,
 		LatestVersion = VersionPlusOne - 1
diff --git a/Engine/Plugins/Mutable/Source/CustomizableObjectEditor/Private/MuCOE/GenerateMutableSource/GenerateMutableSource.cpp b/Engine/Plugins/Mutable/Source/CustomizableObjectEditor/Private/MuCOE/GenerateMutableSource/GenerateMutableSource.cpp
--- a/Engine/Plugins/Mutable/Source/CustomizableObjectEditor/Private/MuCOE/GenerateMutableSource/GenerateMutableSource.cpp
+++ b/Engine/Plugins/Mutable/Source/CustomizableObjectEditor/Private/MuCOE/GenerateMutableSource/GenerateMutableSource.cpp
@@ -1887,7 +1887,24 @@ bool FMutableGraphGenerationContext::FSharedSurface::operator==(const FSharedSur
 }
 
 
-bool FMutableGraphGenerationContext::FGeneratedDataTablesData::HasSameSettings(const UCustomizableObjectNodeTable* Node) const
+FGeneratedMutableDataTableKey::FGeneratedMutableDataTableKey(FString TableName, FName VersionColumn, FName CompilationFilterColumn, TArray<FString> CompilationFilters, ETableCompilationFilterOperationType FilterOperationType)
+	: TableName(TableName), VersionColumn(VersionColumn), CompilationFilterColumn(CompilationFilterColumn), CompilationFilters(CompilationFilters), FilterOperationType(FilterOperationType)
+{}
+
+
+uint32 GetTypeHash(const FGeneratedMutableDataTableKey& Key)
 {
-	return Node->bDisableCheckedRows == bDisableCheckedRows && Node->VersionColumn == VersionColumn;
+	uint32 Hash = GetTypeHash(Key.TableName);
+	Hash = HashCombine(Hash, GetTypeHash(Key.VersionColumn));
+	Hash = HashCombine(Hash, GetTypeHash(Key.CompilationFilterColumn));
+
+	for (int32 NameIndex = 0; NameIndex < Key.CompilationFilters.Num(); ++NameIndex)
+	{
+		Hash = HashCombine(Hash, GetTypeHash(Key.CompilationFilters[NameIndex]));
+	}
+
+	Hash = HashCombine(Hash, GetTypeHash(Key.FilterOperationType));
+
+	return Hash;
 }
+
diff --git a/Engine/Plugins/Mutable/Source/CustomizableObjectEditor/Private/MuCOE/GenerateMutableSource/GenerateMutableSource.h b/Engine/Plugins/Mutable/Source/CustomizableObjectEditor/Private/MuCOE/GenerateMutableSource/GenerateMutableSource.h
--- a/Engine/Plugins/Mutable/Source/CustomizableObjectEditor/Private/MuCOE/GenerateMutableSource/GenerateMutableSource.h
+++ b/Engine/Plugins/Mutable/Source/CustomizableObjectEditor/Private/MuCOE/GenerateMutableSource/GenerateMutableSource.h
@@ -59,7 +59,7 @@ struct FMutableRefSocket;
 struct FMutableSkinWeightProfileInfo;
 struct FMorphTargetVertexData;
 enum class EPinMode;
-
+enum class ETableCompilationFilterOperationType : uint8;
 
 struct FGeneratedImageProperties
 {
@@ -528,6 +528,32 @@ public:
 };
 
 
+struct FGeneratedMutableDataTableKey
+{
+	FGeneratedMutableDataTableKey(FString TableName, FName VersionColumn, FName CompilationFilterColumn, TArray<FString> CompilationFilters, ETableCompilationFilterOperationType FilterOperationType);
+
+	// Name of the Data Table Asset
+	FString TableName;
+
+	// Compilation Restrictions:
+	// Name of the column that determines de version control
+	FName VersionColumn;
+	
+	// Name of the column that determines the compilation filter
+	FName CompilationFilterColumn;
+	
+	// Compilation Filter
+	TArray<FString> CompilationFilters;
+
+	// Restriction option
+	ETableCompilationFilterOperationType FilterOperationType;
+
+	bool operator==(const FGeneratedMutableDataTableKey&) const = default;
+};
+
+uint32 GetTypeHash(const FGeneratedMutableDataTableKey& Key);
+
+
 /** This structure stores the information that is used only during the graph generation stage of the CustomizableObject compilation. 
 * Data needed for node graph processing goes here.
 */
@@ -560,21 +586,10 @@ struct FMutableGraphGenerationContext
 		// Stores the names of the rows that will be compiled
 		TArray<FName> RowNames;
 		TArray<uint32> RowIds;
-
-		// Compilation Restrictions:
-		// If there is a bool column in the table, checked rows will not be compiled
-		bool bDisableCheckedRows;
-
-		// Name of the column that determines de version control
-		FName VersionColumn;
-
-		// Compare the stored compilation settings with the compilation settings of a Table Node
-		// return true if the compilation settings are equal
-		bool HasSameSettings(const UCustomizableObjectNodeTable* Node) const;
 	};
 
 	// Cache of generated Node Tables
-	TMap<FString, FGeneratedDataTablesData> GeneratedTables;
+	TMap<FGeneratedMutableDataTableKey, FGeneratedDataTablesData> GeneratedTables;
 
 	TMap<FGeneratedGroupProjectorsKey, FGroupProjectorTempData> GeneratedGroupProjectors;
 
diff --git a/Engine/Plugins/Mutable/Source/CustomizableObjectEditor/Private/MuCOE/GenerateMutableSource/GenerateMutableSourceMesh.cpp b/Engine/Plugins/Mutable/Source/CustomizableObjectEditor/Private/MuCOE/GenerateMutableSource/GenerateMutableSourceMesh.cpp
--- a/Engine/Plugins/Mutable/Source/CustomizableObjectEditor/Private/MuCOE/GenerateMutableSource/GenerateMutableSourceMesh.cpp
+++ b/Engine/Plugins/Mutable/Source/CustomizableObjectEditor/Private/MuCOE/GenerateMutableSource/GenerateMutableSourceMesh.cpp
@@ -3501,7 +3501,6 @@ mu::NodeMeshPtr GenerateMorphMesh(const UEdGraphPin* Pin,
 		const FString TableName = TypedNodeTable->Table ? GetNameSafe(TypedNodeTable->Table).ToLower() : GetNameSafe(TypedNodeTable->Structure).ToLower();
 		const uint32 TableId = CityHash32(reinterpret_cast<const char*>(*TableName), TableName.Len() * sizeof(FString::ElementType));
 
-
 		UDataTable* DataTable = GetDataTable(TypedNodeTable, GenerationContext);
 
 		// Generate a new Column for each morph
@@ -3510,7 +3509,9 @@ mu::NodeMeshPtr GenerateMorphMesh(const UEdGraphPin* Pin,
 		int32 NumRows = RowNames.Num();
 
 		// Should exist
-		mu::Ptr<mu::FTable> Table = GenerationContext.GeneratedTables[DataTable->GetName()].GeneratedTable;
+		FGeneratedMutableDataTableKey MutableTableKey(DataTable->GetName(), TypedNodeTable->VersionColumn,
+			TypedNodeTable->CompilationFilterColumn, TypedNodeTable->CompilationFilters, TypedNodeTable->FilterOperationType);
+		mu::Ptr<mu::FTable> Table = GenerationContext.GeneratedTables[MutableTableKey].GeneratedTable;
 		check(Table);
 
 		FString ColumnName = TableColumnName + TypedNodeMorphs[MorphIndex].MorphTargetName;
diff --git a/Engine/Plugins/Mutable/Source/CustomizableObjectEditor/Private/MuCOE/GenerateMutableSource/GenerateMutableSourceTable.cpp b/Engine/Plugins/Mutable/Source/CustomizableObjectEditor/Private/MuCOE/GenerateMutableSource/GenerateMutableSourceTable.cpp
--- a/Engine/Plugins/Mutable/Source/CustomizableObjectEditor/Private/MuCOE/GenerateMutableSource/GenerateMutableSourceTable.cpp
+++ b/Engine/Plugins/Mutable/Source/CustomizableObjectEditor/Private/MuCOE/GenerateMutableSource/GenerateMutableSourceTable.cpp
@@ -6,6 +6,7 @@
 #include "Animation/PoseAsset.h"
 #include "AssetRegistry/ARFilter.h"
 #include "AssetRegistry/AssetRegistryModule.h"
+#include "DataTableUtils.h"
 #include "Engine/CompositeDataTable.h"
 #include "Engine/StaticMesh.h"
 #include "Engine/Texture2DArray.h"
@@ -802,7 +803,9 @@ void GenerateUniqueRowIds(const TArray<FName>& RowNames, TArray<uint32>& OutRowI
 
 TArray<FName> GetRowsToCompile(const UDataTable& DataTable, const UCustomizableObjectNodeTable& TableNode, FMutableGraphGenerationContext& GenerationContext, TArray<uint32>& OutRowIds)
 {
-	if (FMutableGraphGenerationContext::FGeneratedDataTablesData* Result = GenerationContext.GeneratedTables.Find(DataTable.GetName()))
+	FGeneratedMutableDataTableKey MutableTableKey(DataTable.GetName(), TableNode.VersionColumn,
+		TableNode.CompilationFilterColumn, TableNode.CompilationFilters, TableNode.FilterOperationType);
+	if (FMutableGraphGenerationContext::FGeneratedDataTablesData* Result = GenerationContext.GeneratedTables.Find(MutableTableKey))
 	{
 		OutRowIds = Result->RowIds;
 		return Result->RowNames;
@@ -1047,21 +1050,13 @@ mu::Ptr<mu::FTable> GenerateMutableSourceTable(const UDataTable* DataTable, cons
 	// Checking if the table is in the cache
 	const FString TableName = DataTable->GetName();
 
-	if (FMutableGraphGenerationContext::FGeneratedDataTablesData* CachedTable = GenerationContext.GeneratedTables.Find(TableName))
+	FGeneratedMutableDataTableKey MutableTableKey(TableName, TableNode->VersionColumn,
+		TableNode->CompilationFilterColumn, TableNode->CompilationFilters, TableNode->FilterOperationType);
+	if (FMutableGraphGenerationContext::FGeneratedDataTablesData* CachedTable = GenerationContext.GeneratedTables.Find(MutableTableKey))
 	{
 		// Generating Parameter Metadata for parameters that reuse a Table
 		GenerateTableParameterUIData(DataTable, TableNode, GenerationContext);
 
-		if (!CachedTable->HasSameSettings(TableNode))
-		{
-			TArray<const UObject*> Nodes;
-			Nodes.Add(TableNode);
-			Nodes.Add(CachedTable->ReferenceNode);
-
-			GenerationContext.Log(LOCTEXT("TableNodesCompilationRestrictionError",
-				"Found one or more Table Nodes with the same data table but different Compilation Restrictions."), Nodes);
-		}
-
 		return CachedTable->GeneratedTable;
 	}
 
@@ -1097,14 +1092,12 @@ mu::Ptr<mu::FTable> GenerateMutableSourceTable(const UDataTable* DataTable, cons
 		// Generating data for Table Cache
 		FMutableGraphGenerationContext::FGeneratedDataTablesData GeneratedTable;
 		GeneratedTable.GeneratedTable = MutableTable;
-		GeneratedTable.bDisableCheckedRows = TableNode->bDisableCheckedRows;
-		GeneratedTable.VersionColumn = TableNode->VersionColumn;
 		GeneratedTable.RowNames = RowNames;
 		GeneratedTable.RowIds = RowIds;
 		GeneratedTable.ReferenceNode = TableNode;
 
 		// Add table to cache
-		GenerationContext.GeneratedTables.Add(TableName, GeneratedTable);
+		GenerationContext.GeneratedTables.Add(MutableTableKey, GeneratedTable);
 	}
 	else
 	{
diff --git a/Engine/Plugins/Mutable/Source/CustomizableObjectEditor/Private/MuCOE/Nodes/CustomizableObjectNodeTable.cpp b/Engine/Plugins/Mutable/Source/CustomizableObjectEditor/Private/MuCOE/Nodes/CustomizableObjectNodeTable.cpp
--- a/Engine/Plugins/Mutable/Source/CustomizableObjectEditor/Private/MuCOE/Nodes/CustomizableObjectNodeTable.cpp
+++ b/Engine/Plugins/Mutable/Source/CustomizableObjectEditor/Private/MuCOE/Nodes/CustomizableObjectNodeTable.cpp
@@ -115,6 +115,17 @@ void UCustomizableObjectNodeTableRemapPins::RemapPins(const UCustomizableObjectN
 
 
 // Table Node -------
+
+const TArray<FFieldClass*> UCustomizableObjectNodeTable::SupportedFilterTypes = {
+		FFloatProperty::StaticClass(),
+		FDoubleProperty::StaticClass(),
+		FBoolProperty::StaticClass(),
+		FNameProperty::StaticClass(),
+		FStrProperty::StaticClass(),
+		FTextProperty::StaticClass(),
+		FIntProperty::StaticClass() };
+
+
 void UCustomizableObjectNodeTable::BackwardsCompatibleFixup(int32 CustomizableObjectCustomVersion)
 {
 	Super::BackwardsCompatibleFixup(CustomizableObjectCustomVersion);
@@ -272,6 +283,40 @@ void UCustomizableObjectNodeTable::BackwardsCompatibleFixup(int32 CustomizableOb
 		// set this option to true for already generated pins with the none option.
 		bUseMaterialColor = bAddNoneOption;
 	}
+
+	if (CustomizableObjectCustomVersion == FCustomizableObjectCustomVersion::AddedTableNodeCompilationFilters)
+	{
+		if (bDisableCheckedRows_DEPRECATED)
+		{
+			if (Structure)
+			{
+				ConditionalPostLoadReference(*Structure);
+			}
+
+			if (Table)
+			{
+				ConditionalPostLoadReference(*Table);
+			}
+
+			const UScriptStruct* TableStruct = GetTableNodeStruct();
+			FBoolProperty* BoolProperty = nullptr;
+
+			for (TFieldIterator<FProperty> PropertyIt(TableStruct); PropertyIt; ++PropertyIt)
+			{
+				BoolProperty = CastField<FBoolProperty>(*PropertyIt);
+
+				if (BoolProperty)
+				{
+					CompilationFilterColumn = FName(*BoolProperty->GetDisplayNameText().ToString());
+					CompilationFilters.Add("false");
+
+					// There should be only one Bool column
+					break;
+				}
+			}
+
+		}
+	}
 }
 
 
@@ -1543,38 +1588,84 @@ TArray<FName> UCustomizableObjectNodeTable::GetEnabledRows(const UDataTable& Dat
 	}
 
 	TArray<FName> TableRowNames = DataTable.GetRowNames();
-	FBoolProperty* BoolProperty = nullptr;
 
 	// Sort them to avoid cooked data indeterminism problems. Rows may come from different tables and their loading order
 	// is not defined.
 	TableRowNames.Sort([](const FName& A, const FName& B) { return A.ToString() < B.ToString(); });
 
-	for (TFieldIterator<FProperty> PropertyIt(TableStruct); PropertyIt && bDisableCheckedRows; ++PropertyIt)
+	// Check if there is any filter condition
+	if (CompilationFilterColumn.IsNone() || !CompilationFilters.Num())
 	{
-		BoolProperty = CastField<FBoolProperty>(*PropertyIt);
+		return TableRowNames;
+	}
+
+	FProperty* ColumnProperty = FindColumnProperty(CompilationFilterColumn);
+
+	if (!ColumnProperty)
+	{
+		return TableRowNames;
+	}
+
+	RowNames.Reserve(TableRowNames.Num());
 
-		if (BoolProperty)
+	if (UCustomizableObjectNodeTable::SupportedFilterTypes.Contains(ColumnProperty->GetClass()))
+	{
+		for (int32 RowIndex = 0; RowIndex < TableRowNames.Num(); ++RowIndex)
 		{
-			for (const FName& RowName : TableRowNames)
+			if (uint8* CellData = UCustomizableObjectNodeTable::GetCellData(TableRowNames[RowIndex], DataTable, *ColumnProperty))
 			{
-				if (uint8* CellData = GetCellData(RowName, DataTable, *BoolProperty))
+				// By getting values as Texts we enable the support to FText columns
+				FString Value = DataTableUtils::GetPropertyValueAsTextDirect(ColumnProperty, CellData).ToString();
+
+				if (CompilationFilters.Contains(Value))
 				{
-					if (!BoolProperty->GetPropertyValue(CellData))
-					{
-						RowNames.Add(RowName);
-					}
+					RowNames.Add(TableRowNames[RowIndex]);
 				}
 			}
-
-			// There should be only one Bool column
-			break;
 		}
 	}
 
-	// There is no Bool column or we don't want to disable rows
-	if (!BoolProperty)
+	else if (FArrayProperty* ArrayProperty = CastField<FArrayProperty>(ColumnProperty))
 	{
-		return TableRowNames;
+		FProperty* InnerProperty = ArrayProperty->Inner;
+		check(InnerProperty);
+
+		if (UCustomizableObjectNodeTable::SupportedFilterTypes.Contains(InnerProperty->GetClass()))
+		{
+			for (int32 RowIndex = 0; RowIndex < TableRowNames.Num(); ++RowIndex)
+			{
+				if (uint8* CellData = UCustomizableObjectNodeTable::GetCellData(TableRowNames[RowIndex], DataTable, *ArrayProperty))
+				{
+					FScriptArrayHelper ArrayHelper(ArrayProperty, CellData);
+					int32 NumMatchingFilters = 0;
+
+					for (int32 NameIndex = 0; NameIndex < ArrayHelper.Num(); ++NameIndex)
+					{
+						// By getting values as Texts we enable the support to FText columns
+						FString Value = DataTableUtils::GetPropertyValueAsTextDirect(InnerProperty, ArrayHelper.GetRawPtr(NameIndex)).ToString();
+
+						if (CompilationFilters.Contains(Value))
+						{
+							if (FilterOperationType == ETableCompilationFilterOperationType::TCFOT_OR)
+							{
+								RowNames.Add(TableRowNames[RowIndex]);
+								break;
+							}
+							else
+							{
+								NumMatchingFilters++;
+
+								if (NumMatchingFilters == CompilationFilters.Num())
+								{
+									RowNames.Add(TableRowNames[RowIndex]);
+									break;
+								}
+							}
+						}
+					}
+				}
+			}
+		}
 	}
 
 	return RowNames;
diff --git a/Engine/Plugins/Mutable/Source/CustomizableObjectEditor/Private/MuCOE/Nodes/CustomizableObjectNodeTable.h b/Engine/Plugins/Mutable/Source/CustomizableObjectEditor/Private/MuCOE/Nodes/CustomizableObjectNodeTable.h
--- a/Engine/Plugins/Mutable/Source/CustomizableObjectEditor/Private/MuCOE/Nodes/CustomizableObjectNodeTable.h
+++ b/Engine/Plugins/Mutable/Source/CustomizableObjectEditor/Private/MuCOE/Nodes/CustomizableObjectNodeTable.h
@@ -55,6 +55,16 @@ enum class ETableDataGatheringSource : uint8
 	ETDGM_AssetRegistry = 1 UMETA(DisplayName = " Struct + Asset Registry")
 };
 
+UENUM()
+enum class ETableCompilationFilterOperationType : uint8
+{
+	/** At least one of the filters should be in the Data Table row */
+	TCFOT_OR = 0 UMETA(DisplayName = "OR"),
+
+	/** All filters should be in the Data Table row */
+	TCFOT_AND = 1 UMETA(DisplayName = "AND")
+};
+
 
 USTRUCT()
 struct FTableNodeColumnData
@@ -193,9 +203,19 @@ public:
 	UPROPERTY(EditAnywhere, Category = TableProperties)
 	bool bAddNoneOption;
 
-	/** If there is a bool column in the table, checked rows will not be compiled */
+	/** Column of the Data Table that contains the compilation filters of each row */
+	UPROPERTY(EditAnywhere, Category = CompilationRestrictions)
+	FName CompilationFilterColumn;
+
+	/** Compilation filters required for this Table node represented in a string
+		Supported types: Names, Strings, int, bools (true, false). */
+	UPROPERTY(EditAnywhere, Category = CompilationRestrictions)
+	TArray<FString> CompilationFilters;
+
+	/** Determines how Filters interact with the values of the Compilation Filter Column.
+		Check the tooltip of each operation for more information. */
 	UPROPERTY(EditAnywhere, Category = CompilationRestrictions)
-	bool bDisableCheckedRows = true;
+	ETableCompilationFilterOperationType FilterOperationType = ETableCompilationFilterOperationType::TCFOT_OR;
 
 	/** Source where table gathers the data */
 	UPROPERTY(EditAnywhere, Category = TableProperties)
@@ -250,6 +270,9 @@ public:
 	UPROPERTY(EditAnywhere, Category = TableProperties, meta = (DisplayAfter = "bAddNoneOption", EditCondition = "bAddNoneOption==true", EditConditionHides))
 	bool bUseMaterialColor = false;
 
+	// Array with all the classes supported to be used as row filters
+	static const TArray<FFieldClass*> SupportedFilterTypes;
+
 public:
 
 	// UObject interface
@@ -416,4 +439,8 @@ private:
 	// Checks if a pin already exists and if it has the same type as before the node refresh
 	bool CheckPinUpdated(const FString& PinName, const FName& PinType) const;
 
+	/** If there is a bool column in the table, checked rows will not be compiled */
+	UPROPERTY()
+	bool bDisableCheckedRows_DEPRECATED = true;
+
 };
diff --git a/Engine/Plugins/Mutable/Source/CustomizableObjectEditor/Private/MuCOE/Nodes/CustomizableObjectNodeTableDetails.cpp b/Engine/Plugins/Mutable/Source/CustomizableObjectEditor/Private/MuCOE/Nodes/CustomizableObjectNodeTableDetails.cpp
--- a/Engine/Plugins/Mutable/Source/CustomizableObjectEditor/Private/MuCOE/Nodes/CustomizableObjectNodeTableDetails.cpp
+++ b/Engine/Plugins/Mutable/Source/CustomizableObjectEditor/Private/MuCOE/Nodes/CustomizableObjectNodeTableDetails.cpp
@@ -53,6 +53,7 @@ void FCustomizableObjectNodeTableDetails::CustomizeDetails(const TSharedPtr<IDet
 		IDetailCategoryBuilder& CustomizableObjectCategory = DetailBuilder->EditCategory("TableProperties");
 		IDetailCategoryBuilder& CompilationRestrictionsCategory = DetailBuilder->EditCategory("CompilationRestrictions");
 		DetailBuilder->HideProperty("VersionColumn");
+		DetailBuilder->HideProperty("CompilationFilterColumn");
 		IDetailCategoryBuilder& UICategory = DetailBuilder->EditCategory("UI");
 		DetailBuilder->HideProperty("ParamUIMetadataColumn");
 		DetailBuilder->HideProperty("ThumbnailColumn");
@@ -65,6 +66,7 @@ void FCustomizableObjectNodeTableDetails::CustomizeDetails(const TSharedPtr<IDet
 		TSharedPtr<FString> CurrentMutableMetadataColumn = GenerateMutableMetaDataColumnComboBoxOptions();
 		TSharedPtr<FString> CurrentVersionColumn = GenerateVersionColumnComboBoxOptions();
 		TSharedPtr<FString> CurrentThumbnailColumn = GenerateThumbnailColumnComboBoxOptions();
+		TSharedPtr<FString> CurrentCompilationFilterColumn = GenerateCompilationFilterColumnComboBoxOptions();
 
 		CustomizableObjectCategory.AddProperty("ParameterName");
 		TSharedRef<IPropertyHandle> AddNoneOptionProperty = DetailBuilder->GetProperty("bAddNoneOption");
@@ -94,6 +96,34 @@ void FCustomizableObjectNodeTableDetails::CustomizeDetails(const TSharedPtr<IDet
 			]
 			.OverrideResetToDefault(FResetToDefaultOverride::Create(FSimpleDelegate::CreateSP(this, &FCustomizableObjectNodeTableDetails::OnVersionColumnComboBoxSelectionReset)));
 
+		IDetailGroup& CompilationFilterGroup = CompilationRestrictionsCategory.AddGroup(TEXT("CompilationFilterGroupo"), LOCTEXT("CompilationFilterColumn_Selector", "Compilation Filter Column"), false, true);
+		CompilationFilterGroup.HeaderRow()
+			.NameContent()
+			[
+				SNew(STextBlock)
+					.Text(LOCTEXT("CompilationFilterColumn_SelectorText", "Compilation Filter Column"))
+					.ToolTipText(LOCTEXT("CompilationFilterColumn_SelectorTooltip", "Select the column that contains the compilation filter of each row."))
+					.Font(IDetailLayoutBuilder::GetDetailFont())
+			]
+			.ValueContent()
+			[
+				SAssignNew(CompilationFilterColumnComboBox, STextComboBox)
+					.InitiallySelectedItem(CurrentCompilationFilterColumn)
+					.OptionsSource(&CompilationFilterColumnOptionNames)
+					.OnComboBoxOpening(this, &FCustomizableObjectNodeTableDetails::OnOpenCompilationFilterColumnComboBox)
+					.OnSelectionChanged(this, &FCustomizableObjectNodeTableDetails::OnCompilationFilterColumnComboBoxSelectionChanged)
+					.Font(IDetailLayoutBuilder::GetDetailFont())
+					.ColorAndOpacity(this, &FCustomizableObjectNodeTableDetails::GetCompilationFilterColumnComboBoxTextColor, &CompilationFilterColumnOptionNames)
+			]
+			.OverrideResetToDefault(FResetToDefaultOverride::Create(FSimpleDelegate::CreateSP(this, &FCustomizableObjectNodeTableDetails::OnCompilationFilterColumnComboBoxSelectionReset)));
+		
+		TSharedRef<IPropertyHandle> CompilationFilterPropertyHandle = DetailBuilder->GetProperty("CompilationFilters");
+		TSharedRef<IPropertyHandle> FilterRequiredPropertyHandle = DetailBuilder->GetProperty("FilterOperationType");
+
+		CompilationFilterGroup.AddPropertyRow(FilterRequiredPropertyHandle);
+		CompilationFilterGroup.AddPropertyRow(CompilationFilterPropertyHandle);
+
+
 		UICategory.AddCustomRow(LOCTEXT("MutableUIMetadataColumn_Selector", "MutableUIMetadataColumn"))
 			.NameContent()
 			[
@@ -749,7 +779,7 @@ void FCustomizableObjectNodeTableDetails::OnThumbnailColumnComboBoxSelectionRese
 }
 
 
-// Version Bridge Category --------------------------------------------------------------------------------
+// Compilation Restrictions Category --------------------------------------------------------------------------------
 
 TSharedPtr<FString> FCustomizableObjectNodeTableDetails::GenerateVersionColumnComboBoxOptions()
 {
@@ -840,4 +870,119 @@ void FCustomizableObjectNodeTableDetails::OnVersionColumnComboBoxSelectionReset(
 }
 
 
+TSharedPtr<FString> FCustomizableObjectNodeTableDetails::GenerateCompilationFilterColumnComboBoxOptions()
+{
+	const UScriptStruct* TableStruct = Node->GetTableNodeStruct();
+	TSharedPtr<FString> CurrentSelection;
+	CompilationFilterColumnOptionNames.Reset();
+
+	FText NothingSelectedText = LOCTEXT("NothingSelectedText", "- Nothing Selected -");
+	CompilationFilterColumnOptionNames.Add(MakeShareable(new FString(NothingSelectedText.ToString())));
+
+	if (!TableStruct)
+	{
+		return CurrentSelection;
+	}
+
+	// Iterating struct Options
+	for (TFieldIterator<FProperty> It(TableStruct); It; ++It)
+	{
+		FProperty* ColumnProperty = *It;
+
+		if (!ColumnProperty)
+		{
+			continue;
+		}
+
+		if (const FArrayProperty* ArrayProperty = CastField<FArrayProperty>(ColumnProperty))
+		{
+			ColumnProperty = ArrayProperty->Inner;
+		}
+
+		if (UCustomizableObjectNodeTable::SupportedFilterTypes.Contains(ColumnProperty->GetClass()))
+		{
+			TSharedPtr<FString> Option = MakeShareable(new FString(ColumnProperty->GetDisplayNameText().ToString()));
+			CompilationFilterColumnOptionNames.Add(Option);
+
+			if (*Option == Node->CompilationFilterColumn)
+			{
+				CurrentSelection = CompilationFilterColumnOptionNames.Last();
+			}
+		}
+	}
+
+	if (!CurrentSelection)
+	{
+		if (!Node->CompilationFilterColumn.IsNone())
+		{
+			CompilationFilterColumnOptionNames.Add(MakeShareable(new FString(Node->CompilationFilterColumn.ToString())));
+			CurrentSelection = CompilationFilterColumnOptionNames.Last();
+		}
+		else
+		{
+			CurrentSelection = CompilationFilterColumnOptionNames[0];
+		}
+	}
+
+	return CurrentSelection;
+}
+
+
+void FCustomizableObjectNodeTableDetails::OnOpenCompilationFilterColumnComboBox()
+{
+	TSharedPtr<FString> CurrentSelection = GenerateCompilationFilterColumnComboBoxOptions();
+
+	if (CompilationFilterColumnComboBox.IsValid())
+	{
+		CompilationFilterColumnComboBox->ClearSelection();
+		CompilationFilterColumnComboBox->RefreshOptions();
+		CompilationFilterColumnComboBox->SetSelectedItem(CurrentSelection);
+	}
+}
+
+
+void FCustomizableObjectNodeTableDetails::OnCompilationFilterColumnComboBoxSelectionChanged(TSharedPtr<FString> Selection, ESelectInfo::Type SelectInfo)
+{
+	if (Selection && Node->CompilationFilterColumn != FName(*Selection)
+		&& (SelectInfo == ESelectInfo::OnKeyPress || SelectInfo == ESelectInfo::OnMouseClick))
+	{
+		if (Selection == CompilationFilterColumnOptionNames[0])
+		{
+			Node->CompilationFilterColumn = NAME_None;
+		}
+		else
+		{
+			Node->CompilationFilterColumn = FName(*Selection);
+		}
+
+		Node->MarkPackageDirty();
+	}
+}
+
+
+FSlateColor FCustomizableObjectNodeTableDetails::GetCompilationFilterColumnComboBoxTextColor(TArray<TSharedPtr<FString>>* CurrentOptions) const
+{
+	if (Node->FindColumnProperty(Node->CompilationFilterColumn) || Node->CompilationFilterColumn.IsNone())
+	{
+		return FSlateColor::UseForeground();
+	}
+
+	// Table Struct null or does not contain the selected property anymore
+	return FSlateColor(FLinearColor(0.9f, 0.05f, 0.05f, 1.0f));
+}
+
+
+void FCustomizableObjectNodeTableDetails::OnCompilationFilterColumnComboBoxSelectionReset()
+{
+	Node->CompilationFilterColumn = NAME_None;
+
+	if (CompilationFilterColumnComboBox.IsValid())
+	{
+		GenerateCompilationFilterColumnComboBoxOptions();
+		CompilationFilterColumnComboBox->SetSelectedItem(CompilationFilterColumnOptionNames[0]);
+		CompilationFilterColumnComboBox->RefreshOptions();
+	}
+}
+
+
 #undef LOCTEXT_NAMESPACE
diff --git a/Engine/Plugins/Mutable/Source/CustomizableObjectEditor/Private/MuCOE/Nodes/CustomizableObjectNodeTableDetails.h b/Engine/Plugins/Mutable/Source/CustomizableObjectEditor/Private/MuCOE/Nodes/CustomizableObjectNodeTableDetails.h
--- a/Engine/Plugins/Mutable/Source/CustomizableObjectEditor/Private/MuCOE/Nodes/CustomizableObjectNodeTableDetails.h
+++ b/Engine/Plugins/Mutable/Source/CustomizableObjectEditor/Private/MuCOE/Nodes/CustomizableObjectNodeTableDetails.h
@@ -106,7 +106,7 @@ private:
 	void OnThumbnailColumnComboBoxSelectionReset();
 
 
-	// Version Control Category -------------
+	// Compilation Restrictions Category -------------
 
 	// Generates MutableMetadata columns combobox options
 	// Returns the current selected option or a null pointer
@@ -124,6 +124,23 @@ private:
 	// OnComboBoxSelectionChanged Callback for Layout ComboBox
 	void OnVersionColumnComboBoxSelectionReset();
 
+	// Generates MutableMetadata columns combobox options
+	// Returns the current selected option or a null pointer
+	TSharedPtr<FString> GenerateCompilationFilterColumnComboBoxOptions();
+
+	// Callback to regenerate the combobox options
+	void OnOpenCompilationFilterColumnComboBox();
+
+	// OnComboBoxSelectionChanged Callback for Layout ComboBox
+	void OnCompilationFilterColumnComboBoxSelectionChanged(TSharedPtr<FString> Selection, ESelectInfo::Type SelectInfo);
+
+	// Sets the combo box selection color
+	FSlateColor GetCompilationFilterColumnComboBoxTextColor(TArray<TSharedPtr<FString>>* CurrentOptions) const;
+
+	// OnComboBoxSelectionChanged Callback for Layout ComboBox
+	void OnCompilationFilterColumnComboBoxSelectionReset();
+
+
 private:
 
 	// Details -------------
@@ -179,4 +196,10 @@ private:
 	// ComboBox widget to select a VersionColumn from the NodeTable
 	TSharedPtr<STextComboBox> VersionColumnsComboBox;
 
+	// Array with the name of the Version columns
+	TArray<TSharedPtr<FString>> CompilationFilterColumnOptionNames;
+	
+	// ComboBox widget to select a VersionColumn from the NodeTable
+	TSharedPtr<STextComboBox> CompilationFilterColumnComboBox;
+
 };
