[Mutable] Fix non-deterministic behavior in .uexp files when cooking multiple platforms at once.

#rnx

diff --git a/Engine/Plugins/Mutable/Source/CustomizableObjectEditor/Private/MuCOE/CustomizableObjectCompiler.cpp b/Engine/Plugins/Mutable/Source/CustomizableObjectEditor/Private/MuCOE/CustomizableObjectCompiler.cpp
--- a/Engine/Plugins/Mutable/Source/CustomizableObjectEditor/Private/MuCOE/CustomizableObjectCompiler.cpp
+++ b/Engine/Plugins/Mutable/Source/CustomizableObjectEditor/Private/MuCOE/CustomizableObjectCompiler.cpp
@@ -857,8 +857,7 @@ void FCustomizableObjectCompiler::CompileInternal(bool bAsync)
 		// The first part of model resources has to be filled from the GenerationContext.
 		// The second part will be filled when the core compilation finishes so that the CompilationContext is complete.
 		const FString ModelResourcesName = GetModelResourcesNameForPlatform(*CurrentObject, *CurrentOptions.TargetPlatform);
-		UObject* Outer = CurrentOptions.bIsCooking ? CurrentObject : StaticCast<UObject*>(GetTransientPackage());
-		TObjectPtr<UModelResources> ModelResources = NewObject<UModelResources>(Outer, FName(*ModelResourcesName), RF_Public);
+		TObjectPtr<UModelResources> ModelResources = NewObject<UModelResources>(GetTransientPackage(), FName(*ModelResourcesName), RF_Public);
 		PlatformData->ModelResources = TStrongObjectPtr<UModelResources>(ModelResources);
 		
 		PlatformData->ModelStreamableBulkData = MakeShared<FModelStreamableBulkData>();
