[Mutable] Force enqueue compile requests from public API calls.

#rnx

diff --git a/Engine/Plugins/Mutable/Source/CustomizableObject/Internal/MuCO/ICustomizableObjectEditorModule.h b/Engine/Plugins/Mutable/Source/CustomizableObject/Internal/MuCO/ICustomizableObjectEditorModule.h
--- a/Engine/Plugins/Mutable/Source/CustomizableObject/Internal/MuCO/ICustomizableObjectEditorModule.h
+++ b/Engine/Plugins/Mutable/Source/CustomizableObject/Internal/MuCO/ICustomizableObjectEditorModule.h
@@ -112,7 +112,7 @@ public:
 
 	/** Compile the given Customizable Object.
 	 * If calling inside the Customizable Object Editor module, consider using ICustomizableObjectEditorModulePrivate::EnqueueCompileRequest. */
-	virtual void CompileCustomizableObject(UCustomizableObject& Object, const FCompileParams* Params, bool bSilentCompilation) = 0;
+	virtual void CompileCustomizableObject(UCustomizableObject& Object, const FCompileParams* Params, bool bSilent, bool bForce) = 0;
 
 	virtual int32 Tick(bool bBlocking) = 0;
 
diff --git a/Engine/Plugins/Mutable/Source/CustomizableObject/Private/MuCO/CustomizableObject.cpp b/Engine/Plugins/Mutable/Source/CustomizableObject/Private/MuCO/CustomizableObject.cpp
--- a/Engine/Plugins/Mutable/Source/CustomizableObject/Private/MuCO/CustomizableObject.cpp
+++ b/Engine/Plugins/Mutable/Source/CustomizableObject/Private/MuCO/CustomizableObject.cpp
@@ -776,7 +776,7 @@ bool UCustomizableObject::ConditionalAutoCompile()
 		}
 		else
 		{
-			EditorModule->CompileCustomizableObject(*this, nullptr, true);
+			EditorModule->CompileCustomizableObject(*this, nullptr, true, false);
 		}
 	}
 
@@ -1938,7 +1938,7 @@ void UCustomizableObject::Compile(const FCompileParams& Params)
 {
 	if (ICustomizableObjectEditorModule* EditorModule = ICustomizableObjectEditorModule::Get())
 	{
-		EditorModule->CompileCustomizableObject(*this, &Params, false);
+		EditorModule->CompileCustomizableObject(*this, &Params, false, true);
 	}
 	else
 	{
diff --git a/Engine/Plugins/Mutable/Source/CustomizableObject/Private/MuCO/CustomizableObjectInstance.cpp b/Engine/Plugins/Mutable/Source/CustomizableObject/Private/MuCO/CustomizableObjectInstance.cpp
--- a/Engine/Plugins/Mutable/Source/CustomizableObject/Private/MuCO/CustomizableObjectInstance.cpp
+++ b/Engine/Plugins/Mutable/Source/CustomizableObject/Private/MuCO/CustomizableObjectInstance.cpp
@@ -1062,7 +1062,7 @@ void UCustomizableInstancePrivate::TickUpdateCloseCustomizableObjects(UCustomiza
 	{
 		if (ICustomizableObjectEditorModule* EditorModule = ICustomizableObjectEditorModule::Get())
 		{
-			EditorModule->CompileCustomizableObject(*Object, nullptr, true);
+			EditorModule->CompileCustomizableObject(*Object, nullptr, true, false);
 		}
 	}
 #endif
diff --git a/Engine/Plugins/Mutable/Source/CustomizableObjectEditor/Private/MuCOE/CustomizableObjectEditorFunctionLibrary.cpp b/Engine/Plugins/Mutable/Source/CustomizableObjectEditor/Private/MuCOE/CustomizableObjectEditorFunctionLibrary.cpp
--- a/Engine/Plugins/Mutable/Source/CustomizableObjectEditor/Private/MuCOE/CustomizableObjectEditorFunctionLibrary.cpp
+++ b/Engine/Plugins/Mutable/Source/CustomizableObjectEditor/Private/MuCOE/CustomizableObjectEditorFunctionLibrary.cpp
@@ -29,19 +29,19 @@ ECustomizableObjectCompilationState UCustomizableObjectEditorFunctionLibrary::Co
 
 	const double StartTime = FPlatformTime::Seconds();
 
-	TSharedRef<FCompilationRequest> CompileRequest = MakeShared<FCompilationRequest>(*CustomizableObject);
-	CompileRequest->bAsync = false;
-	CompileRequest->bSilentCompilation = false;
-	CompileRequest->Options.OptimizationLevel = ConvertOptimizationLevel(InOptimizationLevel); 
-	CompileRequest->Options.TextureCompression = InTextureCompression;
-	CompileRequest->Options.bGatherReferences = bGatherReferences;
-	
-	ICustomizableObjectEditorModulePrivate::GetChecked().EnqueueCompileRequest(CompileRequest);
-
-	check(CompileRequest->GetCompilationState() == ECompilationStatePrivate::Completed);
-
-	const bool bCompilationSuccess = CompileRequest->GetCompilationResult() == ECompilationResultPrivate::Success ||
-		CompileRequest->GetCompilationResult() == ECompilationResultPrivate::Warnings;
+	FCompileParams Params;
+	Params.bAsync = false;
+	Params.OptimizationLevel = InOptimizationLevel;
+	Params.TextureCompression = InTextureCompression;
+	Params.bGatherReferneces = bGatherReferences;
+
+	bool bCompilationSuccess = false;
+	Params.CallbackNative.BindLambda([&bCompilationSuccess](const FCompileCallbackParams& Params)
+		{
+			bCompilationSuccess = Params.bCompiled;
+		});
+
+	CustomizableObject->Compile(Params);
 	
 	const double CurrentTime = FPlatformTime::Seconds();
 	UE_LOG( LogMutable, Display,
diff --git a/Engine/Plugins/Mutable/Source/CustomizableObjectEditor/Private/MuCOE/CustomizableObjectEditorModule.cpp b/Engine/Plugins/Mutable/Source/CustomizableObjectEditor/Private/MuCOE/CustomizableObjectEditorModule.cpp
--- a/Engine/Plugins/Mutable/Source/CustomizableObjectEditor/Private/MuCOE/CustomizableObjectEditorModule.cpp
+++ b/Engine/Plugins/Mutable/Source/CustomizableObjectEditor/Private/MuCOE/CustomizableObjectEditorModule.cpp
@@ -1048,7 +1048,7 @@ void FCustomizableObjectEditorModule::EnqueueCompileRequest(const TSharedRef<FCo
 }
 
 
-void FCustomizableObjectEditorModule::CompileCustomizableObject(UCustomizableObject& Object, const FCompileParams* Params, bool bSilentCompilation)
+void FCustomizableObjectEditorModule::CompileCustomizableObject(UCustomizableObject& Object, const FCompileParams* Params, bool bSilent, bool bForce)
 {
 	UCustomizableObjectSystem* System = UCustomizableObjectSystem::GetInstance();
 	if (!System)
@@ -1085,9 +1085,9 @@ void FCustomizableObjectEditorModule::CompileCustomizableObject(UCustomizableObj
 		}
 	}
 
-	CompileRequest->bSilentCompilation = bSilentCompilation;
+	CompileRequest->bSilentCompilation = bSilent;
 	
-	EnqueueCompileRequest(CompileRequest);
+	EnqueueCompileRequest(CompileRequest, bForce);
 }
 
 
diff --git a/Engine/Plugins/Mutable/Source/CustomizableObjectEditor/Private/MuCOE/CustomizableObjectEditorModule.h b/Engine/Plugins/Mutable/Source/CustomizableObjectEditor/Private/MuCOE/CustomizableObjectEditorModule.h
--- a/Engine/Plugins/Mutable/Source/CustomizableObjectEditor/Private/MuCOE/CustomizableObjectEditorModule.h
+++ b/Engine/Plugins/Mutable/Source/CustomizableObjectEditor/Private/MuCOE/CustomizableObjectEditorModule.h
@@ -60,7 +60,7 @@ public:
 	virtual void CancelCompileRequests() override;
 	virtual int32 GetNumCompileRequests() override;
 	virtual bool IsCompiling(const UCustomizableObject& Object) const override;
-	virtual void CompileCustomizableObject(UCustomizableObject& Object, const FCompileParams* Params, bool bSilentCompilation) override;
+	virtual void CompileCustomizableObject(UCustomizableObject& Object, const FCompileParams* Params, bool bSilent, bool bForce) override;
 	virtual int32 Tick(bool bBlocking) override;
 	virtual void BeginCacheForCookedPlatformData(UCustomizableObject& Object, const ITargetPlatform* TargetPlatform) override;
 	virtual bool IsCachedCookedPlatformDataLoaded(UCustomizableObject& Object, const ITargetPlatform* TargetPlatform) override;
