[mutable] Fixed CO compatibility code not being run during the Cook operation
- The loading of mutable objects was being skipped due to it being performed as part of the cook execution.

#rnx

diff --git a/Engine/Plugins/Mutable/Source/CustomizableObject/Private/MuCO/CustomizableObject.cpp b/Engine/Plugins/Mutable/Source/CustomizableObject/Private/MuCO/CustomizableObject.cpp
--- a/Engine/Plugins/Mutable/Source/CustomizableObject/Private/MuCO/CustomizableObject.cpp
+++ b/Engine/Plugins/Mutable/Source/CustomizableObject/Private/MuCO/CustomizableObject.cpp
@@ -380,11 +380,8 @@ void UCustomizableObject::PostLoad()
 			});
 	}
 	
-	if (!IsRunningCookCommandlet())
-	{
 		GetPrivate()->Status.NextState(FCustomizableObjectStatusTypes::EState::Loading);
 		UCustomizableObjectSystem::GetInstance()->GetPrivate()->AddPendingLoad(this);
-	}
 #endif
 }
 
@@ -638,6 +635,13 @@ void UCustomizableObjectPrivate::LoadCompiledDataFromDisk()
 {
 	MUTABLE_CPUPROFILER_SCOPE(UCustomizableObjectPrivate::LoadCompiledDataFromDisk)
 	
+	// Skip data loading from disk
+	if (IsRunningCookCommandlet())
+	{
+		Status.NextState(FCustomizableObjectStatusTypes::EState::NoModel);
+		return;
+	}
+	
 	ITargetPlatformManagerModule& TargetPlatformManager = GetTargetPlatformManagerRef();
 	const ITargetPlatform* RunningPlatform = TargetPlatformManager.GetRunningTargetPlatform();
 	check(RunningPlatform);
diff --git a/Engine/Plugins/Mutable/Source/CustomizableObject/Private/MuCO/CustomizableObjectSystem.cpp b/Engine/Plugins/Mutable/Source/CustomizableObject/Private/MuCO/CustomizableObjectSystem.cpp
--- a/Engine/Plugins/Mutable/Source/CustomizableObject/Private/MuCO/CustomizableObjectSystem.cpp
+++ b/Engine/Plugins/Mutable/Source/CustomizableObject/Private/MuCO/CustomizableObjectSystem.cpp
@@ -3992,12 +3992,6 @@ int32 UCustomizableObjectSystem::TickInternal(const bool bBlocking)
 		return 0; // Assets are still being loaded, so subobjects won't be found, compiled objects incomplete and thus updates wrong
 	}
 
-	// Do not tick if the CookCommandlet is running.
-	if (IsRunningCookCommandlet())
-	{
-		return 0;
-	}
-
 	// See if any COs pending to load can be completed.
 	for ( TArray<TObjectPtr<UCustomizableObject>>::TIterator It = GetPrivate()->ObjectsPendingLoad.CreateIterator(); It; ++It )
 	{
@@ -4037,6 +4031,12 @@ int32 UCustomizableObjectSystem::TickInternal(const bool bBlocking)
 		}
 	}
 
+	// Do not tick if the CookCommandlet is running.
+	if (IsRunningCookCommandlet())
+	{
+		return GetPrivate()->ObjectsPendingLoad.Num();
+	}
+	
 #endif
 
 	Private->UpdateStats();
@@ -4289,13 +4289,15 @@ int32 UCustomizableObjectSystem::TickInternal(const bool bBlocking)
 		static_cast<int32>(LODUpdateCandidateFound != nullptr) + // Still a pending LOD update. We can not use the size of RequestedLODUpdates since not all requests valid in future ticks.
 		RemainingTasks;
 
+#if WITH_EDITOR
 	if (bBlocking)
 	{
-#if WITH_EDITOR
 		RemainingWork += EditorModule ? EditorModule->Tick(true) : 0;
-#endif
 	}
 	
+	RemainingWork += GetPrivate()->ObjectsPendingLoad.Num();
+#endif
+	
 	RemainingWork += GetPrivate()->StreamableManager->Tick(bBlocking);
 	
 	return RemainingWork;
