[Mutable] Added first macro suggestions.

#rnx

diff --git a/Engine/Plugins/Mutable/Source/CustomizableObjectEditor/Private/MuCOE/CustomizableInstanceDetails.cpp b/Engine/Plugins/Mutable/Source/CustomizableObjectEditor/Private/MuCOE/CustomizableInstanceDetails.cpp
--- a/Engine/Plugins/Mutable/Source/CustomizableObjectEditor/Private/MuCOE/CustomizableInstanceDetails.cpp
+++ b/Engine/Plugins/Mutable/Source/CustomizableObjectEditor/Private/MuCOE/CustomizableInstanceDetails.cpp
@@ -72,13 +72,14 @@ void FCustomizableInstanceDetails::CustomizeDetails(const TSharedPtr<IDetailLayo
 	check(CustomInstance.IsValid());
 
 	TSharedPtr<IToolkit> FoundAssetEditor = FToolkitManager::Get().FindEditorForAsset(CustomInstance.Get());
+	TSharedPtr<ICustomizableObjectInstanceEditor> EditorPtr;
 
 	// Tab spawned in a COEInstanceEditor
 	if (FoundAssetEditor.IsValid()) 
 	{
 		if (TSharedPtr<FCustomizableObjectInstanceEditor> InstanceEditor = StaticCastSharedPtr<FCustomizableObjectInstanceEditor>(FoundAssetEditor))
 		{
-			WeakEditor = StaticCastSharedPtr<ICustomizableObjectInstanceEditor>(InstanceEditor).ToWeakPtr();
+			EditorPtr = StaticCastSharedPtr<ICustomizableObjectInstanceEditor>(InstanceEditor);
 		}
 	}
 
@@ -91,11 +92,13 @@ void FCustomizableInstanceDetails::CustomizeDetails(const TSharedPtr<IDetailLayo
 		{
 			if (TSharedPtr<FCustomizableObjectEditor> ObjectEditor = StaticCastSharedPtr<FCustomizableObjectEditor>(FoundAssetEditor))
 			{
-				WeakEditor = StaticCastSharedPtr<ICustomizableObjectInstanceEditor>(ObjectEditor).ToWeakPtr();
+				EditorPtr = StaticCastSharedPtr<ICustomizableObjectInstanceEditor>(ObjectEditor);
 			}
 		}
 	}
-	check(WeakEditor.IsValid());
+	check(EditorPtr.IsValid());
+
+	WeakEditor = EditorPtr.ToWeakPtr();
 
 	if (CustomInstance->GetPrivate()->IsSelectedParameterProfileDirty())
 	{
diff --git a/Engine/Plugins/Mutable/Source/CustomizableObjectEditor/Private/MuCOE/CustomizableObjectCompiler.cpp b/Engine/Plugins/Mutable/Source/CustomizableObjectEditor/Private/MuCOE/CustomizableObjectCompiler.cpp
--- a/Engine/Plugins/Mutable/Source/CustomizableObjectEditor/Private/MuCOE/CustomizableObjectCompiler.cpp
+++ b/Engine/Plugins/Mutable/Source/CustomizableObjectEditor/Private/MuCOE/CustomizableObjectCompiler.cpp
@@ -654,21 +654,17 @@ mu::Ptr<mu::NodeObject> GenerateMutableRoot(
 		}
 	}
 
-	TSet<UCustomizableObjectNode*> VisitedNodes;
-
+	
 	// First pass. Only used to recollect info required for the primary pass.
 	// Notice that the traversal is different form the primary pass. Here we follow all pins indiscriminately,
 	// while the primary pass follows the Mutable Source structure (which may cut branches).
-	GraphTraversal::RecursiveVisitNodes(*RootNodeObject, [&GenerationContext](UCustomizableObjectNode& Node)
+	GraphTraversal::VisitNodes(*RootNodeObject, [&GenerationContext](UCustomizableObjectNode& Node)
 	{
 		if (UCustomizableObjectNodeComponentMesh* NodeComponentMesh = Cast<UCustomizableObjectNodeComponentMesh>(&Node))
 		{
 			FirstPass(*NodeComponentMesh, GenerationContext);
 		}
-	}, 
-	GenerationContext.MacroNodesStack, VisitedNodes, &GenerationContext.GroupIdToExternalNodeMap);
-
-	check(GenerationContext.MacroNodesStack.Num() == 0); // Ensure Macro context is empty after visiting all the nodes
+	}, &GenerationContext.GroupIdToExternalNodeMap, &GenerationContext.MacroNodesStack);
 	
     GenerationContext.CompilationContext.RealTimeMorphTargetsOverrides = RootNodeObject->RealTimeMorphSelectionOverrides;
 
diff --git a/Engine/Plugins/Mutable/Source/CustomizableObjectEditor/Private/MuCOE/CustomizableObjectEditor.cpp b/Engine/Plugins/Mutable/Source/CustomizableObjectEditor/Private/MuCOE/CustomizableObjectEditor.cpp
--- a/Engine/Plugins/Mutable/Source/CustomizableObjectEditor/Private/MuCOE/CustomizableObjectEditor.cpp
+++ b/Engine/Plugins/Mutable/Source/CustomizableObjectEditor/Private/MuCOE/CustomizableObjectEditor.cpp
@@ -1513,7 +1513,7 @@ bool FCustomizableObjectEditor::ShowProfileManagementOptions()
 
 const UObject* FCustomizableObjectEditor::GetObjectBeingEdited()
 {
-	ensure(GetObjectsCurrentlyBeingEdited()->Num());
+	check(GetObjectsCurrentlyBeingEdited()->Num());
 
 	return (*GetObjectsCurrentlyBeingEdited())[0];
 }
diff --git a/Engine/Plugins/Mutable/Source/CustomizableObjectEditor/Private/MuCOE/CustomizableObjectEditor.h b/Engine/Plugins/Mutable/Source/CustomizableObjectEditor/Private/MuCOE/CustomizableObjectEditor.h
--- a/Engine/Plugins/Mutable/Source/CustomizableObjectEditor/Private/MuCOE/CustomizableObjectEditor.h
+++ b/Engine/Plugins/Mutable/Source/CustomizableObjectEditor/Private/MuCOE/CustomizableObjectEditor.h
@@ -116,7 +116,6 @@ public:
 	virtual FLinearColor GetWorldCentricTabColorScale() const override;
 
 	// ICustomizableObjectInstanceEditor interface
-	UCustomizableObject* GetCustomizableObject();
 	virtual void RefreshTool() override;
 	virtual TSharedPtr<SCustomizableObjectEditorViewportTabBody> GetViewport() override;
 	virtual UCustomizableObjectInstance* GetPreviewInstance() override;
@@ -141,11 +140,8 @@ public:
 	virtual bool ShowProfileManagementOptions() override;
 	virtual const UObject* GetObjectBeingEdited() override;
 
-
-	/** Reconstructs all child the nodes that match the given type.
-	 * @param StartNode Root node to start the graph traversal. This one also will be reconstructed.
-	 * @param NodeType Node types to reconstruct. */
-	void ReconstructAllChildNodes(UCustomizableObjectNode& StartNode, const UClass& NodeType);
+	/** Returns the CO being edited in this editor */
+	UCustomizableObject* GetCustomizableObject();
 
 	/** Utility method: Test whether the CO Node Object given as parameter is linked to any of the CO Node Object Group nodes
 	* in the Test CO given as parameter */
@@ -162,7 +158,10 @@ public:
 	bool CanDeleteNodes() const;
 	void DuplicateSelectedNodes();
 	bool CanDuplicateSelectedNodes() const;
-	void OnSelectedGraphNodesChanged(const FGraphPanelSelectionSet& NewSelection);
+
+	// FCustomizableObjectGraphEditorToolkit interface
+	virtual void OnSelectedGraphNodesChanged(const FGraphPanelSelectionSet& NewSelection) override;
+	virtual void ReconstructAllChildNodes(UCustomizableObjectNode& StartNode, const UClass& NodeType) override;
 
 	/** Callback to notify the editor when the PreviewInstance has been updated */
 	void OnUpdatePreviewInstance(UCustomizableObjectInstance* Instance);
diff --git a/Engine/Plugins/Mutable/Source/CustomizableObjectEditor/Private/MuCOE/CustomizableObjectInstanceEditor.cpp b/Engine/Plugins/Mutable/Source/CustomizableObjectEditor/Private/MuCOE/CustomizableObjectInstanceEditor.cpp
--- a/Engine/Plugins/Mutable/Source/CustomizableObjectEditor/Private/MuCOE/CustomizableObjectInstanceEditor.cpp
+++ b/Engine/Plugins/Mutable/Source/CustomizableObjectEditor/Private/MuCOE/CustomizableObjectInstanceEditor.cpp
@@ -609,7 +609,7 @@ bool FCustomizableObjectInstanceEditor::ShowProfileManagementOptions()
 
 const UObject* FCustomizableObjectInstanceEditor::GetObjectBeingEdited()
 {
-	ensure(GetObjectsCurrentlyBeingEdited()->Num());
+	check(GetObjectsCurrentlyBeingEdited()->Num());
 
 	return (*GetObjectsCurrentlyBeingEdited())[0];
 }
diff --git a/Engine/Plugins/Mutable/Source/CustomizableObjectEditor/Private/MuCOE/CustomizableObjectMacroLibrary/CustomizableObjectGraphEditorToolkit.cpp b/Engine/Plugins/Mutable/Source/CustomizableObjectEditor/Private/MuCOE/CustomizableObjectMacroLibrary/CustomizableObjectGraphEditorToolkit.cpp
--- a/Engine/Plugins/Mutable/Source/CustomizableObjectEditor/Private/MuCOE/CustomizableObjectMacroLibrary/CustomizableObjectGraphEditorToolkit.cpp
+++ b/Engine/Plugins/Mutable/Source/CustomizableObjectEditor/Private/MuCOE/CustomizableObjectMacroLibrary/CustomizableObjectGraphEditorToolkit.cpp
@@ -1,3 +1,4 @@
+#include "CustomizableObjectGraphEditorToolkit.h"
 // Copyright Epic Games, Inc. All Rights Reserved.
 
 #include "MuCOE/CustomizableObjectMacroLibrary/CustomizableObjectGraphEditorToolkit.h"
@@ -106,9 +107,6 @@ void FCustomizableObjectGraphEditorToolkit::PasteNodesHere(const FVector2D& Loca
 
 	// Update UI
 	GraphEditor->NotifyGraphChanged();
-
-	//TODO(Max): Review why do we need to call the line below. Does it make sence?
-	//CustomizableObject->PostEditChange();
 	GraphEditor->GetCurrentGraph()->MarkPackageDirty();
 }
 
@@ -689,6 +687,12 @@ void FCustomizableObjectGraphEditorToolkit::PostUndo(bool bSuccess)
 }
 
 
+void FCustomizableObjectGraphEditorToolkit::PostRedo(bool bSuccess)
+{
+	PostUndo(bSuccess); 
+}
+
+
 void FCustomizableObjectGraphEditorToolkit::UndoGraphAction()
 {
 	GEditor->UndoTransaction();
diff --git a/Engine/Plugins/Mutable/Source/CustomizableObjectEditor/Private/MuCOE/CustomizableObjectMacroLibrary/CustomizableObjectGraphEditorToolkit.h b/Engine/Plugins/Mutable/Source/CustomizableObjectEditor/Private/MuCOE/CustomizableObjectMacroLibrary/CustomizableObjectGraphEditorToolkit.h
--- a/Engine/Plugins/Mutable/Source/CustomizableObjectEditor/Private/MuCOE/CustomizableObjectMacroLibrary/CustomizableObjectGraphEditorToolkit.h
+++ b/Engine/Plugins/Mutable/Source/CustomizableObjectEditor/Private/MuCOE/CustomizableObjectMacroLibrary/CustomizableObjectGraphEditorToolkit.h
@@ -20,7 +20,13 @@ class UCustomizableObjectNode;
 class FCustomizableObjectGraphEditorToolkit : public FAssetEditorToolkit, public FNotifyHook, public FSelfRegisteringEditorUndoClient
 {
 public:
-	
+
+	// FEditorUndoClient Interface
+	virtual void PostUndo(bool bSuccess) override;
+	virtual void PostRedo(bool bSuccess) override;
+
+	// Own Interface ---
+
 	/** Searches a node that contains the inserted word */
 	void OnEnterText(const FText& NewText, ETextCommit::Type TextType);
 	void FindProperty(const FProperty* Property, const void* InContainer, const FString& FindString, const UObject& Context, bool& bFound);
@@ -29,10 +35,6 @@ public:
 	/** Bind common graph commands */
 	void BindGraphCommands();
 
-	// FEditorUndoClient Interface
-	virtual void PostUndo(bool bSuccess) override;
-	virtual void PostRedo(bool bSuccess) override { PostUndo(bSuccess); }
-
 	// Graph Node Actions --------------------------
 	void SelectNode(const UEdGraphNode* Node);
 
@@ -87,6 +89,11 @@ public:
 	/** Graph Editor callbacks */
 	virtual void OnSelectedGraphNodesChanged(const FGraphPanelSelectionSet& NewSelection) = 0;
 
+	/** Reconstructs all child the nodes that match the given type.
+	 * @param StartNode Root node to start the graph traversal. This one also will be reconstructed.
+	 * @param NodeType Node types to reconstruct. */
+	virtual void ReconstructAllChildNodes(UCustomizableObjectNode& StartNode, const UClass& NodeType) = 0;
+
 	virtual void UpdateGraphNodeProperties();
 
 public:
diff --git a/Engine/Plugins/Mutable/Source/CustomizableObjectEditor/Private/MuCOE/CustomizableObjectMacroLibrary/CustomizableObjectMacroDetails.cpp b/Engine/Plugins/Mutable/Source/CustomizableObjectEditor/Private/MuCOE/CustomizableObjectMacroLibrary/CustomizableObjectMacroDetails.cpp
--- a/Engine/Plugins/Mutable/Source/CustomizableObjectEditor/Private/MuCOE/CustomizableObjectMacroLibrary/CustomizableObjectMacroDetails.cpp
+++ b/Engine/Plugins/Mutable/Source/CustomizableObjectEditor/Private/MuCOE/CustomizableObjectMacroLibrary/CustomizableObjectMacroDetails.cpp
@@ -282,7 +282,6 @@ void SCOMacroPinTypeSelector::Construct(const FArguments& InArgs)
 
 TSharedRef<SWidget> SCOMacroPinTypeSelector::OnGenerateRow(TSharedPtr<FPinNameRowData> Option)
 {
-	//return SNew(STextBlock).Text(FText::FromName(Item->PinFriendlyName));
 	const FSlateBrush* IconBrush = FAppStyle::GetBrush(TEXT("Kismet.VariableList.TypeIcon"));
 
 	TSharedPtr<SHorizontalBox> RowWidget = SNew(SHorizontalBox)
@@ -308,7 +307,6 @@ TSharedRef<SWidget> SCOMacroPinTypeSelector::OnGenerateRow(TSharedPtr<FPinNameRo
 
 TSharedRef<SWidget> SCOMacroPinTypeSelector::GenerateSelectedContent() const
 {
-	//return SNew(STextBlock).Text(FText::FromName(Item->PinFriendlyName));
 	const FSlateBrush* IconBrush = FAppStyle::GetBrush(TEXT("Kismet.VariableList.TypeIcon"));
 
 	TSharedPtr<SHorizontalBox> RowWidget = SNew(SHorizontalBox)
diff --git a/Engine/Plugins/Mutable/Source/CustomizableObjectEditor/Private/MuCOE/CustomizableObjectMacroLibrary/CustomizableObjectMacroLibraryEditor.cpp b/Engine/Plugins/Mutable/Source/CustomizableObjectEditor/Private/MuCOE/CustomizableObjectMacroLibrary/CustomizableObjectMacroLibraryEditor.cpp
--- a/Engine/Plugins/Mutable/Source/CustomizableObjectEditor/Private/MuCOE/CustomizableObjectMacroLibrary/CustomizableObjectMacroLibraryEditor.cpp
+++ b/Engine/Plugins/Mutable/Source/CustomizableObjectEditor/Private/MuCOE/CustomizableObjectMacroLibrary/CustomizableObjectMacroLibraryEditor.cpp
@@ -201,6 +201,12 @@ void FCustomizableObjectMacroLibraryEditor::OnSelectedGraphNodesChanged(const FG
 }
 
 
+void FCustomizableObjectMacroLibraryEditor::ReconstructAllChildNodes(UCustomizableObjectNode& StartNode, const UClass& NodeType)
+{
+	// TODO(Max)
+}
+
+
 FName FCustomizableObjectMacroLibraryEditor::GetToolkitFName() const
 {
 	return FName("CustomizableObjectMacroLibraryEditor");
@@ -325,7 +331,7 @@ void FCustomizableObjectMacroLibraryEditor::SetSelectedMacro(UCustomizableObject
 
 			if (bRefreshMacroSelection)
 			{
-				MacroSelector->SetSelectedMacro(NewSelection);
+				MacroSelector->SetSelectedMacro(*NewSelection);
 			}
 		}
 	}
diff --git a/Engine/Plugins/Mutable/Source/CustomizableObjectEditor/Private/MuCOE/CustomizableObjectMacroLibrary/CustomizableObjectMacroLibraryEditor.h b/Engine/Plugins/Mutable/Source/CustomizableObjectEditor/Private/MuCOE/CustomizableObjectMacroLibrary/CustomizableObjectMacroLibraryEditor.h
--- a/Engine/Plugins/Mutable/Source/CustomizableObjectEditor/Private/MuCOE/CustomizableObjectMacroLibrary/CustomizableObjectMacroLibraryEditor.h
+++ b/Engine/Plugins/Mutable/Source/CustomizableObjectEditor/Private/MuCOE/CustomizableObjectMacroLibrary/CustomizableObjectMacroLibraryEditor.h
@@ -28,14 +28,13 @@ public:
 	virtual FString GetWorldCentricTabPrefix() const override;
 	virtual FLinearColor GetWorldCentricTabColorScale() const override;
 
-	/** NotifyHook Interface */
-	//virtual void NotifyPostChange(const FPropertyChangedEvent& PropertyChangedEvent, FProperty* PropertyThatChanged) override;
-
 	/** FEditorUndoClient Interface */
 	virtual void PostUndo(bool bSuccess) override;
 	virtual void PostRedo(bool bSuccess) override;
 
-	/** Sets which macro is edited in the editor. */
+	/** Sets which macro is edited in the editor.
+		@param NewSelection pointer to the new macro to select
+		@param bRefreshMacroSelection force to refresh the selection of the Macros List. This is only needed if the selection is done by code and not by the user. */
 	void SetSelectedMacro(UCustomizableObjectMacro* NewSelection, bool bRefreshMacroSelection = false);
 
 private:
@@ -48,6 +47,7 @@ private:
 	// CO Graph Editor Toolkit Interface
 	/** Graph Editor callbacks */
 	virtual void OnSelectedGraphNodesChanged(const FGraphPanelSelectionSet& NewSelection) override;
+	virtual void ReconstructAllChildNodes(UCustomizableObjectNode& StartNode, const UClass& NodeType) override;
 
 	/** Macro List View callbacks */
 	void OnAddMacroButtonClicked();
diff --git a/Engine/Plugins/Mutable/Source/CustomizableObjectEditor/Private/MuCOE/EdGraphSchema_CustomizableObject.cpp b/Engine/Plugins/Mutable/Source/CustomizableObjectEditor/Private/MuCOE/EdGraphSchema_CustomizableObject.cpp
--- a/Engine/Plugins/Mutable/Source/CustomizableObjectEditor/Private/MuCOE/EdGraphSchema_CustomizableObject.cpp
+++ b/Engine/Plugins/Mutable/Source/CustomizableObjectEditor/Private/MuCOE/EdGraphSchema_CustomizableObject.cpp
@@ -118,10 +118,25 @@ TSharedPtr<FCustomizableObjectGraphEditorToolkit> GetCustomizableObjectEditor(co
 
 	if (CustomizableObjectGraph)
 	{
-		TSharedPtr< IToolkit > FoundAssetEditor = FToolkitManager::Get().FindEditorForAsset(CustomizableObjectGraph);
-		if (FoundAssetEditor.IsValid())
+		// An CO or a Macro Library
+		UObject* AssetBeingEdited;
+
+		if (CustomizableObjectGraph->IsMacro() && CustomizableObjectGraph->GetOuter()) //Macro Library
+		{
+			AssetBeingEdited = CustomizableObjectGraph->GetOuter()->GetOuter();
+		}
+		else // CO
+		{
+			AssetBeingEdited = CustomizableObjectGraph->GetOuter();
+		}
+
+		if (AssetBeingEdited)
 		{
-			return StaticCastSharedPtr<FCustomizableObjectGraphEditorToolkit>(FoundAssetEditor);
+			TSharedPtr< IToolkit > FoundAssetEditor = FToolkitManager::Get().FindEditorForAsset(AssetBeingEdited);
+			if (FoundAssetEditor.IsValid())
+			{
+				return StaticCastSharedPtr<FCustomizableObjectGraphEditorToolkit>(FoundAssetEditor);
+			}
 		}
 	}
 
@@ -913,7 +928,7 @@ void GetContextMenuActionsReconstructAllChildNodes(UToolMenu* Menu, TWeakObjectP
 				return;
 			}
 
-			if (TSharedPtr<FCustomizableObjectEditor> Editor = StaticCastSharedPtr<FCustomizableObjectEditor>(GetCustomizableObjectEditor(Context->Graph)))
+			if (TSharedPtr<FCustomizableObjectGraphEditorToolkit> Editor = GetCustomizableObjectEditor(Context->Graph))
 			{
 				Editor->ReconstructAllChildNodes(*Node, *NodeType);
 			};
diff --git a/Engine/Plugins/Mutable/Source/CustomizableObjectEditor/Private/MuCOE/GenerateMutableSource/GenerateMutableSource.cpp b/Engine/Plugins/Mutable/Source/CustomizableObjectEditor/Private/MuCOE/GenerateMutableSource/GenerateMutableSource.cpp
--- a/Engine/Plugins/Mutable/Source/CustomizableObjectEditor/Private/MuCOE/GenerateMutableSource/GenerateMutableSource.cpp
+++ b/Engine/Plugins/Mutable/Source/CustomizableObjectEditor/Private/MuCOE/GenerateMutableSource/GenerateMutableSource.cpp
@@ -14,6 +14,7 @@
 #include "MuCOE/CustomizableObjectCompiler.h"
 #include "MuCOE/EdGraphSchema_CustomizableObject.h"
 #include "MuCOE/GenerateMutableSource/GenerateMutableSourceGroupProjector.h"
+#include "MuCOE/GenerateMutableSource/GenerateMutableSourceMacro.h"
 #include "MuCOE/GenerateMutableSource/GenerateMutableSourceModifier.h"
 #include "MuCOE/GraphTraversal.h"
 #include "MuCOE/ICustomizableObjectExtensionNode.h"
@@ -429,7 +430,11 @@ uint32 GetTypeHash(const FGeneratedKey& Key)
 	//Hash = HashCombine(Hash, GetTypeHash(Key.MeshMorphStack)); // Does not support array
 	Hash = HashCombine(Hash, GetTypeHash(Key.bOnlyConnectedLOD));
 	Hash = HashCombine(Hash, GetTypeHash(Key.CurrentMeshComponent));
-	//Hash = HashCombine(Hash, GetTypeHash(Key.MacroContext)); // Does not support array
+
+	for (int32 MacroIndex = 0; MacroIndex < Key.MacroContext.Num(); ++MacroIndex)
+	{
+		Hash = HashCombine(Hash, GetTypeHash(Key.MacroContext[MacroIndex]));
+	}
 	
 	return Hash;
 }
@@ -1235,62 +1240,14 @@ mu::Ptr<mu::NodeObject> GenerateMutableSource(const UEdGraphPin * Pin, FMutableG
 	{
 		bCacheNode = false;
 
-		if (const UEdGraphPin* OutputPin = TypedNodeMacro->GetMacroIOPin(ECOMacroIOType::COMVT_Output, Pin->PinName))
-		{
-			if (const UEdGraphPin* FollowPin = FollowInputPin(*OutputPin))
-			{
-				GenerationContext.MacroNodesStack.Push(TypedNodeMacro);
-				Result = GenerateMutableSource(FollowPin, GenerationContext);
-				GenerationContext.MacroNodesStack.Pop();
-			}
-			else
-			{
-				FString Msg = FString::Printf(TEXT("Macro Pin [%s] not linked."), *Pin->PinName.ToString());
-				GenerationContext.Log(FText::FromString(Msg), Node);
-			}
-		}
-		else
-		{
-			FString Msg = FString::Printf(TEXT("Macro Output node does not contain a pin with name [%s]."), *Pin->PinName.ToString());
-			GenerationContext.Log(FText::FromString(Msg), Node);
-		}
+		Result = GenerateMutableSourceMacro<mu::NodeObject>(*Pin, GenerationContext, GenerateMutableSource);
 	}
 
 	else if (const UCustomizableObjectNodeTunnel* TypedNodeTunnel = Cast<UCustomizableObjectNodeTunnel>(Node))
 	{
-		check(TypedNodeTunnel->bIsInputNode);
-		check(GenerationContext.MacroNodesStack.Num());
-
 		bCacheNode = false;
 
-		const UCustomizableObjectNodeMacroInstance* MacroInstanceNode = GenerationContext.MacroNodesStack.Pop();
-		check(MacroInstanceNode);
-
-		if (const UEdGraphPin* InputPin = MacroInstanceNode->FindPin(Pin->PinName, EEdGraphPinDirection::EGPD_Input))
-		{
-			if (const UEdGraphPin* FollowPin = FollowInputPin(*InputPin))
-			{
-				Result = GenerateMutableSource(FollowPin, GenerationContext);
-			}
-			else
-			{
-				FString Msg = FString::Printf(TEXT("Macro Instance Pin [%s] is not linked."), *Pin->PinName.ToString());
-				GenerationContext.Log(FText::FromString(Msg), Node);
-			}
-		}
-		else
-		{
-			FString Msg = FString::Printf(TEXT("Macro Instance Node does not contain a pin with name [%s]."), *Pin->PinName.ToString());
-			GenerationContext.Log(FText::FromString(Msg), Node);
-		}
-
-		if (!Result)
-		{
-			GenerationContext.Log(LOCTEXT("ScalarMacroInstanceEmptyAssets", "Something went wrong"), Node);
-		}
-
-		// Push the Macro again even if the result is null
-		GenerationContext.MacroNodesStack.Push(MacroInstanceNode);
+		Result = GenerateMutableSourceMacro<mu::NodeObject>(*Pin, GenerationContext, GenerateMutableSource);
 	}
 
 	else
@@ -1845,24 +1802,23 @@ FString GetVariationTag(const UCustomizableObjectNodeVariation& VariationNode, i
 	{
 		if (const UEdGraphPin* VariationPin = VariationNode.VariationPin(VariationIndex))
 		{
-			const FString PinName = VariationPin->PinName.ToString() + UCustomizableObjectNodeVariation::VariationTagSufix;
-			const UEdGraphPin* TagPin = VariationNode.FindPin(PinName, EEdGraphPinDirection::EGPD_Input);
-
-			if (TagPin && TagPin->LinkedTo.Num())
+			if (const UEdGraphPin* TagPin = VariationNode.VariationTagPin(VariationIndex))
 			{
-				const UEdGraphPin* ConnectedStringPin = FollowInputPin(*TagPin);
-				const UEdGraphPin* SourceStringPin = ConnectedStringPin ? GraphTraversal::FindStringSource(*ConnectedStringPin, &GenerationContext.MacroNodesStack) : nullptr;
-
-				if (SourceStringPin)
+				if (const UEdGraphPin* ConnectedStringPin = FollowInputPin(*TagPin))
 				{
-					if (const UCustomizableObjectNodeStringConstant* StringNode = Cast<UCustomizableObjectNodeStringConstant>(SourceStringPin->GetOwningNode()))
+					const UEdGraphPin* SourceStringPin = ConnectedStringPin ? GraphTraversal::FindStringSource(*ConnectedStringPin, &GenerationContext.MacroNodesStack) : nullptr;
+
+					if (SourceStringPin)
 					{
-						return StringNode->Value;
+						if (const UCustomizableObjectNodeStringConstant* StringNode = Cast<UCustomizableObjectNodeStringConstant>(SourceStringPin->GetOwningNode()))
+						{
+							return StringNode->Value;
+						}
+					}
+					else
+					{
+						GenerationContext.Log(LOCTEXT("VariationNodeTagError", "Could not found a linked String node."), &VariationNode);
 					}
-				}
-				else
-				{
-					GenerationContext.Log(LOCTEXT("MorphStringNodeFailed", "Could not found a linked String node."), &VariationNode);
 				}
 			}
 		}
@@ -1871,7 +1827,6 @@ FString GetVariationTag(const UCustomizableObjectNodeVariation& VariationNode, i
 	return VariationNode.GetVariation(VariationIndex).Tag;
 }
 
-
 #undef LOCTEXT_NAMESPACE
 
 FMutableGraphGenerationContext::FSharedSurface::FSharedSurface(uint8 InLOD, const mu::Ptr<mu::NodeSurfaceNew>& InNodeSurfaceNew)
diff --git a/Engine/Plugins/Mutable/Source/CustomizableObjectEditor/Private/MuCOE/GenerateMutableSource/GenerateMutableSourceColor.cpp b/Engine/Plugins/Mutable/Source/CustomizableObjectEditor/Private/MuCOE/GenerateMutableSource/GenerateMutableSourceColor.cpp
--- a/Engine/Plugins/Mutable/Source/CustomizableObjectEditor/Private/MuCOE/GenerateMutableSource/GenerateMutableSourceColor.cpp
+++ b/Engine/Plugins/Mutable/Source/CustomizableObjectEditor/Private/MuCOE/GenerateMutableSource/GenerateMutableSourceColor.cpp
@@ -7,7 +7,9 @@
 #include "MuCOE/EdGraphSchema_CustomizableObject.h"
 #include "MuCOE/GenerateMutableSource/GenerateMutableSourceFloat.h"
 #include "MuCOE/GenerateMutableSource/GenerateMutableSourceImage.h"
+#include "MuCOE/GenerateMutableSource/GenerateMutableSourceMacro.h"
 #include "MuCOE/GenerateMutableSource/GenerateMutableSourceTable.h"
+#include "MuCOE/GenerateMutableSource/GenerateMutableSource.h"
 #include "MuCOE/GraphTraversal.h"
 #include "MuCOE/Nodes/CustomizableObjectNodeColorArithmeticOp.h"
 #include "MuCOE/Nodes/CustomizableObjectNodeColorConstant.h"
@@ -348,62 +350,14 @@ mu::Ptr<mu::NodeColour> GenerateMutableSourceColor(const UEdGraphPin* Pin, FMuta
 	{
 		bCacheNode = false;
 
-		if (const UEdGraphPin* OutputPin = TypedNodeMacro->GetMacroIOPin(ECOMacroIOType::COMVT_Output, Pin->PinName))
-		{
-			if (const UEdGraphPin* FollowPin = FollowInputPin(*OutputPin))
-			{
-				GenerationContext.MacroNodesStack.Push(TypedNodeMacro);
-				Result = GenerateMutableSourceColor(FollowPin, GenerationContext);
-				GenerationContext.MacroNodesStack.Pop();
-			}
-			else
-			{
-				FString Msg = FString::Printf(TEXT("Macro Pin %s not linked."), *Pin->PinName.ToString());
-				GenerationContext.Log(FText::FromString(Msg), Node);
-			}
-		}
-		else
-		{
-			FString Msg = FString::Printf(TEXT("Macro Output node does not contain a pin with name %s."), *Pin->PinName.ToString());
-			GenerationContext.Log(FText::FromString(Msg), Node);
-		}
+		Result = GenerateMutableSourceMacro<mu::NodeColour>(*Pin, GenerationContext, GenerateMutableSourceColor);
 	}
 
 	else if (const UCustomizableObjectNodeTunnel* TypedNodeTunnel = Cast<UCustomizableObjectNodeTunnel>(Node))
 	{
-		check(TypedNodeTunnel->bIsInputNode);
-		check(GenerationContext.MacroNodesStack.Num());
-
 		bCacheNode = false;
 
-		const UCustomizableObjectNodeMacroInstance* MacroInstanceNode = GenerationContext.MacroNodesStack.Pop();
-		check(MacroInstanceNode);
-
-		if (const UEdGraphPin* InputPin = MacroInstanceNode->FindPin(Pin->PinName, EEdGraphPinDirection::EGPD_Input))
-		{
-			if (const UEdGraphPin* FollowPin = FollowInputPin(*InputPin))
-			{
-				Result = GenerateMutableSourceColor(FollowPin, GenerationContext);
-			}
-			else
-			{
-				FString Msg = FString::Printf(TEXT("Macro Instance Pin %s is not linked."), *Pin->PinName.ToString());
-				GenerationContext.Log(FText::FromString(Msg), Node);
-			}
-		}
-		else
-		{
-			FString Msg = FString::Printf(TEXT("Macro Instance Node does not contain a pin with name %s."), *Pin->PinName.ToString());
-			GenerationContext.Log(FText::FromString(Msg), Node);
-		}
-
-		if (!Result)
-		{
-			GenerationContext.Log(LOCTEXT("ScalarMacroInstanceEmptyAssets", "Something went wrong"), Node);
-		}
-
-		// Push the Macro again even if the result is null
-		GenerationContext.MacroNodesStack.Push(MacroInstanceNode);
+		Result = GenerateMutableSourceMacro<mu::NodeColour>(*Pin, GenerationContext, GenerateMutableSourceColor);
 	}
 
 	else
diff --git a/Engine/Plugins/Mutable/Source/CustomizableObjectEditor/Private/MuCOE/GenerateMutableSource/GenerateMutableSourceComponent.cpp b/Engine/Plugins/Mutable/Source/CustomizableObjectEditor/Private/MuCOE/GenerateMutableSource/GenerateMutableSourceComponent.cpp
--- a/Engine/Plugins/Mutable/Source/CustomizableObjectEditor/Private/MuCOE/GenerateMutableSource/GenerateMutableSourceComponent.cpp
+++ b/Engine/Plugins/Mutable/Source/CustomizableObjectEditor/Private/MuCOE/GenerateMutableSource/GenerateMutableSourceComponent.cpp
@@ -8,6 +8,7 @@
 #include "MuCOE/GenerateMutableSource/GenerateMutableSourceMesh.h"
 #include "MuCOE/GenerateMutableSource/GenerateMutableSourceSurface.h"
 #include "MuCOE/GenerateMutableSource/GenerateMutableSourceFloat.h"
+#include "MuCOE/GenerateMutableSource/GenerateMutableSourceMacro.h"
 #include "MuCOE/GraphTraversal.h"
 #include "MuCOE/Nodes/CustomizableObjectNodeComponentMesh.h"
 #include "MuCOE/Nodes/CustomizableObjectNodeComponentMeshAddTo.h"
@@ -462,62 +463,13 @@ mu::Ptr<mu::NodeComponent> GenerateMutableSourceComponent(const UEdGraphPin* Pin
 	{
 		bCacheNode = false;
 
-		if (const UEdGraphPin* OutputPin = TypedNodeMacro->GetMacroIOPin(ECOMacroIOType::COMVT_Output, Pin->PinName))
-		{
-			if (const UEdGraphPin* FollowPin = FollowInputPin(*OutputPin))
-			{
-				GenerationContext.MacroNodesStack.Push(TypedNodeMacro);
-				Result = GenerateMutableSourceComponent(FollowPin, GenerationContext);
-				GenerationContext.MacroNodesStack.Pop();
-			}
-			else
-			{
-				FString Msg = FString::Printf(TEXT("Macro Pin [%s] not linked."), *Pin->PinName.ToString());
-				GenerationContext.Log(FText::FromString(Msg), Node);
-			}
-		}
-		else
-		{
-			FString Msg = FString::Printf(TEXT("Macro Output node does not contain a pin with name [%s]."), *Pin->PinName.ToString());
-			GenerationContext.Log(FText::FromString(Msg), Node);
-		}
+		Result = GenerateMutableSourceMacro<mu::NodeComponent>(*Pin, GenerationContext, GenerateMutableSourceComponent);
 	}
 
 	else if (const UCustomizableObjectNodeTunnel* TypedNodeTunnel = Cast<UCustomizableObjectNodeTunnel>(Node))
 	{
-		check(TypedNodeTunnel->bIsInputNode);
-		check(GenerationContext.MacroNodesStack.Num());
-
 		bCacheNode = false;
-
-		const UCustomizableObjectNodeMacroInstance* MacroInstanceNode = GenerationContext.MacroNodesStack.Pop();
-		check(MacroInstanceNode);
-
-		if (const UEdGraphPin* InputPin = MacroInstanceNode->FindPin(Pin->PinName, EEdGraphPinDirection::EGPD_Input))
-		{
-			if (const UEdGraphPin* FollowPin = FollowInputPin(*InputPin))
-			{
-				Result = GenerateMutableSourceComponent(FollowPin, GenerationContext);
-			}
-			else
-			{
-				FString Msg = FString::Printf(TEXT("Macro Instance Pin [%s] is not linked."), *Pin->PinName.ToString());
-				GenerationContext.Log(FText::FromString(Msg), Node);
-			}
-		}
-		else
-		{
-			FString Msg = FString::Printf(TEXT("Macro Instance Node does not contain a pin with name [%s]."), *Pin->PinName.ToString());
-			GenerationContext.Log(FText::FromString(Msg), Node);
-		}
-
-		if (!Result)
-		{
-			GenerationContext.Log(LOCTEXT("ScalarMacroInstanceEmptyAssets", "Something went wrong"), Node);
-		}
-
-		// Push the Macro again even if the result is null
-		GenerationContext.MacroNodesStack.Push(MacroInstanceNode);
+		Result = GenerateMutableSourceMacro<mu::NodeComponent>(*Pin, GenerationContext, GenerateMutableSourceComponent);
 	}
 
 	else
diff --git a/Engine/Plugins/Mutable/Source/CustomizableObjectEditor/Private/MuCOE/GenerateMutableSource/GenerateMutableSourceComponent.h b/Engine/Plugins/Mutable/Source/CustomizableObjectEditor/Private/MuCOE/GenerateMutableSource/GenerateMutableSourceComponent.h
--- a/Engine/Plugins/Mutable/Source/CustomizableObjectEditor/Private/MuCOE/GenerateMutableSource/GenerateMutableSourceComponent.h
+++ b/Engine/Plugins/Mutable/Source/CustomizableObjectEditor/Private/MuCOE/GenerateMutableSource/GenerateMutableSourceComponent.h
@@ -1,4 +1,4 @@
-﻿// Copyright Epic Games, Inc. All Rights Reserved.
+// Copyright Epic Games, Inc. All Rights Reserved.
 
 #pragma once
 
diff --git a/Engine/Plugins/Mutable/Source/CustomizableObjectEditor/Private/MuCOE/GenerateMutableSource/GenerateMutableSourceFloat.cpp b/Engine/Plugins/Mutable/Source/CustomizableObjectEditor/Private/MuCOE/GenerateMutableSource/GenerateMutableSourceFloat.cpp
--- a/Engine/Plugins/Mutable/Source/CustomizableObjectEditor/Private/MuCOE/GenerateMutableSource/GenerateMutableSourceFloat.cpp
+++ b/Engine/Plugins/Mutable/Source/CustomizableObjectEditor/Private/MuCOE/GenerateMutableSource/GenerateMutableSourceFloat.cpp
@@ -9,6 +9,7 @@
 #include "MuCOE/CustomizableObjectGraph.h"
 #include "MuCOE/EdGraphSchema_CustomizableObject.h"
 #include "MuCOE/GenerateMutableSource/GenerateMutableSourceImage.h"
+#include "MuCOE/GenerateMutableSource/GenerateMutableSourceMacro.h"
 #include "MuCOE/GenerateMutableSource/GenerateMutableSourceTable.h"
 #include "MuCOE/GraphTraversal.h"
 #include "MuCOE/Nodes/CustomizableObjectNodeCurve.h"
@@ -392,63 +393,13 @@ mu::Ptr<mu::NodeScalar> GenerateMutableSourceFloat(const UEdGraphPin* Pin, FMuta
 	else if (const UCustomizableObjectNodeMacroInstance* TypedNodeMacro = Cast<UCustomizableObjectNodeMacroInstance>(Node))
 	{
 		bCacheNode = false;
-
-		if (const UEdGraphPin* OutputPin = TypedNodeMacro->GetMacroIOPin(ECOMacroIOType::COMVT_Output, Pin->PinName))
-		{
-			if (const UEdGraphPin* FollowPin = FollowInputPin(*OutputPin))
-			{
-				GenerationContext.MacroNodesStack.Push(TypedNodeMacro);
-				Result = GenerateMutableSourceFloat(FollowPin, GenerationContext);
-				GenerationContext.MacroNodesStack.Pop();
-			}
-			else
-			{
-				FString Msg = FString::Printf(TEXT("Macro Pin [%s] not linked."), *Pin->PinName.ToString());
-				GenerationContext.Log(FText::FromString(Msg), Node);
-			}
-		}
-		else
-		{
-			FString Msg = FString::Printf(TEXT("Macro Output node does not contain a pin with name [%s]."), *Pin->PinName.ToString());
-			GenerationContext.Log(FText::FromString(Msg), Node);
-		}
+		Result = GenerateMutableSourceMacro<mu::NodeScalar>(*Pin, GenerationContext, GenerateMutableSourceFloat);
 	}
 
 	else if (const UCustomizableObjectNodeTunnel* TypedNodeTunnel = Cast<UCustomizableObjectNodeTunnel>(Node))
 	{
-		check(TypedNodeTunnel->bIsInputNode);
-		check(GenerationContext.MacroNodesStack.Num());
-
 		bCacheNode = false;
-
-		const UCustomizableObjectNodeMacroInstance* MacroInstanceNode = GenerationContext.MacroNodesStack.Pop();
-		check(MacroInstanceNode);
-
-		if (const UEdGraphPin* InputPin = MacroInstanceNode->FindPin(Pin->PinName, EEdGraphPinDirection::EGPD_Input))
-		{
-			if (const UEdGraphPin* FollowPin = FollowInputPin(*InputPin))
-			{
-				Result = GenerateMutableSourceFloat(FollowPin, GenerationContext);
-			}
-			else
-			{
-				FString Msg = FString::Printf(TEXT("Macro Instance Pin [%s] is not linked."), *Pin->PinName.ToString());
-				GenerationContext.Log(FText::FromString(Msg), Node);
-			}
-		}
-		else
-		{
-			FString Msg = FString::Printf(TEXT("Macro Instance Node does not contain a pin with name [%s]."), *Pin->PinName.ToString());
-			GenerationContext.Log(FText::FromString(Msg), Node);
-		}
-
-		if (!Result)
-		{
-			GenerationContext.Log(LOCTEXT("ScalarMacroInstanceEmptyAssets", "Something went wrong"), Node);
-		}
-
-		// Push the Macro again even if the result is null
-		GenerationContext.MacroNodesStack.Push(MacroInstanceNode);
+		Result = GenerateMutableSourceMacro<mu::NodeScalar>(*Pin, GenerationContext, GenerateMutableSourceFloat);
 	}
 
 	else
diff --git a/Engine/Plugins/Mutable/Source/CustomizableObjectEditor/Private/MuCOE/GenerateMutableSource/GenerateMutableSourceImage.cpp b/Engine/Plugins/Mutable/Source/CustomizableObjectEditor/Private/MuCOE/GenerateMutableSource/GenerateMutableSourceImage.cpp
--- a/Engine/Plugins/Mutable/Source/CustomizableObjectEditor/Private/MuCOE/GenerateMutableSource/GenerateMutableSourceImage.cpp
+++ b/Engine/Plugins/Mutable/Source/CustomizableObjectEditor/Private/MuCOE/GenerateMutableSource/GenerateMutableSourceImage.cpp
@@ -1160,6 +1160,7 @@ mu::Ptr<mu::NodeImage> GenerateMutableSourceImage(const UEdGraphPin* Pin, FMutab
 
 	else if (const UCustomizableObjectNodeMacroInstance* TypedNodeMacro = Cast<UCustomizableObjectNodeMacroInstance>(Node))
 	{
+		// Can't use the GenerateMutableSourceMacro function here because GenerateMutableSourceImage needs the extra parameter ReferenceTextureSize 
 		bCacheNode = false;
 
 		if (const UEdGraphPin* OutputPin = TypedNodeMacro->GetMacroIOPin(ECOMacroIOType::COMVT_Output, Pin->PinName))
@@ -1172,19 +1173,20 @@ mu::Ptr<mu::NodeImage> GenerateMutableSourceImage(const UEdGraphPin* Pin, FMutab
 			}
 			else
 			{
-				FString Msg = FString::Printf(TEXT("Macro Pin [%s] not linked."), *Pin->PinName.ToString());
-				GenerationContext.Log(FText::FromString(Msg), Node);
+				FText Msg = FText::Format(LOCTEXT("MacroInstanceError_PinNotLinked_Image", "Macro Output node Pin {0} not linked."), FText::FromName(Pin->PinName));
+				GenerationContext.Log(Msg, Node);
 			}
 		}
 		else
 		{
-			FString Msg = FString::Printf(TEXT("Macro Output node does not contain a pin with name [%s]."), *Pin->PinName.ToString());
-			GenerationContext.Log(FText::FromString(Msg), Node);
+			FText Msg = FText::Format(LOCTEXT("MacroInstanceError_PinNameNotFound_Image", "Macro Output node does not contain a pin with name {0}."), FText::FromName(Pin->PinName));
+			GenerationContext.Log(Msg, Node);
 		}
 	}
 
 	else if (const UCustomizableObjectNodeTunnel* TypedNodeTunnel = Cast<UCustomizableObjectNodeTunnel>(Node))
 	{
+		// Can't use the GenerateMutableSourceMacro function here because GenerateMutableSourceImage needs the extra parameter ReferenceTextureSize 
 		check(TypedNodeTunnel->bIsInputNode);
 		check(GenerationContext.MacroNodesStack.Num());
 
@@ -1201,19 +1203,14 @@ mu::Ptr<mu::NodeImage> GenerateMutableSourceImage(const UEdGraphPin* Pin, FMutab
 			}
 			else
 			{
-				FString Msg = FString::Printf(TEXT("Macro Instance Pin [%s] is not linked."), *Pin->PinName.ToString());
-				GenerationContext.Log(FText::FromString(Msg), Node);
+				FText Msg = FText::Format(LOCTEXT("MacroTunnelError_PinNotLinked_Image", "Macro Instance Pin {0} is not linked."), FText::FromName(Pin->PinName));
+				GenerationContext.Log(Msg, Node);
 			}
 		}
 		else
 		{
-			FString Msg = FString::Printf(TEXT("Macro Instance Node does not contain a pin with name [%s]."), *Pin->PinName.ToString());
-			GenerationContext.Log(FText::FromString(Msg), Node);
-		}
-
-		if (!Result)
-		{
-			GenerationContext.Log(LOCTEXT("ScalarMacroInstanceEmptyAssets", "Something went wrong"), Node);
+			FText Msg = FText::Format(LOCTEXT("MacroTunnelError_PinNameNotFound_Image", "Macro Instance Node does not contain a pin with name {0}."), FText::FromName(Pin->PinName));
+			GenerationContext.Log(Msg, Node);
 		}
 
 		// Push the Macro again even if the result is null
diff --git a/Engine/Plugins/Mutable/Source/CustomizableObjectEditor/Private/MuCOE/GenerateMutableSource/GenerateMutableSourceMacro.h b/Engine/Plugins/Mutable/Source/CustomizableObjectEditor/Private/MuCOE/GenerateMutableSource/GenerateMutableSourceMacro.h
new file mode 100644
--- /dev/null
+++ b/Engine/Plugins/Mutable/Source/CustomizableObjectEditor/Private/MuCOE/GenerateMutableSource/GenerateMutableSourceMacro.h
@@ -0,0 +1,74 @@
+// Copyright Epic Games, Inc. All Rights Reserved.
+
+#pragma once
+
+#include "MuCOE/GraphTraversal.h"
+#include "MuCOE/GenerateMutableSource/GenerateMutableSource.h"
+#include "MuCOE/Nodes/CustomizableObjectNodeMacroInstance.h"
+#include "MuCOE/Nodes/CustomizableObjectNodeTunnel.h"
+
+#define LOCTEXT_NAMESPACE "CustomizableObjectEditor"
+
+template<class T>
+mu::Ptr<T> GenerateMutableSourceMacro(const UEdGraphPin& Pin, FMutableGraphGenerationContext& GenerationContext, const TFunction<mu::Ptr<T>(const UEdGraphPin*, FMutableGraphGenerationContext&)>& GenerationFunction)
+{
+	mu::Ptr<T> Result;
+	UCustomizableObjectNode* Node = CastChecked<UCustomizableObjectNode>(Pin.GetOwningNode());
+
+	if (const UCustomizableObjectNodeMacroInstance* TypedNodeMacro = Cast<UCustomizableObjectNodeMacroInstance>(Node))
+	{
+		if (const UEdGraphPin* OutputPin = TypedNodeMacro->GetMacroIOPin(ECOMacroIOType::COMVT_Output, Pin.PinName))
+		{
+			if (const UEdGraphPin* FollowPin = FollowInputPin(*OutputPin))
+			{
+				GenerationContext.MacroNodesStack.Push(TypedNodeMacro);
+				Result = GenerationFunction(FollowPin, GenerationContext);
+				GenerationContext.MacroNodesStack.Pop();
+			}
+			else
+			{
+				FText Msg = FText::Format(LOCTEXT("MacroInstanceError_PinNotLinked", "Macro Output node Pin {0} not linked."), FText::FromName(Pin.PinName));
+				GenerationContext.Log(Msg, Node);
+			}
+		}
+		else
+		{
+			FText Msg = FText::Format(LOCTEXT("MacroInstanceError_PinNameNotFound", "Macro Output node does not contain a pin with name {0}."), FText::FromName(Pin.PinName));
+			GenerationContext.Log(Msg, Node);
+		}
+	}
+
+	else if (const UCustomizableObjectNodeTunnel* TypedNodeTunnel = Cast<UCustomizableObjectNodeTunnel>(Node))
+	{
+		check(TypedNodeTunnel->bIsInputNode);
+		check(GenerationContext.MacroNodesStack.Num());
+
+		const UCustomizableObjectNodeMacroInstance* MacroInstanceNode = GenerationContext.MacroNodesStack.Pop();
+		check(MacroInstanceNode);
+
+		if (const UEdGraphPin* InputPin = MacroInstanceNode->FindPin(Pin.PinName, EEdGraphPinDirection::EGPD_Input))
+		{
+			if (const UEdGraphPin* FollowPin = FollowInputPin(*InputPin))
+			{
+				Result = GenerationFunction(FollowPin, GenerationContext);
+			}
+			else
+			{
+				FText Msg = FText::Format(LOCTEXT("MacroTunnelError_PinNotLinked", "Macro Instance Pin {0} is not linked."), FText::FromName(Pin.PinName));
+				GenerationContext.Log(Msg, Node);
+			}
+		}
+		else
+		{
+			FText Msg = FText::Format(LOCTEXT("MacroTunnelError_PinNameNotFound", "Macro Instance Node does not contain a pin with name {0}."), FText::FromName(Pin.PinName));
+			GenerationContext.Log(Msg, Node);
+		}
+
+		// Push the Macro again even if the result is null
+		GenerationContext.MacroNodesStack.Push(MacroInstanceNode);
+	}
+
+	return Result;
+}
+
+#undef LOCTEXT_NAMESPACE
diff --git a/Engine/Plugins/Mutable/Source/CustomizableObjectEditor/Private/MuCOE/GenerateMutableSource/GenerateMutableSourceMesh.cpp b/Engine/Plugins/Mutable/Source/CustomizableObjectEditor/Private/MuCOE/GenerateMutableSource/GenerateMutableSourceMesh.cpp
--- a/Engine/Plugins/Mutable/Source/CustomizableObjectEditor/Private/MuCOE/GenerateMutableSource/GenerateMutableSourceMesh.cpp
+++ b/Engine/Plugins/Mutable/Source/CustomizableObjectEditor/Private/MuCOE/GenerateMutableSource/GenerateMutableSourceMesh.cpp
@@ -3942,7 +3942,7 @@ mu::Ptr<mu::NodeMesh> GenerateMutableSourceMesh(const UEdGraphPin* Pin,
 		{
 			FString MorphTargetName = TypedNodeMorph->MorphTargetName;
 
-			if (UEdGraphPin* MorphTargetNamePin = Node->FindPin(FString("Morph Target Name"), EEdGraphPinDirection::EGPD_Input))
+			if (UEdGraphPin* MorphTargetNamePin = TypedNodeMorph->MorphTargetNamePin())
 			{
 				const UEdGraphPin* ConnectedStringPin = FollowInputPin(*MorphTargetNamePin);
 				const UEdGraphPin* SourceStringPin = ConnectedStringPin ? GraphTraversal::FindStringSource(*ConnectedStringPin, &GenerationContext.MacroNodesStack) : nullptr;
@@ -4561,6 +4561,7 @@ mu::Ptr<mu::NodeMesh> GenerateMutableSourceMesh(const UEdGraphPin* Pin,
 
 	else if (const UCustomizableObjectNodeMacroInstance* TypedNodeMacro = Cast<UCustomizableObjectNodeMacroInstance>(Node))
 	{
+		// Can't use the GenerateMutableSourceMacro function here because GenerateMutableSourceMesh needs some extra parameters
 		bCacheNode = false;
 
 		if (const UEdGraphPin* OutputPin = TypedNodeMacro->GetMacroIOPin(ECOMacroIOType::COMVT_Output, Pin->PinName))
@@ -4573,19 +4574,20 @@ mu::Ptr<mu::NodeMesh> GenerateMutableSourceMesh(const UEdGraphPin* Pin,
 			}
 			else
 			{
-				FString Msg = FString::Printf(TEXT("Macro Pin [%s] not linked."), *Pin->PinName.ToString());
-				GenerationContext.Log(FText::FromString(Msg), Node);
+				FText Msg = FText::Format(LOCTEXT("MacroInstanceError_PinNotLinked_Mesh", "Macro Output node Pin {0} not linked."), FText::FromName(Pin->PinName));
+				GenerationContext.Log(Msg, Node);
 			}
 		}
 		else
 		{
-			FString Msg = FString::Printf(TEXT("Macro Output node does not contain a pin with name [%s]."), *Pin->PinName.ToString());
-			GenerationContext.Log(FText::FromString(Msg), Node);
+			FText Msg = FText::Format(LOCTEXT("MacroInstanceError_PinNameNotFound_Mesh", "Macro Output node does not contain a pin with name {0}."), FText::FromName(Pin->PinName));
+			GenerationContext.Log(Msg, Node);
 		}
 	}
 
 	else if (const UCustomizableObjectNodeTunnel* TypedNodeTunnel = Cast<UCustomizableObjectNodeTunnel>(Node))
 	{
+		// Can't use the GenerateMutableSourceMacro function here because GenerateMutableSourceMesh needs some extra parameters
 		check(TypedNodeTunnel->bIsInputNode);
 		check(GenerationContext.MacroNodesStack.Num());
 
@@ -4602,19 +4604,14 @@ mu::Ptr<mu::NodeMesh> GenerateMutableSourceMesh(const UEdGraphPin* Pin,
 			}
 			else
 			{
-				FString Msg = FString::Printf(TEXT("Macro Instance Pin [%s] is not linked."), *Pin->PinName.ToString());
-				GenerationContext.Log(FText::FromString(Msg), Node);
+				FText Msg = FText::Format(LOCTEXT("MacroTunnelError_PinNotLinked_Mesh", "Macro Instance Pin {0} is not linked."), FText::FromName(Pin->PinName));
+				GenerationContext.Log(Msg, Node);
 			}
 		}
 		else
 		{
-			FString Msg = FString::Printf(TEXT("Macro Instance Node does not contain a pin with name [%s]."), *Pin->PinName.ToString());
-			GenerationContext.Log(FText::FromString(Msg), Node);
-		}
-
-		if (!Result)
-		{
-			GenerationContext.Log(LOCTEXT("ScalarMacroInstanceEmptyAssets", "Something went wrong"), Node);
+			FText Msg = FText::Format(LOCTEXT("MacroTunnelError_PinNameNotFound_Mesh", "Macro Instance Node does not contain a pin with name {0}."), FText::FromName(Pin->PinName));
+			GenerationContext.Log(Msg, Node);
 		}
 
 		// Push the Macro again even if the result is null
diff --git a/Engine/Plugins/Mutable/Source/CustomizableObjectEditor/Private/MuCOE/GenerateMutableSource/GenerateMutableSourceModifier.cpp b/Engine/Plugins/Mutable/Source/CustomizableObjectEditor/Private/MuCOE/GenerateMutableSource/GenerateMutableSourceModifier.cpp
--- a/Engine/Plugins/Mutable/Source/CustomizableObjectEditor/Private/MuCOE/GenerateMutableSource/GenerateMutableSourceModifier.cpp
+++ b/Engine/Plugins/Mutable/Source/CustomizableObjectEditor/Private/MuCOE/GenerateMutableSource/GenerateMutableSourceModifier.cpp
@@ -12,6 +12,7 @@
 #include "MuCOE/GenerateMutableSource/GenerateMutableSourceLayout.h"
 #include "MuCOE/GenerateMutableSource/GenerateMutableSourceGroupProjector.h"
 #include "MuCOE/GenerateMutableSource/GenerateMutableSourceTransform.h"
+#include "MuCOE/GenerateMutableSource/GenerateMutableSourceMacro.h"
 #include "MuCOE/GraphTraversal.h"
 #include "MuCOE/MutableUtils.h"
 #include "MuCOE/Nodes/CustomizableObjectNodeMacroInstance.h"
@@ -68,7 +69,7 @@ TArray<FString> GetNodeRequiredTags(const UCustomizableObjectNodeModifierBase& N
 			}
 			else
 			{
-				GenerationContext.Log(LOCTEXT("MorphStringNodeFailed", "Could not found a linked String node."), &Node);
+				GenerationContext.Log(LOCTEXT("ModifierNodeTagError", "Could not found a linked String node."), &Node);
 			}
 		}
 	}
@@ -399,7 +400,7 @@ mu::Ptr<mu::NodeModifier> GenerateMutableSourceModifier(const UEdGraphPin * Pin,
 		// Add Macro linked Tags
 		if (TypedNodeExt->IsInMacro())
 		{
-			UEdGraphPin* TagsPin = TypedNodeExt->FindPin(FString("Enable Tags"), EEdGraphPinDirection::EGPD_Input);
+			UEdGraphPin* TagsPin = TypedNodeExt->GetEnableTagsPin();
 			check(TagsPin);
 
 			for (const UEdGraphPin* StringPin : FollowInputPinArray(*TagsPin))
@@ -415,7 +416,7 @@ mu::Ptr<mu::NodeModifier> GenerateMutableSourceModifier(const UEdGraphPin * Pin,
 				}
 				else
 				{
-					GenerationContext.Log(LOCTEXT("MorphStringNodeFailed", "Could not found a linked String node."), Node);
+					GenerationContext.Log(LOCTEXT("ModifierNodeTagError", "Could not found a linked String node."), Node);
 				}
 			}
 		}
@@ -586,7 +587,7 @@ mu::Ptr<mu::NodeModifier> GenerateMutableSourceModifier(const UEdGraphPin * Pin,
 
 		FString MorphTargetName = TypedNodeMorph->MorphTargetName;
 
-		if (UEdGraphPin* MorphTargetNamePin = Node->FindPin(FString("Morph Target Name"), EEdGraphPinDirection::EGPD_Input))
+		if (UEdGraphPin* MorphTargetNamePin = TypedNodeMorph->MorphTargetNamePin())
 		{
 			const UEdGraphPin* ConnectedStringPin = FollowInputPin(*MorphTargetNamePin);
 			const UEdGraphPin* SourceStringPin = ConnectedStringPin ? GraphTraversal::FindStringSource(*ConnectedStringPin, &GenerationContext.MacroNodesStack) : nullptr;
@@ -600,7 +601,7 @@ mu::Ptr<mu::NodeModifier> GenerateMutableSourceModifier(const UEdGraphPin * Pin,
 			}
 			else
 			{
-				GenerationContext.Log(LOCTEXT("MorphStringNodeFailed", "Could not found a linked String node."), Node);
+				GenerationContext.Log(LOCTEXT("ModifierNodeTagError", "Could not found a linked String node."), Node);
 			}
 		}
 
@@ -697,63 +698,13 @@ mu::Ptr<mu::NodeModifier> GenerateMutableSourceModifier(const UEdGraphPin * Pin,
 	else if (const UCustomizableObjectNodeMacroInstance* TypedNodeMacro = Cast<UCustomizableObjectNodeMacroInstance>(Node))
 	{
 		bCacheNode = false;
-
-		if (const UEdGraphPin* OutputPin = TypedNodeMacro->GetMacroIOPin(ECOMacroIOType::COMVT_Output, Pin->PinName))
-		{
-			if (const UEdGraphPin* FollowPin = FollowInputPin(*OutputPin))
-			{
-				GenerationContext.MacroNodesStack.Push(TypedNodeMacro);
-				Result = GenerateMutableSourceModifier(FollowPin, GenerationContext);
-				GenerationContext.MacroNodesStack.Pop();
-			}
-			else
-			{
-				FString Msg = FString::Printf(TEXT("Macro Pin [%s] not linked."), *Pin->PinName.ToString());
-				GenerationContext.Log(FText::FromString(Msg), Node);
-			}
-		}
-		else
-		{
-			FString Msg = FString::Printf(TEXT("Macro Output node does not contain a pin with name [%s]."), *Pin->PinName.ToString());
-			GenerationContext.Log(FText::FromString(Msg), Node);
-		}
+		Result = GenerateMutableSourceMacro<mu::NodeModifier>(*Pin, GenerationContext, GenerateMutableSourceModifier);
 	}
 
 	else if (const UCustomizableObjectNodeTunnel* TypedNodeTunnel = Cast<UCustomizableObjectNodeTunnel>(Node))
 	{
-		check(TypedNodeTunnel->bIsInputNode);
-		check(GenerationContext.MacroNodesStack.Num());
-
 		bCacheNode = false;
-
-		const UCustomizableObjectNodeMacroInstance* MacroInstanceNode = GenerationContext.MacroNodesStack.Pop();
-		check(MacroInstanceNode);
-
-		if (const UEdGraphPin* InputPin = MacroInstanceNode->FindPin(Pin->PinName, EEdGraphPinDirection::EGPD_Input))
-		{
-			if (const UEdGraphPin* FollowPin = FollowInputPin(*InputPin))
-			{
-				Result = GenerateMutableSourceModifier(FollowPin, GenerationContext);
-			}
-			else
-			{
-				FString Msg = FString::Printf(TEXT("Macro Instance Pin [%s] is not linked."), *Pin->PinName.ToString());
-				GenerationContext.Log(FText::FromString(Msg), Node);
-			}
-		}
-		else
-		{
-			FString Msg = FString::Printf(TEXT("Macro Instance Node does not contain a pin with name [%s]."), *Pin->PinName.ToString());
-			GenerationContext.Log(FText::FromString(Msg), Node);
-		}
-
-		if (!Result)
-		{
-			GenerationContext.Log(LOCTEXT("ScalarMacroInstanceEmptyAssets", "Something went wrong"), Node);
-		}
-
-		// Push the Macro again even if the result is null
-		GenerationContext.MacroNodesStack.Push(MacroInstanceNode);
+		Result = GenerateMutableSourceMacro<mu::NodeModifier>(*Pin, GenerationContext, GenerateMutableSourceModifier);
 	}
 
 	else
diff --git a/Engine/Plugins/Mutable/Source/CustomizableObjectEditor/Private/MuCOE/GenerateMutableSource/GenerateMutableSourceSurface.cpp b/Engine/Plugins/Mutable/Source/CustomizableObjectEditor/Private/MuCOE/GenerateMutableSource/GenerateMutableSourceSurface.cpp
--- a/Engine/Plugins/Mutable/Source/CustomizableObjectEditor/Private/MuCOE/GenerateMutableSource/GenerateMutableSourceSurface.cpp
+++ b/Engine/Plugins/Mutable/Source/CustomizableObjectEditor/Private/MuCOE/GenerateMutableSource/GenerateMutableSourceSurface.cpp
@@ -27,6 +27,7 @@
 #include "MuCOE/GenerateMutableSource/GenerateMutableSourceMesh.h"
 #include "MuCOE/GenerateMutableSource/GenerateMutableSourceTable.h"
 #include "MuCOE/GenerateMutableSource/GenerateMutableSourceLayout.h"
+#include "MuCOE/GenerateMutableSource/GenerateMutableSourceMacro.h"
 #include "MuCOE/GraphTraversal.h"
 #include "MuCOE/MutableUtils.h"
 #include "MuCOE/Nodes/CustomizableObjectNodeCopyMaterial.h"
@@ -842,7 +843,7 @@ mu::Ptr<mu::NodeSurface> GenerateMutableSourceSurface(const UEdGraphPin * Pin, F
 					}
 					else
 					{
-						GenerationContext.Log(LOCTEXT("MorphStringNodeFailed", "Could not found a linked String node."), Node);
+						GenerationContext.Log(LOCTEXT("SurfaceNodeTagError", "Could not found a linked String node."), Node);
 					}
 				}
 			}
@@ -1041,63 +1042,13 @@ mu::Ptr<mu::NodeSurface> GenerateMutableSourceSurface(const UEdGraphPin * Pin, F
 	else if (const UCustomizableObjectNodeMacroInstance* TypedNodeMacro = Cast<UCustomizableObjectNodeMacroInstance>(Node))
 	{
 		bCacheNode = false;
-
-		if (const UEdGraphPin* OutputPin = TypedNodeMacro->GetMacroIOPin(ECOMacroIOType::COMVT_Output, Pin->PinName))
-		{
-			if (const UEdGraphPin* FollowPin = FollowInputPin(*OutputPin))
-			{
-				GenerationContext.MacroNodesStack.Push(TypedNodeMacro);
-				Result = GenerateMutableSourceSurface(FollowPin, GenerationContext);
-				GenerationContext.MacroNodesStack.Pop();
-			}
-			else
-			{
-				FString Msg = FString::Printf(TEXT("Macro Pin [%s] not linked."), *Pin->PinName.ToString());
-				GenerationContext.Log(FText::FromString(Msg), Node);
-			}
-		}
-		else
-		{
-			FString Msg = FString::Printf(TEXT("Macro Output node does not contain a pin with name [%s]."), *Pin->PinName.ToString());
-			GenerationContext.Log(FText::FromString(Msg), Node);
-		}
+		Result = GenerateMutableSourceMacro<mu::NodeSurface>(*Pin, GenerationContext, GenerateMutableSourceSurface);
 	}
 
 	else if (const UCustomizableObjectNodeTunnel* TypedNodeTunnel = Cast<UCustomizableObjectNodeTunnel>(Node))
 	{
-		check(TypedNodeTunnel->bIsInputNode);
-		check(GenerationContext.MacroNodesStack.Num());
-
 		bCacheNode = false;
-
-		const UCustomizableObjectNodeMacroInstance* MacroInstanceNode = GenerationContext.MacroNodesStack.Pop();
-		check(MacroInstanceNode);
-
-		if (const UEdGraphPin* InputPin = MacroInstanceNode->FindPin(Pin->PinName, EEdGraphPinDirection::EGPD_Input))
-		{
-			if (const UEdGraphPin* FollowPin = FollowInputPin(*InputPin))
-			{
-				Result = GenerateMutableSourceSurface(FollowPin, GenerationContext);
-			}
-			else
-			{
-				FString Msg = FString::Printf(TEXT("Macro Instance Pin [%s] is not linked."), *Pin->PinName.ToString());
-				GenerationContext.Log(FText::FromString(Msg), Node);
-			}
-		}
-		else
-		{
-			FString Msg = FString::Printf(TEXT("Macro Instance Node does not contain a pin with name [%s]."), *Pin->PinName.ToString());
-			GenerationContext.Log(FText::FromString(Msg), Node);
-		}
-
-		if (!Result)
-		{
-			GenerationContext.Log(LOCTEXT("ScalarMacroInstanceEmptyAssets", "Something went wrong"), Node);
-		}
-
-		// Push the Macro again even if the result is null
-		GenerationContext.MacroNodesStack.Push(MacroInstanceNode);
+		Result = GenerateMutableSourceMacro<mu::NodeSurface>(*Pin, GenerationContext, GenerateMutableSourceSurface);
 	}
 
 	else
diff --git a/Engine/Plugins/Mutable/Source/CustomizableObjectEditor/Private/MuCOE/GraphTraversal.cpp b/Engine/Plugins/Mutable/Source/CustomizableObjectEditor/Private/MuCOE/GraphTraversal.cpp
--- a/Engine/Plugins/Mutable/Source/CustomizableObjectEditor/Private/MuCOE/GraphTraversal.cpp
+++ b/Engine/Plugins/Mutable/Source/CustomizableObjectEditor/Private/MuCOE/GraphTraversal.cpp
@@ -404,72 +404,8 @@ const UCustomizableObject* GraphTraversal::GetRootObject(const UCustomizableObje
 }
 
 
-void GraphTraversal::VisitNodes(UCustomizableObjectNode& StartNode, const TFunction<void(UCustomizableObjectNode&)>& VisitFunction, const TMultiMap<FGuid, UCustomizableObjectNodeObject*>* ObjectGroupMap)
-{
-	TSet<UCustomizableObjectNode*> VisitedNodes;
-	
-	TArray<UCustomizableObjectNode*> NodesToVisit;
-	NodesToVisit.Add(&StartNode);
-
-	TMultiMap<FGuid, UCustomizableObjectNodeObject*> LocalObjectGroupMap;
-	if (!ObjectGroupMap)
-	{
-		LocalObjectGroupMap = GetNodeGroupObjectNodeMapping(GetObject(StartNode));
-		ObjectGroupMap = &LocalObjectGroupMap;
-	}
-	
-	while (!NodesToVisit.IsEmpty())
-	{
-		UCustomizableObjectNode* CurrentNode = NodesToVisit.Pop();
-
-		if (VisitedNodes.Contains(CurrentNode))
-		{
-			continue;
-		}
-
-		VisitedNodes.Add(CurrentNode);
-
-		VisitFunction(*CurrentNode);
-
-		for (UEdGraphPin* Pin : CurrentNode->Pins)
-		{
-			if (Pin->Direction != EGPD_Input)
-			{
-				continue;
-			}
-
-			if (IsPinOrphan(*Pin))
-			{
-				continue;
-			}
-			
-			for (UEdGraphPin* ConnectedPin : FollowInputPinArray(*Pin))
-			{
-				UEdGraphNode* ConnectedNode = ConnectedPin->GetOwningNode();
-				
-				if (UCustomizableObjectNode* Node = Cast<UCustomizableObjectNode>(ConnectedNode))
-				{
-					NodesToVisit.Add(Node);
-					
-					if (UCustomizableObjectNodeObjectGroup* ObjectGroupNode = Cast<UCustomizableObjectNodeObjectGroup>(ConnectedNode))
-					{
-						TArray<UCustomizableObjectNodeObject*> ChildObjectNodes;
-						ObjectGroupMap->MultiFind(ObjectGroupNode->NodeGuid, ChildObjectNodes);
-						
-						for (UCustomizableObjectNodeObject* ChildObjectNode : ChildObjectNodes)
-						{
-							NodesToVisit.Add(ChildObjectNode);
-						}
-					}
-				}
-			}
-		}
-	}
-}
-
-
-void GraphTraversal::RecursiveVisitNodes(UCustomizableObjectNode& CurrentNode, const TFunction<void(UCustomizableObjectNode&)>& VisitFunction,
-	TArray<const UCustomizableObjectNodeMacroInstance*>& MacroContext, TSet<UCustomizableObjectNode*>& VisitedNodes, const TMultiMap<FGuid, UCustomizableObjectNodeObject*>* ObjectGroupMap)
+void RecursiveVisitNodes(UCustomizableObjectNode& CurrentNode, const TFunction<void(UCustomizableObjectNode&)>& VisitFunction,
+	TArray<const UCustomizableObjectNodeMacroInstance*>& MacroContext, TSet<UCustomizableObjectNode*>& VisitedNodes, TMultiMap<FGuid, UCustomizableObjectNodeObject*>& ObjectGroupMap)
 {
 	// Check if we already visited this node
 	if (VisitedNodes.Contains(&CurrentNode))
@@ -477,13 +413,6 @@ void GraphTraversal::RecursiveVisitNodes(UCustomizableObjectNode& CurrentNode, c
 		return;
 	}
 
-	TMultiMap<FGuid, UCustomizableObjectNodeObject*> LocalObjectGroupMap;
-	if (!ObjectGroupMap)
-	{
-		LocalObjectGroupMap = GetNodeGroupObjectNodeMapping(GetObject(CurrentNode));
-		ObjectGroupMap = &LocalObjectGroupMap;
-	}
-
 	// Add it to the cache of already visited nodes
 	VisitedNodes.Add(&CurrentNode);
 
@@ -505,11 +434,11 @@ void GraphTraversal::RecursiveVisitNodes(UCustomizableObjectNode& CurrentNode, c
 			if (UCustomizableObjectNodeObjectGroup* ObjectGroupNode = Cast<UCustomizableObjectNodeObjectGroup>(ConnectedNode))
 			{
 				TArray<UCustomizableObjectNodeObject*> ChildObjectNodes;
-				ObjectGroupMap->MultiFind(ObjectGroupNode->NodeGuid, ChildObjectNodes);
+				ObjectGroupMap.MultiFind(ObjectGroupNode->NodeGuid, ChildObjectNodes);
 
 				for (UCustomizableObjectNode* ChildObjectNode : ChildObjectNodes)
 				{
-					GraphTraversal::RecursiveVisitNodes(*ChildObjectNode, VisitFunction, MacroContext, VisitedNodes, ObjectGroupMap);
+					RecursiveVisitNodes(*ChildObjectNode, VisitFunction, MacroContext, VisitedNodes, ObjectGroupMap);
 				}
 			}
 
@@ -541,7 +470,7 @@ void GraphTraversal::RecursiveVisitNodes(UCustomizableObjectNode& CurrentNode, c
 					if (UCustomizableObjectNode* OutputLinkedNode = Cast<UCustomizableObjectNode>(OutputLinkedPin->GetOwningNode()))
 					{
 						MacroContext.Push(MacroInstanceNode);
-						GraphTraversal::RecursiveVisitNodes(*OutputLinkedNode, VisitFunction, MacroContext, VisitedNodes, ObjectGroupMap);
+						RecursiveVisitNodes(*OutputLinkedNode, VisitFunction, MacroContext, VisitedNodes, ObjectGroupMap);
 						MacroContext.Pop();
 					}
 				}
@@ -573,7 +502,7 @@ void GraphTraversal::RecursiveVisitNodes(UCustomizableObjectNode& CurrentNode, c
 
 							if (UCustomizableObjectNode* InputLinkedNode = Cast<UCustomizableObjectNode>(InputLinkedPin->GetOwningNode()))
 							{
-								GraphTraversal::RecursiveVisitNodes(*InputLinkedNode, VisitFunction, MacroContext, VisitedNodes, ObjectGroupMap);
+								RecursiveVisitNodes(*InputLinkedNode, VisitFunction, MacroContext, VisitedNodes, ObjectGroupMap);
 							}
 						}
 					}
@@ -584,13 +513,32 @@ void GraphTraversal::RecursiveVisitNodes(UCustomizableObjectNode& CurrentNode, c
 
 			else if (UCustomizableObjectNode* Node = Cast<UCustomizableObjectNode>(ConnectedNode))
 			{
-				GraphTraversal::RecursiveVisitNodes(*Node, VisitFunction, MacroContext, VisitedNodes, ObjectGroupMap);
+				RecursiveVisitNodes(*Node, VisitFunction, MacroContext, VisitedNodes, ObjectGroupMap);
 			}
 		}
 	}
 }
 
 
+void GraphTraversal::VisitNodes(UCustomizableObjectNode& StartNode, const TFunction<void(UCustomizableObjectNode&)>& VisitFunction, const TMultiMap<FGuid, UCustomizableObjectNodeObject*>* ObjectGroupMap, TArray<const UCustomizableObjectNodeMacroInstance*>* MacroContext)
+{
+	TSet<UCustomizableObjectNode*> VisitedNodes;
+	
+	TMultiMap<FGuid, UCustomizableObjectNodeObject*> LocalObjectGroupMap;
+	TArray<const UCustomizableObjectNodeMacroInstance*> LocalMacroContext;
+
+	if (!ObjectGroupMap)
+	{
+		LocalObjectGroupMap = GetNodeGroupObjectNodeMapping(GetObject(StartNode));
+		ObjectGroupMap = &LocalObjectGroupMap;
+	}
+
+	RecursiveVisitNodes(StartNode, VisitFunction, MacroContext ? *MacroContext : LocalMacroContext, VisitedNodes, LocalObjectGroupMap);
+
+	check(MacroContext ? MacroContext->IsEmpty() : LocalMacroContext.IsEmpty());
+}
+
+
 UCustomizableObjectNodeObject* GraphTraversal::GetFullGraphRootNode(const UCustomizableObject* Object, TArray<const UCustomizableObject*>& VisitedObjects)
 {
 	if (Object != nullptr)
diff --git a/Engine/Plugins/Mutable/Source/CustomizableObjectEditor/Private/MuCOE/GraphTraversal.h b/Engine/Plugins/Mutable/Source/CustomizableObjectEditor/Private/MuCOE/GraphTraversal.h
--- a/Engine/Plugins/Mutable/Source/CustomizableObjectEditor/Private/MuCOE/GraphTraversal.h
+++ b/Engine/Plugins/Mutable/Source/CustomizableObjectEditor/Private/MuCOE/GraphTraversal.h
@@ -80,25 +80,26 @@ namespace GraphTraversal
 		 * object does not have any parent.	*/
 	CUSTOMIZABLEOBJECTEDITOR_API UCustomizableObject* GetRootObject(UCustomizableObject* ChildObject);
 	CUSTOMIZABLEOBJECTEDITOR_API const UCustomizableObject* GetRootObject(const UCustomizableObject* ChildObject);
+	
+	/** Return the full graph Customizable Object Node root of the node given as parameter */
+	UCustomizableObjectNodeObject* GetFullGraphRootNode(const UCustomizableObject* Object, TArray<const UCustomizableObject*>& VisitedObjects);
+
 
 	/** From the StartNode visit all connected nodes in the hierarchy.
 	 *
 	 *   @param StartNode Node to start the visit.
+	 *   @param VisitFunction Called for each UCustomizableObjectNode node found. Order is not guaranteed. 
 	 *   @param ObjectGroupMap Key is the Object Group node id, values are the attached Child Object nodes. Used to prune the traversal.
-	 *   @param VisitFunction Called for each UCustomizableObjectNode node found. Order is not guaranteed. */
-	void VisitNodes(UCustomizableObjectNode& StartNode, const TFunction<void(UCustomizableObjectNode&)>& VisitFunction, const TMultiMap<FGuid, UCustomizableObjectNodeObject*>* ObjectGroupMap = nullptr);
-	
-	/** Return the full graph Customizable Object Node root of the node given as parameter */
-	UCustomizableObjectNodeObject* GetFullGraphRootNode(const UCustomizableObject* Object, TArray<const UCustomizableObject*>& VisitedObjects);
-
-	void RecursiveVisitNodes(UCustomizableObjectNode& CurrentNode, const TFunction<void(UCustomizableObjectNode&)>& VisitFunction,
-		TArray<const UCustomizableObjectNodeMacroInstance*>& MacroContext, TSet<UCustomizableObjectNode*>& VisitedNodes, const TMultiMap<FGuid, UCustomizableObjectNodeObject*>* ObjectGroupMap = nullptr);
+	 *   @param MacroContext Keeps track of the macro context when going through macros
+	  */
+	void VisitNodes(UCustomizableObjectNode& StartNode, const TFunction<void(UCustomizableObjectNode&)>& VisitFunction,
+		const TMultiMap<FGuid, UCustomizableObjectNodeObject*>* ObjectGroupMap = nullptr, TArray<const UCustomizableObjectNodeMacroInstance*>* MacroContext = nullptr);
 
 	/** Given an output pin, return the output pin where the string node is located.
 	* @param Pin: Input pin where the search starts.
 	* @param MacroContext: Copy of the current macro context.
 	*/
-	const UEdGraphPin* FindStringSource(const UEdGraphPin& Pin, TArray<const UCustomizableObjectNodeMacroInstance*>* MacroContext);
+	const UEdGraphPin* FindStringSource(const UEdGraphPin& Pin, TArray<const UCustomizableObjectNodeMacroInstance*>* MacroContext = nullptr);
 }
 
 /** Given an output pin, return the output pin where the mesh is located. */
diff --git a/Engine/Plugins/Mutable/Source/CustomizableObjectEditor/Private/MuCOE/Nodes/CustomizableObjectNode.cpp b/Engine/Plugins/Mutable/Source/CustomizableObjectEditor/Private/MuCOE/Nodes/CustomizableObjectNode.cpp
--- a/Engine/Plugins/Mutable/Source/CustomizableObjectEditor/Private/MuCOE/Nodes/CustomizableObjectNode.cpp
+++ b/Engine/Plugins/Mutable/Source/CustomizableObjectEditor/Private/MuCOE/Nodes/CustomizableObjectNode.cpp
@@ -778,13 +778,20 @@ UCustomizableObjectNodePinData* UCustomizableObjectNode::GetPinData(const UEdGra
 	}
 }
 
+
+bool UCustomizableObjectNode::IsNodeSupportedInMacros() const
+{
+	return true;
+}
+
+
 bool UCustomizableObjectNode::IsInMacro() const
 {
 	if (UCustomizableObjectGraph* Graph = Cast<UCustomizableObjectGraph>(GetGraph()))
 	{
 		return Graph->IsMacro();
 	}
-
+	
 	return false;
 }
 
diff --git a/Engine/Plugins/Mutable/Source/CustomizableObjectEditor/Private/MuCOE/Nodes/CustomizableObjectNodeExposePin.cpp b/Engine/Plugins/Mutable/Source/CustomizableObjectEditor/Private/MuCOE/Nodes/CustomizableObjectNodeExposePin.cpp
--- a/Engine/Plugins/Mutable/Source/CustomizableObjectEditor/Private/MuCOE/Nodes/CustomizableObjectNodeExposePin.cpp
+++ b/Engine/Plugins/Mutable/Source/CustomizableObjectEditor/Private/MuCOE/Nodes/CustomizableObjectNodeExposePin.cpp
@@ -36,6 +36,12 @@ void UCustomizableObjectNodeExposePin::AllocateDefaultPins(UCustomizableObjectNo
 }
 
 
+bool UCustomizableObjectNodeExposePin::IsNodeSupportedInMacros() const
+{
+	return false;
+}
+
+
 FText UCustomizableObjectNodeExposePin::GetNodeTitle(ENodeTitleType::Type TitleType) const
 {
 	FText NodeTitle;
diff --git a/Engine/Plugins/Mutable/Source/CustomizableObjectEditor/Private/MuCOE/Nodes/CustomizableObjectNodeExposePin.h b/Engine/Plugins/Mutable/Source/CustomizableObjectEditor/Private/MuCOE/Nodes/CustomizableObjectNodeExposePin.h
--- a/Engine/Plugins/Mutable/Source/CustomizableObjectEditor/Private/MuCOE/Nodes/CustomizableObjectNodeExposePin.h
+++ b/Engine/Plugins/Mutable/Source/CustomizableObjectEditor/Private/MuCOE/Nodes/CustomizableObjectNodeExposePin.h
@@ -39,7 +39,7 @@ public:
 	// UCustomizableObjectNode interface
 	virtual bool CanConnect(const UEdGraphPin* InOwnedInputPin, const UEdGraphPin* InOutputPin, bool& bOutIsOtherNodeBlocklisted, bool& bOutArePinsCompatible) const override;
 	virtual void AllocateDefaultPins(UCustomizableObjectNodeRemapPins* RemapPins) override;
-	virtual bool IsNodeSupportedInMacros() override { return false; }
+	virtual bool IsNodeSupportedInMacros() const override;
 
 	/** Return the Expose Pin Node expose pin name. */
 	FString GetNodeName() const;
diff --git a/Engine/Plugins/Mutable/Source/CustomizableObjectEditor/Private/MuCOE/Nodes/CustomizableObjectNodeExternalPin.cpp b/Engine/Plugins/Mutable/Source/CustomizableObjectEditor/Private/MuCOE/Nodes/CustomizableObjectNodeExternalPin.cpp
--- a/Engine/Plugins/Mutable/Source/CustomizableObjectEditor/Private/MuCOE/Nodes/CustomizableObjectNodeExternalPin.cpp
+++ b/Engine/Plugins/Mutable/Source/CustomizableObjectEditor/Private/MuCOE/Nodes/CustomizableObjectNodeExternalPin.cpp
@@ -348,6 +348,11 @@ void UCustomizableObjectNodeExternalPin::PostEditChangeProperty(FPropertyChanged
 	}
 }
 
+bool UCustomizableObjectNodeExternalPin::IsNodeSupportedInMacros() const
+{
+	return false;
+}
+
 
 UCustomizableObjectNodeRemapPins* UCustomizableObjectNodeExternalPin::CreateRemapPinsDefault() const
 {
diff --git a/Engine/Plugins/Mutable/Source/CustomizableObjectEditor/Private/MuCOE/Nodes/CustomizableObjectNodeExternalPin.h b/Engine/Plugins/Mutable/Source/CustomizableObjectEditor/Private/MuCOE/Nodes/CustomizableObjectNodeExternalPin.h
--- a/Engine/Plugins/Mutable/Source/CustomizableObjectEditor/Private/MuCOE/Nodes/CustomizableObjectNodeExternalPin.h
+++ b/Engine/Plugins/Mutable/Source/CustomizableObjectEditor/Private/MuCOE/Nodes/CustomizableObjectNodeExternalPin.h
@@ -43,7 +43,7 @@ public:
 	virtual void UpdateReferencedNodeId(const FGuid& NewGuid) override;
 	virtual void PreEditChange(FProperty* PropertyAboutToChange) override;
 	virtual void PostEditChangeProperty(FPropertyChangedEvent& PropertyChangedEvent) override;
-	virtual bool IsNodeSupportedInMacros() override { return false; }
+	virtual bool IsNodeSupportedInMacros() const override;
 	
 	// Own interface
 	
diff --git a/Engine/Plugins/Mutable/Source/CustomizableObjectEditor/Private/MuCOE/Nodes/CustomizableObjectNodeMaterial.cpp b/Engine/Plugins/Mutable/Source/CustomizableObjectEditor/Private/MuCOE/Nodes/CustomizableObjectNodeMaterial.cpp
--- a/Engine/Plugins/Mutable/Source/CustomizableObjectEditor/Private/MuCOE/Nodes/CustomizableObjectNodeMaterial.cpp
+++ b/Engine/Plugins/Mutable/Source/CustomizableObjectEditor/Private/MuCOE/Nodes/CustomizableObjectNodeMaterial.cpp
@@ -341,7 +341,7 @@ void UCustomizableObjectNodeMaterial::AllocateDefaultPins(UCustomizableObjectNod
 		const FString PinName = PinFriendlyName + FString(TEXT("_Input_Pin"));
 		UEdGraphPin* TagsPin = CustomCreatePin(EGPD_Input, Schema->PC_String, FName(PinName), true);
 		TagsPin->PinFriendlyName = FText::FromString(PinFriendlyName);
-		TagsPin->bDefaultValueIsIgnored = true;
+		EnableTagsPinRef = TagsPin;
 	}
 
 	for (const EMaterialParameterType Type : ParameterTypes)
@@ -887,24 +887,7 @@ UEdGraphPin* UCustomizableObjectNodeMaterial::GetMaterialAssetPin() const
 
 UEdGraphPin* UCustomizableObjectNodeMaterial::GetEnableTagsPin() const
 {
-	if(IsInMacro())
-	{
-		const FString PinFriendlyName = "Enable Tags";
-		const FString PinName = PinFriendlyName + FString(TEXT("_Input_Pin"));
-
-		UEdGraphPin* Pin = FindPin(PinName, EEdGraphPinDirection::EGPD_Input);
-
-		if (Pin)
-		{
-			return Pin;
-		}
-		else
-		{
-			return FindPin(PinFriendlyName, EEdGraphPinDirection::EGPD_Input);
-		}
-	}
-
-	return nullptr;
+	return EnableTagsPinRef.Get();
 }
 
 
diff --git a/Engine/Plugins/Mutable/Source/CustomizableObjectEditor/Private/MuCOE/Nodes/CustomizableObjectNodeMaterial.h b/Engine/Plugins/Mutable/Source/CustomizableObjectEditor/Private/MuCOE/Nodes/CustomizableObjectNodeMaterial.h
--- a/Engine/Plugins/Mutable/Source/CustomizableObjectEditor/Private/MuCOE/Nodes/CustomizableObjectNodeMaterial.h
+++ b/Engine/Plugins/Mutable/Source/CustomizableObjectEditor/Private/MuCOE/Nodes/CustomizableObjectNodeMaterial.h
@@ -203,6 +203,9 @@ private:
 	  * If a deprecated pin and a non-deprecated pin have the same Parameter id, this the non-deprecated one prevails. */
 	UPROPERTY()
 	TMap<FNodeMaterialParameterId, FEdGraphPinReference> PinsParameterMap;
+
+	UPROPERTY()
+	FEdGraphPinReference EnableTagsPinRef;
 	
 	/** Create the pin data of the given parameter type. */
 	UCustomizableObjectNodeMaterialPinDataParameter* CreatePinData(EMaterialParameterType Type, int32 ParameterIndex);
diff --git a/Engine/Plugins/Mutable/Source/CustomizableObjectEditor/Private/MuCOE/Nodes/CustomizableObjectNodeMeshClipMorph.cpp b/Engine/Plugins/Mutable/Source/CustomizableObjectEditor/Private/MuCOE/Nodes/CustomizableObjectNodeMeshClipMorph.cpp
--- a/Engine/Plugins/Mutable/Source/CustomizableObjectEditor/Private/MuCOE/Nodes/CustomizableObjectNodeMeshClipMorph.cpp
+++ b/Engine/Plugins/Mutable/Source/CustomizableObjectEditor/Private/MuCOE/Nodes/CustomizableObjectNodeMeshClipMorph.cpp
@@ -140,7 +140,7 @@ void UCustomizableObjectNodeModifierClipMorph::AllocateDefaultPins(UCustomizable
 {
 	const UEdGraphSchema_CustomizableObject* Schema = GetDefault<UEdGraphSchema_CustomizableObject>();
 	
-	UEdGraphPin* OutputPin = CustomCreatePin(EGPD_Output, Schema->PC_Modifier, FName("Modifier"));
+	CustomCreatePin(EGPD_Output, Schema->PC_Modifier, FName("Modifier"));
 
 	//Create Node Modifier Common Pins
 	Super::AllocateDefaultPins(RemapPins);
diff --git a/Engine/Plugins/Mutable/Source/CustomizableObjectEditor/Private/MuCOE/Nodes/CustomizableObjectNodeMeshMorph.cpp b/Engine/Plugins/Mutable/Source/CustomizableObjectEditor/Private/MuCOE/Nodes/CustomizableObjectNodeMeshMorph.cpp
--- a/Engine/Plugins/Mutable/Source/CustomizableObjectEditor/Private/MuCOE/Nodes/CustomizableObjectNodeMeshMorph.cpp
+++ b/Engine/Plugins/Mutable/Source/CustomizableObjectEditor/Private/MuCOE/Nodes/CustomizableObjectNodeMeshMorph.cpp
@@ -43,8 +43,7 @@ void UCustomizableObjectNodeMeshMorph::AllocateDefaultPins(UCustomizableObjectNo
 	if (IsInMacro())
 	{
 		PinName = TEXT("Morph Target Name");
-		UEdGraphPin* MorphNamePin = CustomCreatePin(EGPD_Input, Schema->PC_String, FName(*PinName));
-		MorphNamePin->bDefaultValueIsIgnored = true;
+		MorphTargetNamePinRef = CustomCreatePin(EGPD_Input, Schema->PC_String, FName(*PinName));
 	}
 }
 
@@ -144,6 +143,12 @@ FText UCustomizableObjectNodeMeshMorph::GetTooltipText() const
 }
 
 
+UEdGraphPin* UCustomizableObjectNodeMeshMorph::MorphTargetNamePin() const
+{
+	return MorphTargetNamePinRef.Get();
+}
+
+
 void UCustomizableObjectNodeMeshMorph::Serialize(FArchive& Ar)
 {
 	Super::Serialize(Ar);
diff --git a/Engine/Plugins/Mutable/Source/CustomizableObjectEditor/Private/MuCOE/Nodes/CustomizableObjectNodeMeshMorph.h b/Engine/Plugins/Mutable/Source/CustomizableObjectEditor/Private/MuCOE/Nodes/CustomizableObjectNodeMeshMorph.h
--- a/Engine/Plugins/Mutable/Source/CustomizableObjectEditor/Private/MuCOE/Nodes/CustomizableObjectNodeMeshMorph.h
+++ b/Engine/Plugins/Mutable/Source/CustomizableObjectEditor/Private/MuCOE/Nodes/CustomizableObjectNodeMeshMorph.h
@@ -46,6 +46,8 @@ public:
 		return FindPin(TEXT("Factor"));
 	}
 
+	UEdGraphPin* MorphTargetNamePin() const;
+
 	void Serialize(FArchive& Ar) override;
 
 	UPROPERTY(Category=MeshMorph, EditAnywhere)
@@ -80,5 +82,10 @@ public:
 	UPROPERTY()
 	bool bDeformAllPhysicsBodies_DEPRECATED = false;
 
+private:
+
+	UPROPERTY()
+	FEdGraphPinReference MorphTargetNamePinRef;
+
 };
 
diff --git a/Engine/Plugins/Mutable/Source/CustomizableObjectEditor/Private/MuCOE/Nodes/CustomizableObjectNodeModifierBase.cpp b/Engine/Plugins/Mutable/Source/CustomizableObjectEditor/Private/MuCOE/Nodes/CustomizableObjectNodeModifierBase.cpp
--- a/Engine/Plugins/Mutable/Source/CustomizableObjectEditor/Private/MuCOE/Nodes/CustomizableObjectNodeModifierBase.cpp
+++ b/Engine/Plugins/Mutable/Source/CustomizableObjectEditor/Private/MuCOE/Nodes/CustomizableObjectNodeModifierBase.cpp
@@ -9,7 +9,6 @@
 #include "MuCO/CustomizableObjectCustomVersion.h"
 #include "MuT/NodeModifier.h"
 
-const FName UCustomizableObjectNodeModifierBase::RequiredTagsPinName(TEXT("Required Tags"));
 
 FLinearColor UCustomizableObjectNodeModifierBase::GetNodeTitleColor() const
 {
@@ -26,7 +25,7 @@ UEdGraphPin* UCustomizableObjectNodeModifierBase::OutputPin() const
 
 UEdGraphPin* UCustomizableObjectNodeModifierBase::RequiredTagsPin() const
 {
-	return FindPin(RequiredTagsPinName, EEdGraphPinDirection::EGPD_Input);
+	return RequiredTagsPinRef.Get();
 }
 
 
@@ -303,7 +302,6 @@ void UCustomizableObjectNodeModifierBase::AllocateDefaultPins(UCustomizableObjec
 {
 	if (IsInMacro())
 	{
-		UEdGraphPin* RquiredTags = CustomCreatePin(EGPD_Input, UEdGraphSchema_CustomizableObject::PC_String, RequiredTagsPinName, true);
-		RquiredTags->bDefaultValueIsIgnored = true;
+		RequiredTagsPinRef = CustomCreatePin(EGPD_Input, UEdGraphSchema_CustomizableObject::PC_String, "Required Tags", true);
 	}
 }
diff --git a/Engine/Plugins/Mutable/Source/CustomizableObjectEditor/Private/MuCOE/Nodes/CustomizableObjectNodeModifierBase.h b/Engine/Plugins/Mutable/Source/CustomizableObjectEditor/Private/MuCOE/Nodes/CustomizableObjectNodeModifierBase.h
--- a/Engine/Plugins/Mutable/Source/CustomizableObjectEditor/Private/MuCOE/Nodes/CustomizableObjectNodeModifierBase.h
+++ b/Engine/Plugins/Mutable/Source/CustomizableObjectEditor/Private/MuCOE/Nodes/CustomizableObjectNodeModifierBase.h
@@ -79,8 +79,9 @@ protected:
 	TArray<FLegacyTag> LegacyBackportsRequiredTags;
 
 private:
-	
-	static const FName RequiredTagsPinName;
+
+	UPROPERTY()
+	FEdGraphPinReference RequiredTagsPinRef;
 
 };
 
diff --git a/Engine/Plugins/Mutable/Source/CustomizableObjectEditor/Private/MuCOE/Nodes/CustomizableObjectNodeModifierExtendMeshSection.cpp b/Engine/Plugins/Mutable/Source/CustomizableObjectEditor/Private/MuCOE/Nodes/CustomizableObjectNodeModifierExtendMeshSection.cpp
--- a/Engine/Plugins/Mutable/Source/CustomizableObjectEditor/Private/MuCOE/Nodes/CustomizableObjectNodeModifierExtendMeshSection.cpp
+++ b/Engine/Plugins/Mutable/Source/CustomizableObjectEditor/Private/MuCOE/Nodes/CustomizableObjectNodeModifierExtendMeshSection.cpp
@@ -120,8 +120,7 @@ void UCustomizableObjectNodeModifierExtendMeshSection::AllocateDefaultPins(UCust
 
 	if (IsInMacro())
 	{
-		UEdGraphPin* EnableTags = CustomCreatePin(EGPD_Input, Schema->PC_String, FName("Enable Tags"), true);
-		EnableTags->bDefaultValueIsIgnored = true;
+		EnableTagsPinRef = CustomCreatePin(EGPD_Input, Schema->PC_String, FName("Enable Tags"), true);
 	}
 
 	// Begin texture pins
@@ -162,6 +161,12 @@ UEdGraphPin* UCustomizableObjectNodeModifierExtendMeshSection::AddMeshPin() cons
 }
 
 
+UEdGraphPin* UCustomizableObjectNodeModifierExtendMeshSection::GetEnableTagsPin() const
+{
+	return EnableTagsPinRef.Get();
+}
+
+
 TArray<UCustomizableObjectLayout*> UCustomizableObjectNodeModifierExtendMeshSection::GetLayouts() const
 {
 	TArray<UCustomizableObjectLayout*> Result;
diff --git a/Engine/Plugins/Mutable/Source/CustomizableObjectEditor/Private/MuCOE/Nodes/CustomizableObjectNodeModifierExtendMeshSection.h b/Engine/Plugins/Mutable/Source/CustomizableObjectEditor/Private/MuCOE/Nodes/CustomizableObjectNodeModifierExtendMeshSection.h
--- a/Engine/Plugins/Mutable/Source/CustomizableObjectEditor/Private/MuCOE/Nodes/CustomizableObjectNodeModifierExtendMeshSection.h
+++ b/Engine/Plugins/Mutable/Source/CustomizableObjectEditor/Private/MuCOE/Nodes/CustomizableObjectNodeModifierExtendMeshSection.h
@@ -51,6 +51,7 @@ public:
 
 	// Own interface
 	UEdGraphPin* AddMeshPin() const;
+	UEdGraphPin* GetEnableTagsPin() const;
 
 	TArray<UCustomizableObjectLayout*> GetLayouts() const;
 
@@ -73,5 +74,8 @@ private:
 
 	UPROPERTY()
 	TArray<FCustomizableObjectNodeExtendMaterialImage> Images_DEPRECATED;
+
+	UPROPERTY()
+	FEdGraphPinReference EnableTagsPinRef;
 };
 
diff --git a/Engine/Plugins/Mutable/Source/CustomizableObjectEditor/Private/MuCOE/Nodes/CustomizableObjectNodeModifierMorphMeshSection.cpp b/Engine/Plugins/Mutable/Source/CustomizableObjectEditor/Private/MuCOE/Nodes/CustomizableObjectNodeModifierMorphMeshSection.cpp
--- a/Engine/Plugins/Mutable/Source/CustomizableObjectEditor/Private/MuCOE/Nodes/CustomizableObjectNodeModifierMorphMeshSection.cpp
+++ b/Engine/Plugins/Mutable/Source/CustomizableObjectEditor/Private/MuCOE/Nodes/CustomizableObjectNodeModifierMorphMeshSection.cpp
@@ -30,8 +30,7 @@ void UCustomizableObjectNodeModifierMorphMeshSection::AllocateDefaultPins(UCusto
 	if (IsInMacro())
 	{
 		PinName = TEXT("Morph Target Name");
-		UEdGraphPin* StringPin = CustomCreatePin(EGPD_Input, Schema->PC_String, FName(*PinName));
-		StringPin->bDefaultValueIsIgnored = true;
+		MorphTargetNamePinRef = CustomCreatePin(EGPD_Input, Schema->PC_String, FName(*PinName));
 	}
 
 	//Create Node Modifier Common Pins
@@ -63,4 +62,10 @@ bool UCustomizableObjectNodeModifierMorphMeshSection::IsSingleOutputNode() const
 }
 
 
+UEdGraphPin* UCustomizableObjectNodeModifierMorphMeshSection::MorphTargetNamePin() const
+{
+	return MorphTargetNamePinRef.Get();
+}
+
+
 #undef LOCTEXT_NAMESPACE
diff --git a/Engine/Plugins/Mutable/Source/CustomizableObjectEditor/Private/MuCOE/Nodes/CustomizableObjectNodeModifierMorphMeshSection.h b/Engine/Plugins/Mutable/Source/CustomizableObjectEditor/Private/MuCOE/Nodes/CustomizableObjectNodeModifierMorphMeshSection.h
--- a/Engine/Plugins/Mutable/Source/CustomizableObjectEditor/Private/MuCOE/Nodes/CustomizableObjectNodeModifierMorphMeshSection.h
+++ b/Engine/Plugins/Mutable/Source/CustomizableObjectEditor/Private/MuCOE/Nodes/CustomizableObjectNodeModifierMorphMeshSection.h
@@ -39,5 +39,12 @@ public:
 		return FindPin(TEXT("Factor"));
 	}
 
+	UEdGraphPin* MorphTargetNamePin() const;
+
+private:
+
+	UPROPERTY()
+	FEdGraphPinReference MorphTargetNamePinRef;
+
 };
 
diff --git a/Engine/Plugins/Mutable/Source/CustomizableObjectEditor/Private/MuCOE/Nodes/CustomizableObjectNodeObject.cpp b/Engine/Plugins/Mutable/Source/CustomizableObjectEditor/Private/MuCOE/Nodes/CustomizableObjectNodeObject.cpp
--- a/Engine/Plugins/Mutable/Source/CustomizableObjectEditor/Private/MuCOE/Nodes/CustomizableObjectNodeObject.cpp
+++ b/Engine/Plugins/Mutable/Source/CustomizableObjectEditor/Private/MuCOE/Nodes/CustomizableObjectNodeObject.cpp
@@ -32,7 +32,6 @@ const FName UCustomizableObjectNodeObject::ChildrenPinName(TEXT("Children"));
 const FName UCustomizableObjectNodeObject::ComponentsPinName(TEXT("Components"));
 const FName UCustomizableObjectNodeObject::ModifiersPinName(TEXT("Modifiers"));
 const FName UCustomizableObjectNodeObject::OutputPinName(TEXT("Object"));
-const FName UCustomizableObjectNodeObject::NamePinName(TEXT("Name"));
 
 
 UCustomizableObjectNodeObject::UCustomizableObjectNodeObject()
@@ -601,13 +600,13 @@ void UCustomizableObjectNodeObject::AllocateDefaultPins(UCustomizableObjectNodeR
 	{
 		if (IsInMacro())
 		{
-			UEdGraphPin* NamePin = CustomCreatePin(EGPD_Input, Schema->PC_String, NamePinName);
+			NamePinRef = CustomCreatePin(EGPD_Input, Schema->PC_String, "Name");
 		}
 	}
 }
 
 
-bool UCustomizableObjectNodeObject::IsNodeSupportedInMacros()
+bool UCustomizableObjectNodeObject::IsNodeSupportedInMacros() const
 {
 	return false;
 }
diff --git a/Engine/Plugins/Mutable/Source/CustomizableObjectEditor/Private/MuCOE/Nodes/CustomizableObjectNodeObject.h b/Engine/Plugins/Mutable/Source/CustomizableObjectEditor/Private/MuCOE/Nodes/CustomizableObjectNodeObject.h
--- a/Engine/Plugins/Mutable/Source/CustomizableObjectEditor/Private/MuCOE/Nodes/CustomizableObjectNodeObject.h
+++ b/Engine/Plugins/Mutable/Source/CustomizableObjectEditor/Private/MuCOE/Nodes/CustomizableObjectNodeObject.h
@@ -161,7 +161,7 @@ public:
 	virtual void BackwardsCompatibleFixup(int32 CustomizableObjectCustomVersion) override;
 	virtual void PostBackwardsCompatibleFixup() override;
 	void AllocateDefaultPins(UCustomizableObjectNodeRemapPins* RemapPins) override;
-	virtual bool IsNodeSupportedInMacros() override;
+	virtual bool IsNodeSupportedInMacros() const override;
 
 	// Own Interface
 	void SetParentObject(UCustomizableObject* CustomizableParentObject);
@@ -204,6 +204,8 @@ private:
 	static const FName ComponentsPinName;
 	static const FName ModifiersPinName;
 	static const FName OutputPinName;
-	static const FName NamePinName;
+
+	UPROPERTY()
+	FEdGraphPinReference NamePinRef;
 };
 
diff --git a/Engine/Plugins/Mutable/Source/CustomizableObjectEditor/Private/MuCOE/Nodes/CustomizableObjectNodeObjectChild.cpp b/Engine/Plugins/Mutable/Source/CustomizableObjectEditor/Private/MuCOE/Nodes/CustomizableObjectNodeObjectChild.cpp
--- a/Engine/Plugins/Mutable/Source/CustomizableObjectEditor/Private/MuCOE/Nodes/CustomizableObjectNodeObjectChild.cpp
+++ b/Engine/Plugins/Mutable/Source/CustomizableObjectEditor/Private/MuCOE/Nodes/CustomizableObjectNodeObjectChild.cpp
@@ -51,7 +51,7 @@ bool UCustomizableObjectNodeObjectChild::CanDuplicateNode() const
 }
 
 
-bool UCustomizableObjectNodeObjectChild::IsNodeSupportedInMacros()
+bool UCustomizableObjectNodeObjectChild::IsNodeSupportedInMacros() const
 {
 	return true;
 }
diff --git a/Engine/Plugins/Mutable/Source/CustomizableObjectEditor/Private/MuCOE/Nodes/CustomizableObjectNodeObjectChild.h b/Engine/Plugins/Mutable/Source/CustomizableObjectEditor/Private/MuCOE/Nodes/CustomizableObjectNodeObjectChild.h
--- a/Engine/Plugins/Mutable/Source/CustomizableObjectEditor/Private/MuCOE/Nodes/CustomizableObjectNodeObjectChild.h
+++ b/Engine/Plugins/Mutable/Source/CustomizableObjectEditor/Private/MuCOE/Nodes/CustomizableObjectNodeObjectChild.h
@@ -26,6 +26,6 @@ public:
 	bool CanDuplicateNode() const override;
 
 	// UCustomizableObjectNode interface
-	virtual bool IsNodeSupportedInMacros() override;
+	virtual bool IsNodeSupportedInMacros() const override;
 };
 
diff --git a/Engine/Plugins/Mutable/Source/CustomizableObjectEditor/Private/MuCOE/Nodes/CustomizableObjectNodeVariation.cpp b/Engine/Plugins/Mutable/Source/CustomizableObjectEditor/Private/MuCOE/Nodes/CustomizableObjectNodeVariation.cpp
--- a/Engine/Plugins/Mutable/Source/CustomizableObjectEditor/Private/MuCOE/Nodes/CustomizableObjectNodeVariation.cpp
+++ b/Engine/Plugins/Mutable/Source/CustomizableObjectEditor/Private/MuCOE/Nodes/CustomizableObjectNodeVariation.cpp
@@ -9,8 +9,6 @@ class UCustomizableObjectNodeRemapPins;
 
 #define LOCTEXT_NAMESPACE "CustomizableObjectEditor"
 
-const FString UCustomizableObjectNodeVariation::VariationTagSufix = " - Tag";
-
 void UCustomizableObjectNodeVariation::PostEditChangeProperty(FPropertyChangedEvent& PropertyChangedEvent)
 {
 	if (PropertyChangedEvent.Property)
@@ -36,6 +34,12 @@ void UCustomizableObjectNodeVariation::AllocateDefaultPins(UCustomizableObjectNo
 	}
 	
 	VariationsPins.SetNum(VariationsData.Num());
+
+	if (IsInMacro())
+	{
+		VariationTagPins.SetNum(VariationsData.Num());
+	}
+
 	for (int32 VariationIndex = VariationsData.Num() - 1; VariationIndex >= 0; --VariationIndex)
 	{
 		const FName PinName = FName(FString::Printf( TEXT("Variation %d"), VariationIndex));
@@ -48,9 +52,11 @@ void UCustomizableObjectNodeVariation::AllocateDefaultPins(UCustomizableObjectNo
 
 		if (IsInMacro())
 		{
-			const FName TagPinName = FName(PinName.ToString() + VariationTagSufix);
+			FText VariationTagSufix = LOCTEXT("TagSufix", " - Tag");
+			const FName TagPinName = FName(PinName.ToString() + VariationTagSufix.ToString());
 			UEdGraphPin* VariationTagPin = CustomCreatePin(EGPD_Input, Schema->PC_String, TagPinName);
-			VariationTagPin->bDefaultValueIsIgnored = true;
+			
+			VariationTagPins[VariationIndex] = VariationTagPin;
 		}
 	}
 
@@ -88,6 +94,18 @@ UEdGraphPin* UCustomizableObjectNodeVariation::VariationPin(int32 Index) const
 	{
 		return VariationsPins[Index].Get();
 	}
+
+	return nullptr;
+}
+
+
+UEdGraphPin* UCustomizableObjectNodeVariation::VariationTagPin(int32 Index) const
+{
+	if (VariationTagPins.IsValidIndex(Index))
+	{
+		return VariationTagPins[Index].Get();
+	}
+
 	return nullptr;
 }
 
diff --git a/Engine/Plugins/Mutable/Source/CustomizableObjectEditor/Private/MuCOE/Nodes/CustomizableObjectNodeVariation.h b/Engine/Plugins/Mutable/Source/CustomizableObjectEditor/Private/MuCOE/Nodes/CustomizableObjectNodeVariation.h
--- a/Engine/Plugins/Mutable/Source/CustomizableObjectEditor/Private/MuCOE/Nodes/CustomizableObjectNodeVariation.h
+++ b/Engine/Plugins/Mutable/Source/CustomizableObjectEditor/Private/MuCOE/Nodes/CustomizableObjectNodeVariation.h
@@ -37,12 +37,13 @@ public:
 	UPROPERTY(EditAnywhere, Category = CustomizableObject)
 	TArray<FCustomizableObjectVariation> VariationsData; // The variable name can not be Variations due issues with the on to UObject Serialization system
 
-	static const FString VariationTagSufix;
-
-private:
+protected:
 
 	UPROPERTY()
 	TArray<FEdGraphPinReference> VariationsPins;
+
+	UPROPERTY()
+	TArray<FEdGraphPinReference> VariationTagPins;
 	
 public:
 
@@ -77,8 +78,8 @@ public:
 	/** Get the Variation Input Pin. */
 	UEdGraphPin* VariationPin(int32 Index) const;
 
-	// Needed to convert the NodeTextureVariation into a NodeVariation
-	friend class UCustomizableObjectNodeTextureVariation;
+	/** Get the Variation Tag Input Pin */
+	UEdGraphPin* VariationTagPin(int32 Index) const;
 
 };
 
diff --git a/Engine/Plugins/Mutable/Source/CustomizableObjectEditor/Private/MuCOE/Widgets/SCustomizableObjectMacroLibraryList.cpp b/Engine/Plugins/Mutable/Source/CustomizableObjectEditor/Private/MuCOE/Widgets/SCustomizableObjectMacroLibraryList.cpp
--- a/Engine/Plugins/Mutable/Source/CustomizableObjectEditor/Private/MuCOE/Widgets/SCustomizableObjectMacroLibraryList.cpp
+++ b/Engine/Plugins/Mutable/Source/CustomizableObjectEditor/Private/MuCOE/Widgets/SCustomizableObjectMacroLibraryList.cpp
@@ -69,7 +69,8 @@ public:
 
 	FText GetMacroName() const
 	{
-		return RowItem->Macro.IsValid() ? FText::FromName(RowItem->Macro->Name) : FText::FromString("#Error# - Null Macro");
+		FText ErrorMsg = LOCTEXT("NullMacroRowError", "#Error# - Null Macro");
+		return RowItem->Macro.IsValid() ? FText::FromName(RowItem->Macro->Name) : ErrorMsg;
 	}
 
 
@@ -185,13 +186,13 @@ void SCustomizableObjectMacroLibraryList::OnRemoveCurrentMacro(UCustomizableObje
 	}
 }
 
-void SCustomizableObjectMacroLibraryList::SetSelectedMacro(UCustomizableObjectMacro* MacroToSelect)
+void SCustomizableObjectMacroLibraryList::SetSelectedMacro(const UCustomizableObjectMacro& MacroToSelect)
 {
 	if (ListView.IsValid())
 	{
 		for (TSharedPtr<FMacroTreeEntry> MacroElement : ListViewSource)
 		{
-			if (MacroElement->Macro == MacroToSelect)
+			if (MacroElement->Macro == &MacroToSelect)
 			{
 				ListView->SetSelection(MacroElement);
 				break;
diff --git a/Engine/Plugins/Mutable/Source/CustomizableObjectEditor/Private/MuCOE/Widgets/SCustomizableObjectMacroLibraryList.h b/Engine/Plugins/Mutable/Source/CustomizableObjectEditor/Private/MuCOE/Widgets/SCustomizableObjectMacroLibraryList.h
--- a/Engine/Plugins/Mutable/Source/CustomizableObjectEditor/Private/MuCOE/Widgets/SCustomizableObjectMacroLibraryList.h
+++ b/Engine/Plugins/Mutable/Source/CustomizableObjectEditor/Private/MuCOE/Widgets/SCustomizableObjectMacroLibraryList.h
@@ -36,7 +36,7 @@ public:
 	void OnRemoveCurrentMacro(UCustomizableObjectMacro* MacroToDelete);
 
 	/** Set the Selected Macro */
-	void SetSelectedMacro(UCustomizableObjectMacro* MacroToSelect);
+	void SetSelectedMacro(const UCustomizableObjectMacro& MacroToSelect);
 
 private:
 
diff --git a/Engine/Plugins/Mutable/Source/CustomizableObjectEditor/Public/MuCOE/Nodes/CustomizableObjectNode.h b/Engine/Plugins/Mutable/Source/CustomizableObjectEditor/Public/MuCOE/Nodes/CustomizableObjectNode.h
--- a/Engine/Plugins/Mutable/Source/CustomizableObjectEditor/Public/MuCOE/Nodes/CustomizableObjectNode.h
+++ b/Engine/Plugins/Mutable/Source/CustomizableObjectEditor/Public/MuCOE/Nodes/CustomizableObjectNode.h
@@ -284,7 +284,7 @@ public:
 	}
 
 	/** returs true if this node can be added to a COMacro Graph */
-	virtual bool IsNodeSupportedInMacros() { return true; } 
+	virtual bool IsNodeSupportedInMacros() const;
 
 	/** Returns true if this node is inside a Macro Graph */
 	bool IsInMacro() const;
