Reconciled offline work

diff --git a/Engine/Plugins/Mutable/Source/CustomizableObjectEditor/Private/MuCOE/GraphTraversal.cpp b/Engine/Plugins/Mutable/Source/CustomizableObjectEditor/Private/MuCOE/GraphTraversal.cpp
--- a/Engine/Plugins/Mutable/Source/CustomizableObjectEditor/Private/MuCOE/GraphTraversal.cpp
+++ b/Engine/Plugins/Mutable/Source/CustomizableObjectEditor/Private/MuCOE/GraphTraversal.cpp
@@ -405,7 +405,7 @@ const UCustomizableObject* GraphTraversal::GetRootObject(const UCustomizableObje
 
 
 void RecursiveVisitNodes(UCustomizableObjectNode& CurrentNode, const TFunction<void(UCustomizableObjectNode&)>& VisitFunction,
-	TArray<const UCustomizableObjectNodeMacroInstance*>& MacroContext, TSet<UCustomizableObjectNode*>& VisitedNodes, TMultiMap<FGuid, UCustomizableObjectNodeObject*>& ObjectGroupMap)
+	TArray<const UCustomizableObjectNodeMacroInstance*>& MacroContext, TSet<UCustomizableObjectNode*>& VisitedNodes, const TMultiMap<FGuid, UCustomizableObjectNodeObject*>& ObjectGroupMap)
 {
 	// Check if we already visited this node
 	if (VisitedNodes.Contains(&CurrentNode))
@@ -533,7 +533,7 @@ void GraphTraversal::VisitNodes(UCustomizableObjectNode& StartNode, const TFunct
 		ObjectGroupMap = &LocalObjectGroupMap;
 	}
 
-	RecursiveVisitNodes(StartNode, VisitFunction, MacroContext ? *MacroContext : LocalMacroContext, VisitedNodes, LocalObjectGroupMap);
+	RecursiveVisitNodes(StartNode, VisitFunction, MacroContext ? *MacroContext : LocalMacroContext, VisitedNodes, *ObjectGroupMap);
 
 	check(MacroContext ? MacroContext->IsEmpty() : LocalMacroContext.IsEmpty());
 }
