[mutable] Extend mutable debugger mesh viewer with initial rendering of UVs.
Added some additional debug safety checks in the image provider.
#rnx

diff --git a/Engine/Plugins/Mutable/Source/CustomizableObject/Private/MuCO/UnrealMutableImageProvider.cpp b/Engine/Plugins/Mutable/Source/CustomizableObject/Private/MuCO/UnrealMutableImageProvider.cpp
--- a/Engine/Plugins/Mutable/Source/CustomizableObject/Private/MuCO/UnrealMutableImageProvider.cpp
+++ b/Engine/Plugins/Mutable/Source/CustomizableObject/Private/MuCO/UnrealMutableImageProvider.cpp
@@ -354,7 +354,13 @@ TTuple<UE::Tasks::FTask, TFunction<void()>> FUnrealMutableResourceProvider::GetR
 	TSharedPtr<mu::FImage> Image = MakeShared<mu::FImage>();
 
 	FScopeLock Lock(&RuntimeReferencedLock);
-	
+	if (!RuntimeReferencedImages.Contains(ModelPtr))
+	{
+		UE_LOG(LogMutable, Error, TEXT("Failed to load image [%i]. Model not registered in the provider."), Id);
+		ResultCallback(CreateDummy());
+		return Invoke(TrivialReturn);
+	}
+
 	TArray<FMutableSourceTextureData>& SourceTextures = RuntimeReferencedImages[ModelPtr].SourceTextures;
 	if (!SourceTextures.IsValidIndex(Id))
 	{
diff --git a/Engine/Plugins/Mutable/Source/CustomizableObjectEditor/Private/MuCOE/SCustomizableObjectLayoutGrid.cpp b/Engine/Plugins/Mutable/Source/CustomizableObjectEditor/Private/MuCOE/SCustomizableObjectLayoutGrid.cpp
--- a/Engine/Plugins/Mutable/Source/CustomizableObjectEditor/Private/MuCOE/SCustomizableObjectLayoutGrid.cpp
+++ b/Engine/Plugins/Mutable/Source/CustomizableObjectEditor/Private/MuCOE/SCustomizableObjectLayoutGrid.cpp
@@ -243,7 +243,7 @@ int32 SCustomizableObjectLayoutGrid::OnPaint(const FPaintArgs& Args, const FGeom
 	
 	ELayoutGridMode GridMode = Mode.Get();
 	UVCanvasDrawer->SetLayoutMode(GridMode);
-	UVCanvasDrawer->InitializeDrawingData(UVLayout, UnassignedUVLayoutVertices, Blocks.Get(), SelectedBlocks);
+	UVCanvasDrawer->InitializeDrawingData(UVLayout.Get(), UnassignedUVLayoutVertices, Blocks.Get(), SelectedBlocks);
 	UVCanvasDrawer->Initialize(CanvasRect, ClippingRect, Origin * AllottedGeometry.Scale, Size * AllottedGeometry.Scale, GridSizePoint, AuxCellSize * AllottedGeometry.Scale);
 
 	FSlateDrawElement::MakeCustom(OutDrawElements, RetLayerId, UVCanvasDrawer);
diff --git a/Engine/Plugins/Mutable/Source/CustomizableObjectEditor/Private/MuCOE/SCustomizableObjectLayoutGrid.h b/Engine/Plugins/Mutable/Source/CustomizableObjectEditor/Private/MuCOE/SCustomizableObjectLayoutGrid.h
--- a/Engine/Plugins/Mutable/Source/CustomizableObjectEditor/Private/MuCOE/SCustomizableObjectLayoutGrid.h
+++ b/Engine/Plugins/Mutable/Source/CustomizableObjectEditor/Private/MuCOE/SCustomizableObjectLayoutGrid.h
@@ -64,7 +64,7 @@ public:
 		SLATE_ATTRIBUTE(ELayoutGridMode, Mode)
 		SLATE_ATTRIBUTE( FIntPoint, GridSize )
 		SLATE_ATTRIBUTE(TArray<FCustomizableObjectLayoutBlock>, Blocks)
-		SLATE_ARGUMENT( TArray<FVector2f>, UVLayout )
+		SLATE_ATTRIBUTE( TArray<FVector2f>, UVLayout )
 		SLATE_ARGUMENT( TArray<FVector2f>, UnassignedUVLayoutVertices )
 		SLATE_ARGUMENT( FColor, SelectionColor )
 		SLATE_EVENT( FBlockChangedDelegate, OnBlockChanged )
@@ -191,7 +191,7 @@ private:
 	TAttribute< TArray<FCustomizableObjectLayoutBlock> > Blocks;
 
 	/** Array with all the UVs to draw in the layout */
-	TArray<FVector2f> UVLayout;
+	TAttribute< TArray<FVector2f> > UVLayout;
 
 	/** Array with all the unassigned UVs */
 	TArray<FVector2f> UnassignedUVLayoutVertices;
diff --git a/Engine/Plugins/Mutable/Source/CustomizableObjectEditor/Private/MuCOE/SMutableMeshViewer.cpp b/Engine/Plugins/Mutable/Source/CustomizableObjectEditor/Private/MuCOE/SMutableMeshViewer.cpp
--- a/Engine/Plugins/Mutable/Source/CustomizableObjectEditor/Private/MuCOE/SMutableMeshViewer.cpp
+++ b/Engine/Plugins/Mutable/Source/CustomizableObjectEditor/Private/MuCOE/SMutableMeshViewer.cpp
@@ -8,6 +8,7 @@
 #include "MuCOE/SMutableSkeletonViewer.h"
 #include "MuCO/CustomizableObject.h"
 #include "MuT/TypeInfo.h"
+#include "MuR/MeshPrivate.h"
 #include "PropertyCustomizationHelpers.h"
 #include "Widgets/Layout/SScrollBox.h"
 #include "Widgets/Views/SListView.h"
@@ -497,10 +498,82 @@ TSharedRef<SWidget> SMutableMeshViewer::GenerateDataTableSlates()
 					.OnGenerateRow(this, &SMutableMeshViewer::GenerateTagRow)
 				]
 			// ---------------------------------
+
+			// UVs
+			+SVerticalBox::Slot()
+			.AutoHeight()
+			[
+				SNew(STextBlock).
+					Text(LOCTEXT("UVs", "UVs"))
+			]
+
+			+ SVerticalBox::Slot()
+			.Padding(IndentationSpace, AfterTitleSpacing)
+			.AutoHeight()
+			.MinHeight(300.0f)
+			[
+				SAssignNew(LayoutGridWidget, SCustomizableObjectLayoutGrid)
+					.Mode(ELGM_Show)
+					.GridSize(this, &SMutableMeshViewer::GetGridSize)
+					.Blocks(this, &SMutableMeshViewer::GetBlocks)
+					.UVLayout(this, &SMutableMeshViewer::GetUVs)
+			]
+			// ---------------------------------
 		];
 }
 
 
+FIntPoint SMutableMeshViewer::GetGridSize() const
+{
+	// TODO
+	return FIntPoint(8,8);
+}
+
+
+TArray<FCustomizableObjectLayoutBlock> SMutableMeshViewer::GetBlocks() const
+{
+	TArray<FCustomizableObjectLayoutBlock> Blocks;
+
+	// TODO
+	return Blocks;
+}
+
+
+TArray<FVector2f> SMutableMeshViewer::GetUVs() const
+{
+	TArray<FVector2f> Lines;
+
+	if (MutableMesh)
+	{
+		int32 NumFaces = MutableMesh->GetFaceCount();
+		Lines.Reserve(NumFaces*6);
+
+		// TODO
+		int32 UVIndex = 0;
+
+		mu::UntypedMeshBufferIteratorConst ItIndices(MutableMesh->IndexBuffers, mu::EMeshBufferSemantic::VertexIndex, 0);
+		mu::UntypedMeshBufferIteratorConst It(MutableMesh->VertexBuffers, mu::EMeshBufferSemantic::TexCoords, UVIndex);
+		for (int32 FaceIndex=0; FaceIndex<NumFaces; ++FaceIndex)
+		{
+			uint32 V0 = (ItIndices + FaceIndex * 3 + 0).GetAsUINT32();
+			uint32 V1 = (ItIndices + FaceIndex * 3 + 1).GetAsUINT32();
+			uint32 V2 = (ItIndices + FaceIndex * 3 + 2).GetAsUINT32();
+
+			Lines.Add((It + V0).GetAsVec2f());
+			Lines.Add((It + V1).GetAsVec2f());
+
+			Lines.Add((It + V1).GetAsVec2f());
+			Lines.Add((It + V2).GetAsVec2f());
+
+			Lines.Add((It + V2).GetAsVec2f());
+			Lines.Add((It + V0).GetAsVec2f());
+		}
+	}
+
+	return Lines;
+}
+
+
 TSharedRef<ITableRow> SMutableMeshViewer::GenerateTagRow(TSharedPtr<FString> InItem, const TSharedRef<STableViewBase>& OwnerTable)
 {
 	return SNew(STableRow<TSharedPtr<FName>>, OwnerTable)
diff --git a/Engine/Plugins/Mutable/Source/CustomizableObjectEditor/Private/MuCOE/SMutableMeshViewer.h b/Engine/Plugins/Mutable/Source/CustomizableObjectEditor/Private/MuCOE/SMutableMeshViewer.h
--- a/Engine/Plugins/Mutable/Source/CustomizableObjectEditor/Private/MuCOE/SMutableMeshViewer.h
+++ b/Engine/Plugins/Mutable/Source/CustomizableObjectEditor/Private/MuCOE/SMutableMeshViewer.h
@@ -5,6 +5,7 @@
 #include "MuR/Mesh.h"
 #include "UObject/GCObject.h"
 #include "Widgets/SCompoundWidget.h"
+#include "MuCOE/SCustomizableObjectLayoutGrid.h"
 
 template <typename ItemType> class SListView;
 
@@ -92,6 +93,9 @@ private:
 	TSharedPtr<SListView<TSharedPtr<FBufferElement>>> VertexBuffersSlateView;
 	TSharedPtr<SListView<TSharedPtr<FBufferElement>>> IndexBuffersSlateView;
 
+	/** Widget to show the UVs */
+	TSharedPtr<SCustomizableObjectLayoutGrid> LayoutGridWidget;
+
 private:
 
 	/** Generates all slate objects related with the Mesh Viewport Slate */
@@ -103,6 +107,11 @@ private:
 	/** Method called each time the mesh selected changes so the UI gets updated reliably */
 	void OnMeshChanged();
 
+	/** UV widget callbacks. */
+	FIntPoint GetGridSize() const;
+	TArray<FCustomizableObjectLayoutBlock> GetBlocks() const;
+	TArray<FVector2f> GetUVs() const;
+
 private:
 	
 	/* Generic UI callbacks used by the Widget */
diff --git a/Engine/Plugins/Mutable/Source/MutableRuntime/Internal/MuR/MeshPrivate.h b/Engine/Plugins/Mutable/Source/MutableRuntime/Internal/MuR/MeshPrivate.h
--- a/Engine/Plugins/Mutable/Source/MutableRuntime/Internal/MuR/MeshPrivate.h
+++ b/Engine/Plugins/Mutable/Source/MutableRuntime/Internal/MuR/MeshPrivate.h
@@ -98,7 +98,7 @@ namespace mu
 			Format = EMeshBufferFormat::None;
 			Components = 0;
 			ElementSize = 0;
-			Buffer = 0;
+			Buffer = nullptr;
 		}
 
 
@@ -140,7 +140,7 @@ namespace mu
 				Format = EMeshBufferFormat::None;
 				Components = 0;
 				ElementSize = 0;
-				Buffer = 0;
+				Buffer = nullptr;
 			}
 		}
 
