[mutable] Fixed components directly connected to child objects not being added to the compilation pre-pass.
#tests Editor
#rnx

diff --git a/Engine/Plugins/Mutable/Source/CustomizableObject/Internal/MuCO/CustomizableObjectPrivate.h b/Engine/Plugins/Mutable/Source/CustomizableObject/Internal/MuCO/CustomizableObjectPrivate.h
--- a/Engine/Plugins/Mutable/Source/CustomizableObject/Internal/MuCO/CustomizableObjectPrivate.h
+++ b/Engine/Plugins/Mutable/Source/CustomizableObject/Internal/MuCO/CustomizableObjectPrivate.h
@@ -1604,6 +1604,8 @@ enum class ECustomizableObjectVersions : int32
 
 	ReaddAddedMeshIDToMeshParamOp,
 
+	AddConnectedChildObjectComponentsToPrepass,
+
 	// -----<new versions can be added above this line>--------
 	LastCustomizableObjectVersion
 };
diff --git a/Engine/Plugins/Mutable/Source/CustomizableObjectEditor/Private/MuCOE/GenerateMutableSource/GenerateMutableSourceComponent.cpp b/Engine/Plugins/Mutable/Source/CustomizableObjectEditor/Private/MuCOE/GenerateMutableSource/GenerateMutableSourceComponent.cpp
--- a/Engine/Plugins/Mutable/Source/CustomizableObjectEditor/Private/MuCOE/GenerateMutableSource/GenerateMutableSourceComponent.cpp
+++ b/Engine/Plugins/Mutable/Source/CustomizableObjectEditor/Private/MuCOE/GenerateMutableSource/GenerateMutableSourceComponent.cpp
@@ -163,7 +163,7 @@ mu::Ptr<mu::NodeComponent> GenerateMutableSourceComponent(const UEdGraphPin* Pin
 	}
 
 	// Bool that determines if a node can be added to the cache of nodes.
-	// Most nodes need to be added to the cache but there are some that don't. For exampel, MacroInstanceNodes
+	// Most nodes need to be added to the cache but there are some that don't. For example, MacroInstanceNodes
 	bool bCacheNode = true;
 	
 	mu::Ptr<mu::NodeComponent> Result;
diff --git a/Engine/Plugins/Mutable/Source/CustomizableObjectEditor/Private/MuCOE/GraphTraversal.cpp b/Engine/Plugins/Mutable/Source/CustomizableObjectEditor/Private/MuCOE/GraphTraversal.cpp
--- a/Engine/Plugins/Mutable/Source/CustomizableObjectEditor/Private/MuCOE/GraphTraversal.cpp
+++ b/Engine/Plugins/Mutable/Source/CustomizableObjectEditor/Private/MuCOE/GraphTraversal.cpp
@@ -433,6 +433,10 @@ void RecursiveVisitNodes(UCustomizableObjectNode& CurrentNode, const TFunction<v
 
 			if (UCustomizableObjectNodeObjectGroup* ObjectGroupNode = Cast<UCustomizableObjectNodeObjectGroup>(ConnectedNode))
 			{
+				// Visit the pins directly connected to the ObjectGroup node such as the child objects and projectors
+				RecursiveVisitNodes(*ObjectGroupNode, VisitFunction, MacroContext, VisitedNodes, ObjectGroupMap);
+
+				// Visit the child objects in other COs
 				TArray<UCustomizableObjectNodeObject*> ChildObjectNodes;
 				ObjectGroupMap.MultiFind(ObjectGroupNode->NodeGuid, ChildObjectNodes);
 
