[Mutable] Set mutable.ForceGeometryOnFirstGeneration to false as default.

- With the cvar set to false we could have rendering crashes on some cases when meshes are streamed.

#rnx

diff --git a/Engine/Plugins/Mutable/Source/CustomizableObject/Private/MuCO/CustomizableObjectSystem.cpp b/Engine/Plugins/Mutable/Source/CustomizableObject/Private/MuCO/CustomizableObjectSystem.cpp
--- a/Engine/Plugins/Mutable/Source/CustomizableObject/Private/MuCO/CustomizableObjectSystem.cpp
+++ b/Engine/Plugins/Mutable/Source/CustomizableObject/Private/MuCO/CustomizableObjectSystem.cpp
@@ -130,8 +130,8 @@ TAutoConsoleVariable<bool> CVarIgnoreFirstAvailableLODCalculation(
 
 TAutoConsoleVariable<bool> CVarForceGeometryOnFirstGeneration(
 	TEXT("mutable.ForceGeometryOnFirstGeneration"),
-	false,
-	TEXT("If set to true, forces geometry generation on firts generation even if the LOD will be streamed."));
+	true,
+	TEXT("If set to true, forces geometry generation on first generation even if the LOD will be streamed."));
 
 #if WITH_EDITOR
 bool bEnableLODManagmentInEditor = false;
