[Mutable] Optimize ME_Morph operation.

#rnx

diff --git a/Engine/Plugins/Mutable/Source/MutableRuntime/Private/MuR/OpMeshMorph.h b/Engine/Plugins/Mutable/Source/MutableRuntime/Private/MuR/OpMeshMorph.h
--- a/Engine/Plugins/Mutable/Source/MutableRuntime/Private/MuR/OpMeshMorph.h
+++ b/Engine/Plugins/Mutable/Source/MutableRuntime/Private/MuR/OpMeshMorph.h
@@ -8,6 +8,8 @@
 #include "MuR/Platform.h"
 #include "MuR/SparseIndexMap.h"
 
+#include "PackedNormal.h"
+
 namespace mu
 {
 	/**
@@ -27,6 +29,7 @@ namespace mu
 			MeshVertexIdIteratorConst MorphIdIter, int32 MorphNum)
 		-> SparseIndexMapSet
 		{	
+			MUTABLE_CPUPROFILER_SCOPE(MakeIndexMap);
 			TArray<SparseIndexMapSet::FRangeDesc> RangeDescs;
 
 			// Detect all ranges and their limits
@@ -68,161 +71,206 @@ namespace mu
 			return IndexMap;
 		};
 
-        const auto ApplyNormalMorph = [](
-			MeshVertexIdIteratorConst BaseIdIter, const TStaticArray<UntypedMeshBufferIterator, 3>& BaseTangentFrameIters, const int32 BaseNum,
-			MeshVertexIdIteratorConst MorphIdIter, UntypedMeshBufferIteratorConst& MorphNormalIter, const int32 MorphNum,
-				const SparseIndexMapSet& IndexMap, const float Factor)
-        -> void
-        {
-			const UntypedMeshBufferIterator& BaseNormalIter = BaseTangentFrameIters[2];
-			const UntypedMeshBufferIterator& BaseTangentIter = BaseTangentFrameIters[1];
-			const UntypedMeshBufferIterator& BaseBiNormalIter = BaseTangentFrameIters[0];
+		const auto ApplyMorph = [](
+			MeshVertexIdIteratorConst BaseIdIter, const MeshBufferIterator<EMeshBufferFormat::Float32, float, 3>& BasePositionIter,
+			const TStaticArray<UntypedMeshBufferIterator, 3>& BaseTangentFrameIters, const TArray<UntypedMeshBufferIterator>& BaseGenericIters, const int32 BaseNum,
+			MeshVertexIdIteratorConst MorphIdIter, MeshBufferIteratorConst<EMeshBufferFormat::Float32, float, 3>& MorphPositionIter,
+			UntypedMeshBufferIteratorConst& MorphNormalIter, const TArray<UntypedMeshBufferIteratorConst>& MorphGenericIters, const int32 MorphNum,
+			const SparseIndexMapSet& IndexMap, const float Factor)
+			-> void
+			{
+				MUTABLE_CPUPROFILER_SCOPE(ApplyMorph);
 
-            const EMeshBufferFormat NormalFormat = BaseNormalIter.GetFormat();
-            const int32 NormalComps = BaseNormalIter.GetComponents();
+				const UntypedMeshBufferIterator& BaseNormalIter = BaseTangentFrameIters[2];
+				const UntypedMeshBufferIterator& BaseTangentIter = BaseTangentFrameIters[1];
+				const UntypedMeshBufferIterator& BaseBiNormalIter = BaseTangentFrameIters[0];
 
-            const EMeshBufferFormat TangentFormat = BaseTangentIter.GetFormat();
-            const int32 TangentComps = BaseTangentIter.GetComponents();
+				const EMeshBufferFormat NormalFormat = BaseNormalIter.GetFormat();
+				const int32 NormalComps = BaseNormalIter.GetComponents();
 
-            const EMeshBufferFormat BiNormalFormat = BaseBiNormalIter.GetFormat();
-            const int32 BiNormalComps = BaseBiNormalIter.GetComponents();
+				const EMeshBufferFormat TangentFormat = BaseTangentIter.GetFormat();
+				const int32 TangentComps = BaseTangentIter.GetComponents();
 
-			// When normal is packed, binormal channel is not expected. It is not a big deal if it's there but we would be doing extra unused work in that case. 
-			ensureAlways(!(NormalFormat == EMeshBufferFormat::PackedDir8_W_TangentSign || NormalFormat == EMeshBufferFormat::PackedDirS8_W_TangentSign) || !BaseBiNormalIter.ptr());
+				const EMeshBufferFormat BiNormalFormat = BaseBiNormalIter.GetFormat();
+				const int32 BiNormalComps = BaseBiNormalIter.GetComponents();
 
-            for (int32 VertexIndex = 0; VertexIndex < BaseNum; ++VertexIndex)
-            {
-                const uint64 BaseId = (BaseIdIter + VertexIndex).Get();
-                const int32 MorphIndex = static_cast<int32>(IndexMap.Find(BaseId));
+				const EMeshBufferFormat MorphNormalFormat = MorphNormalIter.GetFormat();
 
-                if (MorphIndex == SparseIndexMap::NotFoundValue)
-                {
-                    continue;
-                }
+				const bool bHasPositions = BasePositionIter.ptr() && MorphPositionIter.ptr();
+				check((BasePositionIter.GetFormat() == EMeshBufferFormat::Float32 && BasePositionIter.GetComponents() == 3) || !bHasPositions);
 
-                // Find consecutive run.
-				MeshVertexIdIteratorConst RunBaseIter = BaseIdIter + VertexIndex;
-				MeshVertexIdIteratorConst RunMorphIter = MorphIdIter + MorphIndex;
+				const bool bHasOptimizedNormals = NormalFormat == EMeshBufferFormat::PackedDirS8_W_TangentSign && MorphNormalFormat == EMeshBufferFormat::Float32
+					&& (!BaseTangentIter.ptr() || TangentFormat == EMeshBufferFormat::PackedDirS8) && !BaseBiNormalIter.ptr();
+				
+				bool bHasGenericMorphs = false;
+				const int32 ChannelNum = MorphGenericIters.Num();
+				for (int32 ChannelIndex = 0; ChannelIndex < ChannelNum; ++ChannelIndex)
+				{
+					if (!(BaseGenericIters[ChannelIndex].ptr() && MorphGenericIters[ChannelIndex].ptr()))
+					{
+						continue;
+					}
 
-                int32 RunSize = 0;
-                for ( ; VertexIndex + RunSize < BaseNum && MorphIndex + RunSize < MorphNum && RunBaseIter.Get() == RunMorphIter.Get();
-                        ++RunSize, ++RunBaseIter, ++RunMorphIter);
+					bHasGenericMorphs = true;
+				}
 
-				for (int32 RunIndex = 0; RunIndex < RunSize; ++RunIndex)
+				for (int32 VertexIndex = 0; VertexIndex < BaseNum; ++VertexIndex)
 				{
-					UntypedMeshBufferIterator NormalIter = BaseNormalIter + (VertexIndex + RunIndex);
+					const uint64 BaseId = (BaseIdIter + VertexIndex).Get();
+					const uint32 MorphIndex = IndexMap.Find(BaseId);
 
-					const FVector3f BaseNormal = NormalIter.GetAsVec3f();
-					const FVector3f MorphNormal = (MorphNormalIter + (MorphIndex + RunIndex)).GetAsVec3f();
- 
-					const FVector3f Normal = (BaseNormal + MorphNormal*Factor).GetSafeNormal();
-				
-					// Leave the tangent basis sign untouched for packed normals formats.
-					for (int32 C = 0; C < NormalComps && C < 3; ++C)
+					if (MorphIndex == SparseIndexMap::NotFoundValue)
 					{
-						ConvertData(C, NormalIter.ptr(), NormalFormat, &Normal, EMeshBufferFormat::Float32);
+						continue;
 					}
 
-					// Tangent
-					if (BaseTangentIter.ptr())
-					{
-						UntypedMeshBufferIterator TangentIter = BaseTangentIter + (VertexIndex + RunIndex);
-
-						const FVector3f BaseTangent = TangentIter.GetAsVec3f();
+					// Find consecutive run.
+					MeshVertexIdIteratorConst RunBaseIter = BaseIdIter + VertexIndex;
+					MeshVertexIdIteratorConst RunMorphIter = MorphIdIter + MorphIndex;
 
-						// Orthogonalize Tangent based on new Normal. This assumes Normal and BaseTangent are normalized and different.
-						const FVector3f Tangent = (BaseTangent - FVector3f::DotProduct(Normal, BaseTangent) * Normal).GetSafeNormal();
+					int32 RunSize = 0;
+					for (; VertexIndex + RunSize < BaseNum && int32(MorphIndex) + RunSize < MorphNum && RunBaseIter.Get() == RunMorphIter.Get();
+						++RunSize, ++RunBaseIter, ++RunMorphIter);
 
-						for (int32 C = 0; C < TangentComps && C < 3; ++C)
+					// Positions
+					if (bHasPositions)
+					{
+						for (int32 RunIndex = 0; RunIndex < RunSize; ++RunIndex)
 						{
-							ConvertData(C, TangentIter.ptr(), TangentFormat, &Tangent, EMeshBufferFormat::Float32);
+							FVector3f& Position = *reinterpret_cast<FVector3f*>(*(BasePositionIter + VertexIndex + RunIndex));
+							const FVector3f& MorphPosition = *reinterpret_cast<const FVector3f*>(*(MorphPositionIter + MorphIndex + RunIndex));
+							Position += MorphPosition * Factor;
 						}
+					}
 
-						// BiNormal
-						if (BaseBiNormalIter.ptr())
+					// Normals
+					if (bHasOptimizedNormals)
+					{
+						for (int32 RunIndex = 0; RunIndex < RunSize; ++RunIndex)
 						{
-							UntypedMeshBufferIterator BiNormalIter = BaseBiNormalIter + (VertexIndex + RunIndex);
+							// Normal
+							FPackedNormal* PackedNormal = reinterpret_cast<FPackedNormal*>((BaseNormalIter + VertexIndex + RunIndex).ptr());
+							int8 W = PackedNormal->Vector.W;
+							const FVector3f BaseNormal = PackedNormal->ToFVector3f();
+
+							const FVector3f* MorphNormal = reinterpret_cast<const FVector3f*>((MorphNormalIter + MorphIndex + RunIndex).ptr());
 
-							const FVector3f& N = BaseNormal;
-							const FVector3f& T = BaseTangent;
-							const FVector3f  B = BiNormalIter.GetAsVec3f();
+							const FVector3f Normal = (BaseNormal + *MorphNormal * Factor).GetSafeNormal();
 
-							const float BaseTangentBasisDeterminant =  
-									B.X*T.Y*N.Z + B.Z*T.X*N.Y + B.Y*T.Z*N.Y -
-									B.Z*T.Y*N.X - B.Y*T.X*N.Z - B.X*T.Z*N.Y;
+							*PackedNormal = *reinterpret_cast<const FVector3f*>(&Normal);
+							PackedNormal->Vector.W = W;
 
-							const float BaseTangentBasisDeterminantSign = BaseTangentBasisDeterminant >= 0 ? 1.0f : -1.0f;
+							// Tangent
+							if (BaseTangentIter.ptr())
+							{
+								FPackedNormal* PackedTangent = reinterpret_cast<FPackedNormal*>((BaseTangentIter + (VertexIndex + RunIndex)).ptr());
+								const FVector3f BaseTangent = PackedTangent->ToFVector3f();
+
+								// Orthogonalize Tangent based on new Normal. This assumes Normal and BaseTangent are normalized and different.
+								const FVector3f Tangent = (BaseTangent - FVector3f::DotProduct(Normal, BaseTangent) * Normal).GetSafeNormal();
+
+								*PackedTangent = *reinterpret_cast<const FVector3f*>(&Tangent);
+							}
+						}
+					}
+					else if (BaseNormalIter.ptr())
+					{
+						// When normal is packed, binormal channel is not expected. It is not a big deal if it's there but we would be doing extra unused work in that case. 
+						ensure(!(NormalFormat == EMeshBufferFormat::PackedDir8_W_TangentSign || NormalFormat == EMeshBufferFormat::PackedDirS8_W_TangentSign) || !BaseBiNormalIter.ptr());
+						
+						MUTABLE_CPUPROFILER_SCOPE(ApplyNormalMorph_SlowPath);
 
-							const FVector3f BiNormal = FVector3f::CrossProduct(Tangent, Normal) * BaseTangentBasisDeterminantSign;
+						for (int32 RunIndex = 0; RunIndex < RunSize; ++RunIndex)
+						{
+							UntypedMeshBufferIterator NormalIter = BaseNormalIter + (VertexIndex + RunIndex);
 
-							for (int32 C = 0; C < BiNormalComps && C < 3; ++C)
+							const FVector3f BaseNormal = NormalIter.GetAsVec3f();
+							const FVector3f MorphNormal = (MorphNormalIter + (MorphIndex + RunIndex)).GetAsVec3f();
+
+							const FVector3f Normal = (BaseNormal + MorphNormal * Factor).GetSafeNormal();
+
+							// Leave the tangent basis sign untouched for packed normals formats.
+							for (int32 C = 0; C < NormalComps && C < 3; ++C)
 							{
-								ConvertData(C, BiNormalIter.ptr(), BiNormalFormat, &BiNormal, EMeshBufferFormat::Float32);
+								ConvertData(C, NormalIter.ptr(), NormalFormat, &Normal, EMeshBufferFormat::Float32);
+							}
+
+							// Tangent
+							if (BaseTangentIter.ptr())
+							{
+								UntypedMeshBufferIterator TangentIter = BaseTangentIter + (VertexIndex + RunIndex);
+
+								const FVector3f BaseTangent = TangentIter.GetAsVec3f();
+
+								// Orthogonalize Tangent based on new Normal. This assumes Normal and BaseTangent are normalized and different.
+								const FVector3f Tangent = (BaseTangent - FVector3f::DotProduct(Normal, BaseTangent) * Normal).GetSafeNormal();
+
+								for (int32 C = 0; C < TangentComps && C < 3; ++C)
+								{
+									ConvertData(C, TangentIter.ptr(), TangentFormat, &Tangent, EMeshBufferFormat::Float32);
+								}
+
+								// BiNormal
+								if (BaseBiNormalIter.ptr())
+								{
+									UntypedMeshBufferIterator BiNormalIter = BaseBiNormalIter + (VertexIndex + RunIndex);
+
+									const FVector3f& N = BaseNormal;
+									const FVector3f& T = BaseTangent;
+									const FVector3f  B = BiNormalIter.GetAsVec3f();
+
+									const float BaseTangentBasisDeterminant =
+										B.X * T.Y * N.Z + B.Z * T.X * N.Y + B.Y * T.Z * N.Y -
+										B.Z * T.Y * N.X - B.Y * T.X * N.Z - B.X * T.Z * N.Y;
+
+									const float BaseTangentBasisDeterminantSign = BaseTangentBasisDeterminant >= 0 ? 1.0f : -1.0f;
+
+									const FVector3f BiNormal = FVector3f::CrossProduct(Tangent, Normal) * BaseTangentBasisDeterminantSign;
+
+									for (int32 C = 0; C < BiNormalComps && C < 3; ++C)
+									{
+										ConvertData(C, BiNormalIter.ptr(), BiNormalFormat, &BiNormal, EMeshBufferFormat::Float32);
+									}
+								}
 							}
 						}
 					}
-				}
+					
+					// Generic Morphs
+					if (bHasGenericMorphs)
+					{
+						MUTABLE_CPUPROFILER_SCOPE(ApplyNormalMorph_Generic);
+						for (int32 ChannelIndex = 0; ChannelIndex < ChannelNum; ++ChannelIndex)
+						{
+							if (!(BaseGenericIters[ChannelIndex].ptr() && MorphGenericIters[ChannelIndex].ptr()))
+							{
+								continue;
+							}
 
-				VertexIndex += FMath::Max(RunSize - 1, 0);
-            }
-        };
-
-        const auto ApplyGenericMorph = []( 
-			MeshVertexIdIteratorConst BaseIdIter, const TArray<UntypedMeshBufferIterator>& BaseChannelsIters, const int32 BaseNum,
-			MeshVertexIdIteratorConst MorphIdIter, const TArray<UntypedMeshBufferIteratorConst>& MorphChannelsIters, const int32 MorphNum,
-				const SparseIndexMapSet& IndexMap, const float Factor)
-        -> void
-        {
-            for (int32 VertexIndex = 0; VertexIndex < BaseNum; ++VertexIndex)
-            {
-                const uint64 BaseId = (BaseIdIter + VertexIndex).Get();
-                const uint32 MorphIndex = IndexMap.Find(BaseId);
-
-                if (MorphIndex == SparseIndexMap::NotFoundValue)
-                {
-                    continue;
-                }
-
-                // Find consecutive run.
-				MeshVertexIdIteratorConst RunBaseIter = BaseIdIter + VertexIndex;
-				MeshVertexIdIteratorConst RunMorphIter = MorphIdIter + MorphIndex;
-
-                int32 RunSize = 0;
-                for ( ; VertexIndex + RunSize < BaseNum && int32(MorphIndex) + RunSize < MorphNum && RunBaseIter.Get() == RunMorphIter.Get();
-                        ++RunSize, ++RunBaseIter, ++RunMorphIter);
-
-                const int32 ChannelNum = MorphChannelsIters.Num();
-                for (int32 ChannelIndex = 0; ChannelIndex < ChannelNum; ++ChannelIndex)
-                {
-                	if (!(BaseChannelsIters[ChannelIndex].ptr() && MorphChannelsIters[ChannelIndex].ptr()))
-                	{
-                		continue;
-                	}
-                
-                    UntypedMeshBufferIterator ChannelBaseIter = BaseChannelsIters[ChannelIndex] + VertexIndex;
-					UntypedMeshBufferIteratorConst ChannelMorphIter = MorphChannelsIters[ChannelIndex] + MorphIndex;
-                   
-                    const EMeshBufferFormat DestChannelFormat = BaseChannelsIters[ChannelIndex].GetFormat();
-                    const int32 DestChannelComps = BaseChannelsIters[ChannelIndex].GetComponents();
-
-                    // Apply Morph to range found above.
-                    for (int32 R = 0; R < RunSize; ++R, ++ChannelBaseIter, ++ChannelMorphIter)
-                    {
-                        const FVector4f Value = ChannelBaseIter.GetAsVec4f() + ChannelMorphIter.GetAsVec4f()*Factor;
-
-                        // TODO: Optimize this for the specific components.
-                        // Max 4 components
-                        for (int32 Comp = 0; Comp < DestChannelComps && Comp < 4; ++Comp)
-                        {
-                            ConvertData(Comp, ChannelBaseIter.ptr(), DestChannelFormat, &Value, EMeshBufferFormat::Float32);
-                        }
-                    }
-                }
-
-				VertexIndex += FMath::Max(RunSize - 1, 0);
-            }
-        };
+							UntypedMeshBufferIterator ChannelBaseIter = BaseGenericIters[ChannelIndex] + VertexIndex;
+							UntypedMeshBufferIteratorConst ChannelMorphIter = MorphGenericIters[ChannelIndex] + MorphIndex;
+
+							const EMeshBufferFormat DestChannelFormat = BaseGenericIters[ChannelIndex].GetFormat();
+							const int32 DestChannelComps = BaseGenericIters[ChannelIndex].GetComponents();
+
+							// Apply Morph to range found above.
+							for (int32 R = 0; R < RunSize; ++R, ++ChannelBaseIter, ++ChannelMorphIter)
+							{
+								const FVector4f Value = ChannelBaseIter.GetAsVec4f() + ChannelMorphIter.GetAsVec4f() * Factor;
+
+								// TODO: Optimize this for the specific components.
+								// Max 4 components
+								for (int32 Comp = 0; Comp < DestChannelComps && Comp < 4; ++Comp)
+								{
+									ConvertData(Comp, ChannelBaseIter.ptr(), DestChannelFormat, &Value, EMeshBufferFormat::Float32);
+								}
+							}
+						}
+					}
 
+					VertexIndex += FMath::Max(RunSize - 1, 0);
+				}
+			};
 
 		// Number of vertices to modify
 		const int32 MinNum = MinMesh ? MinMesh->GetVertexBuffers().GetElementCount() : 0;
@@ -240,12 +288,17 @@ namespace mu
 			constexpr int32 MorphBufferDataChannel = 0;
 			const int32 ChannelsNum = RefTarget->GetVertexBuffers().GetBufferChannelCount(MorphBufferDataChannel);
 
-			TArray<UntypedMeshBufferIterator> BaseChannelsIters;
-			BaseChannelsIters.SetNum(ChannelsNum);
-			TArray<UntypedMeshBufferIteratorConst> MinChannelsIters;
-			MinChannelsIters.SetNum(ChannelsNum);
-			TArray<UntypedMeshBufferIteratorConst> MaxChannelsIters;
-			MaxChannelsIters.SetNum(ChannelsNum);
+			TArray<UntypedMeshBufferIterator> BaseGenericIters;
+			BaseGenericIters.SetNum(ChannelsNum);
+			TArray<UntypedMeshBufferIteratorConst> MinGenericIters;
+			MinGenericIters.SetNum(ChannelsNum);
+			TArray<UntypedMeshBufferIteratorConst> MaxGenericIters;
+			MaxGenericIters.SetNum(ChannelsNum);
+
+			// Get pointers to vertex position data
+			MeshBufferIterator<EMeshBufferFormat::Float32, float, 3> BasePositionIter(BaseMesh->VertexBuffers, EMeshBufferSemantic::Position, 0);	
+			MeshBufferIteratorConst<EMeshBufferFormat::Float32, float, 3> MinPositionIter(RefTarget->VertexBuffers, EMeshBufferSemantic::Position, 0);
+			MeshBufferIteratorConst<EMeshBufferFormat::Float32, float, 3> MaxPositionIter(RefTarget->VertexBuffers, EMeshBufferSemantic::Position, 0);
 
 			// {BiNormal, Tangent, Normal}
 			TStaticArray<UntypedMeshBufferIterator, 3> BaseTangentFrameChannelsIters;
@@ -259,7 +312,6 @@ namespace mu
 				const FMeshBufferChannel& Channel = MBSPriv.Buffers[MorphBufferDataChannel].Channels[ChannelIndex];
 				EMeshBufferSemantic Sem = Channel.Semantic;
 				int32 SemIndex = Channel.SemanticIndex;
-			
 				if (Sem == EMeshBufferSemantic::Normal && bBaseHasNormals)
 				{
 					BaseTangentFrameChannelsIters[2] = UntypedMeshBufferIterator(BaseMesh->GetVertexBuffers(), Sem, SemIndex);
@@ -281,21 +333,21 @@ namespace mu
 				{
 					BaseTangentFrameChannelsIters[0] = UntypedMeshBufferIterator(BaseMesh->GetVertexBuffers(), Sem, SemIndex);
 				}
-				else
+				else if(Sem != EMeshBufferSemantic::Position)
 				{
-					BaseChannelsIters[ChannelIndex] = UntypedMeshBufferIterator(BaseMesh->GetVertexBuffers(), Sem, SemIndex);
+					BaseGenericIters[ChannelIndex] = UntypedMeshBufferIterator(BaseMesh->GetVertexBuffers(), Sem, SemIndex);
 					if (MinNum > 0)
 					{
-						MinChannelsIters[ChannelIndex] = UntypedMeshBufferIteratorConst(MinMesh->GetVertexBuffers(), Sem, SemIndex);
+						MinGenericIters[ChannelIndex] = UntypedMeshBufferIteratorConst(MinMesh->GetVertexBuffers(), Sem, SemIndex);
 					}
 
 					if (MaxNum > 0)
 					{
-						MaxChannelsIters[ChannelIndex] = UntypedMeshBufferIteratorConst(MaxMesh->GetVertexBuffers(), Sem, SemIndex);
+						MaxGenericIters[ChannelIndex] = UntypedMeshBufferIteratorConst(MaxMesh->GetVertexBuffers(), Sem, SemIndex);
 					}
 				}
 			}
-			
+
 			MeshVertexIdIteratorConst BaseIdIter(BaseMesh);
 
 			if (MinNum > 0)
@@ -303,11 +355,8 @@ namespace mu
 				MeshVertexIdIteratorConst MinIdIter(MinMesh);
 				SparseIndexMapSet IndexMap = MakeIndexMap(BaseIdIter, BaseNum, MinIdIter, MinNum);
 
-				ApplyGenericMorph(BaseIdIter, BaseChannelsIters, BaseNum, MinIdIter, MinChannelsIters, MinNum, IndexMap, 1.0f - Factor);
-				if (MinNormalChannelIter.ptr())
-				{
-					ApplyNormalMorph(BaseIdIter, BaseTangentFrameChannelsIters, BaseNum, MinIdIter, MinNormalChannelIter, MinNum, IndexMap, 1.0f - Factor);
-				}
+				ApplyMorph(BaseIdIter, BasePositionIter, BaseTangentFrameChannelsIters, BaseGenericIters, BaseNum,
+					MinIdIter, MinPositionIter, MinNormalChannelIter, MinGenericIters, MinNum, IndexMap, 1.0f - Factor);
 			}
 
 			if (MaxNum > 0)
@@ -315,12 +364,8 @@ namespace mu
 				MeshVertexIdIteratorConst MaxIdIter(MaxMesh);
 				SparseIndexMapSet IndexMap = MakeIndexMap(BaseIdIter, BaseNum, MaxIdIter, MaxNum);
 
-				ApplyGenericMorph(BaseIdIter, BaseChannelsIters, BaseNum, MaxIdIter, MaxChannelsIters, MaxNum, IndexMap, Factor);
-				
-				if (MaxNormalChannelIter.ptr())
-				{
-					ApplyNormalMorph(BaseIdIter, BaseTangentFrameChannelsIters, BaseNum, MaxIdIter, MaxNormalChannelIter, MaxNum, IndexMap, Factor);
-				}
+				ApplyMorph(BaseIdIter, BasePositionIter, BaseTangentFrameChannelsIters, BaseGenericIters, BaseNum,
+					MaxIdIter, MaxPositionIter, MaxNormalChannelIter, MaxGenericIters, MaxNum, IndexMap, Factor);
 			}
 		}
     }
