[mutable] Ensure a all tasks used to generate mesh options are actually waited for, even if no other task ends depending on them.
#rnx

diff --git a/Engine/Plugins/Mutable/Source/MutableTools/Private/MuT/CodeGenerator.cpp b/Engine/Plugins/Mutable/Source/MutableTools/Private/MuT/CodeGenerator.cpp
--- a/Engine/Plugins/Mutable/Source/MutableTools/Private/MuT/CodeGenerator.cpp
+++ b/Engine/Plugins/Mutable/Source/MutableTools/Private/MuT/CodeGenerator.cpp
@@ -1687,7 +1687,7 @@ namespace mu
 				return SurfaceResult;
 			}
 			, 
-			UE::Tasks::Prerequisites(MeshTask, PreviousLODTask)
+			UE::Tasks::Prerequisites(MeshTask, MeshOptionsTask, PreviousLODTask)
 		);
 
 		for (FirstPassGenerator::FSurface* TargetSurface : TargetSurfaces)
