[Mutable] Fixed component nodes inheritance.

#rnx

diff --git a/Engine/Plugins/Mutable/Source/CustomizableObject/Internal/MuCO/CustomizableObjectCustomVersion.h b/Engine/Plugins/Mutable/Source/CustomizableObject/Internal/MuCO/CustomizableObjectCustomVersion.h
--- a/Engine/Plugins/Mutable/Source/CustomizableObject/Internal/MuCO/CustomizableObjectCustomVersion.h
+++ b/Engine/Plugins/Mutable/Source/CustomizableObject/Internal/MuCO/CustomizableObjectCustomVersion.h
@@ -182,6 +182,8 @@ struct CUSTOMIZABLEOBJECT_API FCustomizableObjectCustomVersion
 
 		AddedTableNodeCompilationFilters,
 
+		ChangedComponentsInheritance,
+
 		// -----<new versions can be added above this line>-------------------------------------------------
 		VersionPlusOne,
 		LatestVersion = VersionPlusOne - 1
diff --git a/Engine/Plugins/Mutable/Source/CustomizableObjectEditor/Private/MuCOE/GenerateMutableSource/GenerateMutableSourceComponent.cpp b/Engine/Plugins/Mutable/Source/CustomizableObjectEditor/Private/MuCOE/GenerateMutableSource/GenerateMutableSourceComponent.cpp
--- a/Engine/Plugins/Mutable/Source/CustomizableObjectEditor/Private/MuCOE/GenerateMutableSource/GenerateMutableSourceComponent.cpp
+++ b/Engine/Plugins/Mutable/Source/CustomizableObjectEditor/Private/MuCOE/GenerateMutableSource/GenerateMutableSourceComponent.cpp
@@ -76,7 +76,7 @@ FName GetComponentName(const UCustomizableObjectNode& Node, FMutableGraphGenerat
 /** Generate LOD pins of the given NodeComponentBase (NodeComponent, NodeComponentExtend...).
  * @param TypedComponentMesh Given component node.
  * @param NodeComponent Core node to connect LOD generated pins. */
-void GenerateMutableSourceComponentMesh(FMutableGraphGenerationContext& GenerationContext, const UCustomizableObjectNodeComponentMeshBase& TypedComponentMesh, mu::Ptr<mu::NodeComponent> NodeComponent)
+void GenerateMutableSourceComponentMesh(FMutableGraphGenerationContext& GenerationContext, const ICustomizableObjectNodeComponentMeshInterface& TypedComponentMesh, mu::Ptr<mu::NodeComponent> NodeComponent)
 {
 	int32 FirstLOD = -1;
 
@@ -94,7 +94,7 @@ void GenerateMutableSourceComponentMesh(FMutableGraphGenerationContext& Generati
 
 		LODNode->SetMessageContext(&TypedComponentMesh);
 
-		const int32 NumLODs = TypedComponentMesh.LODPins.Num();
+		const int32 NumLODs = TypedComponentMesh.GetLODPins().Num();
 
 		const bool bUseAutomaticLods = GenerationContext.CurrentAutoLODStrategy == ECustomizableObjectAutomaticLODStrategy::AutomaticFromMesh;
 		FirstLOD = (CurrentLOD < NumLODs) && (FirstLOD == INDEX_NONE || !bUseAutomaticLods) ? CurrentLOD : FirstLOD;
@@ -112,12 +112,12 @@ void GenerateMutableSourceComponentMesh(FMutableGraphGenerationContext& Generati
 		// Generate all relevant LODs for this object up until the current LODIndex.
 		for (int32 LODIndex = FirstLOD; LODIndex <= CurrentLOD; ++LODIndex)
 		{
-			if (!TypedComponentMesh.LODPins.IsValidIndex(LODIndex))
+			if (!TypedComponentMesh.GetLODPins().IsValidIndex(LODIndex))
 			{
 				continue;
 			}
 			
-			const UEdGraphPin* LODPin = TypedComponentMesh.LODPins[LODIndex].Get();
+			const UEdGraphPin* LODPin = TypedComponentMesh.GetLODPins()[LODIndex].Get();
 			check(LODPin);
 			
 			GenerationContext.FromLOD = LODIndex;
@@ -131,7 +131,7 @@ void GenerateMutableSourceComponentMesh(FMutableGraphGenerationContext& Generati
 				if (Cast<UCustomizableObjectNodeModifierBase>(ChildNodePin->GetOwningNode()))
 				{
 					FString Msg = FString::Printf(TEXT("The object has legacy modifier connections that cannot be generated. Their connections should be updated."));
-					GenerationContext.Log(FText::FromString(Msg), &TypedComponentMesh, EMessageSeverity::Warning);
+					GenerationContext.Log(FText::FromString(Msg), TypedComponentMesh.GetOwningNode(), EMessageSeverity::Warning);
 					continue;
 				}
 				
diff --git a/Engine/Plugins/Mutable/Source/CustomizableObjectEditor/Private/MuCOE/Nodes/CustomizableObjectNode.cpp b/Engine/Plugins/Mutable/Source/CustomizableObjectEditor/Private/MuCOE/Nodes/CustomizableObjectNode.cpp
--- a/Engine/Plugins/Mutable/Source/CustomizableObjectEditor/Private/MuCOE/Nodes/CustomizableObjectNode.cpp
+++ b/Engine/Plugins/Mutable/Source/CustomizableObjectEditor/Private/MuCOE/Nodes/CustomizableObjectNode.cpp
@@ -693,9 +693,9 @@ int32 UCustomizableObjectNode::GetLOD() const
 			{
 				for (UEdGraphPin* LinkedPin : FollowOutputPinArray(*Pin))
 				{
-					if (UCustomizableObjectNodeComponentMeshBase* ComponentMesh = Cast<UCustomizableObjectNodeComponentMeshBase>(LinkedPin->GetOwningNode()))
+					if (ICustomizableObjectNodeComponentMeshInterface* ComponentMesh = Cast<ICustomizableObjectNodeComponentMeshInterface>(LinkedPin->GetOwningNode()))
 					{
-						return ComponentMesh->LODPins.IndexOfByPredicate([&](const FEdGraphPinReference& LODPin)
+						return ComponentMesh->GetLODPins().IndexOfByPredicate([&](const FEdGraphPinReference& LODPin)
 						{
 							return LinkedPin == LODPin.Get();
 						});
diff --git a/Engine/Plugins/Mutable/Source/CustomizableObjectEditor/Private/MuCOE/Nodes/CustomizableObjectNodeComponent.cpp b/Engine/Plugins/Mutable/Source/CustomizableObjectEditor/Private/MuCOE/Nodes/CustomizableObjectNodeComponent.cpp
--- a/Engine/Plugins/Mutable/Source/CustomizableObjectEditor/Private/MuCOE/Nodes/CustomizableObjectNodeComponent.cpp
+++ b/Engine/Plugins/Mutable/Source/CustomizableObjectEditor/Private/MuCOE/Nodes/CustomizableObjectNodeComponent.cpp
@@ -1,7 +1,25 @@
-﻿// Copyright Epic Games, Inc. All Rights Reserved.
+// Copyright Epic Games, Inc. All Rights Reserved.
 
 #include "MuCOE/Nodes/CustomizableObjectNodeComponent.h"
 
+#include "MuCOE/EdGraphSchema_CustomizableObject.h"
+
+void UCustomizableObjectNodeComponent::AllocateDefaultPins(UCustomizableObjectNodeRemapPins* RemapPins)
+{
+	OutputPin = CustomCreatePin(EGPD_Output, UEdGraphSchema_CustomizableObject::PC_Component, FName(TEXT("Component")));
+
+	if (IsInMacro())
+	{
+		ComponentNamePin = CustomCreatePin(EGPD_Input, UEdGraphSchema_CustomizableObject::PC_String, FName("Component Name"));
+	}
+}
+
+
+FLinearColor UCustomizableObjectNodeComponent::GetNodeTitleColor() const
+{
+	return UEdGraphSchema_CustomizableObject::GetPinTypeColor(UEdGraphSchema_CustomizableObject::PC_Component);
+}
+
 
 bool UCustomizableObjectNodeComponent::IsAffectedByLOD() const
 { 
diff --git a/Engine/Plugins/Mutable/Source/CustomizableObjectEditor/Private/MuCOE/Nodes/CustomizableObjectNodeComponent.h b/Engine/Plugins/Mutable/Source/CustomizableObjectEditor/Private/MuCOE/Nodes/CustomizableObjectNodeComponent.h
--- a/Engine/Plugins/Mutable/Source/CustomizableObjectEditor/Private/MuCOE/Nodes/CustomizableObjectNodeComponent.h
+++ b/Engine/Plugins/Mutable/Source/CustomizableObjectEditor/Private/MuCOE/Nodes/CustomizableObjectNodeComponent.h
@@ -15,7 +15,23 @@ class CUSTOMIZABLEOBJECTEDITOR_API UCustomizableObjectNodeComponent : public UCu
 public:
 	GENERATED_BODY()
 
+	// UEdGraphNode interface
+	virtual void AllocateDefaultPins(UCustomizableObjectNodeRemapPins* RemapPins) override;
+	virtual FLinearColor GetNodeTitleColor() const override;
+
 	// UCustomizableObjectNode interface
 	virtual bool IsAffectedByLOD() const override;
+
+public:
+
+	UPROPERTY(EditAnywhere, Category = PassthroughMesh)
+	FName ComponentName = "Default name";
+
+	UPROPERTY()
+	FEdGraphPinReference OutputPin;
+
+	UPROPERTY()
+	FEdGraphPinReference ComponentNamePin;
+
 };
 
diff --git a/Engine/Plugins/Mutable/Source/CustomizableObjectEditor/Private/MuCOE/Nodes/CustomizableObjectNodeComponentMesh.cpp b/Engine/Plugins/Mutable/Source/CustomizableObjectEditor/Private/MuCOE/Nodes/CustomizableObjectNodeComponentMesh.cpp
--- a/Engine/Plugins/Mutable/Source/CustomizableObjectEditor/Private/MuCOE/Nodes/CustomizableObjectNodeComponentMesh.cpp
+++ b/Engine/Plugins/Mutable/Source/CustomizableObjectEditor/Private/MuCOE/Nodes/CustomizableObjectNodeComponentMesh.cpp
@@ -58,17 +58,24 @@ void UCustomizableObjectNodeComponentMesh::BackwardsCompatibleFixup(int32 Custom
 
 void UCustomizableObjectNodeComponentMesh::AllocateDefaultPins(UCustomizableObjectNodeRemapPins* RemapPins)
 {
-	const UEdGraphSchema_CustomizableObject* Schema = GetDefault<UEdGraphSchema_CustomizableObject>();
-
+	// Mesh Component Pins
+	const FString PinFriendlyName = TEXT("Overlay Material");
+	const FString PinName = PinFriendlyName + FString(TEXT("_Input_Pin"));
+	UEdGraphPin* OverlayMaterialPin = CustomCreatePin(EGPD_Input, UEdGraphSchema_CustomizableObject::PC_MaterialAsset, FName(*PinName));
+	OverlayMaterialPin->PinFriendlyName = FText::FromString(PinFriendlyName);
+	OverlayMaterialPin->bDefaultValueIsIgnored = true;
+	OverlayMaterialPin->PinToolTip = "Pin for an Overlay Material from a Table Node";
+	
+	// Base Mesh Interface Pins
+	LODPins.Empty(NumLODs);
+	for (int32 LODIndex = 0; LODIndex < NumLODs; ++LODIndex)
 	{
-		const FString PinFriendlyName = TEXT("Overlay Material");
-		const FString PinName = PinFriendlyName + FString(TEXT("_Input_Pin"));
-		UEdGraphPin* OverlayMaterialPin = CustomCreatePin(EGPD_Input, Schema->PC_MaterialAsset, FName(*PinName));
-		OverlayMaterialPin->PinFriendlyName = FText::FromString(PinFriendlyName);
-		OverlayMaterialPin->bDefaultValueIsIgnored = true;
-		OverlayMaterialPin->PinToolTip = "Pin for an Overlay Material from a Table Node";
+		FString LODName = FString::Printf(TEXT("LOD %d"), LODIndex);
+
+		UEdGraphPin* Pin = CustomCreatePin(EGPD_Input, UEdGraphSchema_CustomizableObject::PC_Material, FName(*LODName), true);
+		LODPins.Add(Pin);
 	}
-	
+
 	Super::AllocateDefaultPins(RemapPins);
 }
 
@@ -127,5 +134,48 @@ UEdGraphPin* UCustomizableObjectNodeComponentMesh::GetOverlayMaterialAssetPin()
 }
 
 
+bool UCustomizableObjectNodeComponentMesh::IsSingleOutputNode() const
+{
+	// todo UE-225446 : By limiting the number of connections this node can have we avoid a check failure. However, this method should be
+	// removed in the future and the inherent issue with 1:n output connections should be fixed in its place
+	return true;
+}
+
+
+int32 UCustomizableObjectNodeComponentMesh::GetNumLODs()
+{
+	return NumLODs;
+}
+
+
+ECustomizableObjectAutomaticLODStrategy UCustomizableObjectNodeComponentMesh::GetAutoLODStrategy()
+{
+	return AutoLODStrategy;
+}
+
+
+const TArray<FEdGraphPinReference>& UCustomizableObjectNodeComponentMesh::GetLODPins() const
+{
+	return LODPins;
+}
+
+
+UEdGraphPin* UCustomizableObjectNodeComponentMesh::GetOutputPin() const
+{
+	return OutputPin.Get();
+}
+
+
+void UCustomizableObjectNodeComponentMesh::SetOutputPin(const UEdGraphPin* Pin)
+{
+	OutputPin = Pin;
+}
+
+
+ const UCustomizableObjectNode* UCustomizableObjectNodeComponentMesh::GetOwningNode() const
+{
+	 return this;
+}
+
 #undef LOCTEXT_NAMESPACE
 
diff --git a/Engine/Plugins/Mutable/Source/CustomizableObjectEditor/Private/MuCOE/Nodes/CustomizableObjectNodeComponentMesh.h b/Engine/Plugins/Mutable/Source/CustomizableObjectEditor/Private/MuCOE/Nodes/CustomizableObjectNodeComponentMesh.h
--- a/Engine/Plugins/Mutable/Source/CustomizableObjectEditor/Private/MuCOE/Nodes/CustomizableObjectNodeComponentMesh.h
+++ b/Engine/Plugins/Mutable/Source/CustomizableObjectEditor/Private/MuCOE/Nodes/CustomizableObjectNodeComponentMesh.h
@@ -46,7 +46,7 @@ struct FLODReductionSettings
 
 
 UCLASS()
-class CUSTOMIZABLEOBJECTEDITOR_API UCustomizableObjectNodeComponentMesh : public UCustomizableObjectNodeComponentMeshBase
+class CUSTOMIZABLEOBJECTEDITOR_API UCustomizableObjectNodeComponentMesh : public UCustomizableObjectNodeComponent, public ICustomizableObjectNodeComponentMeshInterface
 {
 	GENERATED_BODY()
 
@@ -62,12 +62,25 @@ public:
 	virtual FText GetNodeTitle(ENodeTitleType::Type TitleType) const override;
 	virtual void OnRenameNode(const FString& NewName) override;
 
+	//UCustomizableObjectNode interface
+	bool IsSingleOutputNode() const;
+
 	// Own interface
 	TSoftObjectPtr<UMaterialInterface> GetOverlayMaterial() const;
 	UEdGraphPin* GetOverlayMaterialAssetPin() const;
 
-	UPROPERTY(EditAnywhere, Category = ComponentMesh)
-	FName ComponentName = "Default name";
+	// ComponentMesh Interface
+	virtual int32 GetNumLODs() override;
+	virtual ECustomizableObjectAutomaticLODStrategy GetAutoLODStrategy() override;
+	virtual const TArray<FEdGraphPinReference>& GetLODPins() const override;
+	virtual UEdGraphPin* GetOutputPin() const;
+	virtual void SetOutputPin(const UEdGraphPin* Pin);
+	virtual const UCustomizableObjectNode* GetOwningNode() const;
+
+public:
+
+	/*UPROPERTY(EditAnywhere, Category = ComponentMesh)
+	FName ComponentName = "Default name";*/
 	
 	/** All the Skeletal Meshes generated for this component will use the Reference Skeletal Mesh properties
 	 * for everything Mutable doesn't create or modify. This includes data like LOD distances, Physics
@@ -90,4 +103,14 @@ public:
 
 	/** Details selected LOD. */
 	int32 SelectedLOD = 0;
+
+	UPROPERTY(EditAnywhere, Category = ComponentMesh)
+	int32 NumLODs = 1;
+
+	UPROPERTY(EditAnywhere, Category = ComponentMesh, DisplayName = "Auto LOD Strategy")
+	ECustomizableObjectAutomaticLODStrategy AutoLODStrategy = ECustomizableObjectAutomaticLODStrategy::AutomaticFromMesh;
+
+	UPROPERTY()
+	TArray<FEdGraphPinReference> LODPins;
+
 };
diff --git a/Engine/Plugins/Mutable/Source/CustomizableObjectEditor/Private/MuCOE/Nodes/CustomizableObjectNodeComponentMeshAddTo.cpp b/Engine/Plugins/Mutable/Source/CustomizableObjectEditor/Private/MuCOE/Nodes/CustomizableObjectNodeComponentMeshAddTo.cpp
--- a/Engine/Plugins/Mutable/Source/CustomizableObjectEditor/Private/MuCOE/Nodes/CustomizableObjectNodeComponentMeshAddTo.cpp
+++ b/Engine/Plugins/Mutable/Source/CustomizableObjectEditor/Private/MuCOE/Nodes/CustomizableObjectNodeComponentMeshAddTo.cpp
@@ -2,12 +2,34 @@
 
 #include "MuCOE/Nodes/CustomizableObjectNodeComponentMeshAddTo.h"
 
+#include "MuCO/CustomizableObjectCustomVersion.h"
 #include "MuCOE/EdGraphSchema_CustomizableObject.h"
 #include "MuCOE/Nodes/CustomizableObjectNodeComponentMesh.h"
 
 
 #define LOCTEXT_NAMESPACE "CustomizableObjectEditor"
 
+void UCustomizableObjectNodeComponentMeshAddTo::AllocateDefaultPins(UCustomizableObjectNodeRemapPins* RemapPins)
+{
+	Super::AllocateDefaultPins(RemapPins);
+
+	LODPins.Empty(NumLODs);
+	for (int32 LODIndex = 0; LODIndex < NumLODs; ++LODIndex)
+	{
+		FString LODName = FString::Printf(TEXT("LOD %d"), LODIndex);
+
+		UEdGraphPin* Pin = CustomCreatePin(EGPD_Input, UEdGraphSchema_CustomizableObject::PC_Material, FName(*LODName), true);
+		LODPins.Add(Pin);
+	}
+
+	if (IsInMacro())
+	{
+		ParentComponentNamePinRef = CustomCreatePin(EGPD_Input, UEdGraphSchema_CustomizableObject::PC_String, FName("Component Name"));
+	}
+
+	OutputPin = CustomCreatePin(EGPD_Output, UEdGraphSchema_CustomizableObject::PC_Component, FName(TEXT("Component")));
+}
+
 
 void UCustomizableObjectNodeComponentMeshAddTo::PostEditChangeProperty(FPropertyChangedEvent& PropertyChangedEvent)
 {
@@ -19,13 +41,27 @@ void UCustomizableObjectNodeComponentMeshAddTo::PostEditChangeProperty(FProperty
 		return;
 	}
 
-	if (PropertyThatChanged->GetFName() == GET_MEMBER_NAME_CHECKED(UCustomizableObjectNodeComponentMesh, NumLODs))
+	if (PropertyThatChanged->GetFName() == GET_MEMBER_NAME_CHECKED(UCustomizableObjectNodeComponentMeshAddTo, NumLODs))
 	{
 		ReconstructNode();
 	}
 }
 
 
+bool UCustomizableObjectNodeComponentMeshAddTo::IsAffectedByLOD() const
+{
+	return false;
+}
+
+
+bool UCustomizableObjectNodeComponentMeshAddTo::IsSingleOutputNode() const
+{
+	// todo UE-225446 : By limiting the number of connections this node can have we avoid a check failure. However, this method should be
+	// removed in the future and the inherent issue with 1:n output connections should be fixed in its place
+	return true;
+}
+
+
 FText UCustomizableObjectNodeComponentMeshAddTo::GetNodeTitle(ENodeTitleType::Type TitleType) const
 {
 	FText NodeTitle;
@@ -58,5 +94,47 @@ void UCustomizableObjectNodeComponentMeshAddTo::OnRenameNode(const FString& NewN
 }
 
 
+FLinearColor UCustomizableObjectNodeComponentMeshAddTo::GetNodeTitleColor() const
+{
+	return UEdGraphSchema_CustomizableObject::GetPinTypeColor(UEdGraphSchema_CustomizableObject::PC_Component);
+}
+
+
+int32 UCustomizableObjectNodeComponentMeshAddTo::GetNumLODs()
+{
+	return NumLODs;
+}
+
+
+ECustomizableObjectAutomaticLODStrategy UCustomizableObjectNodeComponentMeshAddTo::GetAutoLODStrategy()
+{
+	return AutoLODStrategy;
+}
+
+
+const TArray<FEdGraphPinReference>& UCustomizableObjectNodeComponentMeshAddTo::GetLODPins() const
+{
+	return LODPins;
+}
+
+
+UEdGraphPin* UCustomizableObjectNodeComponentMeshAddTo::GetOutputPin() const
+{
+	return OutputPin.Get();
+}
+
+
+void UCustomizableObjectNodeComponentMeshAddTo::SetOutputPin(const UEdGraphPin* Pin)
+{
+	OutputPin = Pin;
+}
+
+
+const UCustomizableObjectNode* UCustomizableObjectNodeComponentMeshAddTo::GetOwningNode() const
+{
+	return this;
+}
+
+
 #undef LOCTEXT_NAMESPACE
 
diff --git a/Engine/Plugins/Mutable/Source/CustomizableObjectEditor/Private/MuCOE/Nodes/CustomizableObjectNodeComponentMeshAddTo.h b/Engine/Plugins/Mutable/Source/CustomizableObjectEditor/Private/MuCOE/Nodes/CustomizableObjectNodeComponentMeshAddTo.h
--- a/Engine/Plugins/Mutable/Source/CustomizableObjectEditor/Private/MuCOE/Nodes/CustomizableObjectNodeComponentMeshAddTo.h
+++ b/Engine/Plugins/Mutable/Source/CustomizableObjectEditor/Private/MuCOE/Nodes/CustomizableObjectNodeComponentMeshAddTo.h
@@ -7,18 +7,52 @@
 #include "CustomizableObjectNodeComponentMeshAddTo.generated.h"
 
 UCLASS()
-class CUSTOMIZABLEOBJECTEDITOR_API UCustomizableObjectNodeComponentMeshAddTo : public UCustomizableObjectNodeComponentMeshBase
+class CUSTOMIZABLEOBJECTEDITOR_API UCustomizableObjectNodeComponentMeshAddTo : public UCustomizableObjectNode, public ICustomizableObjectNodeComponentMeshInterface
 {
 	GENERATED_BODY()
 
 public:
-	// UObject interface
-	virtual void PostEditChangeProperty(FPropertyChangedEvent& PropertyChangedEvent) override;
 
 	// UEdGraphNode interface
+	virtual void AllocateDefaultPins(UCustomizableObjectNodeRemapPins* RemapPins) override;
+	virtual bool GetCanRenameNode() const override { return true; }
 	virtual FText GetNodeTitle(ENodeTitleType::Type TitleType) const override;
 	virtual void OnRenameNode(const FString& NewName) override;
-	
+	virtual FLinearColor GetNodeTitleColor() const override;
+
+	// UObject interface
+	virtual void PostEditChangeProperty(FPropertyChangedEvent& PropertyChangedEvent) override;
+
+	//UCustomizableObjectNode interface
+	virtual bool IsAffectedByLOD() const override;
+	bool IsSingleOutputNode() const;
+
+	// ComponentMesh Interface
+	virtual int32 GetNumLODs() override;
+	virtual ECustomizableObjectAutomaticLODStrategy GetAutoLODStrategy() override;
+	virtual const TArray<FEdGraphPinReference>& GetLODPins() const override;
+	virtual UEdGraphPin* GetOutputPin() const;
+	virtual void SetOutputPin(const UEdGraphPin* Pin);
+	virtual const UCustomizableObjectNode* GetOwningNode() const;
+
+public:
+
 	UPROPERTY(EditAnywhere, Category = ComponentMesh)
 	FName ParentComponentName;
+
+	UPROPERTY(EditAnywhere, Category = ComponentMesh)
+	int32 NumLODs = 1;
+
+	UPROPERTY(EditAnywhere, Category = ComponentMesh, DisplayName = "Auto LOD Strategy")
+	ECustomizableObjectAutomaticLODStrategy AutoLODStrategy = ECustomizableObjectAutomaticLODStrategy::AutomaticFromMesh;
+
+	UPROPERTY()
+	TArray<FEdGraphPinReference> LODPins;
+
+	UPROPERTY()
+	FEdGraphPinReference OutputPin;
+	
+	UPROPERTY()
+	FEdGraphPinReference ParentComponentNamePinRef;
+
 };
\ No newline at end of file
diff --git a/Engine/Plugins/Mutable/Source/CustomizableObjectEditor/Private/MuCOE/Nodes/CustomizableObjectNodeComponentMeshBase.cpp b/Engine/Plugins/Mutable/Source/CustomizableObjectEditor/Private/MuCOE/Nodes/CustomizableObjectNodeComponentMeshBase.cpp
--- a/Engine/Plugins/Mutable/Source/CustomizableObjectEditor/Private/MuCOE/Nodes/CustomizableObjectNodeComponentMeshBase.cpp
+++ b/Engine/Plugins/Mutable/Source/CustomizableObjectEditor/Private/MuCOE/Nodes/CustomizableObjectNodeComponentMeshBase.cpp
@@ -5,40 +5,8 @@
 #include "MuCOE/EdGraphSchema_CustomizableObject.h"
 
 
-void UCustomizableObjectNodeComponentMeshBase::AllocateDefaultPins(UCustomizableObjectNodeRemapPins* RemapPins)
+TArray<FEdGraphPinReference>& ICustomizableObjectNodeComponentMeshInterface::GetLODPins()
 {
-	Super::AllocateDefaultPins(RemapPins);
-
-	const UEdGraphSchema_CustomizableObject* Schema = GetDefault<UEdGraphSchema_CustomizableObject>();
-
-	LODPins.Empty(NumLODs);
-	for (int32 LODIndex = 0; LODIndex < NumLODs; ++LODIndex)
-	{
-		FString LODName = FString::Printf(TEXT("LOD %d"), LODIndex);
-
-		UEdGraphPin* Pin = CustomCreatePin(EGPD_Input, Schema->PC_Material, FName(*LODName), true);
-		LODPins.Add(Pin);
-	}
-
-	if (IsInMacro())
-	{
-		UEdGraphPin* Pin = CustomCreatePin(EGPD_Input, Schema->PC_String, FName("Component Name"));
-	}
-
-	OutputPin = CustomCreatePin(EGPD_Output, Schema->PC_Component, FName(TEXT("Component")));
+	const ICustomizableObjectNodeComponentMeshInterface* ConstThis = this;
+	return const_cast<TArray<FEdGraphPinReference>&>(ConstThis->GetLODPins());
 }
-
-
-FLinearColor UCustomizableObjectNodeComponentMeshBase::GetNodeTitleColor() const
-{
-	const UEdGraphSchema_CustomizableObject* Schema = GetDefault<UEdGraphSchema_CustomizableObject>();
-	return Schema->GetPinTypeColor(Schema->PC_Component);
-}
-
-bool UCustomizableObjectNodeComponentMeshBase::IsSingleOutputNode() const
-{
-	// todo UE-225446 : By limiting the number of connections this node can have we avoid a check failure. However, this method should be
-	// removed in the future and the inherent issue with 1:n output connections should be fixed in its place
-	return true;
-}
-
diff --git a/Engine/Plugins/Mutable/Source/CustomizableObjectEditor/Private/MuCOE/Nodes/CustomizableObjectNodeComponentMeshBase.h b/Engine/Plugins/Mutable/Source/CustomizableObjectEditor/Private/MuCOE/Nodes/CustomizableObjectNodeComponentMeshBase.h
--- a/Engine/Plugins/Mutable/Source/CustomizableObjectEditor/Private/MuCOE/Nodes/CustomizableObjectNodeComponentMeshBase.h
+++ b/Engine/Plugins/Mutable/Source/CustomizableObjectEditor/Private/MuCOE/Nodes/CustomizableObjectNodeComponentMeshBase.h
@@ -26,29 +26,37 @@ enum class ECustomizableObjectAutomaticLODStrategy : uint8
 };
 
 
-UCLASS()
-class CUSTOMIZABLEOBJECTEDITOR_API UCustomizableObjectNodeComponentMeshBase : public UCustomizableObjectNodeComponent
+UINTERFACE(MinimalAPI, Blueprintable)
+class UCustomizableObjectNodeComponentMeshInterface : public UInterface
 {
 	GENERATED_BODY()
-	
+};
+
+
+class ICustomizableObjectNodeComponentMeshInterface
+{
+	GENERATED_BODY()
+
 public:
-	// UEdGraphNode interface
-	virtual void AllocateDefaultPins(UCustomizableObjectNodeRemapPins* RemapPins) override;
-	virtual FLinearColor GetNodeTitleColor() const override;
-	virtual bool GetCanRenameNode() const override { return true; }
-
-	// UCustomizableObjectNode interface
-	virtual bool IsSingleOutputNode() const override;
-
-	UPROPERTY(EditAnywhere, Category = ComponentMesh)
-	int32 NumLODs = 1;
-	
-	UPROPERTY(EditAnywhere, Category = ComponentMesh, DisplayName = "Auto LOD Strategy")
-	ECustomizableObjectAutomaticLODStrategy AutoLODStrategy = ECustomizableObjectAutomaticLODStrategy::AutomaticFromMesh;
-	
-	UPROPERTY()
-	TArray<FEdGraphPinReference> LODPins;
-	
-	UPROPERTY()
-	FEdGraphPinReference OutputPin;
-};
\ No newline at end of file
+
+	/** Returns the num of Lods set in the Mesh Component Node */
+	virtual int32 GetNumLODs() = 0;
+
+	/** Returns the Lod Strategy set in the Mesh Component Node */
+	virtual ECustomizableObjectAutomaticLODStrategy GetAutoLODStrategy() = 0;
+
+	/** Returns a reference to an array with all the LOD pins of the Mesh Component Node */
+	virtual TArray<FEdGraphPinReference>& GetLODPins();
+
+	/** Returns a const reference to an array with all the LOD pins of the Mesh Component Node */
+	virtual const TArray<FEdGraphPinReference>& GetLODPins() const = 0;
+
+	/** Returns the output pin of the mesh component node */
+	virtual UEdGraphPin* GetOutputPin() const = 0;
+
+	/** Sets the output pin of the mesh component node */
+	virtual void SetOutputPin(const UEdGraphPin* Pin) = 0;
+
+	/** Returns a pointer to the node that owns this interface */
+	virtual const UCustomizableObjectNode* GetOwningNode() const = 0;
+};
diff --git a/Engine/Plugins/Mutable/Source/CustomizableObjectEditor/Private/MuCOE/Nodes/CustomizableObjectNodeComponentPassthroughMesh.cpp b/Engine/Plugins/Mutable/Source/CustomizableObjectEditor/Private/MuCOE/Nodes/CustomizableObjectNodeComponentPassthroughMesh.cpp
--- a/Engine/Plugins/Mutable/Source/CustomizableObjectEditor/Private/MuCOE/Nodes/CustomizableObjectNodeComponentPassthroughMesh.cpp
+++ b/Engine/Plugins/Mutable/Source/CustomizableObjectEditor/Private/MuCOE/Nodes/CustomizableObjectNodeComponentPassthroughMesh.cpp
@@ -74,15 +74,20 @@ void UCustomizableObjectNodeComponentPassthroughMesh::PostEditChangeProperty(FPr
 }
 
 
-void UCustomizableObjectNodeComponentPassthroughMesh::AllocateDefaultPins(UCustomizableObjectNodeRemapPins* RemapPins)
+void UCustomizableObjectNodeComponentPassthroughMesh::BackwardsCompatibleFixup(int32 CustomizableObjectCustomVersion)
 {
-	const UEdGraphSchema_CustomizableObject* Schema = GetDefault<UEdGraphSchema_CustomizableObject>();
-	if (!Schema)
+	Super::BackwardsCompatibleFixup(CustomizableObjectCustomVersion);
+
+	if (CustomizableObjectCustomVersion == FCustomizableObjectCustomVersion::ChangedComponentsInheritance)
 	{
-		return;
+		OutputPin = FindPin(OutputPinName, EGPD_Output);
 	}
+}
 
-	UEdGraphPin* OutputPin = CustomCreatePin(EGPD_Output, Schema->PC_Component, OutputPinName);
+
+void UCustomizableObjectNodeComponentPassthroughMesh::AllocateDefaultPins(UCustomizableObjectNodeRemapPins* RemapPins)
+{
+	Super::AllocateDefaultPins(RemapPins);
 
 	if (!Mesh.IsValid())
 	{
@@ -121,7 +126,7 @@ void UCustomizableObjectNodeComponentPassthroughMesh::AllocateDefaultPins(UCusto
 					FString SectionFriendlyName = MaterialInterface ? MaterialInterface->GetName() : FString::Printf(TEXT("Section %i"), SectionIndex);
 					FString MaterialPinName = FString::Printf(TEXT("LOD %i - Section %i - Material"), LODIndex, SectionIndex);
 
-					UEdGraphPin* Pin = CustomCreatePin(EGPD_Input, Schema->PC_Material, FName(*MaterialPinName), PinData);
+					UEdGraphPin* Pin = CustomCreatePin(EGPD_Input, UEdGraphSchema_CustomizableObject::PC_Material, FName(*MaterialPinName), PinData);
 					Pin->PinFriendlyName = FText::FromString(FString::Printf(TEXT("LOD %i - Section %i - %s"), LODIndex, SectionIndex, *SectionFriendlyName));
 					Pin->PinToolTip = MaterialPinName;
 				}
@@ -161,13 +166,6 @@ FText UCustomizableObjectNodeComponentPassthroughMesh::GetNodeTitle(ENodeTitleTy
 }
 
 
-FLinearColor UCustomizableObjectNodeComponentPassthroughMesh::GetNodeTitleColor() const
-{
-	const UEdGraphSchema_CustomizableObject* Schema = GetDefault<UEdGraphSchema_CustomizableObject>();
-	return Schema->GetPinTypeColor(Schema->PC_Component);
-}
-
-
 void UCustomizableObjectNodeComponentPassthroughMesh::GetPinSection(const UEdGraphPin& Pin, int32& OutLODIndex, int32& OutSectionIndex) const
 {
 	if (const UCustomizableObjectNodeComponentMeshPinDataSection* PinData = Cast<UCustomizableObjectNodeComponentMeshPinDataSection>(GetPinData(Pin)))
diff --git a/Engine/Plugins/Mutable/Source/CustomizableObjectEditor/Private/MuCOE/Nodes/CustomizableObjectNodeComponentPassthroughMesh.h b/Engine/Plugins/Mutable/Source/CustomizableObjectEditor/Private/MuCOE/Nodes/CustomizableObjectNodeComponentPassthroughMesh.h
--- a/Engine/Plugins/Mutable/Source/CustomizableObjectEditor/Private/MuCOE/Nodes/CustomizableObjectNodeComponentPassthroughMesh.h
+++ b/Engine/Plugins/Mutable/Source/CustomizableObjectEditor/Private/MuCOE/Nodes/CustomizableObjectNodeComponentPassthroughMesh.h
@@ -67,18 +67,15 @@ class CUSTOMIZABLEOBJECTEDITOR_API UCustomizableObjectNodeComponentPassthroughMe
 public:
 	GENERATED_BODY()
 
-	UPROPERTY(EditAnywhere, Category = PassthroughMesh)
-	FName ComponentName;
-
 	UPROPERTY(EditAnywhere, Category = PassthroughMesh, Meta = (DisplayName = Mesh, AllowedClasses = "/Script/Engine.StaticMesh, /Script/Engine.SkeletalMesh"))
 	FSoftObjectPath Mesh;
 
 	// UObject interface.
 	virtual void PostEditChangeProperty(FPropertyChangedEvent& PropertyChangedEvent) override;
+	virtual void BackwardsCompatibleFixup(int32 CustomizableObjectCustomVersion) override;
 	
 	// UEdGraphNode interface
 	virtual FText GetNodeTitle(ENodeTitleType::Type TitleType) const override;
-	virtual FLinearColor GetNodeTitleColor() const override;
 	virtual FText GetTooltipText() const override;
 	virtual bool GetCanRenameNode() const override { return true; }
 	virtual void OnRenameNode(const FString& NewName) override;
diff --git a/Engine/Plugins/Mutable/Source/CustomizableObjectEditor/Private/MuCOE/Nodes/CustomizableObjectNodeObject.cpp b/Engine/Plugins/Mutable/Source/CustomizableObjectEditor/Private/MuCOE/Nodes/CustomizableObjectNodeObject.cpp
--- a/Engine/Plugins/Mutable/Source/CustomizableObjectEditor/Private/MuCOE/Nodes/CustomizableObjectNodeObject.cpp
+++ b/Engine/Plugins/Mutable/Source/CustomizableObjectEditor/Private/MuCOE/Nodes/CustomizableObjectNodeObject.cpp
@@ -139,20 +139,20 @@ void UCustomizableObjectNodeObject::BackwardsCompatibleFixup(int32 CustomizableO
 	{
 		auto NodeComponentMeshBaseAllocateDefaultPins = [](UCustomizableObjectNode* Node, UCustomizableObjectNodeRemapPins*)
 		{
-			UCustomizableObjectNodeComponentMeshBase* NodeComponentMesh = CastChecked<UCustomizableObjectNodeComponentMeshBase>(Node);
+			ICustomizableObjectNodeComponentMeshInterface* NodeComponentMesh = CastChecked<ICustomizableObjectNodeComponentMeshInterface>(Node);
 						
 			const UEdGraphSchema_CustomizableObject* Schema = GetDefault<UEdGraphSchema_CustomizableObject>();
 
-			NodeComponentMesh->LODPins.Empty(NodeComponentMesh->NumLODs);
-			for (int32 NodeComponentLODIndex = 0; NodeComponentLODIndex < NodeComponentMesh->NumLODs; ++NodeComponentLODIndex)
+			NodeComponentMesh->GetLODPins().Empty(NodeComponentMesh->GetNumLODs());
+			for (int32 NodeComponentLODIndex = 0; NodeComponentLODIndex < NodeComponentMesh->GetNumLODs(); ++NodeComponentLODIndex)
 			{
 				FString LODName = FString::Printf(TEXT("LOD %d"), NodeComponentLODIndex);
 
-				UEdGraphPin* Pin = NodeComponentMesh->CustomCreatePin(EGPD_Input, Schema->PC_Material, FName(*LODName), true);
-				NodeComponentMesh->LODPins.Add(Pin);
+				UEdGraphPin* Pin = Node->CustomCreatePin(EGPD_Input, Schema->PC_Material, FName(*LODName), true);
+				NodeComponentMesh->GetLODPins().Add(Pin);
 			}
 						
-			NodeComponentMesh->OutputPin = NodeComponentMesh->CustomCreatePin(EGPD_Output, Schema->PC_Component, TEXT("Component"));
+			NodeComponentMesh->SetOutputPin(Node->CustomCreatePin(EGPD_Output, Schema->PC_Component, TEXT("Component")));
 		};
 
 		auto NodeObjectAllocateDefaultPins = [](UCustomizableObjectNode* Node, UCustomizableObjectNodeRemapPins*)
