[Mutable] Merge StartUpdate with LockMeshCache.

#rnx

diff --git a/Engine/Plugins/Mutable/Source/CustomizableObject/Private/MuCO/CustomizableObjectSystem.cpp b/Engine/Plugins/Mutable/Source/CustomizableObject/Private/MuCO/CustomizableObjectSystem.cpp
--- a/Engine/Plugins/Mutable/Source/CustomizableObject/Private/MuCO/CustomizableObjectSystem.cpp
+++ b/Engine/Plugins/Mutable/Source/CustomizableObject/Private/MuCO/CustomizableObjectSystem.cpp
@@ -2159,54 +2159,7 @@ namespace impl
 			}
 		}
 	}
-
-
-	void CreateMutableInstance(const TSharedRef<FUpdateContextPrivate>& Operation)
-	{
-		UCustomizableObjectSystem* System = UCustomizableObjectSystem::GetInstanceChecked(); // Save since UCustomizableObjectSystem::BeginDestroy always waits for all tasks to finish
-		const UCustomizableObjectSystemPrivate* SystemPrivate = System->GetPrivate();
-		
-		if (FLogBenchmarkUtil::IsBenchmarkingReportingEnabled())
-		{
-			// Get the amount of mutable memory in use now
-			Operation->UpdateStartBytes = mu::FGlobalMemoryCounter::GetAbsoluteCounter();
-			// Reset the counter to later get the peak during the updated
-			mu::FGlobalMemoryCounter::Zero();													
-		}
-
-		// Prepare streaming for the current customizable object
-		check(SystemPrivate->Streamer != nullptr);
-		SystemPrivate->Streamer->PrepareStreamingForObject(Operation->Instance->GetCustomizableObject());			
-
-		const TSharedPtr<mu::FSystem> MutableSystem = SystemPrivate->MutableSystem;
-		
-		if (Operation->bLiveUpdateMode)
-		{
-			if (Operation->InstanceID == 0)
-			{
-				// It's the first update since the instance was put in LiveUpdate Mode, this ID will be reused from now on
-				Operation->InstanceID = MutableSystem->NewInstance(Operation->Model);
-				UE_LOG(LogMutable, Verbose, TEXT("Creating Mutable instance with id [%d] for reuse "), Operation->InstanceID);
-			}
-			else
-			{
-				// The instance was already in LiveUpdate Mode, the ID is reused
-				check(Operation->InstanceID);
-				UE_LOG(LogMutable, Verbose, TEXT("Reusing Mutable instance with id [%d] "), Operation->InstanceID);
-			}
-		}
-		else
-		{
-			// In non-LiveUpdate mode, we are forcing the recreation of mutable-side instances with every update.
-			check(Operation->InstanceID == 0);
-			Operation->InstanceID = MutableSystem->NewInstance(Operation->Model);
-			UE_LOG(LogMutable, Verbose, TEXT("Creating Mutable instance with id [%d] "), Operation->InstanceID);
-		}
-
-		Operation->MutableInstance = MutableSystem->BeginUpdate(Operation->InstanceID, Operation->Parameters, Operation->GetCapturedDescriptor().GetState(), mu::FSystem::AllLODs);
-		Operation->NumInstanceComponents = Operation->MutableInstance->GetComponentCount();
-	}
-
+	
 	
 	void FixLODs(const TSharedRef<FUpdateContextPrivate>& Operation)
 	{
@@ -2831,13 +2784,7 @@ namespace impl
 			{
 				OperationData->MutableSystem->SetImagePixelConversionOverride( OperationData->PixelFormatOverride );
 			}
-
-			if (!OperationData->bUseMeshCache)
-			{
-				CreateMutableInstance(OperationData);
-				FixLODs(OperationData);
-			}
-
+			
 			// Main instance generation step
 			const mu::FInstance* Instance = OperationData->MutableInstance.Get(); // TODO GMTFuture remove
 			if (!Instance)
@@ -2932,7 +2879,7 @@ namespace impl
 	}
 
 	
-		/** End of the GetImages tasks. */
+	/** End of the GetImages tasks. */
 	void Task_Mutable_GetImages_End(const TSharedRef<FUpdateContextPrivate>& OperationData, double StartTime)
 	{
 		MUTABLE_CPUPROFILER_SCOPE(Task_Mutable_GetImages_End)
@@ -3577,79 +3524,7 @@ namespace impl
 				});
 		}
 	}
-
-
-	void Task_Game_LockMeshCache(const TSharedRef<FUpdateContextPrivate>& Operation, UE::Tasks::FTask Requisite)
-	{
-		MUTABLE_CPUPROFILER_SCOPE(Task_Game_LockMeshCache);
-
-		UCustomizableObjectSystem* System = UCustomizableObjectSystem::GetInstanceChecked();
-		UCustomizableObjectSystemPrivate* SystemPrivate = System->GetPrivate();
-
-		UCustomizableObject* CustomizableObject = Operation->Instance->GetCustomizableObject();
-		UCustomizableObjectPrivate* CustomizableObjectPrivate = CustomizableObject->GetPrivate();
-		
-		for (const TArray<mu::FResourceID>& MeshId : Operation->GetMeshDescriptors())
-		{
-			if (USkeletalMesh* CachedMesh = CustomizableObject->GetPrivate()->MeshCache.Get(MeshId))
-			{
-				Operation->Objects.Emplace(CachedMesh);
-			}
-		}
-
-		UE::Tasks::FTask Dependency = SystemPrivate->MutableTaskGraph.AddMutableThreadTask(
-			TEXT("Task_Mutable_GetMeshes"),
-			[Operation]()
-			{
-				Task_Mutable_GetMeshes(Operation);
-			}, { Requisite });
-
-		SystemPrivate->MutableTaskGraph.AddGameThreadTask(
-			TEXT("Task_Game_LockCache"),
-			[Operation]()
-			{
-				impl::Task_Game_LockCache(Operation);
-			},
-			false,
-			{ Dependency });
-	}
 	
-
-	void Task_Mutable_GetMeshID(const TSharedRef<FUpdateContextPrivate>& Operation)
-	{
-		MUTABLE_CPUPROFILER_SCOPE(Task_Mutable_GetMeshID);
-		
-		CreateMutableInstance(Operation);
-		FixLODs(Operation);
-
-		Operation->InitMeshDescriptors(Operation->NumObjectComponents);
-
-		const TMap<FName, uint8>& MinLOD = Operation->GetMinLOD();
-		const TMap<FName, uint8>& NumLODsAvailable = Operation->NumLODsAvailable;
-		
-		for (int32 InstanceComponentIndex = 0; InstanceComponentIndex < Operation->NumInstanceComponents; ++InstanceComponentIndex)
-		{
-			FCustomizableObjectComponentIndex ObjectComponentIndex(Operation->MutableInstance->GetComponentId(InstanceComponentIndex));
-			TArray<mu::FResourceID>* MeshId = Operation->GetMeshDescriptors(ObjectComponentIndex);
-			if (!MeshId)
-			{
-				continue;
-			}
-
-			MeshId->Init(MAX_uint64, MAX_MESH_LOD_COUNT);
-
-			const FName& ComponentName = Operation->ComponentNames[ObjectComponentIndex.GetValue()];
-			
-			for (int32 LODIndex = MinLOD[ComponentName]; LODIndex < NumLODsAvailable[ComponentName]; ++LODIndex)
-			{
-				const bool bGenerateLOD = Operation->GetRequestedLODs()[ComponentName] <= LODIndex;
-				if (bGenerateLOD)
-				{
-					(*MeshId)[LODIndex] = Operation->MutableInstance->GetMeshId(InstanceComponentIndex, LODIndex);
-				}
-			}
-		}
-	}
 	
 	/** "Start Update" */
 	void Task_Game_StartUpdate(const TSharedRef<FUpdateContextPrivate>& Operation)
@@ -3689,6 +3564,8 @@ if (CVarMutableLogObjectMemoryOnUpdate.GetValueOnAnyThread())
 			return;
 		}
 
+		const UCustomizableObjectSystemPrivate* SystemPrivate = System->GetPrivate();
+
 		UCustomizableObjectInstance* CandidateInstance = Operation->Instance.Get();
 		
 		UCustomizableInstancePrivate* CandidateInstancePrivateData = CandidateInstance->GetPrivate();
@@ -3919,49 +3796,106 @@ if (CVarMutableLogObjectMemoryOnUpdate.GetValueOnAnyThread())
 			Operation->SetRequestedLODs(RequestedLODs);
 		}
 
-		UE::Tasks::FTask Mutable_GetMeshTask;
-		
-		if (Operation->bUseMeshCache)
+		// CreateMutableInstance
 		{
-			Mutable_GetMeshTask = SystemPrivateData->MutableTaskGraph.AddMutableThreadTask(
-				TEXT("Task_Mutable_GetMeshID"),
-				[Operation]()
-				{
-					impl::Task_Mutable_GetMeshID(Operation);
-				});
+			if (FLogBenchmarkUtil::IsBenchmarkingReportingEnabled())
+			{
+				// Get the amount of mutable memory in use now
+				Operation->UpdateStartBytes = mu::FGlobalMemoryCounter::GetAbsoluteCounter();
+				// Reset the counter to later get the peak during the updated
+				mu::FGlobalMemoryCounter::Zero();													
+			}
+
+			// Prepare streaming for the current customizable object
+			check(SystemPrivate->Streamer != nullptr);
+			SystemPrivate->Streamer->PrepareStreamingForObject(Operation->Instance->GetCustomizableObject());			
 
-			SystemPrivateData->MutableTaskGraph.AddGameThreadTask(
-				TEXT("Task_Game_LockMeshCache"),
-				 [Operation, CacheRuntimeTexturesEvent]()
+			const TSharedPtr<mu::FSystem> MutableSystem = SystemPrivate->MutableSystem;
+			
+			if (Operation->bLiveUpdateMode)
+			{
+				if (Operation->InstanceID == 0)
 				{
-					impl::Task_Game_LockMeshCache(Operation, CacheRuntimeTexturesEvent);
-				},
-				false,
-				{ Mutable_GetMeshTask });
-	}
-		else
-		{
-			Mutable_GetMeshTask = SystemPrivateData->MutableTaskGraph.AddMutableThreadTask(
-				TEXT("Task_Mutable_GetMeshes"),
-				[Operation]()
+					// It's the first update since the instance was put in LiveUpdate Mode, this ID will be reused from now on
+					Operation->InstanceID = MutableSystem->NewInstance(Operation->Model);
+					UE_LOG(LogMutable, Verbose, TEXT("Creating Mutable instance with id [%d] for reuse "), Operation->InstanceID);
+				}
+				else
 				{
-					Task_Mutable_GetMeshes(Operation);
-				}, { CacheRuntimeTexturesEvent });
+					// The instance was already in LiveUpdate Mode, the ID is reused
+					check(Operation->InstanceID);
+					UE_LOG(LogMutable, Verbose, TEXT("Reusing Mutable instance with id [%d] "), Operation->InstanceID);
+				}
+			}
+			else
+			{
+				// In non-LiveUpdate mode, we are forcing the recreation of mutable-side instances with every update.
+				check(Operation->InstanceID == 0);
+				Operation->InstanceID = MutableSystem->NewInstance(Operation->Model);
+				UE_LOG(LogMutable, Verbose, TEXT("Creating Mutable instance with id [%d] "), Operation->InstanceID);
+			}
 
-			// Task: Lock cache
-			//-------------------------------------------------------------
+			Operation->MutableInstance = MutableSystem->BeginUpdate(Operation->InstanceID, Operation->Parameters, Operation->GetCapturedDescriptor().GetState(), mu::FSystem::AllLODs);
+			Operation->NumInstanceComponents = Operation->MutableInstance->GetComponentCount();
+		}
+		
+		FixLODs(Operation);
+		
+		if (Operation->bUseMeshCache)
+		{
+			Operation->InitMeshDescriptors(Operation->NumObjectComponents);
+
+			const TMap<FName, uint8>& MinLOD = Operation->GetMinLOD();
+			const TMap<FName, uint8>& NumLODsAvailable = Operation->NumLODsAvailable;
+		
+			for (int32 InstanceComponentIndex = 0; InstanceComponentIndex < Operation->NumInstanceComponents; ++InstanceComponentIndex)
 			{
-				// Task inputs
-				SystemPrivateData->MutableTaskGraph.AddGameThreadTask(
-					TEXT("Task_Game_LockCache"),
-					[Operation]()
+				FCustomizableObjectComponentIndex ObjectComponentIndex(Operation->MutableInstance->GetComponentId(InstanceComponentIndex));
+				TArray<mu::FResourceID>* MeshId = Operation->GetMeshDescriptors(ObjectComponentIndex);
+				if (!MeshId)
+				{
+					continue;
+				}
+
+				MeshId->Init(MAX_uint64, MAX_MESH_LOD_COUNT);
+
+				const FName& ComponentName = Operation->ComponentNames[ObjectComponentIndex.GetValue()];
+			
+				for (int32 LODIndex = MinLOD[ComponentName]; LODIndex < NumLODsAvailable[ComponentName]; ++LODIndex)
+				{
+					const bool bGenerateLOD = Operation->GetRequestedLODs()[ComponentName] <= LODIndex;
+					if (bGenerateLOD)
 					{
-						impl::Task_Game_LockCache(Operation);
-					},
-					false,
-					{ Mutable_GetMeshTask });
+						(*MeshId)[LODIndex] = Operation->MutableInstance->GetMeshId(InstanceComponentIndex, LODIndex);
+					}
+				}
+			}
+			
+			for (const TArray<mu::FResourceID>& MeshId : Operation->GetMeshDescriptors())
+			{
+				if (USkeletalMesh* CachedMesh = CustomizableObject->GetPrivate()->MeshCache.Get(MeshId))
+				{
+					Operation->Objects.Emplace(CachedMesh);
+				}
 			}
 		}
+		
+		UE::Tasks::FTask Mutable_GetMeshTask = SystemPrivateData->MutableTaskGraph.AddMutableThreadTask(
+		TEXT("Task_Mutable_GetMeshes"),
+		[Operation]()
+		{
+			Task_Mutable_GetMeshes(Operation);
+		},
+		{ CacheRuntimeTexturesEvent });
+
+		SystemPrivateData->MutableTaskGraph.AddGameThreadTask(
+			TEXT("Task_Game_LockCache"),
+			[Operation]()
+			{
+				impl::Task_Game_LockCache(Operation);
+			},
+			false,
+			{ Mutable_GetMeshTask });
 	}
 } // namespace impl
 
@@ -4853,7 +4787,7 @@ void UCustomizableObjectSystemPrivate::StartUpdateSkeletalMesh(const TSharedRef<
 		{
 			impl::Task_Game_StartUpdate(Context);
 		},
-		false);
+		true);
 }
 
 
