[Mutable] Fix Morph issues with Streaming.

diff --git a/Engine/Plugins/Mutable/Source/CustomizableObject/Private/MuCO/CustomizableObjectMeshUpdate.cpp b/Engine/Plugins/Mutable/Source/CustomizableObject/Private/MuCO/CustomizableObjectMeshUpdate.cpp
--- a/Engine/Plugins/Mutable/Source/CustomizableObject/Private/MuCO/CustomizableObjectMeshUpdate.cpp
+++ b/Engine/Plugins/Mutable/Source/CustomizableObject/Private/MuCO/CustomizableObjectMeshUpdate.cpp
@@ -288,8 +288,11 @@ void FCustomizableObjectMeshStreamIn::RequestMeshUpdate(const FContext& Context)
 				{
 					continue;
 				}
-				
-				LoadMorphTargets(StreamRequest, Mesh.ToSharedRef(), RefThis->OperationData->MorphTargetMeshData);
+
+				if (CVarEnableRealTimeMorphTargets.GetValueOnAnyThread())
+				{
+					LoadMorphTargets(StreamRequest, Mesh.ToSharedRef(), RefThis->OperationData->MorphTargetMeshData);
+				}
 			}
 
 			// Stream data.
diff --git a/Engine/Plugins/Mutable/Source/CustomizableObject/Private/MuCO/UnrealConversionUtils.cpp b/Engine/Plugins/Mutable/Source/CustomizableObject/Private/MuCO/UnrealConversionUtils.cpp
--- a/Engine/Plugins/Mutable/Source/CustomizableObject/Private/MuCO/UnrealConversionUtils.cpp
+++ b/Engine/Plugins/Mutable/Source/CustomizableObject/Private/MuCO/UnrealConversionUtils.cpp
@@ -773,6 +773,11 @@ namespace UnrealConversionUtils
 
 	void MorphTargetVertexInfoBuffers(FSkeletalMeshLODRenderData& LODResource, const USkeletalMesh& Owner, const mu::FMesh& MutableMesh, const TMap<uint32, FMorphTargetMeshData>& MorphTargetMeshData, int32 LODIndex)
 	{
+		if (Owner.GetMorphTargets().IsEmpty())
+		{
+			return;
+		}
+		
 		// Get the global names.
 		const TMap<FName, int32>& IndexMap = Owner.GetMorphTargetIndexMap();
 	
