[Mutable] Allow multiple GameThread tasks per tick.

The time limit can be changed with a CVar. Currently set to 0.5ms.

diff --git a/Engine/Plugins/Mutable/Source/CustomizableObject/Internal/MuCO/CustomizableObjectSystemPrivate.h b/Engine/Plugins/Mutable/Source/CustomizableObject/Internal/MuCO/CustomizableObjectSystemPrivate.h
--- a/Engine/Plugins/Mutable/Source/CustomizableObject/Internal/MuCO/CustomizableObjectSystemPrivate.h
+++ b/Engine/Plugins/Mutable/Source/CustomizableObject/Internal/MuCO/CustomizableObjectSystemPrivate.h
@@ -915,6 +915,9 @@ public:
 	
 	/** Mutable TaskGraph system (Mutable Thread). */
 	FMutableTaskGraph MutableTaskGraph;
+
+	/** Last Mutable task from the previous update. The next update can not start until this task has has completed. */
+	UE::Tasks::FTask LastUpdateMutableTask = UE::Tasks::MakeCompletedTask<void>();
 	
 #if WITH_EDITORONLY_DATA
 	/** Mutable default image provider. Used by the COIEditor and Instance/Descriptor APIs. */
diff --git a/Engine/Plugins/Mutable/Source/CustomizableObject/Internal/MuCO/FMutableTaskGraph.h b/Engine/Plugins/Mutable/Source/CustomizableObject/Internal/MuCO/FMutableTaskGraph.h
--- a/Engine/Plugins/Mutable/Source/CustomizableObject/Internal/MuCO/FMutableTaskGraph.h
+++ b/Engine/Plugins/Mutable/Source/CustomizableObject/Internal/MuCO/FMutableTaskGraph.h
@@ -92,6 +92,9 @@ private:
 	/** Check if the Mutable Thread is locked. */
 	bool IsMutableThreadLocked() const;
 
+	/** Check if the Mutable Thread is locked without using the lock. Does not update the region. */
+	bool IsMutableThreadLockedNoLock() const;
+
 	/** Unlock the Mutable Thread. Releases all blocked tasks. */
 	void UnlockMutableThread();
 	
diff --git a/Engine/Plugins/Mutable/Source/CustomizableObject/Private/MuCO/CustomizableObjectSystem.cpp b/Engine/Plugins/Mutable/Source/CustomizableObject/Private/MuCO/CustomizableObjectSystem.cpp
--- a/Engine/Plugins/Mutable/Source/CustomizableObject/Private/MuCO/CustomizableObjectSystem.cpp
+++ b/Engine/Plugins/Mutable/Source/CustomizableObject/Private/MuCO/CustomizableObjectSystem.cpp
@@ -3312,7 +3312,7 @@ namespace impl
 		// Next Task: Release Mutable. We need this regardless if we cancel or not
 		//-------------------------------------------------------------		
 		const  TSharedPtr<mu::FSystem> MutableSystem = CustomizableObjectSystemPrivateData->MutableSystem;
-		UE::Tasks::FTask Mutable_ReleaseInstanceTask = CustomizableObjectSystemPrivateData->MutableTaskGraph.AddMutableThreadTask(
+		CustomizableObjectSystemPrivateData->LastUpdateMutableTask = CustomizableObjectSystemPrivateData->MutableTaskGraph.AddMutableThreadTask(
 			TEXT("Task_Mutable_ReleaseInstance"),
 			[InstanceID = OperationData->InstanceID, MutableSystem, bLiveUpdateMode = OperationData->bLiveUpdateMode]()
 			{
@@ -3355,20 +3355,13 @@ namespace impl
 
 			// Next Task: Callbacks
 			//-------------------------------------------------------------
-			TArray<UE::Tasks::FTask> Dependencies;
-			if (CVarFixLowPriorityTasksOverlap.GetValueOnGameThread())
-			{
-				Dependencies.Add(Mutable_ReleaseInstanceTask);
-			}
-			
 			CustomizableObjectSystemPrivateData->MutableTaskGraph.AddGameThreadTask(
 				TEXT("Task_Game_Callbacks"),
 				[OperationData]()
 				{
 					Task_Game_Callbacks(OperationData);
 				},
-				false,
-				Dependencies);
+				false);
 		}
 		else
 		{
@@ -4787,7 +4780,8 @@ void UCustomizableObjectSystemPrivate::StartUpdateSkeletalMesh(const TSharedRef<
 		{
 			impl::Task_Game_StartUpdate(Context);
 		},
-		true);
+		true,
+		{ LastUpdateMutableTask });
 }
 
 
diff --git a/Engine/Plugins/Mutable/Source/CustomizableObject/Private/MuCO/FMutableTaskGraph.cpp b/Engine/Plugins/Mutable/Source/CustomizableObject/Private/MuCO/FMutableTaskGraph.cpp
--- a/Engine/Plugins/Mutable/Source/CustomizableObject/Private/MuCO/FMutableTaskGraph.cpp
+++ b/Engine/Plugins/Mutable/Source/CustomizableObject/Private/MuCO/FMutableTaskGraph.cpp
@@ -19,6 +19,12 @@ static TAutoConsoleVariable<float> CVarMutableTaskLowPriorityMaxWaitTime(
 	TEXT("Max time a Mutable Task with Low priority can wait. Once this time has passed it will be launched unconditionally."),
 	ECVF_Scalability);
 
+static TAutoConsoleVariable<float> CVarGameThreadTaskMaxTime(
+	TEXT("mutable.GameThreadTaskMaxTime"),
+	0.5 / 1000.0, // 0.5ms
+	TEXT("Max time Mutable can execute Game Thread Task each tick. A single task may take longer than the specified time."),
+	ECVF_Scalability);
+
 static TAutoConsoleVariable<bool> CVarEnableMutableTaskLowPriority(
 	TEXT("mutable.EnableMutableTaskLowPriority"),
 	true,
@@ -282,39 +288,43 @@ void FMutableTaskGraph::TryLaunchGameThreadTask()
 {
 	MUTABLE_CPUPROFILER_SCOPE(AdvanceCurrentOperation);
 
-	// Scope where PendingTask is valid.
+	double CurrentTime = FPlatformTime::Seconds(); 
+	const double EndTime = CurrentTime + CVarGameThreadTaskMaxTime.GetValueOnGameThread();
+
+	while (CurrentTime < EndTime)
 	{
-		FGameThreadTask* PendingTask = GameThreadTasks.Peek();
+		const FGameThreadTask* PendingTask = GameThreadTasks.Peek();
 		if (!PendingTask)
 		{
-			return;
+			break;
 		}
-
+		
 		if (!PendingTask->AreDependenciesComplete())
 		{
-			return;
+			break;
 		}
-	
+
 		if (PendingTask->bLockMutableThread)
 		{
 			AsyncLockMutableThread();
 
 			if (!IsMutableThreadLocked())
 			{
-				return;
+				break;
 			}
 		}
-	}
-	
-	FGameThreadTask CurrentTask;
-	GameThreadTasks.Dequeue(CurrentTask);
 	
-	CurrentTask.Body();
-	
-	if (CurrentTask.bLockMutableThread)
-	{
-		UnlockMutableThread();
+		FGameThreadTask CurrentTask;
+		GameThreadTasks.Dequeue(CurrentTask); // Dequeue now since current task can enqueue new tasks.
+
+		// Careful! From this point, PendingTask is no longer valid.
+		
+		CurrentTask.Body();
+
+		CurrentTime = FPlatformTime::Seconds(); 
 	}
+
+	UnlockMutableThread();
 }
 
 
@@ -409,13 +419,19 @@ void FMutableTaskGraph::AsyncLockMutableThread()
 }
 
 
+bool FMutableTaskGraph::IsMutableThreadLockedNoLock() const
+{
+	return !MutableThreadUnlockEvent.IsCompleted() && LastMutableTaskBeforeLock.IsCompleted();
+}
+
+
 bool FMutableTaskGraph::IsMutableThreadLocked() const
 {
 	MUTABLE_CPUPROFILER_SCOPE(FMutableTaskGraph::IsMutableThreadLocked)
 
 	FScopeLock Lock(&MutableTaskLock);
 
-	const bool bLocked = !MutableThreadUnlockEvent.IsCompleted() && LastMutableTaskBeforeLock.IsCompleted();
+	const bool bLocked = IsMutableThreadLockedNoLock();
 
 	UpdateRegions(bLocked ? ERegionLockStatus::Locked : RegionLockStatus);
 	
@@ -426,9 +442,14 @@ bool FMutableTaskGraph::IsMutableThreadLocked() const
 void FMutableTaskGraph::UnlockMutableThread()
 {
 	MUTABLE_CPUPROFILER_SCOPE(FMutableTaskGraph::UnlockMutableThread)
-
+	
 	FScopeLock Lock(&MutableTaskLock);
 
+	if (!IsMutableThreadLockedNoLock())
+	{
+		return;
+	}
+	
 	UpdateRegions(ERegionLockStatus::Unlocked);
 	
 	LastMutableTaskBeforeLock = {};
