[Mutable] Fix non-deterministic participating objects.

#rnx

diff --git a/Engine/Plugins/Mutable/Source/CustomizableObjectEditor/Private/MuCOE/CustomizableObjectEditorModule.cpp b/Engine/Plugins/Mutable/Source/CustomizableObjectEditor/Private/MuCOE/CustomizableObjectEditorModule.cpp
--- a/Engine/Plugins/Mutable/Source/CustomizableObjectEditor/Private/MuCOE/CustomizableObjectEditorModule.cpp
+++ b/Engine/Plugins/Mutable/Source/CustomizableObjectEditor/Private/MuCOE/CustomizableObjectEditorModule.cpp
@@ -476,6 +476,12 @@ void GetReferencingCustomizableObjects(FName CustomizableObjectName, TArray<FNam
 
 	AssetRegistryModule.Get().GetAssets(Filter, AssetDataArray);
 
+	// Required to be deterministic.
+	AssetDataArray.Sort([](const FAssetData& A, const FAssetData& B)
+		{
+			return A.PackageName.LexicalLess(B.PackageName);
+		});
+
 	for (FAssetData AssetData : AssetDataArray)
 	{
 		if (AssetData.GetClass() == UCustomizableObject::StaticClass())
@@ -853,10 +859,16 @@ TMap<FName, FGuid> FCustomizableObjectEditorModule::GetParticipatingObjects(cons
 			}
 		}
 
-		TArray<FAssetData> Assets;
-		AssetRegistry.GetAssets(Filter, Assets);
+		TArray<FAssetData> AssetDataArray;
+		AssetRegistry.GetAssets(Filter, AssetDataArray);
+
+		// Required to be deterministic.
+		AssetDataArray.Sort([](const FAssetData& A, const FAssetData& B)
+			{
+				return A.PackageName.LexicalLess(B.PackageName);
+			});
 
-		for (const FAssetData& Asset : Assets)
+		for (const FAssetData& Asset : AssetDataArray)
 		{
 			if (!FilterClasses.Contains(Asset.GetClass(EResolveClass::Yes)))
 			{
