[Mutable] Fix MutableMeshPreviewUtils crash when generating the mesh due to missing skin weights buffers.

#rnx

diff --git a/Engine/Plugins/Mutable/Source/CustomizableObjectEditor/Private/MuCOE/MutableMeshPreviewUtils.cpp b/Engine/Plugins/Mutable/Source/CustomizableObjectEditor/Private/MuCOE/MutableMeshPreviewUtils.cpp
--- a/Engine/Plugins/Mutable/Source/CustomizableObjectEditor/Private/MuCOE/MutableMeshPreviewUtils.cpp
+++ b/Engine/Plugins/Mutable/Source/CustomizableObjectEditor/Private/MuCOE/MutableMeshPreviewUtils.cpp
@@ -227,6 +227,18 @@ namespace MutableMeshPreviewUtils
 				bCanBeConverted = false;
 			}
 
+			// SKIN WEIGHTS CHECK
+			VertexBuffers.FindChannel(mu::EMeshBufferSemantic::BoneIndices, 0, &Buffer, &Channel);
+			if (Buffer == -1)
+			{
+				if (bLogFindingsToConsole)
+				{
+					UE_LOG(LogMutable, Warning, TEXT("\tSkinning Buffer not found."));
+				}
+
+				bCanBeConverted = false;
+			}
+
 			// TANGENT BUFFER INDEX CHECK
 			VertexBuffers.FindChannel(mu::EMeshBufferSemantic::Tangent, 0, &Buffer, &Channel);
 			if (Buffer == -1)
