[Mutable] Fix mesh streaming crash due to missing RHI buffers in instances where the render resources had been released before the stream in task starts.
- Fix DiscardResources releasing render resources of meshes with pending init or streaming requests.
- Fix re-initialization of mesh resources when reusing meshes from previous updates. This double initialization messed with the RHI state and triggered rare crashes when the system was under stress.

#rnx

diff --git a/Engine/Plugins/Mutable/Source/CustomizableObject/Private/MuCO/CustomizableObjectInstance.cpp b/Engine/Plugins/Mutable/Source/CustomizableObject/Private/MuCO/CustomizableObjectInstance.cpp
--- a/Engine/Plugins/Mutable/Source/CustomizableObject/Private/MuCO/CustomizableObjectInstance.cpp
+++ b/Engine/Plugins/Mutable/Source/CustomizableObject/Private/MuCO/CustomizableObjectInstance.cpp
@@ -972,7 +972,7 @@ void UCustomizableInstancePrivate::PostEditChangePropertyWithoutEditor()
 	{
 		USkeletalMesh* SkeletalMesh = Tuple.Get<1>();
 		
-		if (SkeletalMesh && SkeletalMesh->GetResourceForRendering())
+		if (SkeletalMesh && SkeletalMesh->GetResourceForRendering() && !SkeletalMesh->GetResourceForRendering()->IsInitialized())
 		{
 			MUTABLE_CPUPROFILER_SCOPE(InitResources);
 
@@ -3188,7 +3188,7 @@ void UCustomizableInstancePrivate::DiscardResources()
 			{
 				USkeletalMesh* SkeletalMesh = Tuple.Get<1>();
 			
-				if (SkeletalMesh->IsValidLowLevel())
+				if (SkeletalMesh->IsValidLowLevel() && !SkeletalMesh->HasPendingInitOrStreaming())
 				{
 					SkeletalMesh->ReleaseResources();
 				}
diff --git a/Engine/Plugins/Mutable/Source/CustomizableObject/Private/MuCO/CustomizableObjectMeshUpdate.cpp b/Engine/Plugins/Mutable/Source/CustomizableObject/Private/MuCO/CustomizableObjectMeshUpdate.cpp
--- a/Engine/Plugins/Mutable/Source/CustomizableObject/Private/MuCO/CustomizableObjectMeshUpdate.cpp
+++ b/Engine/Plugins/Mutable/Source/CustomizableObject/Private/MuCO/CustomizableObjectMeshUpdate.cpp
@@ -370,7 +370,7 @@ void FCustomizableObjectMeshStreamIn::ConvertMesh(const FContext& Context)
 			return;
 		}
 
-		if (MutableMesh->GetVertexCount() == 0 || MutableMesh->GetSurfaceCount() == 0)
+		if (MutableMesh->GetVertexCount() == 0 || MutableMesh->GetSurfaceCount() == 0 || MutableMesh->GetVertexBuffers().IsDescriptor())
 		{
 			check(false);
 			Abort();
diff --git a/Engine/Plugins/Mutable/Source/CustomizableObject/Private/MuCO/CustomizableObjectSkeletalMesh.cpp b/Engine/Plugins/Mutable/Source/CustomizableObject/Private/MuCO/CustomizableObjectSkeletalMesh.cpp
--- a/Engine/Plugins/Mutable/Source/CustomizableObject/Private/MuCO/CustomizableObjectSkeletalMesh.cpp
+++ b/Engine/Plugins/Mutable/Source/CustomizableObject/Private/MuCO/CustomizableObjectSkeletalMesh.cpp
@@ -10,6 +10,7 @@
 #include "MuCO/CustomizableObjectInstancePrivate.h"
 #include "MuCO/CustomizableObjectMeshUpdate.h"
 
+#include "Rendering/SkeletalMeshRenderData.h"
 #include "Animation/SkinWeightProfile.h"
 
 FName GenerateUniqueNameFromCOInstance(const UCustomizableObjectInstance& Instance, FName ComponentName)
@@ -72,6 +73,13 @@ void UCustomizableObjectSkeletalMesh::InitMutableStreamingData(const TSharedRef<
 bool UCustomizableObjectSkeletalMesh::StreamIn(int32 NewMipCount, bool bHighPrio)
 {
 	check(IsInGameThread());
+
+	FSkeletalMeshRenderData* RenderData = GetResourceForRendering();
+	if (!RenderData || !RenderData->IsInitialized())
+	{
+		return false;
+	}
+
 	if (!HasPendingInitOrStreaming() && CachedSRRState.StreamIn(NewMipCount))
 	{
 		FRenderAssetUpdate::EThreadType CreateResourcesThread = GRHISupportsAsyncTextureCreation
