[mutable] Mutable asset naming and COI Baking changes.
- Made the mutable generated resources expose a hash of the source of the data (implementation varies depending on resource).
- Updated the baking process for better handling of the resource names.
- Added messages that will show when the baking update has been scheduled and when it has ended.
- Removed unnecessary includes.
- The user is no longer required to provide the bake operation with a name for the files.

#rnx

diff --git a/Engine/Plugins/Mutable/Source/CustomizableObject/Internal/MuCO/CustomizableObjectInstancePrivate.h b/Engine/Plugins/Mutable/Source/CustomizableObject/Internal/MuCO/CustomizableObjectInstancePrivate.h
--- a/Engine/Plugins/Mutable/Source/CustomizableObject/Internal/MuCO/CustomizableObjectInstancePrivate.h
+++ b/Engine/Plugins/Mutable/Source/CustomizableObject/Internal/MuCO/CustomizableObjectInstancePrivate.h
@@ -4,8 +4,6 @@
 
 
 #include "Materials/MaterialInterface.h"
-#include "Rendering/SkeletalMeshLODModel.h"
-#include "Rendering/SkeletalMeshModel.h"
 #include "MuCO/CustomizableObjectInstance.h"
 #include "MuCO/CustomizableObjectPrivate.h"
 #include "MuCO/CustomizableObjectSystemPrivate.h"
@@ -317,7 +315,7 @@ private:
 	UPhysicsAsset* GetOrBuildMainPhysicsAsset(const TSharedRef<FUpdateContextPrivate>& OperationData, TObjectPtr<class UPhysicsAsset> TamplateAsset, const mu::FPhysicsBody* PhysicsBody, bool bDisableCollisionBetweenAssets, FCustomizableObjectInstanceComponentIndex InstanceComponentIndex);
 	
 	// Create a transient texture and add it to the TextureTrackerArray
-	UTexture2D* CreateTexture();
+	UTexture2D* CreateTexture(const FString& TextureName);
 
 	void InvalidateGeneratedData();
 
diff --git a/Engine/Plugins/Mutable/Source/CustomizableObject/Private/MuCO/CustomizableObjectInstance.cpp b/Engine/Plugins/Mutable/Source/CustomizableObject/Private/MuCO/CustomizableObjectInstance.cpp
--- a/Engine/Plugins/Mutable/Source/CustomizableObject/Private/MuCO/CustomizableObjectInstance.cpp
+++ b/Engine/Plugins/Mutable/Source/CustomizableObject/Private/MuCO/CustomizableObjectInstance.cpp
@@ -57,6 +57,7 @@
 #include "HAL/PlatformFileManager.h"
 
 #include "MuCO/Plugins/IMutableClothingModule.h"
+#include "Rendering/SkeletalMeshModel.h"
 
 #include UE_INLINE_GENERATED_CPP_BY_NAME(CustomizableObjectInstance)
 
@@ -189,11 +190,11 @@ uint32 GetTypeHash(const FMutableMaterialPlaceholder& PlaceHolder)
 }
 
 
-UTexture2D* UCustomizableInstancePrivate::CreateTexture()
+UTexture2D* UCustomizableInstancePrivate::CreateTexture(const FString& TextureName)
 {
 	UTexture2D* NewTexture = NewObject<UTexture2D>(
 		GetTransientPackage(),
-		NAME_None,
+		GetData(TextureName),
 		RF_Transient
 		);
 	UCustomizableObjectSystem::GetInstance()->GetPrivate()->LogBenchmarkUtil.AddTexture(*NewTexture);
@@ -1316,6 +1317,10 @@ USkeleton* UCustomizableInstancePrivate::MergeSkeletons(UCustomizableObject& Cus
 	}
 	else
 	{
+#if WITH_EDITOR
+		uint32 CombinedSkeletonHash = INDEX_NONE;
+#endif
+		
 		// Make the final skeleton compatible with all the merged skeletons and their compatible skeletons.
 		for (USkeleton* Skeleton : Params.SkeletonsToMerge)
 		{
@@ -1328,9 +1333,19 @@ USkeleton* UCustomizableInstancePrivate::MergeSkeletons(UCustomizableObject& Cus
 				{
 					FinalSkeleton->AddCompatibleSkeletonSoft(CompatibleSkeleton);
 				}
+
+#if WITH_EDITOR
+				const uint32 SkeletonHash = GetTypeHash(Skeleton->GetName());
+				CombinedSkeletonHash = HashCombine(CombinedSkeletonHash, SkeletonHash);
+#endif
 			}
 		}
 
+		// Add the hash based on the sources for the merged skeleton to make its name unique
+#if WITH_EDITOR
+		FinalSkeleton->Rename(*FString::Printf(TEXT("%s_%lu"), *FinalSkeleton->GetName(), CombinedSkeletonHash));
+#endif
+		
 		// Add Skeleton to the cache
 		CustomizableObject.GetPrivate()->SkeletonCache.Add(ReferencedSkeletons.SkeletonIds, FinalSkeleton);
 		ReferencedSkeletons.SkeletonIds.Empty();
@@ -2332,9 +2347,11 @@ bool UCustomizableInstancePrivate::UpdateSkeletalMesh_PostBeginUpdate0(UCustomiz
 
 		// We need the first valid mesh. get it from the component, considering that some LOSs may have been skipped.
 		TSharedPtr<const mu::FMesh> ComponentMesh;
+		uint32 MeshHash = INDEX_NONE;
 		for (int32 FirstValidLODIndex = Component->FirstLOD; FirstValidLODIndex < OperationData->InstanceUpdateData.LODs.Num() && !ComponentMesh; ++FirstValidLODIndex)
 		{
 			ComponentMesh = OperationData->InstanceUpdateData.LODs[FirstValidLODIndex].Mesh;
+			MeshHash = HashCombine(MeshHash,GetTypeHash(OperationData->InstanceUpdateData.LODs[FirstValidLODIndex].MeshID));
 		}
 		
 		if (!ComponentMesh)
@@ -2383,18 +2400,19 @@ bool UCustomizableInstancePrivate::UpdateSkeletalMesh_PostBeginUpdate0(UCustomiz
 		MUTABLE_CPUPROFILER_SCOPE(ConstructMesh);
 
 		USkeletalMesh* SkeletalMesh = nullptr;
-
+		
+		FString SkeletalMeshName = FString::Printf (TEXT("SK_%s_%s_%lu"), *Public->GetCustomizableObject()->GetName(), *ComponentName.ToString(), MeshHash);
+		SkeletalMeshName = MakeUniqueObjectName(GetTransientPackage(), USkeletalMesh::StaticClass(), *SkeletalMeshName, EUniqueObjectNameOptions::GloballyUnique).ToString();
 		if (OperationData->bStreamMeshLODs)
 		{
-			SkeletalMesh = UCustomizableObjectSkeletalMesh::CreateSkeletalMesh(*Public, ComponentName);
+			SkeletalMesh = NewObject<UCustomizableObjectSkeletalMesh>(GetTransientPackage(), FName(*SkeletalMeshName), RF_Transient);
 		}
 		else
 		{
-			FName SkeletalMeshName = GenerateUniqueNameFromCOInstance(*Public, ComponentName);
-			SkeletalMesh = NewObject<USkeletalMesh>(GetTransientPackage(), SkeletalMeshName, RF_Transient);
+			SkeletalMesh = NewObject<USkeletalMesh>(GetTransientPackage(), FName(*SkeletalMeshName), RF_Transient);
 		}
 		
-		check(SkeletalMesh);		
+		check(SkeletalMesh);
 		SkeletalMeshes.Add(ComponentName, SkeletalMesh);
 		
 		const FMutableRefSkeletalMeshData& RefSkeletalMeshData = ModelResources.ReferenceSkeletalMeshesData[ObjectComponentIndex.GetValue()];
@@ -6164,6 +6182,11 @@ void UCustomizableInstancePrivate::BuildMaterials(const TSharedRef<FUpdateContex
 								// and delete it automatically
 								TSharedPtr<FTexturePlatformData, ESPMode::ThreadSafe> ReusedTexturePlatformData;
 
+								// Ensure the name of the texture is unique (prevents black textures in editor previews)
+								FString MutableTextureName = FString::Printf(TEXT("T_%s_%llu"), *Props.TextureParameterName, Image.ImageID);
+								MutableTextureName.ReplaceInline(TEXT(" "), TEXT("_"));
+								MutableTextureName = MakeUniqueObjectName(GetTransientPackage(), UTexture2D::StaticClass(), *MutableTextureName, EUniqueObjectNameOptions::GloballyUnique).ToString();
+									
 								if (ReusedTexture && (*ReusedTexture).IsValid() && !(*ReusedTexture)->HasAnyFlags(RF_BeginDestroyed))
 								{
 									// Only uncompressed textures can be reused. This also fixes an issue in the editor where textures supposedly 
@@ -6179,8 +6202,7 @@ void UCustomizableInstancePrivate::BuildMaterials(const TSharedRef<FUpdateContex
 									else
 									{
 										ReusedTexture = nullptr;
-										MutableTexture = CreateTexture();
-
+										MutableTexture = CreateTexture(MutableTextureName);
 #if WITH_EDITOR
 										UE_LOG(LogMutable, Warning,
 											TEXT("Tried to reuse an uncompressed texture with name %s. Make sure the selected Mutable state disables texture compression/streaming, that one of the state's runtime parameters affects the texture and that the CO is compiled with max. optimization settings."),
@@ -6191,7 +6213,7 @@ void UCustomizableInstancePrivate::BuildMaterials(const TSharedRef<FUpdateContex
 								else
 								{
 									ReusedTexture = nullptr;
-									MutableTexture = CreateTexture();
+									MutableTexture = CreateTexture(MutableTextureName);
 								}
 
 								if (MutableTexture)
@@ -6383,7 +6405,7 @@ void UCustomizableInstancePrivate::BuildMaterials(const TSharedRef<FUpdateContex
 #endif
 
 					UMaterialInstanceDynamic* MaterialInstance = nullptr;
-					
+
 					if (const int32 OldMaterialIndex = OldGeneratedMaterials.Find(Material); bReuseMaterials && OldMaterialIndex != INDEX_NONE)
 					{
 						const FGeneratedMaterial& OldMaterial = OldGeneratedMaterials[OldMaterialIndex];
@@ -6393,7 +6415,37 @@ void UCustomizableInstancePrivate::BuildMaterials(const TSharedRef<FUpdateContex
 					
 					if (!MaterialInstance && MutableMaterialPlaceholder.GetParams().Num() != 0)
 					{
+#if WITH_EDITOR
+						// Remove the MI_ or M_ prefixes from the material string so we use it as the name of the MID
+						FString MIDName = MaterialTemplate->GetName();
+						{
+							const FString MaterialPrefix = TEXT("M_");						// Material
+							const FString MaterialInstancePrefix = TEXT("MI_");				// Material Instance
+							const FString MaterialInstanceConstantPrefix = TEXT("MIC_");	// Material Instance Constant
+							
+							if (MIDName.Find(MaterialInstancePrefix, ESearchCase::CaseSensitive) == 0)
+							{
+								MIDName = MIDName.RightChop(MaterialInstancePrefix.Len());
+							}
+							else if (MIDName.Find(MaterialPrefix, ESearchCase::CaseSensitive) == 0)
+							{
+								MIDName = MIDName.RightChop(MaterialPrefix.Len());
+							}
+							else if (MIDName.Find(MaterialInstanceConstantPrefix, ESearchCase::CaseSensitive) == 0)
+							{
+								MIDName = MIDName.RightChop(MaterialInstanceConstantPrefix.Len());
+							}
+
+							const uint32 MaterialPlaceHolderHash = GetTypeHash(MutableMaterialPlaceholder);
+							MIDName = FString::Printf(TEXT("%s%s_%u"),TEXT("MID_"),*MIDName, MaterialPlaceHolderHash);
+							MIDName = MakeUniqueObjectName(GetTransientPackage(), UMaterialInstanceDynamic::StaticClass(), *MIDName).ToString();
+						}
+						
+						MaterialInstance = UMaterialInstanceDynamic::Create(MaterialTemplate, GetTransientPackage(), FName(MIDName));
+#else
 						MaterialInstance = UMaterialInstanceDynamic::Create(MaterialTemplate, GetTransientPackage());
+#endif
+						
 						Material.MaterialInterface = MaterialInstance;
 					}
 
diff --git a/Engine/Plugins/Mutable/Source/CustomizableObject/Private/MuCO/CustomizableObjectSkeletalMesh.cpp b/Engine/Plugins/Mutable/Source/CustomizableObject/Private/MuCO/CustomizableObjectSkeletalMesh.cpp
--- a/Engine/Plugins/Mutable/Source/CustomizableObject/Private/MuCO/CustomizableObjectSkeletalMesh.cpp
+++ b/Engine/Plugins/Mutable/Source/CustomizableObject/Private/MuCO/CustomizableObjectSkeletalMesh.cpp
@@ -2,34 +2,16 @@
 
 #include "MuCO/CustomizableObjectSkeletalMesh.h"
 
-#include "MuCO/CustomizableObjectSystem.h"
 #include "MuCO/CustomizableObjectSystemPrivate.h"
 
 #include "MuCO/CustomizableObject.h"
-#include "MuCO/CustomizableObjectInstance.h"
 #include "MuCO/CustomizableObjectInstancePrivate.h"
 #include "MuCO/CustomizableObjectMeshUpdate.h"
 
 #include "Rendering/SkeletalMeshRenderData.h"
 #include "Animation/SkinWeightProfile.h"
 
-FName GenerateUniqueNameFromCOInstance(const UCustomizableObjectInstance& Instance, FName ComponentName)
-{
-	FString BaseName = FString("SK_") + Instance.GetCustomizableObject()->GetName()+ FString("_") + ComponentName.ToString();
-	FName SkeletalMeshName = MakeUniqueObjectName(GetTransientPackage(), USkeletalMesh::StaticClass(), FName(BaseName), 
-							                      EUniqueObjectNameOptions::GloballyUnique);
-
-	return SkeletalMeshName;
-}
-
 
-UCustomizableObjectSkeletalMesh* UCustomizableObjectSkeletalMesh::CreateSkeletalMesh(const UCustomizableObjectInstance& Instance, FName ComponentName)
-{
-	FName SkeletalMeshName = GenerateUniqueNameFromCOInstance(Instance, ComponentName);
-	UCustomizableObjectSkeletalMesh* OutSkeletalMesh = NewObject<UCustomizableObjectSkeletalMesh>(GetTransientPackage(), SkeletalMeshName, RF_Transient);
-	
-	return OutSkeletalMesh;
-}
 
 
 void UCustomizableObjectSkeletalMesh::InitMutableStreamingData(const TSharedRef<FUpdateContextPrivate>& InOperationData, const FName& ComponentName, const int32 FirstLOD, const int32 LODCount)
diff --git a/Engine/Plugins/Mutable/Source/CustomizableObject/Public/MuCO/CustomizableObjectSkeletalMesh.h b/Engine/Plugins/Mutable/Source/CustomizableObject/Public/MuCO/CustomizableObjectSkeletalMesh.h
--- a/Engine/Plugins/Mutable/Source/CustomizableObject/Public/MuCO/CustomizableObjectSkeletalMesh.h
+++ b/Engine/Plugins/Mutable/Source/CustomizableObject/Public/MuCO/CustomizableObjectSkeletalMesh.h
@@ -22,7 +22,6 @@ class FUpdateContextPrivate;
 class FCustomizableObjectMeshStreamIn;
 struct FModelStreamableBulkData;
 
-FName GenerateUniqueNameFromCOInstance(const UCustomizableObjectInstance& Instance, FName ComponentName);
 
 /**
  * CustomizableObjectSkeletalMesh is a class inheriting from USkeletalMesh with the sole purpose of allowing skeletal meshes generated by Mutable
@@ -35,11 +34,7 @@ class CUSTOMIZABLEOBJECT_API UCustomizableObjectSkeletalMesh : public USkeletalM
 
 	friend FCustomizableObjectMeshStreamIn;
 	friend UCustomizableInstancePrivate;
-
-private:
-
-	static UCustomizableObjectSkeletalMesh* CreateSkeletalMesh(const UCustomizableObjectInstance& InInstance, FName ComponentName);
-
+	
 	void InitMutableStreamingData(const TSharedRef<FUpdateContextPrivate>& InOperationData, const FName& ComponentName, const int32 FirstLOD, const int32 LODCount);
 
 	//~ Begin UStreamableRenderAsset Interface.
diff --git a/Engine/Plugins/Mutable/Source/CustomizableObjectEditor/Private/MuCOE/CustomizableObjectInstanceBakingUtils.cpp b/Engine/Plugins/Mutable/Source/CustomizableObjectEditor/Private/MuCOE/CustomizableObjectInstanceBakingUtils.cpp
--- a/Engine/Plugins/Mutable/Source/CustomizableObjectEditor/Private/MuCOE/CustomizableObjectInstanceBakingUtils.cpp
+++ b/Engine/Plugins/Mutable/Source/CustomizableObjectEditor/Private/MuCOE/CustomizableObjectInstanceBakingUtils.cpp
@@ -26,13 +26,51 @@
 
 #define LOCTEXT_NAMESPACE "CustomizableObjectEditor" 
 
+// Recommended prefixes for various types of UE assets
+const FString SkeletalMeshAssetPrefix = TEXT("SK_");
+const FString MaterialAssetPrefix = TEXT("M_");
+const FString TextureAssetPrefix = TEXT("T_");
+const FString MaterialInstanceAssetPrefix = TEXT("MI_");
+const FString MaterialDynamicInstanceAssetPrefix = TEXT("MID_");
+const FString MaterialConstantInstanceAssetPrefix = TEXT("MIC_");
+const FString SkeletonAssetPrefix = TEXT("SKEL_");
+
+
+/**
+ * Remove the prefix from the provided string.
+ * @param InOutString The string to update
+ * @param ToCheckPrefixes The list of possible prefixes that, if present, will get removed from the target string
+ */
+void RemovePrefix(FString& InOutString, const TArray<FString>& ToCheckPrefixes)
+{
+	for (const FString& PrefixToRemove : ToCheckPrefixes)
+	{
+		if (InOutString.Find(PrefixToRemove, ESearchCase::CaseSensitive) == 0)
+		{
+			InOutString = InOutString.RightChop(PrefixToRemove.Len());
+			return;
+		}
+	}
+}
+
+/**
+ * Compose the resource name to be used by a resource.
+ * @param InPrefix The prefix we want to use for a given resource.
+ * @param InObjectName The name of the object this resource is part of.
+ * @param InOutResourceName The name of the resource itself. This string will be the target of the concatenations.
+ */
+void ComposeResourceName( const FString& InPrefix, const FString& InObjectName, FString& InOutResourceName)
+{
+	InOutResourceName = InPrefix + InObjectName + InOutResourceName;
+}
+
 
 /**
  * Simple wrapper to be able to invoke the generation of a popup or log message depending on the execution context in which this code is being ran
  * @param InMessage The message to display
  * @param InTitle The title to be used for the popup or the log generated
  */
-void ShowErrorNotification(const FText& InMessage, const FText& InTitle = LOCTEXT("CustomizableObjecInstanceBakingUtils_GenericBakingError","Baking Error") )
+void ShowErrorNotification(const FText& InMessage, const FText& InTitle = LOCTEXT("CustomizableObjectInstanceBakingUtils_GenericBakingError","Baking Error") )
 {
 	if (!FApp::IsUnattended())
 	{
@@ -55,38 +93,30 @@ void ShowErrorNotification(const FText& InMessage, const FText& InTitle = LOCTEX
  */
 bool ValidateProvidedFileName(const FString& FileName)
 {
-	if (FileName.IsEmpty())
-	{
-		UE_LOG(LogMutable, Error, TEXT("Invalid baking configuration : FileName string is empty.."));
-		return false;
-	}
-
 	// Check for invalid characters in the name of the object to be serialized
+	TCHAR InvalidCharacter = '0';
 	{
-		TCHAR InvalidCharacter = '0';
+		FString InvalidCharacters = FPaths::GetInvalidFileSystemChars();
+		for (int32 InvalidCharIndex = 0; InvalidCharIndex < InvalidCharacters.Len(); ++InvalidCharIndex)
 		{
-			FString InvalidCharacters = FPaths::GetInvalidFileSystemChars();
-			for (int32 i = 0; i < InvalidCharacters.Len(); ++i)
+			TCHAR Char = InvalidCharacters[InvalidCharIndex];
+			FString SearchedChar = FString::Chr(Char);
+			if (FileName.Contains(SearchedChar))
 			{
-				TCHAR Char = InvalidCharacters[i];
-				FString SearchedChar = FString::Chr(Char);
-				if (FileName.Contains(SearchedChar))
-				{
-					InvalidCharacter = InvalidCharacters[i];
-					break;
-				}
+				InvalidCharacter = InvalidCharacters[InvalidCharIndex];
+				break;
 			}
 		}
+	}
 
-		if (InvalidCharacter != '0')
-		{
-			const FText InvalidCharacterText = FText::FromString(FString::Chr(InvalidCharacter));
-			const FText ErrorText = FText::Format(LOCTEXT("CustomizableObjecInstanceBakingUtils_InvalidCharacter", "The selected contains an invalid character ({0})."), InvalidCharacterText);
+	if (InvalidCharacter != '0')
+	{
+		const FText InvalidCharacterText = FText::FromString(FString::Chr(InvalidCharacter));
+		const FText ErrorText = FText::Format(LOCTEXT("CustomizableObjectInstanceBakingUtils_InvalidCharacter", "The selected contains an invalid character ({0})."), InvalidCharacterText);
 
-			ShowErrorNotification(ErrorText);
+		ShowErrorNotification(ErrorText);
 		
-			return false;
-		}
+		return false;
 	}
 
 	return true;
@@ -110,10 +140,9 @@ bool ValidateProvidedAssetPath(const FString& FileName, const FString& AssetPath
 
 	// Ensure we are not overriding the parent CO
 	const FString FullAssetPath = AssetPath + FString("/") + FileName + FString(".") + FileName;		// Full asset path to the new asset we want to create
-	const bool bWouldOverrideParentCO = InstanceCO->GetPathName() == FullAssetPath;
-	if (bWouldOverrideParentCO)
+	if (const bool bWouldOverrideParentCO = InstanceCO->GetPathName() == FullAssetPath)
 	{
-		const FText ErrorText = LOCTEXT("CustomizableObjecInstanceBakingUtils_OverwriteCO", "The selected path would overwrite the instance's parent Customizable Object.");
+		const FText ErrorText = LOCTEXT("CustomizableObjectInstanceBakingUtils_OverwriteCO", "The selected path would overwrite the instance's parent Customizable Object.");
 
 		ShowErrorNotification(ErrorText);
 		
@@ -124,26 +153,19 @@ bool ValidateProvidedAssetPath(const FString& FileName, const FString& AssetPath
 }
 
 
+
 /**
  * Outputs a string that we know it is unique.
- * @param InResource The resource we are working with
  * @param ResourceName The name of the resource we have provided. This should have the name of the current resource and will have the unique name for the resource once the method exits
- * @param InCachedResources Collection with all the already processed resources.
  * @param InCachedResourceNames Collection with all the already processed resources name's
- * @return True if the generation of the unique resource name was successful, false otherwise.
  */
-bool GetUniqueResourceName(const UObject* InResource, FString& ResourceName, TArray<UObject*>& InCachedResources, const TArray<FString>& InCachedResourceNames)
+void MakeResourceNameUnique(FString& ResourceName, const TArray<FString>& InCachedResourceNames)
 {
-	check(InResource);
-
+	// Look for the resource name provided to see if we have already worked with it.
 	int32 FindResult = InCachedResourceNames.Find(ResourceName);
 	if (FindResult != INDEX_NONE)
 	{
-		if (InResource == InCachedResources[FindResult])
-		{
-			return false;
-		}
-
+		// Add an integer suffix to create the unique name
 		uint32 Count = 0;
 		while (FindResult != INDEX_NONE)
 		{
@@ -153,10 +175,9 @@ bool GetUniqueResourceName(const UObject* InResource, FString& ResourceName, TAr
 
 		ResourceName += "_" + FString::FromInt(--Count);
 	}
-
-	return true;
 }
 
+
 /**
  * Ensures the resource we want to save is ready to be saved. It handles closing it's editor and warning the user about possible overriding of resources.
  * @param InAssetSavePath The directory path where to save the baked object
@@ -166,9 +187,9 @@ bool GetUniqueResourceName(const UObject* InResource, FString& ResourceName, TAr
  * @param OutSaveResolution
  * @return True if the operation was successful, false otherwise
  */
-bool ManageBakingAction(const FString& InAssetSavePath, const FString& InObjName, bool& bOverridePermissionGranted, const bool& bIsUnattended, EPackageSaveResolutionType& OutSaveResolution)
+bool CreateAssetAtPath(const FString& InAssetSavePath, const FString& InObjName, bool& bOverridePermissionGranted, const bool& bIsUnattended, EPackageSaveResolutionType& OutSaveResolution)
 {
-	// Before the value provided by "bOverridePermissionGranted" was being updated due to user request but not it is not. It will stay as is if unatended an
+	// Before the value provided by "bOverridePermissionGranted" was being updated due to user request but not it is not. It will stay as is if unattended an
 	// will get updated if this gets to be an attended execution.
 	
 	const FString PackagePath = InAssetSavePath + "/" + InObjName;
@@ -209,8 +230,8 @@ bool ManageBakingAction(const FString& InAssetSavePath, const FString& InObjName
 				// Close the editors that contains this asset
 				if (!ObjectEditorInstance->CloseWindow(EAssetEditorCloseReason::AssetEditorHostClosed))
 				{
-					const FText Caption = LOCTEXT("CustomizableObjecInstanceBakingUtils_OpenExisitngFile", "Open File");
-					const FText Message = FText::Format(LOCTEXT("CustomizableObjecInstanceBakingUtils_CantCloseAsset", "This Obejct \"{0}\" is open in an editor and can't be closed automatically. Please close the editor and try to bake it again"), FText::FromString(InObjName));
+					const FText Caption = LOCTEXT("CustomizableObjectInstanceBakingUtils_OpenExistingFile", "Open File");
+					const FText Message = FText::Format(LOCTEXT("CustomizableObjectInstanceBakingUtils_CantCloseAsset", "This Object \"{0}\" is open in an editor and can't be closed automatically. Please close the editor and try to bake it again"), FText::FromString(InObjName));
 
 					ShowErrorNotification(Message, Caption);
 					
@@ -219,12 +240,12 @@ bool ManageBakingAction(const FString& InAssetSavePath, const FString& InObjName
 			}
 		}
 
-		// If the execution requires user interaction and we have no permission to override the existing file ask him if he wants or not to override data
+		// If the execution requires user interaction, and we have no permission to override the existing file ask him if he wants or not to override data
 		if (!bIsUnattended && !bOverridePermissionGranted)
 		{
 			check (!FApp::IsUnattended())
-			const FText Caption = LOCTEXT("CustomizableObjecInstanceBakingUtils_AlreadyExistingBakedFiles", "Already existing baked files");
-			const FText Message = FText::Format(LOCTEXT("CustomizableObjecInstanceBakingUtils_OverwriteBakedInstance", "Instance baked files already exist in selected destination \"{0}\", this action will overwrite them."), FText::AsCultureInvariant(InAssetSavePath));
+			const FText Caption = LOCTEXT("CustomizableObjectInstanceBakingUtils_AlreadyExistingBakedFiles", "Already existing baked files");
+			const FText Message = FText::Format(LOCTEXT("CustomizableObjectInstanceBakingUtils_OverwriteBakedInstance", "Instance baked files already exist in selected destination \"{0}\", this action will overwrite them."), FText::AsCultureInvariant(InAssetSavePath));
 
 			if (FMessageDialog::Open(EAppMsgType::OkCancel, Message, Caption) == EAppReturnType::Cancel)
 			{
@@ -240,8 +261,7 @@ bool ManageBakingAction(const FString& InAssetSavePath, const FString& InObjName
 		// At this point we may or may not have permission to delete the existing asset
 		
 		// Delete the old asset if we have permission to do so
-		UObject* ExistingObject = StaticFindObject(UObject::StaticClass(), ExistingPackage, *InObjName);
-		if (ExistingObject)
+		if (UObject* ExistingObject = StaticFindObject(UObject::StaticClass(), ExistingPackage, *InObjName))
 		{
 			// Based on if we have or not permission to override the file do or do not so
 			if (bOverridePermissionGranted)
@@ -251,7 +271,7 @@ bool ManageBakingAction(const FString& InAssetSavePath, const FString& InObjName
 				TArray<UObject*> ObjectsToDelete;
 				ObjectsToDelete.Add(ExistingObject);
 				
-				const FText Message = FText::Format(LOCTEXT("CustomizableObjecInstanceBakingUtils_AssetOverriden", "The COI asset \"{0}\" already exists and will be overriden due to user demand."), FText::FromString(ExistingPackage->GetName()));
+				const FText Message = FText::Format(LOCTEXT("CustomizableObjectInstanceBakingUtils_AssetOverriden", "The COI asset \"{0}\" already exists and will be overriden due to user demand."), FText::FromString(ExistingPackage->GetName()));
 				UE_LOG(LogMutable,Warning,TEXT("%s"), *Message.ToString());
 
 				// Notify the caller we did proceed with the override (performed later)
@@ -267,7 +287,7 @@ bool ManageBakingAction(const FString& InAssetSavePath, const FString& InObjName
 				OutSaveResolution = EPackageSaveResolutionType::UnableToOverride;
 				
 				// Report that the file will not get overriden since we have no permission to do so
-				const FText UnableToOverrideMessage = FText::Format(LOCTEXT("CustomizableObjecInstanceBakingUtils_AssetCanNotBeOverriden", "Could not replace the COI asset \"{0}\" as it already exists."), FText::FromString(ExistingPackage->GetName()));
+				const FText UnableToOverrideMessage = FText::Format(LOCTEXT("CustomizableObjectInstanceBakingUtils_AssetCanNotBeOverriden", "Could not replace the COI asset \"{0}\" as it already exists."), FText::FromString(ExistingPackage->GetName()));
 				UE_LOG(LogMutable,Error,TEXT("%s"), *UnableToOverrideMessage.ToString());
 
 				return false;
@@ -286,8 +306,16 @@ bool ManageBakingAction(const FString& InAssetSavePath, const FString& InObjName
 static bool bIsBakeOperationAlreadyScheduled = false;
 
 
-void ResetBakeOperationAlreadyRunning(const FUpdateContext& Result)
+void OnInstanceUpdateFinish(const FUpdateContext& Result)
 {
+	FCustomizableObjectEditorLogger::CreateLog(
+	LOCTEXT("CustomizableObjectInstanceBakingUtils_UpdateFinished", "The COInstance Update operation for baking has finished."))
+	.Category(ELoggerCategory::COInstanceBaking)
+	.CustomNotification()
+	.Notification(true)
+	.Log();
+
+	// Allow the baking of more instances once the bake of this one has completed (or at least until the callbacks have been broadcast)
 	bIsBakeOperationAlreadyScheduled = false;
 }
 
@@ -302,8 +330,7 @@ void ScheduleInstanceUpdateForBaking(UCustomizableObjectInstance& InInstance, FI
 		return;
 	}
 
-	// Allow the baking of more instances once the bake of this one has completed (or at least until the callbacks have been broadcast)
-	InInstanceUpdateDelegate.AddStatic(&ResetBakeOperationAlreadyRunning);
+	InInstanceUpdateDelegate.AddStatic(&OnInstanceUpdateFinish);
 	
 	// Desired System settings for the update part of a baking operation
 	TSharedPtr<FMutableSystemSettingsOverrides> SystemSettingsOverride =
@@ -318,7 +345,13 @@ void ScheduleInstanceUpdateForBaking(UCustomizableObjectInstance& InInstance, FI
 	
 	// Tell the baking system that an instance for baking will be processed
 	bIsBakeOperationAlreadyScheduled = true;
-	UE_LOG(LogMutable, Display, TEXT("The COInstance Update operation for baking has been scheduled for the instance '%s'."), *InInstance.GetName());
+
+	FCustomizableObjectEditorLogger::CreateLog(
+	LOCTEXT("CustomizableObjectInstanceBakingUtils_UpdateScheduled", "The COInstance Update operation for baking has been scheduled. Please hold."))
+	.Category(ELoggerCategory::COInstanceBaking)
+	.CustomNotification()
+	.Notification(true)
+	.Log();
 }
 
 
@@ -410,38 +443,48 @@ bool BakeCustomizableObjectInstance(
 	{
 		bIsUnattendedExecution = true;
 	}
-	
+
+	// Prefixes used for UMaterial assets. 
+	const TArray<FString> MaterialResourcePrefixes = {MaterialAssetPrefix, MaterialInstanceAssetPrefix, MaterialDynamicInstanceAssetPrefix, MaterialConstantInstanceAssetPrefix};
+
+	// Array with the already processed resource names and resources (UObjects)
+	TArray<FString> CachedResourceNames;
+	TArray<UObject*> CachedObjects;
 	
 	const int32 NumObjectComponents = InstanceCO->GetComponentCount();
 	for (int32 ObjectComponentIndex = 0; ObjectComponentIndex < NumObjectComponents; ++ObjectComponentIndex)
 	{
-		FName ComponentName = InstanceCO->GetPrivate()->GetComponentName(FCustomizableObjectComponentIndex(ObjectComponentIndex));
+		const FName ComponentName = InstanceCO->GetPrivate()->GetComponentName(FCustomizableObjectComponentIndex(ObjectComponentIndex));
 		USkeletalMesh* Mesh = InInstance.GetComponentMeshSkeletalMesh(ComponentName);
 
 		if (!Mesh)
 		{
 			continue;
 		}
-
-		FString ObjectName = FileName;
-		if (NumObjectComponents > 1)
+		
+		FString ObjectBaseName = FileName;
+		
+		if (FileName.IsEmpty())
+		{
+			ObjectBaseName = ComponentName.ToString();
+		}
+		else
 		{
-			ObjectName = FileName + "_Component_" + ComponentName.ToString();
+			ObjectBaseName = FileName + TEXT("_") + ComponentName.ToString();
 		}
 
+		ObjectBaseName += TEXT("_");
+		
 		TMap<UObject*, UObject*> ReplacementMap;
-		TArray<FString> ArrayCachedElement;
-		TArray<UObject*> ArrayCachedObject;
 
 		if (bExportAllResources)
 		{
-			UMaterialInstance* Inst;
+			UMaterialInstance* MaterialInstance;
 			UMaterial* Material;
 			UTexture* Texture;
-			FString MaterialName;
-			FString ResourceName;
 			FString PackageName;
 			UObject* DuplicatedObject;
+			
 			TArray<TMap<int, UTexture*>> TextureReplacementMaps;
 
 			// Duplicate Textures found in the Material Instances of the SkeletalMesh so we can later assign them to the
@@ -451,79 +494,75 @@ bool BakeCustomizableObjectInstance(
 			{
 				UMaterialInterface* Interface = Mesh->GetMaterials()[m].MaterialInterface;
 				Material = Interface->GetMaterial();
-				MaterialName = Material ? Material->GetName() : "Material";
-				Inst = Cast<UMaterialInstance>(Interface);
+				MaterialInstance = Cast<UMaterialInstance>(Interface);
 
 				TextureReplacementMaps.AddDefaulted();
 				
-				if (Material != nullptr && Inst != nullptr)
+				if (Material != nullptr && MaterialInstance != nullptr)
 				{
-					TArray<FGuid> Guids;
-					TArray<FMaterialParameterInfo> OutParameterInfo;
-					Material->GetAllTextureParameterInfo(OutParameterInfo, Guids);
+					TArray<FGuid> ParameterIds;
+					TArray<FMaterialParameterInfo> ParameterInfoObjects;
+					Material->GetAllTextureParameterInfo(ParameterInfoObjects, ParameterIds);
 					
-					for (int32 Index = 0; Index < OutParameterInfo.Num(); Index++)
+					for (int32 ParameterInfoIndex = 0; ParameterInfoIndex < ParameterInfoObjects.Num(); ParameterInfoIndex++)
 					{
-						const FName& ParameterName = OutParameterInfo[Index].Name;
-						
-						if (Inst->GetTextureParameterValue(ParameterName, Texture))
+						const FName& ParameterName = ParameterInfoObjects[ParameterInfoIndex].Name;
+						if (MaterialInstance->GetTextureParameterValue(ParameterName, Texture))
 						{
-							UTexture2D* SrcTex = Cast<UTexture2D>(Texture);
-							if (!SrcTex)
+							UTexture2D* SourceTexture = Cast<UTexture2D>(Texture);
+							if (!SourceTexture)
 							{
 								continue;
 							}
-
-							FString ParameterSanitized = ParameterName.GetPlainNameString();
-							RemoveRestrictedChars(ParameterSanitized);
-							ResourceName = ObjectName + "_" + MaterialName + "_" + ParameterSanitized;
-							if (!GetUniqueResourceName(SrcTex, ResourceName, ArrayCachedObject, ArrayCachedElement))
+							
+							// Check that the resource has not yet been processed, and if not, skip it
+							if (CachedObjects.Contains(SourceTexture))
 							{
 								continue;
 							}
 							
+							FString TextureName = SourceTexture->GetName();
+							RemovePrefix(TextureName, {TextureAssetPrefix});
+							ComposeResourceName(TextureAssetPrefix, ObjectBaseName, TextureName);
+							MakeResourceNameUnique(TextureName, CachedResourceNames);
+							
 							EPackageSaveResolutionType SaveType = EPackageSaveResolutionType::None;
-							if (!ManageBakingAction(AssetPath, ResourceName, bUsedGrantedOverridingRights, bIsUnattendedExecution, SaveType))
+							if (!CreateAssetAtPath(AssetPath, TextureName, bUsedGrantedOverridingRights, bIsUnattendedExecution, SaveType))
 							{
 								return false;
 							}
-
-							// Skip already processed resource
-							if (ArrayCachedElement.Find(ResourceName) != INDEX_NONE)
-							{
-								continue;
-							}
-
+							
 							bool bIsMutableTexture = false;
-							for (UAssetUserData* UserData : *SrcTex->GetAssetUserDataArray())
+							for (UAssetUserData* UserData : *SourceTexture->GetAssetUserDataArray())
 							{
 								if (Cast<UMutableTextureMipDataProviderFactory>(UserData))
 								{
 									bIsMutableTexture = true;
+									break;
 								}
 							}
 							
 							// Duplicating mutable generated textures
 							if (bIsMutableTexture)
 							{
-								if (SrcTex->GetPlatformData() && SrcTex->GetPlatformData()->Mips.Num() > 0)
+								if (SourceTexture->GetPlatformData() && SourceTexture->GetPlatformData()->Mips.Num() > 0)
 								{
 									// Recover original name of the texture parameter value, now substituted by the generated Mutable texture
 									UTexture* OriginalTexture = nullptr;
-									Inst->Parent->GetTextureParameterValue(FName(*ParameterName.GetPlainNameString()), OriginalTexture);
+									MaterialInstance->Parent->GetTextureParameterValue(FName(*ParameterName.GetPlainNameString()), OriginalTexture);
 
-									PackageName = AssetPath + FString("/") + ResourceName;
-									TMap<UObject*, UObject*> FakeReplacementMap;
-									UTexture2D* DupTex = FUnrealBakeHelpers::BakeHelper_CreateAssetTexture(SrcTex, ResourceName, PackageName, OriginalTexture, true, FakeReplacementMap, bUsedGrantedOverridingRights);
-									ArrayCachedElement.Add(ResourceName);
-									ArrayCachedObject.Add(DupTex);
+									PackageName = AssetPath + FString("/") + TextureName;
+									TMap<UObject*, UObject*> DummyReplacementMap;
+									UTexture2D* DuplicatedTexture = FUnrealBakeHelpers::BakeHelper_CreateAssetTexture(SourceTexture, TextureName, PackageName, OriginalTexture, true, DummyReplacementMap, bUsedGrantedOverridingRights);
+									CachedResourceNames.Add(TextureName);
+									CachedObjects.Add(DuplicatedTexture);
 								
-									TPair<EPackageSaveResolutionType, UPackage*> PackageToSave {SaveType, DupTex->GetPackage()};
+									TPair<EPackageSaveResolutionType, UPackage*> PackageToSave {SaveType, DuplicatedTexture->GetPackage()};
 									OutSavedPackages.Add(PackageToSave);
 
 									if (OriginalTexture != nullptr)
 									{
-										TextureReplacementMaps[m].Add(Index, DupTex);
+										TextureReplacementMaps[m].Add(ParameterInfoIndex, DuplicatedTexture);
 									}
 								}
 							}
@@ -531,17 +570,17 @@ bool BakeCustomizableObjectInstance(
 							{
 								// Duplicate the non-mutable textures of the Material instance (pass-through textures)
 								
-								PackageName = AssetPath + FString("/") + ResourceName;
-								TMap<UObject*, UObject*> FakeReplacementMap;
-								DuplicatedObject = FUnrealBakeHelpers::BakeHelper_DuplicateAsset(Texture, ResourceName, PackageName, true, FakeReplacementMap, bUsedGrantedOverridingRights, false);
-								ArrayCachedElement.Add(ResourceName);
-								ArrayCachedObject.Add(DuplicatedObject);
+								PackageName = AssetPath + FString("/") + TextureName;
+								TMap<UObject*, UObject*> DummyReplacementMap;
+								DuplicatedObject = FUnrealBakeHelpers::BakeHelper_DuplicateAsset(Texture, TextureName, PackageName, true, DummyReplacementMap, bUsedGrantedOverridingRights, false);
+								CachedResourceNames.Add(TextureName);
+								CachedObjects.Add(DuplicatedObject);
 
 								TPair<EPackageSaveResolutionType, UPackage*> PackageToSave {SaveType, DuplicatedObject->GetPackage()};
 								OutSavedPackages.Add(PackageToSave);
 
 								UTexture* DupTexture = Cast<UTexture>(DuplicatedObject);
-								TextureReplacementMaps[m].Add(Index, DupTexture);
+								TextureReplacementMaps[m].Add(ParameterInfoIndex, DupTexture);
 							}
 						}
 					}
@@ -551,32 +590,36 @@ bool BakeCustomizableObjectInstance(
 
 			// Duplicate the materials used by each material instance so that the replacement map has proper information 
 			// when duplicating the material instances
-			for (int32 m = 0; m < Mesh->GetMaterials().Num(); ++m)
+			for (int32 MaterialIndex = 0; MaterialIndex < Mesh->GetMaterials().Num(); ++MaterialIndex)
 			{
-				UMaterialInterface* Interface = Mesh->GetMaterials()[m].MaterialInterface;
+				UMaterialInterface* Interface = Mesh->GetMaterials()[MaterialIndex].MaterialInterface;
 				Material = Interface ? Interface->GetMaterial() : nullptr;
-
+				
 				if (Material)
 				{
-					ResourceName = ObjectName + "_Material_" + Material->GetName();
-
-					if (!GetUniqueResourceName(Material, ResourceName, ArrayCachedObject, ArrayCachedElement))
+					// Check that the resource has not yet been processed, and if not, skip it
+					if (CachedObjects.Contains(Material))
 					{
 						continue;
 					}
 
+					FString MaterialName = Material->GetName();
+					RemovePrefix(MaterialName, {MaterialAssetPrefix});
+					ComposeResourceName(MaterialAssetPrefix, ObjectBaseName, MaterialName);
+					MakeResourceNameUnique(MaterialName, CachedResourceNames);
+
 					EPackageSaveResolutionType SaveType = EPackageSaveResolutionType::None;
-					if (!ManageBakingAction(AssetPath, ResourceName, bUsedGrantedOverridingRights, bIsUnattendedExecution, SaveType))
+					if (!CreateAssetAtPath(AssetPath, MaterialName, bUsedGrantedOverridingRights, bIsUnattendedExecution, SaveType))
 					{
 						return false;
 					}
 
-					PackageName = AssetPath + FString("/") + ResourceName;
+					PackageName = AssetPath + FString("/") + MaterialName;
 					TMap<UObject*, UObject*> FakeReplacementMap;
-					DuplicatedObject = FUnrealBakeHelpers::BakeHelper_DuplicateAsset(Material, ResourceName, PackageName, 
+					DuplicatedObject = FUnrealBakeHelpers::BakeHelper_DuplicateAsset(Material, MaterialName, PackageName, 
 						false, FakeReplacementMap, bUsedGrantedOverridingRights, bGenerateConstantMaterialInstances);
-					ArrayCachedElement.Add(ResourceName);
-					ArrayCachedObject.Add(DuplicatedObject);
+					CachedResourceNames.Add(MaterialName);
+					CachedObjects.Add(DuplicatedObject);
 					ReplacementMap.Add(Interface, DuplicatedObject);
 
 					TPair<EPackageSaveResolutionType, UPackage*> PackageToSave {SaveType, DuplicatedObject->GetPackage()};
@@ -584,7 +627,7 @@ bool BakeCustomizableObjectInstance(
 
 					if (UMaterial* DuplicatedMaterial = Cast<UMaterial>(DuplicatedObject))
 					{
-						FUnrealBakeHelpers::CopyAllMaterialParameters<UMaterial>(*DuplicatedMaterial, *Interface, TextureReplacementMaps[m]);
+						FUnrealBakeHelpers::CopyAllMaterialParameters<UMaterial>(*DuplicatedMaterial, *Interface, TextureReplacementMaps[MaterialIndex]);
 					}
 				}
 			}
@@ -596,37 +639,66 @@ bool BakeCustomizableObjectInstance(
 			{
 				UMaterialInterface* Interface = Mesh->GetMaterials()[MaterialIndex].MaterialInterface;
 				UMaterial* ParentMaterial = Interface->GetMaterial();
-				FString MaterialName = ParentMaterial ? ParentMaterial->GetName() : "Material";
-
-				// Material
-				FString MatObjName = ObjectName + "_" + MaterialName;
 
-				if (!GetUniqueResourceName(Interface, MatObjName, ArrayCachedObject, ArrayCachedElement))
+				// Check that the resource has not yet been processed, and if not, skip it
+				if (CachedObjects.Contains(Interface))
 				{
 					continue;
 				}
 
+				FString MaterialName;
+				{
+					MaterialName = ParentMaterial ? ParentMaterial->GetName() : MaterialAssetPrefix;
+
+					RemovePrefix(MaterialName, MaterialResourcePrefixes);
+					
+					if (bGenerateConstantMaterialInstances && Interface->IsA<UMaterialInstance>())
+					{
+						// Change the prefix to the Material Constant Instance since in this situation the new asset based on the Interface
+						// will be of the Material Constant Instance type
+						ComposeResourceName(MaterialConstantInstanceAssetPrefix, ObjectBaseName, MaterialName);
+					}
+					else
+					{
+						if (Interface->IsA<UMaterial>())
+						{
+							ComposeResourceName(MaterialAssetPrefix, ObjectBaseName, MaterialName);
+						}
+						else if (Interface->IsA<UMaterialInstanceConstant>())
+						{
+							ComposeResourceName(MaterialConstantInstanceAssetPrefix, ObjectBaseName, MaterialName);
+						}
+						else if (Interface->IsA<UMaterialInstanceDynamic>())
+						{
+							ComposeResourceName(MaterialDynamicInstanceAssetPrefix, ObjectBaseName, MaterialName);
+						}
+						else
+						{
+							checkNoEntry();		// Invalid material type.
+						}
+					}
+					
+					MakeResourceNameUnique(MaterialName, CachedResourceNames);
+				}
+
 				EPackageSaveResolutionType SaveType = EPackageSaveResolutionType::None;
-				if (!ManageBakingAction(AssetPath, MatObjName, bUsedGrantedOverridingRights, bIsUnattendedExecution, SaveType))
+				if (!CreateAssetAtPath(AssetPath, MaterialName, bUsedGrantedOverridingRights, bIsUnattendedExecution, SaveType))
 				{
 					return false;
 				}
 
-				FString MatPkgName = AssetPath + FString("/") + MatObjName;
-				UObject* DupMat = FUnrealBakeHelpers::BakeHelper_DuplicateAsset(Interface, MatObjName, 
+				FString MatPkgName = AssetPath + FString("/") + MaterialName;
+				UObject* DupMat = FUnrealBakeHelpers::BakeHelper_DuplicateAsset(Interface, MaterialName, 
 					MatPkgName, false, ReplacementMap, bUsedGrantedOverridingRights, bGenerateConstantMaterialInstances);
-				ArrayCachedObject.Add(DupMat);
-				ArrayCachedElement.Add(MatObjName);
+				CachedObjects.Add(DupMat);
+				CachedResourceNames.Add(MaterialName);
 
 				TPair<EPackageSaveResolutionType, UPackage*> PackageToSave {SaveType, DupMat->GetPackage()};
 				OutSavedPackages.Add(PackageToSave);
-		
-
-				UMaterialInstance* Inst = Cast<UMaterialInstance>(Interface);
 
 				// Only need to duplicate the generate textures if the original material is a dynamic instance
 				// If the material has Mutable textures, then it will be a dynamic material instance for sure
-				if (Inst)
+				if (UMaterialInstance* MaterialInstance = Cast<UMaterialInstance>(Interface))
 				{
 					// Duplicate generated textures
 					UMaterialInstanceDynamic* InstDynamic = Cast<UMaterialInstanceDynamic>(DupMat);
@@ -634,59 +706,59 @@ bool BakeCustomizableObjectInstance(
 
 					if (InstDynamic || InstConstant)
 					{
-						for (int32 TextureIndex = 0; TextureIndex < Inst->TextureParameterValues.Num(); ++TextureIndex)
+						for (int32 TextureIndex = 0; TextureIndex < MaterialInstance->TextureParameterValues.Num(); ++TextureIndex)
 						{
-							if (Inst->TextureParameterValues[TextureIndex].ParameterValue)
+							if (MaterialInstance->TextureParameterValues[TextureIndex].ParameterValue)
 							{
-								if (Inst->TextureParameterValues[TextureIndex].ParameterValue->HasAnyFlags(RF_Transient))
+								if (MaterialInstance->TextureParameterValues[TextureIndex].ParameterValue->HasAnyFlags(RF_Transient))
 								{
-									UTexture2D* SrcTex = Cast<UTexture2D>(Inst->TextureParameterValues[TextureIndex].ParameterValue);
-
-									if (SrcTex)
+									if (UTexture2D* SourceTexture = Cast<UTexture2D>(MaterialInstance->TextureParameterValues[TextureIndex].ParameterValue))
 									{
-										FString ParameterSanitized = Inst->TextureParameterValues[TextureIndex].ParameterInfo.Name.ToString();
-										RemoveRestrictedChars(ParameterSanitized);
-
-										FString TexObjName = ObjectName + "_" + MaterialName + "_" + ParameterSanitized;
-
-										if (!GetUniqueResourceName(SrcTex, TexObjName, ArrayCachedObject, ArrayCachedElement))
+										if (CachedObjects.Contains(SourceTexture))
 										{
-											UTexture* PrevTexture = Cast<UTexture>(ArrayCachedObject[ArrayCachedElement.Find(TexObjName)]);
-
+											UTexture* PrevTexture = Cast<UTexture>(CachedObjects[CachedResourceNames.Find(SourceTexture->GetName())]);
+											check(PrevTexture);
+											
 											if (InstDynamic)
 											{
-												InstDynamic->SetTextureParameterValue(Inst->TextureParameterValues[TextureIndex].ParameterInfo.Name, PrevTexture);
+												InstDynamic->SetTextureParameterValue(MaterialInstance->TextureParameterValues[TextureIndex].ParameterInfo.Name, PrevTexture);
 											}
 											else if (InstConstant)
 											{
-												InstConstant->SetTextureParameterValueEditorOnly(Inst->TextureParameterValues[TextureIndex].ParameterInfo.Name, PrevTexture);
+												InstConstant->SetTextureParameterValueEditorOnly(MaterialInstance->TextureParameterValues[TextureIndex].ParameterInfo.Name, PrevTexture);
 											}
 											
 											continue;
 										}
 
+										// Generate a new name for the duplicated texture so it does not collide with the original one
+										FString TextureName = SourceTexture->GetName();
+										RemovePrefix(TextureName, {TextureAssetPrefix});
+										ComposeResourceName(TextureAssetPrefix, ObjectBaseName, TextureName);
+										MakeResourceNameUnique(TextureName, CachedResourceNames);
+										
 										EPackageSaveResolutionType TextureSaveType = EPackageSaveResolutionType::None;
-										if (!ManageBakingAction(AssetPath, TexObjName, bUsedGrantedOverridingRights, bIsUnattendedExecution, TextureSaveType))
+										if (!CreateAssetAtPath(AssetPath, TextureName, bUsedGrantedOverridingRights, bIsUnattendedExecution, TextureSaveType))
 										{
 											return false;
 										}
 
-										FString TexPkgName = AssetPath + FString("/") + TexObjName;
+										FString TexPkgName = AssetPath + FString("/") + TextureName;
 										TMap<UObject*, UObject*> FakeReplacementMap;
-										UTexture2D* DupTex = FUnrealBakeHelpers::BakeHelper_CreateAssetTexture(SrcTex, TexObjName, TexPkgName, nullptr, false, FakeReplacementMap, bUsedGrantedOverridingRights);
-										ArrayCachedObject.Add(DupTex);
-										ArrayCachedElement.Add(TexObjName);
+										UTexture2D* DupTex = FUnrealBakeHelpers::BakeHelper_CreateAssetTexture(SourceTexture, TextureName, TexPkgName, nullptr, false, FakeReplacementMap, bUsedGrantedOverridingRights);
+										CachedObjects.Add(DupTex);
+										CachedResourceNames.Add(TextureName);
 
 										TPair<EPackageSaveResolutionType, UPackage*> TexturePackageToSave {TextureSaveType, DupTex->GetPackage()};
 										OutSavedPackages.Add(TexturePackageToSave);
 										
 										if (InstDynamic)
 										{
-											InstDynamic->SetTextureParameterValue(Inst->TextureParameterValues[TextureIndex].ParameterInfo.Name, DupTex);
+											InstDynamic->SetTextureParameterValue(MaterialInstance->TextureParameterValues[TextureIndex].ParameterInfo.Name, DupTex);
 										}
 										else if(InstConstant)
 										{
-											InstConstant->SetTextureParameterValueEditorOnly(Inst->TextureParameterValues[TextureIndex].ParameterInfo.Name, DupTex);
+											InstConstant->SetTextureParameterValueEditorOnly(MaterialInstance->TextureParameterValues[TextureIndex].ParameterInfo.Name, DupTex);
 										}
 									}
 									else
@@ -700,11 +772,11 @@ bool BakeCustomizableObjectInstance(
 									// Just set the original texture
 									if (InstDynamic)
 									{
-										InstDynamic->SetTextureParameterValue(Inst->TextureParameterValues[TextureIndex].ParameterInfo.Name, Inst->TextureParameterValues[TextureIndex].ParameterValue);
+										InstDynamic->SetTextureParameterValue(MaterialInstance->TextureParameterValues[TextureIndex].ParameterInfo.Name, MaterialInstance->TextureParameterValues[TextureIndex].ParameterValue);
 									}
 									else if (InstConstant)
 									{
-										InstConstant->SetTextureParameterValueEditorOnly(Inst->TextureParameterValues[TextureIndex].ParameterInfo.Name, Inst->TextureParameterValues[TextureIndex].ParameterValue);
+										InstConstant->SetTextureParameterValueEditorOnly(MaterialInstance->TextureParameterValues[TextureIndex].ParameterInfo.Name, MaterialInstance->TextureParameterValues[TextureIndex].ParameterValue);
 									}
 								}
 							}
@@ -713,52 +785,66 @@ bool BakeCustomizableObjectInstance(
 				}
 			}
 		}
+
+		// Get the clean name of the Skeletal mesh (without the embedded CO name)
+		FString SkeletalMeshName = Mesh->GetName();
+		if (Mesh->HasAnyFlags(RF_Transient))
+		{
+			const FString CustomizableObjectName = InInstance.GetCustomizableObject()->GetName() + TEXT("_");
+			SkeletalMeshName = SkeletalMeshName.Replace(*CustomizableObjectName,TEXT(""));
+		}
+		RemovePrefix(SkeletalMeshName, {SkeletalMeshAssetPrefix});
+		ComposeResourceName(SkeletalMeshAssetPrefix, ObjectBaseName, SkeletalMeshName);
+		MakeResourceNameUnique(SkeletalMeshName, CachedResourceNames);
 		
 		// Skeletal Mesh's Skeleton
-		if (Mesh->GetSkeleton())
+		if (USkeleton* Skeleton = Mesh->GetSkeleton())
 		{
-			const bool bTransient = Mesh->GetSkeleton()->GetPackage() == GetTransientPackage();
+			const bool bTransient = Skeleton->GetPackage() == GetTransientPackage();
 
 			// Don't duplicate if not transient or export all assets.
 			if (bTransient || bExportAllResources)
 			{
-				FString SkeletonName = ObjectName + "_Skeleton";
+				FString SkeletonName = Mesh->GetName();
+				RemovePrefix(SkeletonName, {SkeletalMeshAssetPrefix});
+				ComposeResourceName(SkeletonAssetPrefix, ObjectBaseName, SkeletonName);
+				MakeResourceNameUnique(SkeletonName, CachedResourceNames);
+				
 				EPackageSaveResolutionType SaveType = EPackageSaveResolutionType::None;
-				if (!ManageBakingAction(AssetPath, SkeletonName, bUsedGrantedOverridingRights, bIsUnattendedExecution, SaveType))
+				if (!CreateAssetAtPath(AssetPath, SkeletonName, bUsedGrantedOverridingRights, bIsUnattendedExecution, SaveType))
 				{
 					return false;
 				}
 
 				FString SkeletonPkgName = AssetPath + FString("/") + SkeletonName;
-				UObject* DuplicatedSkeleton = FUnrealBakeHelpers::BakeHelper_DuplicateAsset(Mesh->GetSkeleton(), SkeletonName, 
+				UObject* DuplicatedSkeleton = FUnrealBakeHelpers::BakeHelper_DuplicateAsset(Skeleton, SkeletonName, 
 					SkeletonPkgName, false, ReplacementMap, bUsedGrantedOverridingRights, false);
 
-				ArrayCachedObject.Add(DuplicatedSkeleton);
+				CachedObjects.Add(DuplicatedSkeleton);
 				TPair<EPackageSaveResolutionType, UPackage*> PackageToSave {SaveType, DuplicatedSkeleton->GetPackage()};
 				OutSavedPackages.Add(PackageToSave);
-				ReplacementMap.Add(Mesh->GetSkeleton(), DuplicatedSkeleton);
+				ReplacementMap.Add(Skeleton, DuplicatedSkeleton);
 			}
 		}
 
 		// Skeletal Mesh
 		EPackageSaveResolutionType SaveType = EPackageSaveResolutionType::None;
-		if (!ManageBakingAction(AssetPath, ObjectName, bUsedGrantedOverridingRights, bIsUnattendedExecution, SaveType))
+		if (!CreateAssetAtPath(AssetPath, SkeletalMeshName, bUsedGrantedOverridingRights, bIsUnattendedExecution, SaveType))
 		{
 			return false;
 		}
 
-		FString PkgName = AssetPath + FString("/") + ObjectName;
-		UObject* DupObject = FUnrealBakeHelpers::BakeHelper_DuplicateAsset(Mesh, ObjectName, PkgName, 
+		FString PkgName = AssetPath + FString("/") + SkeletalMeshName;
+		UObject* DupObject = FUnrealBakeHelpers::BakeHelper_DuplicateAsset(Mesh, SkeletalMeshName, PkgName, 
 			false, ReplacementMap, bUsedGrantedOverridingRights, false);
-		ArrayCachedObject.Add(DupObject);
+		CachedObjects.Add(DupObject);
 
 		TPair<EPackageSaveResolutionType, UPackage*> PackageToSave {SaveType, DupObject->GetPackage()};
 		OutSavedPackages.Add(PackageToSave);
 
 		Mesh->Build();
 
-		USkeletalMesh* SkeletalMesh = Cast<USkeletalMesh>(DupObject);
-		if (SkeletalMesh)
+		if (USkeletalMesh* SkeletalMesh = Cast<USkeletalMesh>(DupObject))
 		{
 			SkeletalMesh->ResetLODInfo();
 			for (int32 LODIndex = 0; LODIndex < Mesh->GetLODNum(); ++LODIndex)
@@ -808,9 +894,7 @@ bool BakeCustomizableObjectInstance(
 
 			// Duplicate PhysicsAsset
 			{
-				const UPhysicsAsset* PhysicsAsset = Mesh->GetPhysicsAsset();
-
-				if (PhysicsAsset)
+				if (const UPhysicsAsset* PhysicsAsset = Mesh->GetPhysicsAsset())
 				{
 					// Duplicate to change from the Transient Package to the baked mesh one
 					UPhysicsAsset* NewPhysicsAsset = Cast<UPhysicsAsset>(StaticDuplicateObject(PhysicsAsset, SkeletalMesh));
@@ -832,7 +916,7 @@ bool BakeCustomizableObjectInstance(
 		}
 
 		// Remove duplicated UObjects from Root (previously added to avoid objects from being GC in the middle of the bake process)
-		for (UObject* Obj : ArrayCachedObject)
+		for (UObject* Obj : CachedObjects)
 		{
 			Obj->RemoveFromRoot();
 		}
