[mutable] Add support for mutable runtime meshes that have been cooked for mesh streaming.
Expose method for mesh data serialization to use in different modules.
#rnx

diff --git a/Engine/Plugins/Mutable/Source/CustomizableObject/Internal/MuCO/UnrealConversionUtils.h b/Engine/Plugins/Mutable/Source/CustomizableObject/Internal/MuCO/UnrealConversionUtils.h
--- a/Engine/Plugins/Mutable/Source/CustomizableObject/Internal/MuCO/UnrealConversionUtils.h
+++ b/Engine/Plugins/Mutable/Source/CustomizableObject/Internal/MuCO/UnrealConversionUtils.h
@@ -10,6 +10,7 @@
 #include "MuR/Ptr.h"
 #include "Containers/Array.h"
 #include "Containers/Map.h"
+#include "Tasks/Task.h"
 
 struct FMorphTargetMeshData;
 
@@ -22,7 +23,6 @@ namespace mu
 struct FReferenceSkeleton;
 struct FMutableSurfaceMetadata;
 class UModelResources;
-
 class FSkeletalMeshLODRenderData;
 class USkeletalMesh;
 class USkeleton;
@@ -136,6 +136,6 @@ namespace UnrealConversionUtils
 
 	void MorphTargetVertexInfoBuffers(FSkeletalMeshLODRenderData& LODResource, const USkeletalMesh& Owner, const mu::FMesh& MutableMesh, const TMap<uint32, FMorphTargetMeshData>& MorphTargetMeshData, int32 LODIndex);
 
-	CUSTOMIZABLEOBJECT_API void ConvertSkeletalMeshFromRuntimeData(USkeletalMesh& SkeletalMesh, int32 LODIndex, int32 SectionIndex, UModelResources* ModelResources, mu::FMesh& Result);
+	CUSTOMIZABLEOBJECT_API UE::Tasks::FTask ConvertSkeletalMeshFromRuntimeData(USkeletalMesh* SkeletalMesh, int32 LODIndex, int32 SectionIndex, UModelResources* ModelResources, mu::FMesh* Result);
 
 }
diff --git a/Engine/Plugins/Mutable/Source/CustomizableObject/Private/MuCO/UnrealConversionUtils.cpp b/Engine/Plugins/Mutable/Source/CustomizableObject/Private/MuCO/UnrealConversionUtils.cpp
--- a/Engine/Plugins/Mutable/Source/CustomizableObject/Private/MuCO/UnrealConversionUtils.cpp
+++ b/Engine/Plugins/Mutable/Source/CustomizableObject/Private/MuCO/UnrealConversionUtils.cpp
@@ -829,34 +829,40 @@ namespace UnrealConversionUtils
 	}
 
 
-	void ConvertSkeletalMeshFromRuntimeData(USkeletalMesh& SkeletalMesh, int32 LODIndex, int32 SectionIndex, UModelResources* ModelResources, mu::FMesh& Result)
+	UE::Tasks::FTask ConvertSkeletalMeshFromRuntimeData(USkeletalMesh* SkeletalMesh, int32 LODIndex, int32 SectionIndex, UModelResources* ModelResources, mu::FMesh* Result)
 	{
 		MUTABLE_CPUPROFILER_SCOPE(ConvertSkeletalMeshFromRuntimeData);
 
-		FSkeletalMeshRenderData* Data = SkeletalMesh.GetResourceForRendering();
+		if (!SkeletalMesh || !Result)
+		{
+			ensure(false);
+			return UE::Tasks::MakeCompletedTask<void>();
+		}
+
+		FSkeletalMeshRenderData* Data = SkeletalMesh->GetResourceForRendering();
 
 		if (!Data->LODRenderData.IsValidIndex(LODIndex))
 		{
 			// This may happen if we are requesting a LOD that is not present in the provided mesh.
 			// TODO: an empty mesh is returned, but should we use the last LOD instead?
 			UE_LOG(LogMutable, Warning, TEXT("Trying to convert mesh [%s] LOD %d but it only has %d LODs"), 
-				*SkeletalMesh.GetName(), LODIndex, Data->LODRenderData.Num() );
+				*SkeletalMesh->GetName(), LODIndex, Data->LODRenderData.Num() );
 
-			return;
+			return UE::Tasks::MakeCompletedTask<void>();
 		}
 
-		FSkeletalMeshLODRenderData& LOD = Data->LODRenderData[LODIndex];
+		FSkeletalMeshLODRenderData* LOD = &Data->LODRenderData[LODIndex];
 
-		if (!LOD.RenderSections.IsValidIndex(SectionIndex))
+		if (!LOD->RenderSections.IsValidIndex(SectionIndex))
 		{
 			// This may happen in the provided mesh in the parameter value doesn't have the expected number of sections
 			UE_LOG(LogMutable, Warning, TEXT("Trying to convert mesh [%s] section %d in LOD %d but it only has %d sections"),
-				*SkeletalMesh.GetName(), SectionIndex, LODIndex, LOD.RenderSections.Num());
+				*SkeletalMesh->GetName(), SectionIndex, LODIndex, LOD->RenderSections.Num());
 
-			return;
+			return UE::Tasks::MakeCompletedTask<void>();
 		}
 
-		const FSkelMeshRenderSection& Section = LOD.RenderSections[SectionIndex];
+		const FSkelMeshRenderSection& Section = LOD->RenderSections[SectionIndex];
 
 		// Skeleton data
 		bool bBoneMapModified = false;
@@ -866,21 +872,21 @@ namespace UnrealConversionUtils
 		bool bIgnoreSkeleton = false;
 		if (!bIgnoreSkeleton)
 		{
-			USkeleton* Skeleton = SkeletalMesh.GetSkeleton();
+			USkeleton* Skeleton = SkeletalMesh->GetSkeleton();
 			if (!Skeleton)
 			{
 				ensure(false);
-				return;
+				return UE::Tasks::MakeCompletedTask<void>();
 			}
 
 			// Add the skeleton to the list of referenced skeletons and add its index to the mesh
 			if (ModelResources)
 			{
 				const int32 SkeletonID = ModelResources->Skeletons.AddUnique(Skeleton);
-				Result.AddSkeletonID(SkeletonID);
+				Result->AddSkeletonID(SkeletonID);
 			}
 
-			const TArray<uint16>& SourceRequiredBones = LOD.RequiredBones;
+			const TArray<uint16>& SourceRequiredBones = LOD->RequiredBones;
 
 			// Build RequiredBones array
 			TArray<FBoneIndexType> RequiredBones;
@@ -907,10 +913,10 @@ namespace UnrealConversionUtils
 
 			// Create the skeleton, poses, and BoneMap for this mesh
 			TSharedPtr<mu::FSkeleton> MutableSkeleton = MakeShared<mu::FSkeleton>();
-			Result.SetSkeleton(MutableSkeleton);
+			Result->SetSkeleton(MutableSkeleton);
 
 			const int32 NumRequiredBones = RequiredBones.Num();
-			Result.SetBonePoseCount(NumRequiredBones);
+			Result->SetBonePoseCount(NumRequiredBones);
 			MutableSkeleton->SetBoneCount(NumRequiredBones);
 
 			// MutableBoneMap will not keep an index to the Skeleton, but to the BoneName
@@ -920,7 +926,7 @@ namespace UnrealConversionUtils
 			TArray<FMatrix> ComposedRefPoseMatrices;
 			ComposedRefPoseMatrices.SetNum(NumRequiredBones);
 
-			const TArray<FMeshBoneInfo>& RefBoneInfo = SkeletalMesh.GetRefSkeleton().GetRefBoneInfo();
+			const TArray<FMeshBoneInfo>& RefBoneInfo = SkeletalMesh->GetRefSkeleton().GetRefBoneInfo();
 			for (int32 BoneIndex = 0; BoneIndex < NumRequiredBones; ++BoneIndex)
 			{
 				const int32 RefSkeletonBoneIndex = RequiredBones[BoneIndex];
@@ -989,11 +995,11 @@ namespace UnrealConversionUtils
 
 				if (ParentBoneIndex >= 0)
 				{
-					ComposedRefPoseMatrices[BoneIndex] = SkeletalMesh.GetRefPoseMatrix(RefSkeletonBoneIndex) * ComposedRefPoseMatrices[ParentBoneIndex];
+					ComposedRefPoseMatrices[BoneIndex] = SkeletalMesh->GetRefPoseMatrix(RefSkeletonBoneIndex) * ComposedRefPoseMatrices[ParentBoneIndex];
 				}
 				else
 				{
-					ComposedRefPoseMatrices[BoneIndex] = SkeletalMesh.GetRefPoseMatrix(RefSkeletonBoneIndex);
+					ComposedRefPoseMatrices[BoneIndex] = SkeletalMesh->GetRefPoseMatrix(RefSkeletonBoneIndex);
 				}
 
 				// Set bone pose
@@ -1004,194 +1010,265 @@ namespace UnrealConversionUtils
 				EnumAddFlags(BoneUsageFlags, BoneMapIndex != INDEX_NONE ? mu::EBoneUsageFlags::Skinning : mu::EBoneUsageFlags::None);
 				EnumAddFlags(BoneUsageFlags, ParentBoneIndex == INDEX_NONE ? mu::EBoneUsageFlags::Root : mu::EBoneUsageFlags::None);
 
-				Result.SetBonePose(BoneIndex, BoneName, BoneTransform, BoneUsageFlags);
+				Result->SetBonePose(BoneIndex, BoneName, BoneTransform, BoneUsageFlags);
 			}
 
-			Result.SetBoneMap(MutableBoneMap);
+			Result->SetBoneMap(MutableBoneMap);
 		}
 
 
+		const bool bUseStreaming = !LOD->bStreamedDataInlined;
+		
+		UE::Tasks::FTaskEvent LoadedEvent(TEXT("MutableMeshParamLoadCompleted"));
 
-		// Mesh data
-		int32 FirstVertexIndex = LOD.RenderSections[SectionIndex].BaseVertexIndex;
-		int32 VertexCount = Section.GetNumVertices();
+		// This context stores the information that travels between tasks.
+		struct FMeshConversionContext
+		{
+			FSkeletalMeshLODRenderData StreamedLOD;
+			FSkeletalMeshLODRenderData* OriginalLOD = nullptr;
+			FSkeletalMeshLODRenderData* DataLOD = nullptr;
+			TUniquePtr<IBulkDataIORequest> Request;
 
-		mu::FMeshBufferSet& Vertices = Result.GetVertexBuffers();
-		Vertices.SetElementCount(VertexCount, mu::EMemoryInitPolicy::Zeroed);
-		Vertices.SetBufferCount(5);
+			FMeshConversionContext()
+				: StreamedLOD(false)
+			{
+			}
+		};
 
-		// Position buffer
-		int32 CurrentVertexBuffer = 0;
+		TSharedPtr<FMeshConversionContext> Context = MakeShared<FMeshConversionContext>();
+		Context->OriginalLOD = LOD;
+		Context->DataLOD = LOD;
+		
+		if (!bUseStreaming)
 		{
-			MutableMeshBufferUtils::SetupVertexPositionsBuffer(CurrentVertexBuffer, Vertices);
+			LoadedEvent.Trigger();
+		}
+		else
+		{
+			MUTABLE_CPUPROFILER_SCOPE(MeshStreaming);
 
-			int32 ElementSize = Vertices.Buffers[CurrentVertexBuffer].ElementSize;
-			const uint8* SourceVertexData = ((const uint8*)LOD.StaticVertexBuffers.PositionVertexBuffer.GetVertexData()) + FirstVertexIndex * ElementSize;
+			FBulkDataIORequestCallBack RequestCallback = [LoadedEvent, SkeletalMesh, Context, LODIndex](bool bWasCancelled, IBulkDataIORequest* IORequest) mutable
+				{
+					if (!bWasCancelled)
+					{
+						uint8* BulkData = IORequest->GetReadResults();
+						check(BulkData != nullptr);
 
-			check(ElementSize * (FirstVertexIndex+VertexCount) <= LOD.StaticVertexBuffers.PositionVertexBuffer.GetAllocatedSize());
-			FMemory::Memcpy(Vertices.GetBufferData(CurrentVertexBuffer), SourceVertexData, ElementSize * VertexCount);
+						{
+							MUTABLE_CPUPROFILER_SCOPE(MeshStreamingSerialization);
 
-			++CurrentVertexBuffer;
-		}
+							FMemoryReaderView Ar(TArrayView<uint8>(BulkData, IORequest->GetSize()), true);
 
-		// Tangent buffer
-		{
-			bool bIsHighPrecision = LOD.StaticVertexBuffers.StaticMeshVertexBuffer.GetUseHighPrecisionTangentBasis();
+							constexpr uint8 DummyStripFlags = 0;
+							const bool bForceKeepCPUResources = true;
+							const bool bNeedsCPUAccess = true;
+							Context->StreamedLOD.SerializeStreamedData(Ar, const_cast<USkeletalMesh*>(SkeletalMesh), LODIndex, DummyStripFlags, bNeedsCPUAccess, bForceKeepCPUResources);
 
-			const mu::EMeshBufferSemantic Semantics[2] = { mu::EMeshBufferSemantic::Tangent, mu::EMeshBufferSemantic::Normal };
-			const int32 SemanticIndices[2] = { 0, 0 };
-			const int32 Components[2] = { 4, 4 };
-			mu::EMeshBufferFormat Formats[2] = { mu::EMeshBufferFormat::PackedDirS8, mu::EMeshBufferFormat::PackedDirS8_W_TangentSign };
-			int32 Offsets[2] = { 0, 4 };
-			int32 ElementSize = 8;
+							Context->DataLOD = &Context->StreamedLOD;
+						}
+					}
+					else
+					{
+						// Can this happen even if we don't cancel ourselves?
+						check(false);
+					}
 
-			if (bIsHighPrecision)
+					LoadedEvent.Trigger();
+				};
+
+			 Context->Request.Reset( LOD->StreamingBulkData.CreateStreamingRequest( EAsyncIOPriorityAndFlags::AIOP_High, &RequestCallback, nullptr) );
+		}
+
+		UE::Tasks::FTask FinalConversionTask = UE::Tasks::Launch( TEXT("MutableRuntimeMeshConversionFinal"),
+			[Context, SectionIndex, Result]()
 			{
-				// Not really supported
-				ensure(false);
-				Formats[0] = mu::EMeshBufferFormat::Int16;
-				Formats[1] = mu::EMeshBufferFormat::Int16;
-				Offsets[1] = 8;
-				ElementSize = 16;
-			}
-			Vertices.SetBuffer(CurrentVertexBuffer, ElementSize, 2, Semantics, SemanticIndices, Formats, Components, Offsets);
+				MUTABLE_CPUPROFILER_SCOPE(MeshConversion);
 
-			const uint8* SourceVertexData = ((const uint8*)LOD.StaticVertexBuffers.StaticMeshVertexBuffer.GetTangentData()) + FirstVertexIndex * ElementSize;
+				const FSkelMeshRenderSection& Section = Context->OriginalLOD->RenderSections[SectionIndex];
 
-			check((int32)ElementSize * (FirstVertexIndex + VertexCount) <= LOD.StaticVertexBuffers.StaticMeshVertexBuffer.GetTangentSize());
-			FMemory::Memcpy(Vertices.GetBufferData(CurrentVertexBuffer), SourceVertexData, ElementSize * VertexCount);
+				// Mesh data
+				int32 FirstVertexIndex = Context->OriginalLOD->RenderSections[SectionIndex].BaseVertexIndex;
+				int32 VertexCount = Section.GetNumVertices();
 
-			++CurrentVertexBuffer;
-		}
+				mu::FMeshBufferSet& Vertices = Result->GetVertexBuffers();
+				Vertices.SetElementCount(VertexCount, mu::EMemoryInitPolicy::Zeroed);
+				Vertices.SetBufferCount(5);
 
-		// Texture coords buffer
-		{
-			int32 NumTexCoords = LOD.GetNumTexCoords();
+				// Position buffer
+				int32 CurrentVertexBuffer = 0;
+				{
+					MutableMeshBufferUtils::SetupVertexPositionsBuffer(CurrentVertexBuffer, Vertices);
 
-			constexpr int32 MaxChannelCount = 4;
-			mu::EMeshBufferSemantic Semantics[MaxChannelCount];
-			int32 SemanticIndices[MaxChannelCount];
-			mu::EMeshBufferFormat Formats[MaxChannelCount];
-			int32 Components[MaxChannelCount];
-			int32 Offsets[MaxChannelCount];
+					int32 ElementSize = Vertices.Buffers[CurrentVertexBuffer].ElementSize;
+					const uint8* SourceVertexData = ((const uint8*)Context->DataLOD->StaticVertexBuffers.PositionVertexBuffer.GetVertexData()) + FirstVertexIndex * ElementSize;
 
-			int32 UVSize = LOD.StaticVertexBuffers.StaticMeshVertexBuffer.GetUseFullPrecisionUVs() ? 8 : 4;
-			const int32 ElementSize = UVSize * NumTexCoords;
+					check(ElementSize * (FirstVertexIndex + VertexCount) <= Context->DataLOD->StaticVertexBuffers.PositionVertexBuffer.GetAllocatedSize());
+					FMemory::Memcpy(Vertices.GetBufferData(CurrentVertexBuffer), SourceVertexData, ElementSize * VertexCount);
 
-			for (int32 UV = 0; UV < NumTexCoords; ++UV)
-			{
-				Semantics[UV] = mu::EMeshBufferSemantic::TexCoords;
-				SemanticIndices[UV] = UV;
-				Formats[UV] = (UVSize == 8) ? mu::EMeshBufferFormat::Float32 : mu::EMeshBufferFormat::Float16;
-				Components[UV] = 2;
-				Offsets[UV] = UVSize * UV;
-			}
-			Vertices.SetBuffer(CurrentVertexBuffer, ElementSize, NumTexCoords, Semantics, SemanticIndices, Formats, Components, Offsets);
+					++CurrentVertexBuffer;
+				}
 
-			const uint8* SourceVertexData = ((const uint8*)LOD.StaticVertexBuffers.StaticMeshVertexBuffer.GetTexCoordData()) + FirstVertexIndex * UVSize * NumTexCoords;
-			check(UVSize * NumTexCoords * (FirstVertexIndex + VertexCount) <= LOD.StaticVertexBuffers.StaticMeshVertexBuffer.GetTexCoordSize());
+				// Tangent buffer
+				{
+					bool bIsHighPrecision = Context->DataLOD->StaticVertexBuffers.StaticMeshVertexBuffer.GetUseHighPrecisionTangentBasis();
 
-			FMemory::Memcpy(Vertices.GetBufferData(CurrentVertexBuffer), SourceVertexData, Vertices.Buffers[CurrentVertexBuffer].ElementSize * VertexCount);
+					const mu::EMeshBufferSemantic Semantics[2] = { mu::EMeshBufferSemantic::Tangent, mu::EMeshBufferSemantic::Normal };
+					const int32 SemanticIndices[2] = { 0, 0 };
+					const int32 Components[2] = { 4, 4 };
+					mu::EMeshBufferFormat Formats[2] = { mu::EMeshBufferFormat::PackedDirS8, mu::EMeshBufferFormat::PackedDirS8_W_TangentSign };
+					int32 Offsets[2] = { 0, 4 };
+					int32 ElementSize = 8;
 
-			++CurrentVertexBuffer;
-		}
+					if (bIsHighPrecision)
+					{
+						// Not really supported
+						ensure(false);
+						Formats[0] = mu::EMeshBufferFormat::Int16;
+						Formats[1] = mu::EMeshBufferFormat::Int16;
+						Offsets[1] = 8;
+						ElementSize = 16;
+					}
+					Vertices.SetBuffer(CurrentVertexBuffer, ElementSize, 2, Semantics, SemanticIndices, Formats, Components, Offsets);
 
-		// Skin buffer
-		if (Section.MaxBoneInfluences > 0)
-		{
-			const FSkinWeightVertexBuffer* SkinBuffer = LOD.GetSkinWeightVertexBuffer();
-			const int32 MaxBoneIndexTypeSizeBytes = SkinBuffer->GetBoneIndexByteSize();
-			const int32 MaxBoneWeightTypeSizeBytes = SkinBuffer->GetBoneWeightByteSize();
-			const int32 MaxBonesPerVertex = SkinBuffer->GetMaxBoneInfluences();
-			MutableMeshBufferUtils::SetupSkinBuffer(CurrentVertexBuffer, MaxBoneIndexTypeSizeBytes, MaxBoneWeightTypeSizeBytes, MaxBonesPerVertex, Vertices);
+					const uint8* SourceVertexData = ((const uint8*)Context->DataLOD->StaticVertexBuffers.StaticMeshVertexBuffer.GetTangentData()) + FirstVertexIndex * ElementSize;
 
-			int32 ElementSize = Vertices.Buffers[CurrentVertexBuffer].ElementSize;
-			const uint8* SourceVertexData = ((const uint8*)SkinBuffer->GetDataVertexBuffer()->GetWeightData()) + FirstVertexIndex * ElementSize;
+					check((int32)ElementSize * (FirstVertexIndex + VertexCount) <= Context->DataLOD->StaticVertexBuffers.StaticMeshVertexBuffer.GetTangentSize());
+					FMemory::Memcpy(Vertices.GetBufferData(CurrentVertexBuffer), SourceVertexData, ElementSize * VertexCount);
 
-			check(ElementSize * (FirstVertexIndex + VertexCount) <= (int32)SkinBuffer->GetVertexDataSize());
+					++CurrentVertexBuffer;
+				}
 
-			FMemory::Memcpy(Vertices.GetBufferData(CurrentVertexBuffer), SourceVertexData, ElementSize * VertexCount);
+				// Texture coords buffer
+				{
+					int32 NumTexCoords = Context->OriginalLOD->GetNumTexCoords();
 
-			++CurrentVertexBuffer;
-		}
+					constexpr int32 MaxChannelCount = 4;
+					mu::EMeshBufferSemantic Semantics[MaxChannelCount];
+					int32 SemanticIndices[MaxChannelCount];
+					mu::EMeshBufferFormat Formats[MaxChannelCount];
+					int32 Components[MaxChannelCount];
+					int32 Offsets[MaxChannelCount];
 
-		// Colour buffer
-		if (LOD.StaticVertexBuffers.ColorVertexBuffer.GetNumVertices() > 0)
-		{			
-			const FColorVertexBuffer& ColorBuffer = LOD.StaticVertexBuffers.ColorVertexBuffer;
+					int32 UVSize = Context->DataLOD->StaticVertexBuffers.StaticMeshVertexBuffer.GetUseFullPrecisionUVs() ? 8 : 4;
+					const int32 ElementSize = UVSize * NumTexCoords;
 
-			MutableMeshBufferUtils::SetupVertexColorBuffer(CurrentVertexBuffer, Vertices);
+					for (int32 UV = 0; UV < NumTexCoords; ++UV)
+					{
+						Semantics[UV] = mu::EMeshBufferSemantic::TexCoords;
+						SemanticIndices[UV] = UV;
+						Formats[UV] = (UVSize == 8) ? mu::EMeshBufferFormat::Float32 : mu::EMeshBufferFormat::Float16;
+						Components[UV] = 2;
+						Offsets[UV] = UVSize * UV;
+					}
+					Vertices.SetBuffer(CurrentVertexBuffer, ElementSize, NumTexCoords, Semantics, SemanticIndices, Formats, Components, Offsets);
 
-			int32 ElementSize = Vertices.Buffers[CurrentVertexBuffer].ElementSize;
-			const uint8* SourceVertexData = ((const uint8*)ColorBuffer.GetVertexData()) + FirstVertexIndex * ElementSize;
+					const uint8* SourceVertexData = ((const uint8*)Context->DataLOD->StaticVertexBuffers.StaticMeshVertexBuffer.GetTexCoordData()) + FirstVertexIndex * UVSize * NumTexCoords;
+					check(UVSize * NumTexCoords * (FirstVertexIndex + VertexCount) <= Context->DataLOD->StaticVertexBuffers.StaticMeshVertexBuffer.GetTexCoordSize());
 
-			check(ElementSize * (FirstVertexIndex + VertexCount) <= (int32)ColorBuffer.GetAllocatedSize());
-			check(ElementSize == ColorBuffer.GetStride());
+					FMemory::Memcpy(Vertices.GetBufferData(CurrentVertexBuffer), SourceVertexData, Vertices.Buffers[CurrentVertexBuffer].ElementSize * VertexCount);
 
-			FMemory::Memcpy(Vertices.GetBufferData(CurrentVertexBuffer), SourceVertexData, ElementSize * VertexCount);
+					++CurrentVertexBuffer;
+				}
 
-			++CurrentVertexBuffer;
-		}
+				// Skin buffer
+				if (Section.MaxBoneInfluences > 0)
+				{
+					const FSkinWeightVertexBuffer* SkinBuffer = Context->DataLOD->GetSkinWeightVertexBuffer();
+					const int32 MaxBoneIndexTypeSizeBytes = SkinBuffer->GetBoneIndexByteSize();
+					const int32 MaxBoneWeightTypeSizeBytes = SkinBuffer->GetBoneWeightByteSize();
+					const int32 MaxBonesPerVertex = SkinBuffer->GetMaxBoneInfluences();
+					MutableMeshBufferUtils::SetupSkinBuffer(CurrentVertexBuffer, MaxBoneIndexTypeSizeBytes, MaxBoneWeightTypeSizeBytes, MaxBonesPerVertex, Vertices);
 
-		// Indices
-		{
-			int32 FirstIndexIndex = LOD.RenderSections[SectionIndex].BaseIndex;
-			int32 IndexCount = Section.NumTriangles*3;
-			int32 ElementSize = LOD.MultiSizeIndexContainer.GetDataTypeSize();
-
-			mu::FMeshBufferSet& Indices = Result.GetIndexBuffers();
-			Indices.SetBufferCount(1);
-			Indices.SetElementCount(IndexCount);
-			constexpr int32 ChannelCount = 1;
-			const mu::EMeshBufferSemantic Semantics[ChannelCount] = { mu::EMeshBufferSemantic::VertexIndex };
-			const int32 SemanticIndices[ChannelCount] = { 0 };
-			mu::EMeshBufferFormat Formats[ChannelCount] = { ElementSize == 4 ? mu::EMeshBufferFormat::UInt32 : mu::EMeshBufferFormat::UInt16 };
-			const int32 Components[ChannelCount] = { 1 };
-			const int32 Offsets[ChannelCount] = { 0 };
-			Indices.SetBuffer(0, ElementSize, ChannelCount, Semantics, SemanticIndices, Formats, Components, Offsets);
-
-			const uint8* SourceData = (const uint8*)LOD.MultiSizeIndexContainer.GetIndexBuffer()->GetPointerTo(FirstIndexIndex);
-			uint8* TargetData = Indices.GetBufferData(0);
-
-			if (FirstVertexIndex == 0)
-			{
-				FMemory::Memcpy(TargetData, SourceData, IndexCount* ElementSize);
-			}
-			else
-			{
-				// Apply vertex offset
-				switch (ElementSize)
+					int32 ElementSize = Vertices.Buffers[CurrentVertexBuffer].ElementSize;
+					const uint8* SourceVertexData = ((const uint8*)SkinBuffer->GetDataVertexBuffer()->GetWeightData()) + FirstVertexIndex * ElementSize;
+
+					check(ElementSize * (FirstVertexIndex + VertexCount) <= (int32)SkinBuffer->GetVertexDataSize());
+
+					FMemory::Memcpy(Vertices.GetBufferData(CurrentVertexBuffer), SourceVertexData, ElementSize * VertexCount);
+
+					++CurrentVertexBuffer;
+				}
+
+				// Colour buffer
+				if (Context->DataLOD->StaticVertexBuffers.ColorVertexBuffer.GetNumVertices() > 0)
 				{
-				case 2:
-					for (int32 Index=0;Index<IndexCount; ++Index)
+					const FColorVertexBuffer& ColorBuffer = Context->DataLOD->StaticVertexBuffers.ColorVertexBuffer;
+
+					MutableMeshBufferUtils::SetupVertexColorBuffer(CurrentVertexBuffer, Vertices);
+
+					int32 ElementSize = Vertices.Buffers[CurrentVertexBuffer].ElementSize;
+					const uint8* SourceVertexData = ((const uint8*)ColorBuffer.GetVertexData()) + FirstVertexIndex * ElementSize;
+
+					check(ElementSize * (FirstVertexIndex + VertexCount) <= (int32)ColorBuffer.GetAllocatedSize());
+					check(ElementSize == ColorBuffer.GetStride());
+
+					FMemory::Memcpy(Vertices.GetBufferData(CurrentVertexBuffer), SourceVertexData, ElementSize * VertexCount);
+
+					++CurrentVertexBuffer;
+				}
+
+				// Indices
+				{
+					int32 FirstIndexIndex = Section.BaseIndex;
+					int32 IndexCount = Section.NumTriangles * 3;
+					int32 ElementSize = Context->DataLOD->MultiSizeIndexContainer.GetDataTypeSize();
+
+					mu::FMeshBufferSet& Indices = Result->GetIndexBuffers();
+					Indices.SetBufferCount(1);
+					Indices.SetElementCount(IndexCount);
+					constexpr int32 ChannelCount = 1;
+					const mu::EMeshBufferSemantic Semantics[ChannelCount] = { mu::EMeshBufferSemantic::VertexIndex };
+					const int32 SemanticIndices[ChannelCount] = { 0 };
+					mu::EMeshBufferFormat Formats[ChannelCount] = { ElementSize == 4 ? mu::EMeshBufferFormat::UInt32 : mu::EMeshBufferFormat::UInt16 };
+					const int32 Components[ChannelCount] = { 1 };
+					const int32 Offsets[ChannelCount] = { 0 };
+					Indices.SetBuffer(0, ElementSize, ChannelCount, Semantics, SemanticIndices, Formats, Components, Offsets);
+
+					const uint8* SourceData = (const uint8*)Context->DataLOD->MultiSizeIndexContainer.GetIndexBuffer()->GetPointerTo(FirstIndexIndex);
+					uint8* TargetData = Indices.GetBufferData(0);
+
+					if (FirstVertexIndex == 0)
 					{
-						const uint16* Source = ((const uint16*)SourceData) + Index;
-						uint16* Target = ((uint16*)TargetData) + Index;
-						check(*Source>=FirstVertexIndex);
-						*Target = *Source - FirstVertexIndex;
+						FMemory::Memcpy(TargetData, SourceData, IndexCount * ElementSize);
 					}
-					break;
-				case 4:
-					for (int32 Index = 0; Index < IndexCount; ++Index)
+					else
 					{
-						const uint32* Source = ((const uint32*)SourceData) + Index;
-						uint32* Target = ((uint32*)TargetData) + Index;
-						check(*Source >= uint32(FirstVertexIndex));
-						*Target = *Source - FirstVertexIndex;
+						// Apply vertex offset
+						switch (ElementSize)
+						{
+						case 2:
+							for (int32 Index = 0; Index < IndexCount; ++Index)
+							{
+								const uint16* Source = ((const uint16*)SourceData) + Index;
+								uint16* Target = ((uint16*)TargetData) + Index;
+								check(*Source >= FirstVertexIndex);
+								*Target = *Source - FirstVertexIndex;
+							}
+							break;
+						case 4:
+							for (int32 Index = 0; Index < IndexCount; ++Index)
+							{
+								const uint32* Source = ((const uint32*)SourceData) + Index;
+								uint32* Target = ((uint32*)TargetData) + Index;
+								check(*Source >= uint32(FirstVertexIndex));
+								*Target = *Source - FirstVertexIndex;
+							}
+							break;
+						default:
+							// Index size not implemented
+							break;
+						}
 					}
-					break;
-				default:
-					// Index size not implemented
-					break;
 				}
-			}
-		}
 
-		//FString DebugLog;
-		//Result.Log(DebugLog, 8);
-		//UE_LOG(LogMutable, Warning, TEXT("[%s]"), *DebugLog);
+				Result->EnsureSurfaceData();
+			},
+
+			// Dependencies
+			LoadedEvent
+			);
 
-		Result.EnsureSurfaceData();
+		return FinalConversionTask;
 	}
 
 }
diff --git a/Engine/Plugins/Mutable/Source/CustomizableObject/Private/MuCO/UnrealMutableImageProvider.cpp b/Engine/Plugins/Mutable/Source/CustomizableObject/Private/MuCO/UnrealMutableImageProvider.cpp
--- a/Engine/Plugins/Mutable/Source/CustomizableObject/Private/MuCO/UnrealMutableImageProvider.cpp
+++ b/Engine/Plugins/Mutable/Source/CustomizableObject/Private/MuCO/UnrealMutableImageProvider.cpp
@@ -491,42 +491,98 @@ TTuple<UE::Tasks::FTask, TFunction<void()>> FUnrealMutableResourceProvider::GetM
 	TSharedPtr<mu::FMesh> Result = MakeShared<mu::FMesh>();
 
 	UE::Tasks::FTaskEvent Completion(TEXT("MutableMeshParameterLoadInGameThread"));
-	ExecuteOnGameThread(TEXT("MeshParameterLoad"),
-		[Id, Completion, LODIndex, SectionIndex, Result, CustomizableObject=CurrentCustomizableObject] () mutable
+
+	auto MeshLoadCallback = [Id, LODIndex, SectionIndex, Result, Completion, Provider = this](const FSoftObjectPath&, UObject* LoadedObject) mutable
 		{
 			check(IsInGameThread());
 
-			MUTABLE_CPUPROFILER_SCOPE(MutableMeshParameterLoadInGameThread);
-
-			// TODO: Do it async
-			USkeletalMesh* SkeletalMesh = nullptr;
-			{
-				MUTABLE_CPUPROFILER_SCOPE(ActualLoad);
-				SkeletalMesh = Cast<USkeletalMesh>(FSoftObjectPath(Id.ToString()).TryLoad());
-			}
+			MUTABLE_CPUPROFILER_SCOPE(ActualLoad);
 
+			USkeletalMesh* SkeletalMesh = Cast<USkeletalMesh>(LoadedObject);
 			if (!SkeletalMesh)
 			{
 				UE_LOG(LogMutable, Error, TEXT("Failed to load the skeletal mesh [%s] set for a mesh parameter. Please check that it was packaged."), *Id.ToString());
+				Completion.Trigger();
+				return;
 			}
-			else
+
+			// Add a reference to the skeletalmesh
 			{
-				UModelResources* ModelResources = nullptr;
+				bool bFound = false;
+				for (FReferencedSkeletalMesh& Entry : Provider->ReferencedSkeletalMeshes)
+				{
+					if (Entry.SkeletalMesh == SkeletalMesh)
+					{
+						bFound = true;
+						++Entry.ReferenceCount;
+						break;
+					}
+				}
 
-				TStrongObjectPtr<UCustomizableObject> CO = CustomizableObject.Pin();
-				if (CO)
+				if (!bFound)
 				{
-					ModelResources = CO->GetPrivate()->GetModelResources();
+					Provider->ReferencedSkeletalMeshes.Add({ SkeletalMesh,1 });
 				}
+			}
 
-				// It is valid not to have a CO or MutableStreamedResources. It is only used for skeleton data. 
-				// This may happen when updating mips, and for those operations we don't need skeleton data.
-				//ensure(ModelResources);
-				UnrealConversionUtils::ConvertSkeletalMeshFromRuntimeData(*SkeletalMesh, LODIndex, SectionIndex, ModelResources, *Result.Get());
+			UModelResources* ModelResources = nullptr;
+
+			TStrongObjectPtr<UCustomizableObject> CO = Provider->CurrentCustomizableObject.Pin();
+			if (CO)
+			{
+				ModelResources = CO->GetPrivate()->GetModelResources();
 			}
 
-			// Signal
-			Completion.Trigger();
+			// It is valid not to have a CO or MutableStreamedResources. It is only used for skeleton data. 
+			// This may happen when updating mips, and for those operations we don't need skeleton data.
+			//ensure(ModelResources);
+
+			UE::Tasks::FTask ConversionTask = UnrealConversionUtils::ConvertSkeletalMeshFromRuntimeData(SkeletalMesh, LODIndex, SectionIndex, ModelResources, Result.Get());
+
+			// The rest of the conversion may happen in a worker thread.
+			UE::Tasks::Launch(TEXT("MeshParameterLoadFinalize"),
+				[Completion, Provider, SkeletalMesh]() mutable
+				{					
+					// Signal
+					Completion.Trigger();
+
+					// Remove a reference to the skeletalmesh
+					ExecuteOnGameThread(TEXT("MeshParameterLoadReleaseMesh"),
+						[Provider, SkeletalMesh]() mutable
+						{
+							bool bFound = false;
+							for (int32 RefIndex = 0; RefIndex < Provider->ReferencedSkeletalMeshes.Num(); ++RefIndex)
+							{
+								FReferencedSkeletalMesh& Entry = Provider->ReferencedSkeletalMeshes[RefIndex];
+								if (Entry.SkeletalMesh == SkeletalMesh)
+								{
+									bFound = true;
+									--Entry.ReferenceCount;
+									if (Entry.ReferenceCount == 0)
+									{
+										Provider->ReferencedSkeletalMeshes.RemoveAtSwap(RefIndex);
+									}
+									break;
+								}
+							}
+
+							check(bFound);
+						});
+				},
+				// Dependencies
+				ConversionTask
+			);
+		};
+
+	// LoadAsync is only thread-safe when using the zenloader
+	ExecuteOnGameThread(TEXT("MeshParameterLoad"),
+		[Id, MeshLoadCallback]() mutable
+		{
+			check(IsInGameThread());
+
+			MUTABLE_CPUPROFILER_SCOPE(MutableMeshParameterLoadInGameThread);
+		
+			FSoftObjectPath(Id.ToString()).LoadAsync( FLoadSoftObjectPathAsyncDelegate::CreateLambda( MeshLoadCallback ));
 		}
 	);
 
@@ -877,6 +933,11 @@ void FUnrealMutableResourceProvider::AddReferencedObjects(FReferenceCollector& C
 		ExternalImagesLock.Unlock();
 	}
 #endif
+
+	for (FReferencedSkeletalMesh& Mesh : ReferencedSkeletalMeshes)
+	{
+		Collector.AddReferencedObject(Mesh.SkeletalMesh);
+	}
 }
 
 
diff --git a/Engine/Plugins/Mutable/Source/CustomizableObject/Private/MuCO/UnrealMutableImageProvider.h b/Engine/Plugins/Mutable/Source/CustomizableObject/Private/MuCO/UnrealMutableImageProvider.h
--- a/Engine/Plugins/Mutable/Source/CustomizableObject/Private/MuCO/UnrealMutableImageProvider.h
+++ b/Engine/Plugins/Mutable/Source/CustomizableObject/Private/MuCO/UnrealMutableImageProvider.h
@@ -53,15 +53,23 @@ public:
 	/** List of actual image providers that have been registered to the CustomizableObjectSystem. */
 	TArray< TWeakObjectPtr<class UCustomizableSystemImageProvider> > ImageProviders;
 
+	// Register and unregister the current CO that we are working on. This is called by internal code when building instances.
+	void SetCurrentObject(const TWeakObjectPtr<UCustomizableObject>&);
+
+	struct FReferencedSkeletalMesh
+	{
+		TObjectPtr<class USkeletalMesh> SkeletalMesh;
+		uint32 ReferenceCount = 0;
+	};
+	TArray<FReferencedSkeletalMesh> ReferencedSkeletalMeshes;
+
+	// FGCObject interface
 	virtual void AddReferencedObjects(FReferenceCollector& Collector) override;
 	virtual FString GetReferencerName() const override
 	{
 		return TEXT("FUnrealMutableImageProvider");
 	}
 
-	// Register and unregister the current CO that we are working on. This is called by internal code when building instances.
-	void SetCurrentObject( const TWeakObjectPtr<UCustomizableObject>& );
-
 private:
 	struct FUnrealMutableImageInfo
 	{
diff --git a/Engine/Plugins/Mutable/Source/CustomizableObjectEditor/Private/MuCOE/SMutableCodeViewer.cpp b/Engine/Plugins/Mutable/Source/CustomizableObjectEditor/Private/MuCOE/SMutableCodeViewer.cpp
--- a/Engine/Plugins/Mutable/Source/CustomizableObjectEditor/Private/MuCOE/SMutableCodeViewer.cpp
+++ b/Engine/Plugins/Mutable/Source/CustomizableObjectEditor/Private/MuCOE/SMutableCodeViewer.cpp
@@ -2808,11 +2808,18 @@ namespace
 			USkeletalMesh* SkeletalMesh = Cast<USkeletalMesh>(FSoftObjectPath(Id.ToString()).TryLoad());
 			check(SkeletalMesh);
 
-			UnrealConversionUtils::ConvertSkeletalMeshFromRuntimeData(*SkeletalMesh, LODIndex, SectionIndex, nullptr, *Result.Get());
+			UE::Tasks::FTask ConversionTask = UnrealConversionUtils::ConvertSkeletalMeshFromRuntimeData(SkeletalMesh, LODIndex, SectionIndex, nullptr, Result.Get());
 
-			ResultCallback(Result);
-
-			return MakeTuple(UE::Tasks::MakeCompletedTask<void>(), []() -> void {});
+			return MakeTuple(
+				UE::Tasks::Launch( TEXT("FinalizeGetMesh"),
+					[Result, ResultCallback]()
+					{
+						ResultCallback(Result);
+					}, 
+					ConversionTask)
+				, 
+				[]() -> void {}
+			);
 		}
 	};
 }
diff --git a/Engine/Source/Runtime/Engine/Public/Rendering/SkeletalMeshLODRenderData.h b/Engine/Source/Runtime/Engine/Public/Rendering/SkeletalMeshLODRenderData.h
--- a/Engine/Source/Runtime/Engine/Public/Rendering/SkeletalMeshLODRenderData.h
+++ b/Engine/Source/Runtime/Engine/Public/Rendering/SkeletalMeshLODRenderData.h
@@ -230,7 +230,7 @@ public:
 	* @param bForceKeepCPUResources - Whether we want to force keeping CPU resources
 	* @return Size of streamed LOD data in bytes
 	*/
-	void SerializeStreamedData(FArchive& Ar, USkinnedAsset* Owner, int32 LODIdx, uint8 ClassStripFlags, bool bNeedsCPUAccess, bool bForceKeepCPUResources);
+	ENGINE_API void SerializeStreamedData(FArchive& Ar, USkinnedAsset* Owner, int32 LODIdx, uint8 ClassStripFlags, bool bNeedsCPUAccess, bool bForceKeepCPUResources);
 
 	void SerializeAvailabilityInfo(FArchive& Ar, int32 LODIdx, bool bHasAdjacencyData, bool bNeedsCPUAccess);
 
