[Mutable] Reduce Started Update Skeletal Mesh Async log frequency.

#rnx

diff --git a/Engine/Plugins/Mutable/Source/CustomizableObject/Private/MuCO/CustomizableObjectSystem.cpp b/Engine/Plugins/Mutable/Source/CustomizableObject/Private/MuCO/CustomizableObjectSystem.cpp
--- a/Engine/Plugins/Mutable/Source/CustomizableObject/Private/MuCO/CustomizableObjectSystem.cpp
+++ b/Engine/Plugins/Mutable/Source/CustomizableObject/Private/MuCO/CustomizableObjectSystem.cpp
@@ -4756,7 +4756,7 @@ void UCustomizableObjectSystemPrivate::StartUpdateSkeletalMesh(const TSharedRef<
 
 		const float CurrentTime = FPlatformTime::Seconds();
 
-		constexpr float LogInterval = 1.0 / 10.0; // Allow maximum 10 logs per second.
+		constexpr float LogInterval = 1.0 / 2.0; // Allow maximum 2 logs per second.
 		const bool bMute = CurrentTime - LogStartedUpdateLast < LogInterval;
 		LogStartedUpdateLast = CurrentTime;
 			
