[mutable] Prevent crash in render section conversion during debug.
#rnx

diff --git a/Engine/Plugins/Mutable/Source/CustomizableObject/Private/MuCO/UnrealConversionUtils.cpp b/Engine/Plugins/Mutable/Source/CustomizableObject/Private/MuCO/UnrealConversionUtils.cpp
--- a/Engine/Plugins/Mutable/Source/CustomizableObject/Private/MuCO/UnrealConversionUtils.cpp
+++ b/Engine/Plugins/Mutable/Source/CustomizableObject/Private/MuCO/UnrealConversionUtils.cpp
@@ -228,7 +228,14 @@ namespace UnrealConversionUtils
 			Section.BoneMap.Reserve(BoneCount);
 			for (int32 BoneMapIndex = 0; BoneMapIndex < BoneCount; ++BoneMapIndex, ++FirstBone)
 			{
-				Section.BoneMap.Add(BoneInfoMap[InBoneMap[FirstBone]].Value);
+				if (InBoneMap.IsValidIndex(FirstBone))
+				{
+					mu::FBoneName FirstBoneName = InBoneMap[FirstBone];
+					if (const TPair<FName, uint16>* Found = BoneInfoMap.Find(FirstBoneName))
+					{
+						Section.BoneMap.Add(Found->Value);
+					}
+				}
 			}
 		}
 	}
