[mutable] CustomizableObjectBulkValidations commandlet implementation.
- Implemented commandlet to handle the CO Compilation+Update of multiple COs at a given target path.
- Added option to only test Root objects (false by default)
- Updated the FScopeLogSection object to make it able to set the target object of the section and to handle multiple scopes open at the same time
- Moved some now shared utility methods to the ValidationUtils files.
- Moved all public files from the module MutableValidation to the Private dir
- Minor refactor afecting multiple files (Ptr to Ref)

#rnx

diff --git a/Engine/Plugins/Mutable/Source/CustomizableObjectEditor/Internal/MuCOE/CustomizableObjectBenchmarkingUtils.h b/Engine/Plugins/Mutable/Source/CustomizableObjectEditor/Internal/MuCOE/CustomizableObjectBenchmarkingUtils.h
--- a/Engine/Plugins/Mutable/Source/CustomizableObjectEditor/Internal/MuCOE/CustomizableObjectBenchmarkingUtils.h
+++ b/Engine/Plugins/Mutable/Source/CustomizableObjectEditor/Internal/MuCOE/CustomizableObjectBenchmarkingUtils.h
@@ -20,7 +20,7 @@ namespace CustomizableObjectBenchmarkingUtils
 	 * @param OutSuccesfullyGeneratedInstanceCount The total of instances generated successfully [0 , InstancesPerState * COStateCount]
 	 * @return 
 	 */
-	CUSTOMIZABLEOBJECTEDITOR_API bool GenerateDeterministicSetOfInstances(TObjectPtr<UCustomizableObject> TargetCustomizableObject, const uint16 InstancesPerState, TSpscQueue<TStrongObjectPtr<UCustomizableObjectInstance>>& OutGeneratedInstances, uint32& OutSuccesfullyGeneratedInstanceCount);
+	CUSTOMIZABLEOBJECTEDITOR_API bool GenerateDeterministicSetOfInstances(UCustomizableObject& TargetCustomizableObject, const uint16 InstancesPerState, TSpscQueue<TStrongObjectPtr<UCustomizableObjectInstance>>& OutGeneratedInstances, uint32& OutSuccesfullyGeneratedInstanceCount);
 
 	
 	/**
diff --git a/Engine/Plugins/Mutable/Source/CustomizableObjectEditor/Private/MuCOE/CustomizableObjectBenchmarkingUtils.cpp b/Engine/Plugins/Mutable/Source/CustomizableObjectEditor/Private/MuCOE/CustomizableObjectBenchmarkingUtils.cpp
--- a/Engine/Plugins/Mutable/Source/CustomizableObjectEditor/Private/MuCOE/CustomizableObjectBenchmarkingUtils.cpp
+++ b/Engine/Plugins/Mutable/Source/CustomizableObjectEditor/Private/MuCOE/CustomizableObjectBenchmarkingUtils.cpp
@@ -5,21 +5,21 @@
 #include "MuCO/CustomizableObjectSystem.h"
 
 
-bool CustomizableObjectBenchmarkingUtils::GenerateDeterministicSetOfInstances(TObjectPtr<UCustomizableObject> TargetCustomizableObject, const uint16 InstancesPerState,TSpscQueue<TStrongObjectPtr<UCustomizableObjectInstance>>& OutGeneratedInstances,  uint32& OutSuccesfullyGeneratedInstanceCount)
+bool CustomizableObjectBenchmarkingUtils::GenerateDeterministicSetOfInstances(UCustomizableObject& TargetCustomizableObject, const uint16 InstancesPerState,TSpscQueue<TStrongObjectPtr<UCustomizableObjectInstance>>& OutGeneratedInstances,  uint32& OutSuccesfullyGeneratedInstanceCount)
 {
-	if (!TargetCustomizableObject || !TargetCustomizableObject->IsCompiled())
+	if (!TargetCustomizableObject.IsCompiled())
 	{
 		return false;
 	}
 
-	const UModelResources* ModelResources = TargetCustomizableObject->GetPrivate()->GetModelResources();
+	const UModelResources* ModelResources = TargetCustomizableObject.GetPrivate()->GetModelResources();
 	if (!ModelResources)
 	{
 		return false;
 	}
 	
 	// Test this parameter configuration in all the states of the CO
-	const uint32 StateCount = TargetCustomizableObject->GetStateCount();
+	const uint32 StateCount = TargetCustomizableObject.GetStateCount();
 	check(StateCount >= 1);
 
 	UE_LOG(LogMutable, Display, TEXT("Requested Instances Count : %i"), InstancesPerState);
@@ -35,7 +35,7 @@ bool CustomizableObjectBenchmarkingUtils::GenerateDeterministicSetOfInstances(TO
 	// Generate a series of instances to later update
 	for (uint32 CustomizableObjectInstanceIndex = 0; CustomizableObjectInstanceIndex < InstancesPerState; CustomizableObjectInstanceIndex++)
 	{
-		if (UCustomizableObjectInstance* GeneratedInstance = TargetCustomizableObject->CreateInstance())
+		if (UCustomizableObjectInstance* GeneratedInstance = TargetCustomizableObject.CreateInstance())
 		{
 			// Force generation of all LODS
 			TMap<FName, uint8> RequestedLodLevels;
@@ -58,7 +58,7 @@ bool CustomizableObjectBenchmarkingUtils::GenerateDeterministicSetOfInstances(TO
 		}
 		else
 		{
-			UE_LOG(LogMutable,Error,TEXT("Failed to generate COI for the %s CO."),*TargetCustomizableObject->GetName());
+			UE_LOG(LogMutable,Error,TEXT("Failed to generate COI for the %s CO."),*TargetCustomizableObject.GetName());
 			return false;
 		}
 	}
diff --git a/Engine/Plugins/Mutable/Source/CustomizableObjectEditor/Private/MuCOE/CustomizableObjectEditorPerformanceAnalyzer.cpp b/Engine/Plugins/Mutable/Source/CustomizableObjectEditor/Private/MuCOE/CustomizableObjectEditorPerformanceAnalyzer.cpp
--- a/Engine/Plugins/Mutable/Source/CustomizableObjectEditor/Private/MuCOE/CustomizableObjectEditorPerformanceAnalyzer.cpp
+++ b/Engine/Plugins/Mutable/Source/CustomizableObjectEditor/Private/MuCOE/CustomizableObjectEditorPerformanceAnalyzer.cpp
@@ -487,7 +487,7 @@ FReply SCustomizableObjectEditorPerformanceAnalyzer::OnInstanceUpdateButtonClick
 	// Set the CVArs as in CIS
 	//SetForBenchmarkMutableCVarValues();
 		
-	if (!CustomizableObjectBenchmarkingUtils::GenerateDeterministicSetOfInstances(CustomizableObject.Get(),InstancesToGeneratePerState,InstancesToUpdate,TotalScheduledUpdates))
+	if (!CustomizableObjectBenchmarkingUtils::GenerateDeterministicSetOfInstances(*CustomizableObject ,InstancesToGeneratePerState,InstancesToUpdate,TotalScheduledUpdates))
 	{
 		UE_LOG(LogMutable,Error, TEXT("Mutable Customizable Object Instance generation failed. Aborting perf test."))
 		return FReply::Handled();
diff --git a/Engine/Plugins/Mutable/Source/MutableValidation/Private/MuV/AssetValidator_CustomizableObjects.h b/Engine/Plugins/Mutable/Source/MutableValidation/Private/MuV/AssetValidator_CustomizableObjects.h
new file mode 100644
--- /dev/null
+++ b/Engine/Plugins/Mutable/Source/MutableValidation/Private/MuV/AssetValidator_CustomizableObjects.h
@@ -0,0 +1,49 @@
+// Copyright Epic Games, Inc. All Rights Reserved.
+
+#pragma once
+
+#include "Editor.h"
+#include "EditorValidatorBase.h"
+#include "AssetRegistry/AssetData.h"
+#include "Delegates/DelegateSignatureImpl.inl"
+
+
+#include "AssetValidator_CustomizableObjects.generated.h"
+
+class FText;
+class UObject;
+class UCustomizableObject;
+
+UCLASS()
+class UAssetValidator_CustomizableObjects : public UEditorValidatorBase
+{
+	GENERATED_BODY()
+
+public:
+	UAssetValidator_CustomizableObjects();
+
+	/** Checks the validity of the provided CustomizableObject by compiling and recalling the compilation status and the raised compilation warning and error messages.
+	 * @param InCustomizableObject The mutable customizable object to test.
+	 * @param OutValidationErrors A list of error messages produced during the provided customizable object's compilation.
+	 * @param OutValidationWarnings A list of warning messages produced during the provided customizable object's compilation.
+	 * @return Validation result for the provided object.
+	 */
+	static EDataValidationResult IsCustomizableObjectValid(UCustomizableObject* InCustomizableObject, TArray<FText>& OutValidationErrors, TArray<FText>& OutValidationWarnings);
+	
+protected:
+	// UEditorValidatorBase
+	virtual bool CanValidateAsset_Implementation(const FAssetData& InAssetData, UObject* InAsset, FDataValidationContext& InContext) const override;
+	virtual EDataValidationResult ValidateLoadedAsset_Implementation(const FAssetData& InAssetData, UObject* InAsset, FDataValidationContext& InContext) override;
+	// UEditorValidatorBase
+
+private:
+	
+	/** Cached handle to be able later to remove the bound method from the FEditorDelegates::OnPostAssetValidation delegate */
+	inline static FDelegateHandle OnPostCOValidationHandle;
+	
+	/** Collection with all root objects tested during this IsDataValidRun. Shared with all COs */
+	inline static TSet<UCustomizableObject*> AlreadyValidatedRootObjects;
+	
+	/** Method invoked once the validation of all assets has been completed. */
+	static void OnPostCOsValidation();
+};
diff --git a/Engine/Plugins/Mutable/Source/MutableValidation/Private/MuV/AssetValidator_ReferencedCustomizableObjects.h b/Engine/Plugins/Mutable/Source/MutableValidation/Private/MuV/AssetValidator_ReferencedCustomizableObjects.h
new file mode 100644
--- /dev/null
+++ b/Engine/Plugins/Mutable/Source/MutableValidation/Private/MuV/AssetValidator_ReferencedCustomizableObjects.h
@@ -0,0 +1,54 @@
+// Copyright Epic Games, Inc. All Rights Reserved.
+
+#pragma once
+
+#include "EditorValidatorBase.h"
+#include "AssetRegistry/AssetData.h"
+#include "AssetRegistry/AssetRegistryModule.h"
+
+#include "AssetValidator_ReferencedCustomizableObjects.generated.h"
+
+class FText;
+class FName;
+class UObject;
+class UCustomizableObject;
+
+UCLASS()
+class UAssetValidator_ReferencedCustomizableObjects : public UEditorValidatorBase
+{
+	GENERATED_BODY()
+
+public:
+	UAssetValidator_ReferencedCustomizableObjects();
+	
+protected:
+
+	// UEditorValidatorBase
+	virtual bool CanValidateAsset_Implementation(const FAssetData& InAssetData, UObject* InAsset, FDataValidationContext& InContext) const override;
+	virtual EDataValidationResult ValidateLoadedAsset_Implementation(const FAssetData& InAssetData, UObject* InAsset, FDataValidationContext& InContext) override;
+	// UEditorValidatorBase
+
+private:
+
+	/** Get a set of referencer objects to a provided UObject.
+	 * @param InAsset A Pointer to the asset to get all referencers of.
+	 * @param InAssetRegistry AssetRegistry object to be used to perform the referencers search.
+	 * @return A set with all the referencing packages.
+	 */
+	TSet<FName> GetAllAssetReferencers(const UObject* InAsset,const IAssetRegistry& InAssetRegistry) const;
+
+	/** Returns a set of customizable objects from the provided collection of packages.
+	 * @param InPackagesToCheck List of packages to scan for CustomizableObjects.
+	 * @param InAssetRegistry AssetRegistry object used to perform the asset search.
+	 * @return A collection of unique Customizable Objets found in the provided packages.
+	 */
+	TSet<UCustomizableObject*> FindCustomizableObjects(const TSet<FName>& InPackagesToCheck,const IAssetRegistry& InAssetRegistry) const;
+
+	/** Validates all the customizable objects provided and sets the validator status accordingly. It will not make the validator fail.
+	 * @param InAsset Input asset provided to the validator.
+	 * @param InCustomizableObjectsToValidate Customizable Objects we want to validate with IsDataValid()
+	 * @param InValidationErrors List to fill with the warnings and errors generated during the validation of the Customizable objects
+	 */
+	void ValidateCustomizableObjects(UObject* InAsset, const TSet<UCustomizableObject*>& InCustomizableObjectsToValidate);
+
+};
diff --git a/Engine/Plugins/Mutable/Source/MutableValidation/Private/MuV/COIBakingTestCommandlet.cpp b/Engine/Plugins/Mutable/Source/MutableValidation/Private/MuV/COIBakingTestCommandlet.cpp
--- a/Engine/Plugins/Mutable/Source/MutableValidation/Private/MuV/COIBakingTestCommandlet.cpp
+++ b/Engine/Plugins/Mutable/Source/MutableValidation/Private/MuV/COIBakingTestCommandlet.cpp
@@ -86,7 +86,7 @@ int32 UCOIBakingTestCommandlet::Main(const FString& Params)
 	UCustomizableObject* InstanceCustomizableObject = TargetInstance->GetCustomizableObject();
 	if (!InstanceCustomizableObject)
 	{
-		UE_LOG(LogMutable,Error,TEXT("The instance %s does not have a CO to compile : Exitting commandlet."), *TargetInstance->GetName());
+		UE_LOG(LogMutable,Error,TEXT("The instance %s does not have a CO to compile : Exiting commandlet."), *TargetInstance->GetName());
 		return 1;
 	}
 	
@@ -97,7 +97,7 @@ int32 UCOIBakingTestCommandlet::Main(const FString& Params)
 	CompilationOptions.bUseDiskCompilation = false;
 
 	TSharedRef<FCustomizableObjectCompilationUtility> CompilationUtility = MakeShared<FCustomizableObjectCompilationUtility>();
-	if( !CompilationUtility->CompileCustomizableObject(InstanceCustomizableObject, true, &CompilationOptions ))
+	if(!CompilationUtility->CompileCustomizableObject(*InstanceCustomizableObject, true, &CompilationOptions ))
 	{
 		UE_LOG(LogMutable,Error,TEXT("Failed to compile the target CO. Exiting commandlet."));
 		return 1;
@@ -135,7 +135,7 @@ int32 UCOIBakingTestCommandlet::Main(const FString& Params)
 		// Check the end status of the instance update
 		if (!bWasInstanceUpdateSuccessful)
 		{
-			UE_LOG(LogMutable,Error,TEXT("Failed to successfully update the target COI. Exiting commandlet."));
+			UE_LOG(LogMutable, Error, TEXT("Failed to successfully update the target COI. Exiting commandlet."));
 			return 1;
 		}
 	}
@@ -161,7 +161,7 @@ int32 UCOIBakingTestCommandlet::Main(const FString& Params)
 			UE_LOG(LogMutable,Warning,TEXT("The directory with path  \" %s \" does already exist. This may be produced by an incompelte execution of a previous test. Clearing it out before continuing..."), *GlobalBakingDirectory);
 			if (!IFileManager::Get().DeleteDirectory(*GlobalBakingDirectory,false,true))
 			{
-				UE_LOG(LogMutable,Error,TEXT("Failed to deleta baking directory at path \" %s \"."), *GlobalBakingDirectory);
+				UE_LOG(LogMutable,Error,TEXT("Failed to delete the baking directory at path \" %s \"."), *GlobalBakingDirectory);
 				return 1;
 			}
 		}
diff --git a/Engine/Plugins/Mutable/Source/MutableValidation/Private/MuV/COIBakingTestCommandlet.h b/Engine/Plugins/Mutable/Source/MutableValidation/Private/MuV/COIBakingTestCommandlet.h
new file mode 100644
--- /dev/null
+++ b/Engine/Plugins/Mutable/Source/MutableValidation/Private/MuV/COIBakingTestCommandlet.h
@@ -0,0 +1,32 @@
+﻿// Copyright Epic Games, Inc. All Rights Reserved.
+#pragma once
+
+#include "CoreMinimal.h"
+#include "Commandlets/Commandlet.h"
+#include "MuV/CustomizableObjectCompilationUtility.h"
+
+#include "COIBakingTestCommandlet.generated.h"
+
+class UCustomizableObjectInstance;
+struct FUpdateContext;
+
+/**
+ * Commandlet designed to test the baking of instances.
+ * It expects an instance UAsset being provided to it so it can later compile it's CO, update the instance itself, and then bake it.
+ * EX : -Run=COIBakingTest -EnablePlugins=MutableTesting -CustomizableObjectInstance="/MutableTesting/Cyborg_Character/Character_Inst" -AllowCommandletRendering
+ */
+UCLASS()
+class MUTABLEVALIDATION_API UCOIBakingTestCommandlet : public UCommandlet
+{
+	GENERATED_BODY()
+	
+public:
+	virtual int32 Main(const FString& Params) override;
+	
+private:
+	
+	/** Instance targeted for baking */
+	UPROPERTY()
+	TObjectPtr<UCustomizableObjectInstance> TargetInstance = nullptr;
+};
+
diff --git a/Engine/Plugins/Mutable/Source/MutableValidation/Private/MuV/COIBulkUpdateTestCommandlet.cpp b/Engine/Plugins/Mutable/Source/MutableValidation/Private/MuV/COIBulkUpdateTestCommandlet.cpp
--- a/Engine/Plugins/Mutable/Source/MutableValidation/Private/MuV/COIBulkUpdateTestCommandlet.cpp
+++ b/Engine/Plugins/Mutable/Source/MutableValidation/Private/MuV/COIBulkUpdateTestCommandlet.cpp
@@ -27,7 +27,7 @@ int32 UCOIBulkUpdateTestCommandlet::Main(const FString& Params)
 	GIsRunningUnattendedScript = true;
 
 	// Get the path where to look for the Customizable Object Instances we want to validate
-	FString InstancesPackagePath;
+	FName InstancesPackagePath;
 	{
 		if (!FParse::Value(*Params, TEXT("InstancesPackagePath="), InstancesPackagePath))
 		{
@@ -35,14 +35,14 @@ int32 UCOIBulkUpdateTestCommandlet::Main(const FString& Params)
 			return 1;
 		}
 
-		if (InstancesPackagePath.IsEmpty())
+		if (InstancesPackagePath.IsNone())
 		{
 			UE_LOG(LogMutable,Error,TEXT("The path to scan can not be empty"))
 			return 1;
 		}
 	}
 
-	UE_LOG(LogMutable, Log, TEXT("(string) instance_search_path : %s "), *InstancesPackagePath);
+	UE_LOG(LogMutable, Log, TEXT("(string) instance_search_path : %s "), *InstancesPackagePath.ToString());
 	
 	// Load the asset registry system so we can proceed without issues
 	PrepareAssetRegistry();
@@ -50,20 +50,8 @@ int32 UCOIBulkUpdateTestCommandlet::Main(const FString& Params)
 	LogGlobalSettings();
 	
 	// Cache all UAssets (find a way to not have them in memory yet, not until we need them)
-	TArray<FAssetData> FoundAssetData;
-	{
-		FARFilter Filter;
-		Filter.ClassPaths.Add(UCustomizableObjectInstance::StaticClass()->GetClassPathName());
-		Filter.PackagePaths.Add( FName(*InstancesPackagePath) );
-		Filter.bRecursivePaths = true;
-		
-		const FAssetRegistryModule& AssetRegistryModule = FModuleManager::LoadModuleChecked<FAssetRegistryModule>(AssetRegistryConstants::ModuleName);
-		
-		UE_LOG(LogMutable,Display,TEXT("Searching for COIs to test at path : %s ."), *InstancesPackagePath);
-		AssetRegistryModule.Get().GetAssets(Filter, FoundAssetData);
-		UE_LOG(LogMutable,Display,TEXT("Search of instances completed. Found %i mutable instances to validate."), FoundAssetData.Num());
-	}
-
+	const TArray<FAssetData> FoundAssetData = FindAllAssetsAtPath(InstancesPackagePath, UCustomizableObjectInstance::StaticClass());
+	
 	// Early exit if no instances could be found for testing.
 	if (FoundAssetData.IsEmpty())
 	{
@@ -73,7 +61,7 @@ int32 UCOIBulkUpdateTestCommandlet::Main(const FString& Params)
 
 	// Get all the CO's that need compilation before proceeding
 	TMap<UCustomizableObject*, TArray<TStrongObjectPtr<UCustomizableObjectInstance>>>  MutableResources;
-	for (FAssetData& Data : FoundAssetData)
+	for (const FAssetData& Data : FoundAssetData)
 	{
 		UObject* LoadedAsset =  MutablePrivate::LoadObject(Data);
 		check (LoadedAsset);
@@ -153,28 +141,30 @@ int32 UCOIBulkUpdateTestCommandlet::Main(const FString& Params)
 		CompilationOptions.TargetPlatform = TargetCompilationPlatform;
 		
 		// Compile the current CO object
-		if (!CompilationUtility->CompileCustomizableObject(CustomizableObject, false, &CompilationOptions))	// Do not log mutable data since mongoDB will not be able to handle it correctly 
+		if (!CompilationUtility->CompileCustomizableObject(*CustomizableObject, false, &CompilationOptions))	// Do not log mutable data since mongoDB will not be able to handle it correctly 
 		{
-			UE_LOG(LogMutable,Error,TEXT("The CO %s could not be compiled succesfully. Skipping the update of all COIs that use it."), *CustomizableObjectName )
+			UE_LOG(LogMutable, Error,TEXT("The CO %s could not be compiled successfully. Skipping the update of all COIs that use it."), *CustomizableObjectName )
 			continue;
 		}
 
 		// Now that CO has been compiled proceed with the update of the instances that use it
 		const uint32 CustomizableObjectInstanceCount = ResourcesIterator->Value.Num();
-		UE_LOG(LogMutable,Display,TEXT("Starting update of the \"%u\" instances with CO : \"%s\"."), CustomizableObjectInstanceCount, *CustomizableObjectName);
+		UE_LOG(LogMutable, Display, TEXT("Starting update of the \"%u\" instances with CO : \"%s\"."), CustomizableObjectInstanceCount, *CustomizableObjectName);
 
 		// Iterate over all the COIs of the CO and update them
 		TArray<TStrongObjectPtr<UCustomizableObjectInstance>>& Instances = ResourcesIterator->Value;
 		while (!Instances.IsEmpty())
 		{
 			TStrongObjectPtr<UCustomizableObjectInstance> Instance = Instances[0];
+			check(Instance);
+			
 			Instances.RemoveAt(0);
-			UE_LOG(LogMutable,Display,TEXT("\t( %u / %u ) Processing instance : \"%s\" ."),CurrentInstanceIndex++, TotalAmountOfInstances ,*Instance->GetName());
+			UE_LOG(LogMutable, Display, TEXT("\t( %u / %u ) Processing instance : \"%s\" ."),CurrentInstanceIndex++, TotalAmountOfInstances ,*Instance->GetName());
 
 			CollectGarbage(GARBAGE_COLLECTION_KEEPFLAGS, true);
 			
 			// Update each one of the instances and notify if the update failed in any manner
-			InstanceUpdatingUtility->UpdateInstance(Instance.Get());
+			InstanceUpdatingUtility->UpdateInstance(*Instance);
 			
 			// Remove standalone flag from the instance so we can GC it while keeping other standalone objects
 			Instance->ClearFlags(EObjectFlags::RF_Standalone);
diff --git a/Engine/Plugins/Mutable/Source/MutableValidation/Private/MuV/COIBulkUpdateTestCommandlet.h b/Engine/Plugins/Mutable/Source/MutableValidation/Private/MuV/COIBulkUpdateTestCommandlet.h
new file mode 100644
--- /dev/null
+++ b/Engine/Plugins/Mutable/Source/MutableValidation/Private/MuV/COIBulkUpdateTestCommandlet.h
@@ -0,0 +1,24 @@
+﻿// Copyright Epic Games, Inc. All Rights Reserved.
+#pragma once
+
+#include "Commandlets/Commandlet.h"
+#include "MuV/CustomizableObjectCompilationUtility.h"
+#include "MuV/CustomizableObjectInstanceUpdateUtility.h"
+
+#include "COIBulkUpdateTestCommandlet.generated.h"
+
+/**
+ * Testing commandlet designed to allow the mass update of instances in a given project path.
+ * It will run in sequence where each of the instances will get updated after their CO gets compiled.
+ *
+ * Example commandline that updates all instances found in the plugin content directory (and all subdirectories):
+ * -run=COIBulkUpdateTest -EnablePlugins=MutableTesting -InstancesPackagePath="/MutableTesting" -AllowCommandletRendering
+ */
+UCLASS()
+class UCOIBulkUpdateTestCommandlet : public UCommandlet 
+{
+	GENERATED_BODY()
+
+public:
+	virtual int32 Main(const FString& Params) override;
+};
diff --git a/Engine/Plugins/Mutable/Source/MutableValidation/Private/MuV/CacheAssetRegistryCommandlet.h b/Engine/Plugins/Mutable/Source/MutableValidation/Private/MuV/CacheAssetRegistryCommandlet.h
new file mode 100644
--- /dev/null
+++ b/Engine/Plugins/Mutable/Source/MutableValidation/Private/MuV/CacheAssetRegistryCommandlet.h
@@ -0,0 +1,17 @@
+// Copyright Epic Games, Inc. All Rights Reserved.
+
+#pragma once
+
+#include "Commandlets/Commandlet.h"
+
+#include "CacheAssetRegistryCommandlet.generated.h"
+
+/** Simple commandlet used to fill-up the asset registry cache for the execution of the mutable tests */
+UCLASS()
+class UCacheAssetRegistryCommandlet : public UCommandlet
+{
+	GENERATED_BODY()
+
+public:
+	virtual int32 Main(const FString& Params) override;
+};
diff --git a/Engine/Plugins/Mutable/Source/MutableValidation/Private/MuV/CustomizableObjectBulkValidationCommandlet.cpp b/Engine/Plugins/Mutable/Source/MutableValidation/Private/MuV/CustomizableObjectBulkValidationCommandlet.cpp
new file mode 100644
--- /dev/null
+++ b/Engine/Plugins/Mutable/Source/MutableValidation/Private/MuV/CustomizableObjectBulkValidationCommandlet.cpp
@@ -0,0 +1,156 @@
+// Copyright Epic Games, Inc. All Rights Reserved.
+
+#include "MuV/CustomizableObjectBulkValidationCommandlet.h"
+
+#include "ScopedLogSection.h"
+#include "ValidationUtils.h"
+#include "MuCO/CustomizableObject.h"
+#include "MuCO/CustomizableObjectSystemPrivate.h"
+#include "MuCO/ICustomizableObjectEditorModule.h"
+#include "MuCO/LoadUtils.h"
+#include "MuCO/LogBenchmarkUtil.h"
+
+
+int32 UCustomizableObjectBulkValidationCommandlet::Main(const FString& Params)
+{
+	LLM_SCOPE_BYNAME(TEXT("CustomizableObjectBulkValidationCommandlet"));
+
+	// Prepare the environment for the testing -------------------------------------------------------------------------------------------------------
+	
+	// Ensure we have set the mutable system to the benchmarking mode and that we are reporting benchmarking data
+	FLogBenchmarkUtil::SetBenchmarkReportingStateOverride(true);
+	UCustomizableObjectSystemPrivate::SetUsageOfBenchmarkingSettings(true);
+	
+	// Ensure we do not show any OK dialog since we are not a user that can interact with them
+	GIsRunningUnattendedScript = true;
+
+	// Parse the arguments provided ------------------------------------------------------------------------------------------------------------------
+	
+	// Parse the value for the platform to be used for the compilation of the COs
+	const ITargetPlatform* TargetCompilationPlatform = GetCompilationPlatform(Params);
+	if (!TargetCompilationPlatform)
+	{
+		return 1;
+	}
+	// Parse if we want to use disk compilation or not
+	const bool bUseDiskCompilation = GetDiskCompilationArg(Params);
+	// Get the amount of instances to generate if parameter was provided (it will get multiplied by the amount of states later so this is a minimun value)
+	const uint32 InstancesToGenerate = GetTargetAmountOfInstances(Params);
+
+	// Work only for the root Customizable Objects found
+	bool bOnlyTestRootObjects = false;
+	FParse::Bool(*Params,TEXT("SkipNonRootObjects="),bOnlyTestRootObjects);
+	if (bOnlyTestRootObjects)
+	{
+		UE_LOG(LogMutable, Display, TEXT("Only the root COs will be tested")) ;
+	}
+
+	// Find the assets to be tested based on the path provided ---------------------------------------------------------------------------------------
+	
+	// Parse the path to search for COs on
+	FName CustomizableObjectsSearchPath;
+	{
+		if (!FParse::Value(*Params, TEXT("CustomizableObjectsSearchPath="), CustomizableObjectsSearchPath))
+		{
+			UE_LOG(LogMutable, Error, TEXT("Failed to parse Customizable Object search path from the provided argument : %s"), *Params)
+			return 1;
+		}
+
+		if (CustomizableObjectsSearchPath.IsNone())
+		{
+			UE_LOG(LogMutable, Error, TEXT("The path to scan can not be empty"))
+			return 1;
+		}
+	}
+
+	// Perform a blocking search to ensure all assets used by mutable are reachable using the AssetRegistry
+	PrepareAssetRegistry();
+	
+	// Then run what we know as the UCustomizableObjectValidationCommandlet body for each one of them. Keep in mind that you may need to do some cleanup
+	// between runs.
+	TArray<FAssetData> FoundAssetData;
+	{
+		LLM_SCOPE_BYNAME(TEXT("CustomizableObjectBulkValidationCommandlet/AssetsSearch"));
+
+		FoundAssetData = FindAllAssetsAtPath(CustomizableObjectsSearchPath, UCustomizableObject::StaticClass());
+
+		// Early exit if no instances could be found for testing.
+		if (FoundAssetData.IsEmpty())
+		{
+			UE_LOG(LogMutable, Error, TEXT("Aborting Bulk Customizable Object Validation Test: No assets could be found at the provided package path"));
+			return 1;
+		}
+		
+		// Log all the Customizable Objects to be tested:
+		UE_LOG(LogMutable, Display, TEXT("Found a total of %u Customizable Objects to validate. Some may be discarded based on the test settings."), FoundAssetData.Num())
+		for (const FAssetData& MutableAssetData : FoundAssetData)
+		{
+			UE_LOG(LogMutable, Display, TEXT("\t - %s (%s)"), *MutableAssetData.AssetName.ToString(), *MutableAssetData.PackageName.ToString());
+		}
+	}
+	
+	// Make sure there is nothing else that the engine needs to do before starting our test
+	Wait(60);
+	
+	LogGlobalSettings();
+	
+	// All pre-testing operations completed. Starting the testing of COs -----------------------------------------------------------------------------
+
+	const int32 FoundAssetsCount = FoundAssetData.Num();
+	for (int32 AssetIndex = 0; AssetIndex < FoundAssetsCount; ++AssetIndex)
+	{
+		LLM_SCOPE_BYNAME(TEXT("CustomizableObjectBulkValidationCommandlet/COTest"));
+
+		const FAssetData& CustomizableObjectAssetData = FoundAssetData[AssetIndex];
+		TObjectPtr<UObject> FoundObject = MutablePrivate::LoadObject(CustomizableObjectAssetData);
+		if (!FoundObject)
+		{
+			UE_LOG(LogMutable, Error, TEXT("Failed to load the asset with path : %s ."), *CustomizableObjectAssetData.GetSoftObjectPath().ToString())
+			continue;
+		}
+
+		TObjectPtr<UCustomizableObject> TargetCustomizableObject = Cast<UCustomizableObject>(FoundObject);
+		if (!TargetCustomizableObject)
+		{
+			UE_LOG(LogMutable, Error, TEXT("Failed to cast found UObject to UCustomizableObject."));
+			continue;
+		}		
+		
+		const FString CustomizableObjectName = TargetCustomizableObject->GetName();
+		
+		// If we only want to test root objects and this is not one go to the next CO
+		if (bOnlyTestRootObjects && !ICustomizableObjectEditorModule::GetChecked().GetRootObject(TargetCustomizableObject.Get()))
+		{
+			UE_LOG(LogMutable, Display, TEXT("Skipping CO \"%s\" as it is not a root CO."), *CustomizableObjectName)
+			continue;
+		}
+		
+		// Get the CustomizableObject.
+		{
+			// TODO: MOVE THIS Scope section inside TestCustomizableObject once the parsing system has been updated to handle this new section type
+			// By keeping it here we avoid our current implementation of the single CO Validation from using this new type of not yet supported by the
+			// parser section.
+			const FScopedLogSection ObjectSection (EMutableLogSection::Object, FName(CustomizableObjectName));
+			
+			// Body of the test ----------------------------------------------------------------------------------------------------------------------
+			TestCustomizableObject(*TargetCustomizableObject, *TargetCompilationPlatform, InstancesToGenerate, bUseDiskCompilation);
+			// ---------------------------------------------------------------------------------------------------------------------------------------
+		}
+
+		
+		// Try to collect the garbage
+		{
+			if (GIsInitialLoad)
+			{
+				UE_LOG(LogMutable, Warning, TEXT("GC will not run as GIsInitialLoad is currently set to true."));
+			}
+			
+			CollectGarbage(RF_NoFlags, true);
+		}
+
+		UE_LOG(LogMutable, Display, TEXT("Validated %u/%u assets."), AssetIndex + 1, FoundAssetsCount);
+	}
+
+	UE_LOG(LogMutable, Display, TEXT("Mutable commandlet finished."));
+	return 0;
+}
diff --git a/Engine/Plugins/Mutable/Source/MutableValidation/Private/MuV/CustomizableObjectBulkValidationCommandlet.h b/Engine/Plugins/Mutable/Source/MutableValidation/Private/MuV/CustomizableObjectBulkValidationCommandlet.h
new file mode 100644
--- /dev/null
+++ b/Engine/Plugins/Mutable/Source/MutableValidation/Private/MuV/CustomizableObjectBulkValidationCommandlet.h
@@ -0,0 +1,19 @@
+// Copyright Epic Games, Inc. All Rights Reserved.
+
+#pragma once
+
+#include "Commandlets/Commandlet.h"
+#include "CustomizableObjectBulkValidationCommandlet.generated.h"
+
+/**
+ * Commandlet similar to "UCustomizableObjectValidationCommandlet" but in this case it performs the same tests as in the mentioned commandlet but over
+ * a series of COs located at a given path.
+ */
+UCLASS()
+class UCustomizableObjectBulkValidationCommandlet : public UCommandlet
+{
+	GENERATED_BODY()
+	
+public:
+	virtual int32 Main(const FString& Params) override;
+};
diff --git a/Engine/Plugins/Mutable/Source/MutableValidation/Private/MuV/CustomizableObjectCompilationUtility.cpp b/Engine/Plugins/Mutable/Source/MutableValidation/Private/MuV/CustomizableObjectCompilationUtility.cpp
--- a/Engine/Plugins/Mutable/Source/MutableValidation/Private/MuV/CustomizableObjectCompilationUtility.cpp
+++ b/Engine/Plugins/Mutable/Source/MutableValidation/Private/MuV/CustomizableObjectCompilationUtility.cpp
@@ -16,13 +16,12 @@
 #include "HAL/FileManager.h"
 
 
-bool FCustomizableObjectCompilationUtility::CompileCustomizableObject(UCustomizableObject* InCustomizableObject, const bool bShouldLogMutableLogs /* = true */, const FCompilationOptions* InCompilationOptionsOverride  /* nullptr */)
+bool FCustomizableObjectCompilationUtility::CompileCustomizableObject(UCustomizableObject& InCustomizableObject, const bool bShouldLogMutableLogs /* = true */, const FCompilationOptions* InCompilationOptionsOverride  /* nullptr */)
 {
 	LLM_SCOPE_BYNAME(TEXT("FCustomizableObjectCompilationUtility/Compile"));
 	const FScopedLogSection CompilationSection (EMutableLogSection::Compilation);	
-	check(InCustomizableObject);
 	
-	CustomizableObject = TStrongObjectPtr(InCustomizableObject);
+	CustomizableObject = TStrongObjectPtr(&InCustomizableObject);
 	
 	// Handle the overriding of the compilation configuration of the CO ------------------------------------------------------------------------------
 	
@@ -36,7 +35,7 @@ bool FCustomizableObjectCompilationUtility::CompileCustomizableObject(UCustomiza
 	}
 	else
 	{
-		CompilationOptions = InCustomizableObject->GetPrivate()->GetCompileOptions();
+		CompilationOptions = CustomizableObject->GetPrivate()->GetCompileOptions();
 		UE_LOG(LogMutable,Display,TEXT("Compiling CO using it's own compilation options."));
 	}
 
@@ -44,7 +43,7 @@ bool FCustomizableObjectCompilationUtility::CompileCustomizableObject(UCustomiza
 	// Mutable is able to run without one but we want to be explicit in the context of testing.
 	if (!CompilationOptions.TargetPlatform)
 	{
-		UE_LOG(LogMutable, Error, TEXT("The compilation of the %s model could not be started : No explicit platform was provided."), *InCustomizableObject->GetName());
+		UE_LOG(LogMutable, Error, TEXT("The compilation of the %s model could not be started : No explicit platform was provided."), *CustomizableObject->GetName());
 		return false;
 	}
 
@@ -108,7 +107,7 @@ bool FCustomizableObjectCompilationUtility::CompileCustomizableObject(UCustomiza
 		
 		UE_LOG(LogMutable,Display,TEXT("Compiling Customizable Object..."));
 		
-		TSharedRef<FCompilationRequest> CompileRequest = MakeShared<FCompilationRequest>(*InCustomizableObject);
+		TSharedRef<FCompilationRequest> CompileRequest = MakeShared<FCompilationRequest>(*CustomizableObject);
 		CompileRequest->bSilentCompilation = false;
 		CompileRequest->Options = CompilationOptions;
 		
@@ -149,7 +148,7 @@ bool FCustomizableObjectCompilationUtility::CompileCustomizableObject(UCustomiza
 
 		// Report the time we took to run the compilation
 		const double CompilationEndSeconds = FPlatformTime::Seconds() - CompilationStartSeconds;
-		UE_LOG(LogMutable, Display, TEXT("The compilation of the %s CO model took %f seconds."), *InCustomizableObject->GetName(), CompilationEndSeconds);
+		UE_LOG(LogMutable, Display, TEXT("The compilation of the %s CO model took %f seconds."), *CustomizableObject->GetName(), CompilationEndSeconds);
 		if (bShouldLogMutableLogs)
 		{
 			UE_LOG(LogMutable, Log, TEXT("(double) model_compile_time_ms : %f "), CompilationEndSeconds * 1000);
diff --git a/Engine/Plugins/Mutable/Source/MutableValidation/Private/MuV/CustomizableObjectCompilationUtility.h b/Engine/Plugins/Mutable/Source/MutableValidation/Private/MuV/CustomizableObjectCompilationUtility.h
--- a/Engine/Plugins/Mutable/Source/MutableValidation/Private/MuV/CustomizableObjectCompilationUtility.h
+++ b/Engine/Plugins/Mutable/Source/MutableValidation/Private/MuV/CustomizableObjectCompilationUtility.h
@@ -22,7 +22,7 @@ public:
 	 * @param InCompilationOptionsOverride The configuration for the compilation of the CO we want to use instead of the one part of the CO.
 	 * @return True if the compilation was successful and false if it failed.
 	 */
-	bool CompileCustomizableObject(UCustomizableObject* InCustomizableObject, const bool bShouldLogMutableLogs = true, const FCompilationOptions* InCompilationOptionsOverride = nullptr );
+	bool CompileCustomizableObject(UCustomizableObject& InCustomizableObject, const bool bShouldLogMutableLogs = true, const FCompilationOptions* InCompilationOptionsOverride = nullptr );
 
 private:
 
diff --git a/Engine/Plugins/Mutable/Source/MutableValidation/Private/MuV/CustomizableObjectInstanceUpdateUtility.cpp b/Engine/Plugins/Mutable/Source/MutableValidation/Private/MuV/CustomizableObjectInstanceUpdateUtility.cpp
--- a/Engine/Plugins/Mutable/Source/MutableValidation/Private/MuV/CustomizableObjectInstanceUpdateUtility.cpp
+++ b/Engine/Plugins/Mutable/Source/MutableValidation/Private/MuV/CustomizableObjectInstanceUpdateUtility.cpp
@@ -12,16 +12,15 @@
 #include "Components/SkeletalMeshComponent.h"
 
 
-bool FCustomizableObjectInstanceUpdateUtility::UpdateInstance(UCustomizableObjectInstance* InInstance)
+bool FCustomizableObjectInstanceUpdateUtility::UpdateInstance(UCustomizableObjectInstance& InInstance)
 {
 	LLM_SCOPE_BYNAME(TEXT("FCustomizableObjectInstanceUpdateUtility/UpdateInstance"));
 	const FScopedLogSection UpdateSection (EMutableLogSection::Update);	
 	
-	check (InInstance);
 	check (ComponentsBeingUpdated.IsEmpty());
 	
 	// Cache the instance being updated for reference once in the update end callback
-	Instance = TStrongObjectPtr(InInstance);
+	Instance = TStrongObjectPtr(&InInstance);
 	
 	// Schedule the update of the COI
 	{
diff --git a/Engine/Plugins/Mutable/Source/MutableValidation/Private/MuV/CustomizableObjectInstanceUpdateUtility.h b/Engine/Plugins/Mutable/Source/MutableValidation/Private/MuV/CustomizableObjectInstanceUpdateUtility.h
--- a/Engine/Plugins/Mutable/Source/MutableValidation/Private/MuV/CustomizableObjectInstanceUpdateUtility.h
+++ b/Engine/Plugins/Mutable/Source/MutableValidation/Private/MuV/CustomizableObjectInstanceUpdateUtility.h
@@ -20,7 +20,7 @@ public:
 	 * @param InInstance Instance to be compiled
 	 * @return True if the update was successful and false otherwise
 	 */
-	bool UpdateInstance(UCustomizableObjectInstance* InInstance);
+	bool UpdateInstance(UCustomizableObjectInstance& InInstance);
 	
 
 private:
diff --git a/Engine/Plugins/Mutable/Source/MutableValidation/Private/MuV/CustomizableObjectValidationCommandlet.cpp b/Engine/Plugins/Mutable/Source/MutableValidation/Private/MuV/CustomizableObjectValidationCommandlet.cpp
--- a/Engine/Plugins/Mutable/Source/MutableValidation/Private/MuV/CustomizableObjectValidationCommandlet.cpp
+++ b/Engine/Plugins/Mutable/Source/MutableValidation/Private/MuV/CustomizableObjectValidationCommandlet.cpp
@@ -37,43 +37,41 @@ int32 UCustomizableObjectValidationCommandlet::Main(const FString& Params)
 	GIsRunningUnattendedScript = true;
 	
 	// Get the package name of the CO to test
-	FString CustomizableObjectAssetPath = "";
-	if (!FParse::Value(*Params, TEXT("CustomizableObject="), CustomizableObjectAssetPath))
+	TStrongObjectPtr<UCustomizableObject> TargetCustomizableObject = nullptr;
 	{
-		UE_LOG(LogMutable,Error,TEXT("Failed to parse Customizable Object package name from provided argument : %s"),*Params)
-		return 1;
-	}
+		FString CustomizableObjectAssetPath = "";
+		if (!FParse::Value(*Params, TEXT("CustomizableObject="), CustomizableObjectAssetPath))
+		{
+			UE_LOG(LogMutable, Error, TEXT("Failed to parse Customizable Object package name from provided argument : %s"), *Params)
+			return 1;
+		}
 	
-	// Get the amount of instances to generate if parameter was provided (it will get multiplied by the amount of states later so this is a minimun value)
-	uint32 InstancesToGenerate = 16;
-	if (!FParse::Value(*Params, TEXT("InstanceGenerationCount="),InstancesToGenerate))
-	{
-		UE_LOG(LogMutable,Display,TEXT("Instance generation count not specified. Using default value : %u"),InstancesToGenerate);
-	}
+		// Load the resource
+		UObject* FoundObject = MutablePrivate::LoadObject(FSoftObjectPath(CustomizableObjectAssetPath));
+		if (!FoundObject)
+		{
+			UE_LOG(LogMutable, Error, TEXT("Failed to retrieve UObject from path %s"), *CustomizableObjectAssetPath);
+			return 1;
+		}
 	
-	// Load the resource
-	UObject* FoundObject = MutablePrivate::LoadObject(FSoftObjectPath(CustomizableObjectAssetPath));
-	if (!FoundObject)
-	{
-		UE_LOG(LogMutable,Error,TEXT("Failed to retrieve UObject from path %s"),*CustomizableObjectAssetPath);
-		return 1;
+		// Get the CustomizableObject.
+		TargetCustomizableObject = TStrongObjectPtr(Cast<UCustomizableObject>(FoundObject));
+		if (!TargetCustomizableObject)
+		{
+			UE_LOG(LogMutable, Error, TEXT("Failed to cast found UObject to UCustomizableObject."));
+			return 1;
+		}
 	}
 	
-	// Get the CustomizableObject.
-	ToTestCustomizableObject = Cast<UCustomizableObject>(FoundObject);
-	if (!ToTestCustomizableObject)
-	{
-		UE_LOG(LogMutable,Error,TEXT("Failed to cast found UObject to UCustomizableObject."));
-		return 1;
-	}
-
 	// What platform we want to compile the CO against
-	ITargetPlatform* TargetCompilationPlatform = ParseCompilationPlatform(Params);
+	ITargetPlatform* TargetCompilationPlatform = GetCompilationPlatform(Params);
 	if (!TargetCompilationPlatform)
 	{
-		UE_LOG(LogMutable,Error,TEXT("No Target Compilation Platform was provided for the compilation of the CO."));
 		return 1;
 	}
+	// Parse if we want to use disk compilation or not
+	const bool bUseDiskCompilation = GetDiskCompilationArg(Params);
+	const uint32 InstancesToGenerate = GetTargetAmountOfInstances(Params);
 	
 	// Perform a blocking search to ensure all assets used by mutable are reachable using the AssetRegistry
 	PrepareAssetRegistry();
@@ -82,176 +80,11 @@ int32 UCustomizableObjectValidationCommandlet::Main(const FString& Params)
 	Wait(60);
 
 	LogGlobalSettings();
-	
-	// Compile the Customizable Object ------------------------------------------------------------------------------ //
-	bool bWasCoCompilationSuccessful = false;
-	{
-		LLM_SCOPE_BYNAME(TEXT("CustomizableObjectValidationCommandlet/Compile"));
-
-		// Override some configurations that may have been changed by the user
-		FCompilationOptions CompilationOptions = GetCompilationOptionsForBenchmarking(*ToTestCustomizableObject);
-		
-		// Set the target compilation platform based on what the caller wants
-		CompilationOptions.TargetPlatform = TargetCompilationPlatform;
-
-		// Disk cache usage for compilation operation : Override if the user provided an argument with a different value than the default one of the CO
-		CompilationOptions.bUseDiskCompilation = false;
-		FParse::Bool(*Params,TEXT("UseDiskCompilation="),CompilationOptions.bUseDiskCompilation);
-		
-		TSharedRef<FCustomizableObjectCompilationUtility> CompilationUtility = MakeShared<FCustomizableObjectCompilationUtility>();
-		bWasCoCompilationSuccessful = CompilationUtility->CompileCustomizableObject(ToTestCustomizableObject, true, &CompilationOptions);
-	}
-	// -------------------------------------------------------------------------------------------------------------- //
-
-	if (!bWasCoCompilationSuccessful)
-	{
-		UE_LOG(LogMutable,Error,TEXT("The compilation of the Customizable object was not succesfull : No instances will be generated."));
-		return 1;		// Validation failed
-	}
-
-	UE_LOG(LogMutable,Display,TEXT("Customizable Object was compiled succesfully."));
-	
-	// GHet the total size of the streaming data of the model ---------------------------------------------- //
-	{
-		const TSharedPtr<const mu::FModel> MutableModel = ToTestCustomizableObject->GetPrivate()->GetModel();
-		check (MutableModel);
-
-		// Roms ---------------------- //
-		{
-			const int32 RomCount =  MutableModel->GetRomCount();
-			int64 TotalRomSizeBytes = 0;
-			for (int32 RomIndex = 0; RomIndex < RomCount; RomIndex++)
-			{
-				const uint32 RomByteSize = MutableModel->GetRomSize(RomIndex);
-				TotalRomSizeBytes += RomByteSize;
-			}
-
-			// Print MTU parseable logs
-			UE_LOG(LogMutable, Log,TEXT("(int) model_rom_count : %d "), RomCount);
-			UE_LOG(LogMutable, Log,TEXT("(int) model_roms_size : %lld "), TotalRomSizeBytes);
-		}
-
-		// CO embedded data size ------ //
-		{
-			TArray<uint8> EmbeddedDataBytes{};
-			FMemoryWriter SerializationTarget{EmbeddedDataBytes, false};
-		
-			ToTestCustomizableObject->GetPrivate()->SaveEmbeddedData(SerializationTarget);
-			const int64 COEmbeddedDataSizeBytes = EmbeddedDataBytes.Num();
-		
-			UE_LOG(LogMutable, Log,TEXT("(int) co_embedded_data_bytes : %lld "), COEmbeddedDataSizeBytes);
-		}
-	}
-	
-	// Skip instances updating if no instances should be updated 
-	if (InstancesToGenerate <= 0)
-	{
-		UE_LOG(LogMutable,Display,TEXT("Instances to generate are 0 : No instances will be generated."));
-		return 0;	// No instances are targeted for generation, this will be taken as compilation only test.
-	}
-
-	// Do not generate instances if the selected platform is not the running platform
-	if (TargetCompilationPlatform != GetTargetPlatformManagerRef().GetRunningTargetPlatform())
-	{
-		UE_LOG(LogMutable,Display,TEXT("RunningPlatform != UserProvidedCompilationPlatform : No instances will be generated."));
-		return 0;
-	}
-
-	// Generate target random instances to be tested ------------------------------------------------------------ //
-	bool bWasInstancesCreationSuccessful = true;
-	uint32 GeneratedInstances = 0;
-	{
-		LLM_SCOPE_BYNAME(TEXT("CustomizableObjectValidationCommandlet/GenerateInstances"));
-		
-		// Create a set of instances so we can later test them out
-		bWasInstancesCreationSuccessful = CustomizableObjectBenchmarkingUtils::GenerateDeterministicSetOfInstances(ToTestCustomizableObject,InstancesToGenerate, InstancesToProcess, GeneratedInstances);
-	}
-	// ---------------------------------------------------------------------------------------------------------- //
-
-	UE_LOG(LogMutable, Log,TEXT("(int) generated_instances_count : %u "), GeneratedInstances);
-	
-	// Update the instances generated --------------------------------------------------------------------------- //
-	UE_LOG(LogMutable,Display,TEXT("Updating generated instances..."));
-	bool bInstanceFailedUpdate = false;
-	const double InstancesUpdateStartSeconds = FPlatformTime::Seconds();
-	{
-		LLM_SCOPE_BYNAME(TEXT("CustomizableObjectValidationCommandlet/Update"));
-		
-		TSharedRef<FCustomizableObjectInstanceUpdateUtility> InstanceUpdatingUtility = MakeShared<FCustomizableObjectInstanceUpdateUtility>();
 
-		TStrongObjectPtr<UCustomizableObjectInstance> InstanceToUpdate;
-		while (InstancesToProcess.Dequeue(InstanceToUpdate))
-		{
-			CollectGarbage(GARBAGE_COLLECTION_KEEPFLAGS, true);
-
-			if (!InstanceUpdatingUtility->UpdateInstance(InstanceToUpdate.Get()))
-			{
-				bInstanceFailedUpdate = true;
-			}
-		}
-	}
-	const double InstancesUpdateEndSeconds = FPlatformTime::Seconds();
-	
-	// Notify and log time required by the instances to get updated
-	const double CombinedInstanceUpdateSeconds = InstancesUpdateEndSeconds - InstancesUpdateStartSeconds;
-	UE_LOG(LogMutable, Log,TEXT("(double) combined_update_time_ms : %f "), CombinedInstanceUpdateSeconds * 1000);
-
-	check(GeneratedInstances > 0);
-	const double AverageInstanceUpdateSeconds = CombinedInstanceUpdateSeconds / GeneratedInstances;
-	UE_LOG(LogMutable, Log,TEXT("(double) avg_update_time_ms : %f "), AverageInstanceUpdateSeconds * 1000);
-
-	UE_LOG(LogMutable,Display,TEXT("Generation of Customizable object instances took %f seconds (%f seconds avg)."), CombinedInstanceUpdateSeconds, AverageInstanceUpdateSeconds);
-	// ---------------------------------------------------------------------------------------------------------- //
-
-	// Compute instance update result
-	const bool bInstancesTestedSuccessfully = !bInstanceFailedUpdate && bWasInstancesCreationSuccessful;
-	if (bInstancesTestedSuccessfully)
-    {
-        UE_LOG(LogMutable,Display,TEXT("Generation of Customizable object instances was succesfull."));
-    }
-    else
-    {
-        UE_LOG(LogMutable,Error,TEXT("The generation of Customizable object instances was not succesfull."));
-    }
-	
+	// Body of the test ------------------------------------------------------------------------------------------------------------------------------
+	TestCustomizableObject(*TargetCustomizableObject, *TargetCompilationPlatform, InstancesToGenerate, bUseDiskCompilation);
+	// -----------------------------------------------------------------------------------------------------------------------------------------------
 	
-	// If something failed then fail the commandlet execution
-	UE_LOG(LogMutable,Display,TEXT("Mutable commandlet finished."));
+	UE_LOG(LogMutable, Display, TEXT("Mutable commandlet finished."));
 	return 0;
 }
-
-
-ITargetPlatform* UCustomizableObjectValidationCommandlet::ParseCompilationPlatform(const FString& Params) const
-{
-	// Get the package name of the CO to test
-	FString TargetPlatformName = "";
-	if (!FParse::Value(*Params, TEXT("CompilationPlatformName="), TargetPlatformName))
-	{
-		UE_LOG(LogMutable,Error,TEXT("Failed to parse the target compilation platformm."))
-		return nullptr;
-	}
-
-	// Set the target platform in the context. For now it is the current platform.
-	ITargetPlatformManagerModule* TPM = GetTargetPlatformManager();
-	check(TPM);
-	
-	ITargetPlatform* TargetCompilationPlatform = nullptr;
-	const TArray<ITargetPlatform*> TPMPlatforms = TPM->GetTargetPlatforms();
-	for (ITargetPlatform* Platform : TPMPlatforms)
-	{
-		FString PlatformName = Platform->PlatformName();
-		if (PlatformName.Compare(TargetPlatformName) == 0)
-		{
-			// We have found the platform provided
-			TargetCompilationPlatform = Platform;
-			break;
-		}
-	}
-	
-	if (!TargetCompilationPlatform)
-	{
-		UE_LOG(LogMutable, Error, TEXT("Unable to relate the provided platform name (%s) with the available platforms in this machine."), *TargetPlatformName);
-	}
-
-	return TargetCompilationPlatform;
-}
diff --git a/Engine/Plugins/Mutable/Source/MutableValidation/Private/MuV/CustomizableObjectValidationCommandlet.h b/Engine/Plugins/Mutable/Source/MutableValidation/Private/MuV/CustomizableObjectValidationCommandlet.h
new file mode 100644
--- /dev/null
+++ b/Engine/Plugins/Mutable/Source/MutableValidation/Private/MuV/CustomizableObjectValidationCommandlet.h
@@ -0,0 +1,21 @@
+// Copyright Epic Games, Inc. All Rights Reserved.
+
+#pragma once
+
+#include "Commandlets/Commandlet.h"
+
+#include "CustomizableObjectValidationCommandlet.generated.h"
+
+/**
+ * Commandlet designed to compile and then update a deterministically random set of instances.
+ */
+UCLASS()
+class UCustomizableObjectValidationCommandlet : public UCommandlet 
+{
+	GENERATED_BODY()
+
+public:
+	virtual int32 Main(const FString& Params) override;
+};
+
+
diff --git a/Engine/Plugins/Mutable/Source/MutableValidation/Private/MuV/MutableValidationModule.cpp b/Engine/Plugins/Mutable/Source/MutableValidation/Private/MuV/MutableValidationModule.cpp
new file mode 100644
--- /dev/null
+++ b/Engine/Plugins/Mutable/Source/MutableValidation/Private/MuV/MutableValidationModule.cpp
@@ -0,0 +1,76 @@
+// Copyright Epic Games, Inc. All Rights Reserved.
+
+#include "ISettingsModule.h"
+#include "ISettingsSection.h"
+#include "Modules/ModuleManager.h"
+#include "MuV/MutableValidationSettings.h"
+
+#define LOCTEXT_NAMESPACE "MutableValidationSettings"
+
+/**
+ * StaticMesh editor module
+ */
+class FMutableValidationModule : public FDefaultModuleImpl
+{
+public:
+	// IModuleInterface interface
+	virtual void StartupModule() override;
+	virtual void ShutdownModule() override;
+
+	bool HandleSettingsSaved() const;
+
+private:
+	ISettingsSectionPtr SettingsSectionPtr = nullptr;
+};
+
+IMPLEMENT_MODULE(FMutableValidationModule, MutableValidation);
+
+bool FMutableValidationModule::HandleSettingsSaved() const
+{
+	UMutableValidationSettings* CustomizableObjectSettings = GetMutableDefault<UMutableValidationSettings>();
+
+	if (CustomizableObjectSettings != nullptr)
+	{
+		CustomizableObjectSettings->SaveConfig();
+	}
+
+    return true;
+}
+
+
+void FMutableValidationModule::StartupModule()
+{
+	ISettingsModule* SettingsModule = FModuleManager::GetModulePtr<ISettingsModule>("Settings");
+	if (SettingsModule != nullptr)
+	{ 		
+		 SettingsSectionPtr = SettingsModule->RegisterSettings("Project", "Plugins", "MutableValidationSettings",
+			LOCTEXT("MutableValidationSettings_Setting", "Mutable Validation"),
+			LOCTEXT("MutableValidatioinSettings_Setting_Desc", "Mutable resources validation settings"),
+			GetMutableDefault<UMutableValidationSettings>()
+		);
+		
+		if (SettingsSectionPtr.IsValid())
+        {
+        	SettingsSectionPtr->OnModified().BindRaw(this, &FMutableValidationModule::HandleSettingsSaved);
+        }
+	}
+}
+
+
+void FMutableValidationModule::ShutdownModule()
+{
+	// Unbind OnModified delegate
+	if (SettingsSectionPtr)
+	{
+		SettingsSectionPtr->OnModified().Unbind();
+	}
+	
+	ISettingsModule* SettingsModule = FModuleManager::GetModulePtr<ISettingsModule>("Settings");
+	if (SettingsModule != nullptr)
+	{ 		
+		SettingsModule->UnregisterSettings("Project", "Plugins", "MutableValidationSettings");
+	}
+}
+
+
+#undef LOCTEXT_NAMESPACE
diff --git a/Engine/Plugins/Mutable/Source/MutableValidation/Private/MuV/ScopedLogSection.cpp b/Engine/Plugins/Mutable/Source/MutableValidation/Private/MuV/ScopedLogSection.cpp
--- a/Engine/Plugins/Mutable/Source/MutableValidation/Private/MuV/ScopedLogSection.cpp
+++ b/Engine/Plugins/Mutable/Source/MutableValidation/Private/MuV/ScopedLogSection.cpp
@@ -1,30 +1,46 @@
 ﻿// Copyright Epic Games, Inc. All Rights Reserved.
+
 #include "ScopedLogSection.h"
 #include "Logging/LogMacros.h"
-//#include "UObject/Class.h"
+#include "UObject/NameTypes.h"
+
 
 DECLARE_LOG_CATEGORY_EXTERN(LogMutableValidation, Log, All);
 DEFINE_LOG_CATEGORY(LogMutableValidation);
 
 
-FScopedLogSection::FScopedLogSection(EMutableLogSection Section)
+FScopedLogSection::FScopedLogSection(EMutableLogSection Section, const FName& InSectionTarget)
 {
-	// Two scopes can not be active at the same time (not handled at the parsing level)
-	check(CurrentSection == EMutableLogSection::Undefined);
-	
-	UE_LOG(LogMutableValidation,Log,TEXT(" SECTION START : %s "), *GetLogSectionName(Section));
-	
 	// Set the current section handled by this object
 	CurrentSection = Section;
+
+	// Cache the name of the current target. Will only be displayed for sections of type Object.
+	if (!InSectionTarget.IsNone())
+	{
+		SectionObject = InSectionTarget;
+	}
+
+	if (CurrentSection == EMutableLogSection::Object)
+	{
+		UE_LOG(LogMutableValidation, Log, TEXT(" SECTION START : %s - [%s] "), *GetLogSectionName(CurrentSection), *SectionObject.ToString());
+	}
+	else
+	{
+		UE_LOG(LogMutableValidation, Log, TEXT(" SECTION START : %s "), *GetLogSectionName(CurrentSection));
+	}
 }
 
 
 FScopedLogSection::~FScopedLogSection()
 {
-	// You are trying to close a section that is already closed. This should not be possible as we handle it in this objects
-	check(CurrentSection != EMutableLogSection::Undefined);
-	
-	UE_LOG(LogMutableValidation,Log,TEXT(" SECTION END : %s "), *GetLogSectionName(CurrentSection));
+	if (CurrentSection == EMutableLogSection::Object)
+	{
+		UE_LOG(LogMutableValidation, Log, TEXT(" SECTION END : %s - [%s] "), *GetLogSectionName(CurrentSection), *SectionObject.ToString());
+	}
+	else
+	{
+		UE_LOG(LogMutableValidation, Log, TEXT(" SECTION END : %s "), *GetLogSectionName(CurrentSection));
+	}
 	
 	// Set the current section to none (undefined)
 	CurrentSection = EMutableLogSection::Undefined;
@@ -50,6 +66,9 @@ FString FScopedLogSection::GetLogSectionName(EMutableLogSection Section) const
 	 case EMutableLogSection::Bake:
 	 	Output = TEXT("bake");
  		break;
+	 case EMutableLogSection::Object:
+	 	Output = TEXT("object");
+	 	break;
 	 default:
 	 	checkNoEntry();
 	 	break;
diff --git a/Engine/Plugins/Mutable/Source/MutableValidation/Private/MuV/ScopedLogSection.h b/Engine/Plugins/Mutable/Source/MutableValidation/Private/MuV/ScopedLogSection.h
--- a/Engine/Plugins/Mutable/Source/MutableValidation/Private/MuV/ScopedLogSection.h
+++ b/Engine/Plugins/Mutable/Source/MutableValidation/Private/MuV/ScopedLogSection.h
@@ -2,6 +2,8 @@
 #pragma once
 
 #include  "HAL/Platform.h"
+#include "UObject/NameTypes.h"
+
 
 class FString;
 
@@ -12,9 +14,12 @@ class FString;
 enum class EMutableLogSection : uint8
 {
 	Undefined = 0,
+	
 	Compilation,
 	Update,
-	Bake
+	Bake,
+	
+	Object,
 };
 
 
@@ -24,7 +29,7 @@ enum class EMutableLogSection : uint8
 class FScopedLogSection
 {
 public:
-	FScopedLogSection(EMutableLogSection Section);
+	FScopedLogSection(EMutableLogSection Section, const FName& SectionTarget = FName());
 	~FScopedLogSection();
 
 private:
@@ -33,5 +38,8 @@ private:
 	FString GetLogSectionName(EMutableLogSection Section) const;
 	
   	/** Section that this object is currently representing */
-	static inline EMutableLogSection CurrentSection = EMutableLogSection::Undefined;
+	EMutableLogSection CurrentSection = EMutableLogSection::Undefined;
+
+	/** The object the current section is representing */
+	static inline FName SectionObject;
 };
\ No newline at end of file
diff --git a/Engine/Plugins/Mutable/Source/MutableValidation/Private/MuV/ValidationUtils.cpp b/Engine/Plugins/Mutable/Source/MutableValidation/Private/MuV/ValidationUtils.cpp
--- a/Engine/Plugins/Mutable/Source/MutableValidation/Private/MuV/ValidationUtils.cpp
+++ b/Engine/Plugins/Mutable/Source/MutableValidation/Private/MuV/ValidationUtils.cpp
@@ -2,13 +2,20 @@
 
 #include "ValidationUtils.h"
 
+#include "CustomizableObjectCompilationUtility.h"
+#include "CustomizableObjectInstanceUpdateUtility.h"
+#include "RHIGlobals.h"
 #include "AssetRegistry/AssetRegistryModule.h"
 #include "Commandlets/Commandlet.h"
+#include "Interfaces/ITargetPlatform.h"
+#include "Interfaces/ITargetPlatformManagerModule.h"
+#include "Misc/CoreMisc.h"
 #include "MuCO/CustomizableObject.h"
-#include "MuCO/CustomizableObjectSystem.h"
-#include "RHIGlobals.h"
 #include "MuCO/CustomizableObjectPrivate.h"
+#include "MuCO/CustomizableObjectSystem.h"
 #include "MuCOE/CustomizableObjectBenchmarkingUtils.h"
+#include "MuR/Model.h"
+#include "Serialization/MemoryWriter.h"
 #include "UObject/ObjectPtr.h"
 
 void PrepareAssetRegistry()
@@ -61,15 +68,241 @@ void Wait(const double ToWaitSeconds)
 }
 
 
-FCompilationOptions GetCompilationOptionsForBenchmarking (UCustomizableObject& ReferenceCustomizableObject)
+FCompilationOptions GetCompilationOptionsForBenchmarking (const UCustomizableObject& ReferenceCustomizableObject)
 {
 	// Override some configurations that may have been changed by the user
 	FCompilationOptions CISCompilationOptions = ReferenceCustomizableObject.GetPrivate()->GetCompileOptions();
 	CISCompilationOptions.OptimizationLevel = CustomizableObjectBenchmarkingUtils::GetOptimizationLevelForBenchmarking();
 	CISCompilationOptions.TextureCompression = ECustomizableObjectTextureCompression::Fast;	// Does not affect instance update speed but does compilation
-
 	return CISCompilationOptions;
 }
 
 
+TArray<FAssetData> FindAllAssetsAtPath(FName SearchPath, const UClass* TargetObjectClass)
+{
+	TArray<FAssetData> FoundAssetData;
+	
+	if (TargetObjectClass)
+	{
+		FARFilter Filter;
+		Filter.ClassPaths.Add(TargetObjectClass->GetClassPathName());
+		Filter.PackagePaths.Add( SearchPath);
+		Filter.bRecursivePaths = true;
+	
+		const FAssetRegistryModule& AssetRegistryModule = FModuleManager::LoadModuleChecked<FAssetRegistryModule>(AssetRegistryConstants::ModuleName);
+
+		// Ensure the AR module is ready to search for stuff
+		AssetRegistryModule.Get().SearchAllAssets(true);
+	
+		UE_LOG(LogMutable, Display, TEXT("Searching for all %s objects to test at path : %s ."), *TargetObjectClass->GetName(), *SearchPath.ToString());
+		AssetRegistryModule.Get().GetAssets(Filter, FoundAssetData);
+		UE_LOG(LogMutable, Display, TEXT("Search of %s objects completed. Found %i objects."), *TargetObjectClass->GetName(), FoundAssetData.Num());
+	}
+	else
+	{
+		UE_LOG(LogMutable, Error, TEXT("No objects can be retrieved using a null class."));
+	}
+
+	return FoundAssetData;
+}
+
+
+bool GetDiskCompilationArg(const FString& Params)
+{
+	bool bUseDiskCompilation = false;
+	if (! FParse::Bool(*Params,TEXT("UseDiskCompilation="),bUseDiskCompilation))
+	{
+		UE_LOG(LogMutable, Display, TEXT("Disk compilation setting for the compilation of the CO not specified. Using default value : %hhd"), bUseDiskCompilation) ;
+	}
+	return bUseDiskCompilation;
+}
+
+
+uint32 GetTargetAmountOfInstances(const FString& Params)
+{
+	// Get the amount of instances to generate if parameter was provided (it will get multiplied by the amount of states later so this is a minimun value)
+	uint32 InstancesToGenerate = 16;
+	if (!FParse::Value(*Params, TEXT("InstanceGenerationCount="),InstancesToGenerate))
+	{
+		UE_LOG(LogMutable, Display, TEXT("Instance generation count not specified. Using default value : %u"), InstancesToGenerate);
+	}
+	return InstancesToGenerate;
+}
+
+
+ITargetPlatform* GetCompilationPlatform(const FString& Params)
+{
+	// Get the package name of the CO to test
+	FString TargetPlatformName = "";
+	if (!FParse::Value(*Params, TEXT("CompilationPlatformName="), TargetPlatformName))
+	{
+		UE_LOG(LogMutable, Error, TEXT("Failed to parse the target compilation platform. Have you even provided the argument?"))
+		return nullptr;
+	}
+
+	// Set the target platform in the context. For now it is the current platform.
+	ITargetPlatformManagerModule* TPM = GetTargetPlatformManager();
+	check(TPM);
+	
+	ITargetPlatform* TargetCompilationPlatform = nullptr;
+	const TArray<ITargetPlatform*> TPMPlatforms = TPM->GetTargetPlatforms();
+	for (ITargetPlatform* Platform : TPMPlatforms)
+	{
+		FString PlatformName = Platform->PlatformName();
+		if (PlatformName.Compare(TargetPlatformName) == 0)
+		{
+			// We have found the platform provided
+			TargetCompilationPlatform = Platform;
+			break;
+		}
+	}
+	
+	if (!TargetCompilationPlatform)
+	{
+		UE_LOG(LogMutable, Error, TEXT("Unable to relate the provided platform name (%s) with the available platforms in this machine."), *TargetPlatformName);
+	}
+
+	return TargetCompilationPlatform;
+}
+
+
+bool TestCustomizableObject(UCustomizableObject& InTargetCustomizableObject, const ITargetPlatform& TargetCompilationPlatform,
+	const uint32 InstancesToGenerate, const bool bUseDiskCompilation)
+{
+	// Keep a strong object pointer pointing at the CO to prevent it from being GCd during the test
+	const TStrongObjectPtr<UCustomizableObject> TargetCO = TStrongObjectPtr{&InTargetCustomizableObject};
+	
+	// Compile the Customizable Object ------------------------------------------------------------------------------ //
+	bool bWasCoCompilationSuccessful = false;
+	{
+		LLM_SCOPE_BYNAME(TEXT("MutableValidation/Compile"));
+		// Override some configurations that may have been changed by the user
+		FCompilationOptions CompilationOptions = GetCompilationOptionsForBenchmarking(InTargetCustomizableObject);
+		
+		// Set the target compilation platform based on what the caller wants
+		CompilationOptions.TargetPlatform = &TargetCompilationPlatform;
+		// Disk cache usage for compilation operation
+		CompilationOptions.bUseDiskCompilation = bUseDiskCompilation;
+		
+		TSharedRef<FCustomizableObjectCompilationUtility> CompilationUtility = MakeShared<FCustomizableObjectCompilationUtility>();
+		bWasCoCompilationSuccessful = CompilationUtility->CompileCustomizableObject(InTargetCustomizableObject, true, &CompilationOptions);
+	}
+	// -------------------------------------------------------------------------------------------------------------- //
+
+	if (!bWasCoCompilationSuccessful)
+	{
+		UE_LOG(LogMutable, Error, TEXT("The compilation of the Customizable object was not successful : No instances will be generated."));
+		return false;		// Validation failed
+	}
+	
+	// GHet the total size of the streaming data of the model ---------------------------------------------- //
+	{
+		const TSharedPtr<const mu::FModel> MutableModel = InTargetCustomizableObject.GetPrivate()->GetModel();
+		check (MutableModel);
+
+		// Roms ---------------------- //
+		{
+			const int32 RomCount =  MutableModel->GetRomCount();
+			int64 TotalRomSizeBytes = 0;
+			for (int32 RomIndex = 0; RomIndex < RomCount; RomIndex++)
+			{
+				const uint32 RomByteSize = MutableModel->GetRomSize(RomIndex);
+				TotalRomSizeBytes += RomByteSize;
+			}
+
+			// Print MTU parseable logs
+			UE_LOG(LogMutable, Log, TEXT("(int) model_rom_count : %d "), RomCount);
+			UE_LOG(LogMutable, Log, TEXT("(int) model_roms_size : %lld "), TotalRomSizeBytes);
+		}
+
+		// CO embedded data size ------ //
+		{
+			TArray<uint8> EmbeddedDataBytes{};
+			FMemoryWriter SerializationTarget{EmbeddedDataBytes, false};
+		
+			InTargetCustomizableObject.GetPrivate()->SaveEmbeddedData(SerializationTarget);
+			const int64 COEmbeddedDataSizeBytes = EmbeddedDataBytes.Num();
+		
+			UE_LOG(LogMutable, Log, TEXT("(int) co_embedded_data_bytes : %lld "), COEmbeddedDataSizeBytes);
+		}
+	}
+	
+	// Skip instances updating if no instances should be updated 
+	if (InstancesToGenerate <= 0)
+	{
+		UE_LOG(LogMutable, Display, TEXT("Instances to generate are 0 : No instances will be generated."));
+		return true;	// No instances are targeted for generation, this will be taken as compilation only test.
+	}
+
+	// Do not generate instances if the selected platform is not the running platform
+	if (&TargetCompilationPlatform != GetTargetPlatformManagerRef().GetRunningTargetPlatform())
+	{
+		UE_LOG(LogMutable, Display, TEXT("RunningPlatform != UserProvidedCompilationPlatform : No instances will be generated."));
+		return true;
+	}
+
+	// At this point we know the compilation has been successful. Generate a deterministically random set of instances now.
+	
+	// Generate target random instances to be tested ------------------------------------------------------------ //
+	bool bWasInstancesCreationSuccessful = true;
+	TSpscQueue<TStrongObjectPtr<UCustomizableObjectInstance>> InstancesToProcess;
+	uint32 GeneratedInstances = 0;
+	{
+		LLM_SCOPE_BYNAME(TEXT("MutableValidation/GenerateInstances"));
+		
+		// Create a set of instances so we can later test them out
+		bWasInstancesCreationSuccessful = CustomizableObjectBenchmarkingUtils::GenerateDeterministicSetOfInstances(InTargetCustomizableObject, InstancesToGenerate, InstancesToProcess, GeneratedInstances);
+	}
+	// ---------------------------------------------------------------------------------------------------------- //
+
+	UE_LOG(LogMutable, Log, TEXT("(int) generated_instances_count : %u "), GeneratedInstances);
+	
+	// Update the instances generated --------------------------------------------------------------------------- //
+	UE_LOG(LogMutable, Display, TEXT("Updating generated instances..."));
+	bool bInstanceFailedUpdate = false;
+	const double InstancesUpdateStartSeconds = FPlatformTime::Seconds();
+	{
+		LLM_SCOPE_BYNAME(TEXT("MutableValidation/Update"));
+		
+		TSharedRef<FCustomizableObjectInstanceUpdateUtility> InstanceUpdatingUtility = MakeShared<FCustomizableObjectInstanceUpdateUtility>();
+
+		TStrongObjectPtr<UCustomizableObjectInstance> InstanceToUpdate;
+		while (InstancesToProcess.Dequeue(InstanceToUpdate))
+		{
+			CollectGarbage(GARBAGE_COLLECTION_KEEPFLAGS, true);
+			check(InstanceToUpdate);
+
+			if (!InstanceUpdatingUtility->UpdateInstance(*InstanceToUpdate))
+			{
+				bInstanceFailedUpdate = true;
+			}
+		}
+	}
+	const double InstancesUpdateEndSeconds = FPlatformTime::Seconds();
+	
+	// Notify and log time required by the instances to get updated
+	const double CombinedInstanceUpdateSeconds = InstancesUpdateEndSeconds - InstancesUpdateStartSeconds;
+	UE_LOG(LogMutable, Log, TEXT("(double) combined_update_time_ms : %f "), CombinedInstanceUpdateSeconds * 1000);
+
+	check(GeneratedInstances > 0);
+	const double AverageInstanceUpdateSeconds = CombinedInstanceUpdateSeconds / GeneratedInstances;
+	UE_LOG(LogMutable, Log, TEXT("(double) avg_update_time_ms : %f "), AverageInstanceUpdateSeconds * 1000);
+
+	UE_LOG(LogMutable, Display, TEXT("Generation of Customizable object instances took %f seconds (%f seconds avg)."), CombinedInstanceUpdateSeconds, AverageInstanceUpdateSeconds);
+	// ---------------------------------------------------------------------------------------------------------- //
+
+	// Compute instance update result
+	const bool bInstancesTestedSuccessfully = !bInstanceFailedUpdate && bWasInstancesCreationSuccessful;
+	if (bInstancesTestedSuccessfully)
+    {
+        UE_LOG(LogMutable, Display, TEXT("Generation of Customizable object instances was successful."));
+    }
+    else
+    {
+        UE_LOG(LogMutable, Error, TEXT("The generation of Customizable object instances was not successful."));
+    }
+	
+	return bInstancesTestedSuccessfully;
+}
+
 
diff --git a/Engine/Plugins/Mutable/Source/MutableValidation/Private/MuV/ValidationUtils.h b/Engine/Plugins/Mutable/Source/MutableValidation/Private/MuV/ValidationUtils.h
--- a/Engine/Plugins/Mutable/Source/MutableValidation/Private/MuV/ValidationUtils.h
+++ b/Engine/Plugins/Mutable/Source/MutableValidation/Private/MuV/ValidationUtils.h
@@ -1,9 +1,41 @@
 ﻿// Copyright Epic Games, Inc. All Rights Reserved.
 #pragma once
 
+#include "AssetRegistry/AssetData.h"
+
 struct FCompilationOptions;
 class UCustomizableObject;
 
+
+// ARGUMENT PARSING ----------------------------------------------------------------------------------------------------------------------------------
+
+/**
+ * Get the type of compilation the caller wants to run. 
+ * @param Params The parameters fed to this commandlet.
+ * @return True if the use of the disk cache is wanted for the compilation, false otherwise.
+ */
+bool GetDiskCompilationArg(const FString& Params);
+
+
+/**
+ * Get the amount of instances we want to generate.
+ * @param Params The parameters fed to this commandlet.
+ * @return The amount of instances to be generated for a given CO and state.
+ */
+uint32 GetTargetAmountOfInstances(const FString& Params);
+
+
+/**
+ * Extracts the targeted compilation platform provided by the user. It will look for "-CompilationPlatformName="PlatformName".
+ * Examples : -CompilationPlatformName=WindowsEditor or -CompilationPlatformName=Switch
+ * @param Params The arguments provided to this commandlet.
+ * @return The target platform to be used for the CO compilation.
+ */
+ITargetPlatform* GetCompilationPlatform(const FString& Params);
+
+
+// SETUP ---------------------------------------------------------------------------------------------------------------------------------------------
+
 /**
  * Prepare the asset registry so we can later use it to search assets. It is required by Mutable to compile.
  */
@@ -31,4 +63,28 @@ void LogGlobalSettings();
  * @return The FCompilationOptions for the provided CO but with some settings changed to be adecuate for a benchmark
  * oriented compilation.
  */
-FCompilationOptions GetCompilationOptionsForBenchmarking (UCustomizableObject& ReferenceCustomizableObject);
+FCompilationOptions GetCompilationOptionsForBenchmarking (const UCustomizableObject& ReferenceCustomizableObject);
+
+
+// TEST ----------------------------------------------------------------------------------------------------------------------------------------------
+
+/**
+ * Get a list of AssetData objects filled with the objets of the class specified found at the provided path.
+ * @param SearchPath The path where the search should be conducted.
+ * @param TargetObjectClass The class to be searching for in the provided path.
+ * @return An array of AssetData objects. Each one will represent one of the found assets. The names of the objects may be repeated, but the paths not.
+ */
+TArray<FAssetData> FindAllAssetsAtPath(FName SearchPath, const UClass* TargetObjectClass);
+
+
+/**
+ * Compiles and then generates a series of instances while reporting the performance of all the processes involved.
+ * @param InTargetCustomizableObject The Customizable Object to be compiled and used as the base for the generated instances.
+ * @param TargetCompilationPlatform The target platform used to compile the Customizable Object against.
+ * @param InstancesToGenerate The amount of instances to generate per state.
+ * @param bUseDiskCompilation Determines if a disk cache will be used for the compilation of the Customizable Object.
+ * @return True if all went well and false if the testing operation did throw errors.
+ */
+bool TestCustomizableObject(UCustomizableObject& InTargetCustomizableObject, const ITargetPlatform& TargetCompilationPlatform,
+	const uint32 InstancesToGenerate,const bool bUseDiskCompilation);
+	
\ No newline at end of file
diff --git a/Engine/Plugins/Mutable/Source/MutableValidation/Public/MuV/AssetValidator_CustomizableObjects.h b/Engine/Plugins/Mutable/Source/MutableValidation/Public/MuV/AssetValidator_CustomizableObjects.h
deleted file mode 100644
diff --git a/Engine/Plugins/Mutable/Source/MutableValidation/Public/MuV/AssetValidator_ReferencedCustomizableObjects.h b/Engine/Plugins/Mutable/Source/MutableValidation/Public/MuV/AssetValidator_ReferencedCustomizableObjects.h
deleted file mode 100644
diff --git a/Engine/Plugins/Mutable/Source/MutableValidation/Public/MuV/COIBakingTestCommandlet.h b/Engine/Plugins/Mutable/Source/MutableValidation/Public/MuV/COIBakingTestCommandlet.h
deleted file mode 100644
diff --git a/Engine/Plugins/Mutable/Source/MutableValidation/Public/MuV/COIBulkUpdateTestCommandlet.h b/Engine/Plugins/Mutable/Source/MutableValidation/Public/MuV/COIBulkUpdateTestCommandlet.h
deleted file mode 100644
diff --git a/Engine/Plugins/Mutable/Source/MutableValidation/Public/MuV/CacheAssetRegistryCommandlet.h b/Engine/Plugins/Mutable/Source/MutableValidation/Public/MuV/CacheAssetRegistryCommandlet.h
deleted file mode 100644
diff --git a/Engine/Plugins/Mutable/Source/MutableValidation/Public/MuV/CustomizableObjectValidationCommandlet.h b/Engine/Plugins/Mutable/Source/MutableValidation/Public/MuV/CustomizableObjectValidationCommandlet.h
deleted file mode 100644
diff --git a/Engine/Plugins/Mutable/Source/MutableValidation/Public/MuV/MutableValidationModule.cpp b/Engine/Plugins/Mutable/Source/MutableValidation/Public/MuV/MutableValidationModule.cpp
deleted file mode 100644
