[mutable] Fixed crash when compiling a Pass-Trough Mesh Component when it tried to access the component's reference mesh.
#tests Editor
#rnx

diff --git a/Engine/Plugins/Mutable/Source/CustomizableObject/Private/MuCO/CustomizableObjectInstance.cpp b/Engine/Plugins/Mutable/Source/CustomizableObject/Private/MuCO/CustomizableObjectInstance.cpp
--- a/Engine/Plugins/Mutable/Source/CustomizableObject/Private/MuCO/CustomizableObjectInstance.cpp
+++ b/Engine/Plugins/Mutable/Source/CustomizableObject/Private/MuCO/CustomizableObjectInstance.cpp
@@ -5161,7 +5161,16 @@ UE::Tasks::FTask UCustomizableInstancePrivate::LoadAdditionalAssetsAndData(const
 			const FInstanceUpdateData::FComponent& Component = Components[InstanceComponentIndex];
 			FCustomizableInstanceComponentData* ComponentData = GetComponentData(Component.Id);
 
+			TSharedPtr<const mu::FMesh> FirstComponentMesh = OperationData->InstanceUpdateData.LODs.IsValidIndex(Component.FirstLOD) ?
+				OperationData->InstanceUpdateData.LODs[Component.FirstLOD].Mesh :
+				nullptr;
 
+			if (FirstComponentMesh && FirstComponentMesh->IsReference())
+			{
+				// Pass-through components don't have a Reference Mesh so don't access it
+				continue;
+			}
+			
 			const FCustomizableObjectComponentIndex ObjectComponentIndex = Component.Id;
 			const FMutableRefSkeletalMeshData& RefSkeletalMeshData = ModelResources.ReferenceSkeletalMeshesData[ObjectComponentIndex.GetValue()];
 
