[Mutable] Fix mutliple streaming issues.

- Avoid crash when streaming without morphs.
- Fixed Morphs using non-existing data from initial generation.
- Allow disabling morphs when using streaming with the CVar.
- Removed MinLOD ussages from initial generation.
- Documented LOD variables
- Renamed LOD variables.

diff --git a/Engine/Plugins/Mutable/Source/CustomizableObject/Internal/MuCO/CustomizableObjectSystemPrivate.h b/Engine/Plugins/Mutable/Source/CustomizableObject/Internal/MuCO/CustomizableObjectSystemPrivate.h
--- a/Engine/Plugins/Mutable/Source/CustomizableObject/Internal/MuCO/CustomizableObjectSystemPrivate.h
+++ b/Engine/Plugins/Mutable/Source/CustomizableObject/Internal/MuCO/CustomizableObjectSystemPrivate.h
@@ -164,15 +164,15 @@ public:
 	// These are the LODs that would be copied to the descriptor to trigger mesh updates on quality setting changes.
 	TMap<FName, uint8> QualitySettingMinLODs;
 
-	TMap<FName, uint8> RequestedLODLevels;
+	TMap<FName, uint8> FirstRequestedLOD;
 
 	FMutableUpdateCandidate(UCustomizableObjectInstance* InCustomizableObjectInstance);
 	
 	FMutableUpdateCandidate(const UCustomizableObjectInstance* InCustomizableObjectInstance, const TMap<FName, uint8>& InMinLOD,
-		const TMap<FName, uint8>& InRequestedLODLevels) :
+		const TMap<FName, uint8>& InFirstRequestedLOD) :
 		CustomizableObjectInstance(const_cast<UCustomizableObjectInstance*>(InCustomizableObjectInstance)),
 		MinLOD(InMinLOD),
-		RequestedLODLevels(InRequestedLODLevels) {}
+		FirstRequestedLOD(InFirstRequestedLOD) {}
 
 	bool HasBeenIssued() const;
 
@@ -459,7 +459,7 @@ struct FInstanceUpdateData
 
 	struct FLOD
 	{		
-		mu::FResourceID MeshID;
+		mu::FResourceID MeshID = TNumericLimits<mu::FResourceID>::Max();
 		TSharedPtr<const mu::FMesh> Mesh;
 
 		/** Range in the Surfaces array */
@@ -473,9 +473,6 @@ struct FInstanceUpdateData
 		/** Range in the external Bones array */
 		uint32 FirstBoneMap = 0;
 		uint32 BoneMapCount = 0;
-
-		// True if the Mesh is valid
-		bool bGenerated = false;
 	};
 
 	struct FComponent
@@ -573,14 +570,12 @@ public:
 
 	bool IsContextValid() const;
 	
-	const TMap<FName, uint8>& GetMinLOD() const;
-
 	void SetMinLOD(const TMap<FName, uint8>& MinLOD);
 
 	/** Return an array of LODs per object component. */
-	const TMap<FName, uint8>& GetRequestedLODs() const;
+	const TMap<FName, uint8>& GetFirstRequestedLOD() const;
 	
-	void SetRequestedLODs(const TMap<FName, uint8>& RequestedLODs);
+	void SetFirstRequestedLOD(const TMap<FName, uint8>& RequestedLODs);
 
 	void SetQualitySettingMinLODs(const TMap<FName, uint8>& FirstLODs);
 
@@ -757,9 +752,9 @@ struct FPendingReleaseSkeletalMeshInfo
 
 /** Private part, hidden from outside the plugin.
  *
- * Small diagram to help understanding Mutable LOD nomenclature:
+ * UE STREAMING:
  * 
- * [- NumLODsAvailable -----------------------------] = 8
+ * [- NumLODsAvailable -----------------------------] = 8 (State and platform dependent)
  * [- Stripped --[- Packaged -----------------------]
  * 0      1      2      3      4      5      6      7    
  * |------|------|------|------|------|------|------|
@@ -767,9 +762,38 @@ struct FPendingReleaseSkeletalMeshInfo
  *               ^                    ^
  *               |                    |
  *               FirstLODAvailable    FirstResidentLOD
- *               MinLOD (when not using Quality Levels)
+ *                                    FirstRequestedLOD (LODs generated by Core)
  * 
- * [- NumLODsToStream ----------------] = 6 */
+ * [- NumLODsToStream ----------------] = 5 (Compiled constant, ModelResources)
+ *
+ *
+ * HACKY MUTABLE STREAMING:
+ * 
+ * [- NumLODsAvailable -----------------------------] = 8
+ * [- Stripped --[- Packaged -----------------------]
+ * 0      1      2      3      4      5      6      7    
+ * |------|------|------|------|------|------|------|
+ *               [- Residents ----------------------]
+ *               [------] Data copied from FirstRequestedLODs. Hacky Mutable LOD Streaming.
+ *               ^      ^         
+ *               |      |
+ *               |      FirstRequestedLOD (LODs generated by Core)
+ *               |      
+ *               FirstLODAvailable (Compilation constant)
+ *               FirstResidentLOD
+ *               MinLOD
+ * 
+ * NumLODsToStream = 0
+ *
+ *
+ * DEFINITIONS:
+ * 
+ * QualitySettingMinLODs =	MinLOD based on the active quality settings.													COI Descriptor.
+ * MinLOD =					From user. Artificial limit.											                     	Skeletal Mesh Component.
+ * FirstLODAvailable =	    First available lod per platform.																Skeletal Mesh Component.
+ * FirstResidentLOD =       First LOD generated with geometry.																Skeletal Mesh Component.
+ * FirstRequestedLOD =		From user. Usually from USkeletalMeshComponent::GetPredictedLODLevel							Skeletal Mesh Component. 
+ */
 UCLASS()
 class UCustomizableObjectSystemPrivate : public UObject, public IStreamingManager
 {
@@ -1008,3 +1032,7 @@ namespace impl
 /** Set OnlyLOD to -1 to generate all mips */
 CUSTOMIZABLEOBJECT_API FTexturePlatformData* MutableCreateImagePlatformData(TSharedPtr<const mu::FImage> MutableImage, int32 OnlyLOD, uint16 FullSizeX, uint16 FullSizeY);
 
+
+/** Return true if Streaming is enabled for the given Object. */
+bool IsStreamingEnabled(const UCustomizableObject& Object);
+
diff --git a/Engine/Plugins/Mutable/Source/CustomizableObject/Internal/MuCO/DescriptorHash.h b/Engine/Plugins/Mutable/Source/CustomizableObject/Internal/MuCO/DescriptorHash.h
--- a/Engine/Plugins/Mutable/Source/CustomizableObject/Internal/MuCO/DescriptorHash.h
+++ b/Engine/Plugins/Mutable/Source/CustomizableObject/Internal/MuCO/DescriptorHash.h
@@ -30,5 +30,5 @@ public:
 	TMap<FName, uint8> QualitySettingMinLODs;
 	
 	// Array of bitmasks that indicate which LODs of each component have been requested
-	TMap<FName, uint8> RequestedLODsPerComponent;
+	TMap<FName, uint8> FirstRequestedLOD;
 };
diff --git a/Engine/Plugins/Mutable/Source/CustomizableObject/Private/MuCO/CustomizableObjectInstance.cpp b/Engine/Plugins/Mutable/Source/CustomizableObject/Private/MuCO/CustomizableObjectInstance.cpp
--- a/Engine/Plugins/Mutable/Source/CustomizableObject/Private/MuCO/CustomizableObjectInstance.cpp
+++ b/Engine/Plugins/Mutable/Source/CustomizableObject/Private/MuCO/CustomizableObjectInstance.cpp
@@ -1975,11 +1975,11 @@ bool UCustomizableInstancePrivate::DoComponentsNeedUpdate(UCustomizableObjectIns
 		const FInstanceUpdateData::FComponent& Component = OperationData->InstanceUpdateData.Components[InstanceComponentIndex];
 		const FName ComponentName = OperationData->ComponentNames[Component.Id.GetValue()];
 
-		for (int32 LODIndex = OperationData->GetMinLOD()[ComponentName]; LODIndex < Component.LODCount; ++LODIndex)
+		for (int32 LODIndex = OperationData->GetFirstRequestedLOD()[ComponentName]; LODIndex < Component.LODCount; ++LODIndex)
 		{
 			const FInstanceUpdateData::FLOD& LOD = OperationData->InstanceUpdateData.LODs[Component.FirstLOD + LODIndex];
 
-			if (!LOD.bGenerated || !LOD.Mesh )
+			if (!LOD.Mesh)
 			{
 				continue;
 			}
@@ -2078,11 +2078,11 @@ bool UCustomizableInstancePrivate::DoComponentsNeedUpdate(UCustomizableObjectIns
 		// than the new FirstLODAvailable. This is an artifact of our LOD management system setting MinLOD.
 		if (!OperationData->MeshChangedPerInstanceComponent[InstanceComponentIndex] && bHadSkeletalMesh) 
 		{
-			const int32 MinLOD = OperationData->GetMinLOD()[ComponentName];
-			for (int32 LODIndex = OperationData->FirstLODAvailable[ComponentName]; LODIndex < MinLOD; ++LODIndex)
+			const int32 FirstRequestedLOD = OperationData->GetFirstRequestedLOD()[ComponentName];
+			for (int32 LODIndex = OperationData->FirstLODAvailable[ComponentName]; LODIndex < FirstRequestedLOD; ++LODIndex)
 			{
 				TIndirectArray<FSkeletalMeshLODRenderData>& LODRenderData = (*SkeletalMeshPtr)->GetResourceForRendering()->LODRenderData;
-				if (LODRenderData[LODIndex].RenderSections.Num() != LODRenderData[MinLOD].RenderSections.Num())
+				if (LODRenderData[LODIndex].RenderSections.Num() != LODRenderData[FirstRequestedLOD].RenderSections.Num())
 				{
 					OperationData->MeshChangedPerInstanceComponent[InstanceComponentIndex] = true;
 					break;
@@ -3342,7 +3342,7 @@ FMutableUpdateCandidate::FMutableUpdateCandidate(UCustomizableObjectInstance* In
 	const FCustomizableObjectInstanceDescriptor& Descriptor = InCustomizableObjectInstance->GetPrivate()->GetDescriptor();
 	MinLOD = Descriptor.MinLOD;
 	QualitySettingMinLODs = Descriptor.QualitySettingMinLODs;
-	RequestedLODLevels = Descriptor.GetRequestedLODLevels();
+	FirstRequestedLOD = Descriptor.GetFirstRequestedLOD();
 }
 
 
@@ -3362,11 +3362,11 @@ void FMutableUpdateCandidate::ApplyLODUpdateParamsToInstance(FUpdateContextPriva
 {
 	CustomizableObjectInstance->Descriptor.MinLOD = MinLOD;
 	CustomizableObjectInstance->Descriptor.QualitySettingMinLODs = QualitySettingMinLODs;
-	CustomizableObjectInstance->Descriptor.RequestedLODLevels = RequestedLODLevels;
+	CustomizableObjectInstance->Descriptor.FirstRequestedLOD = FirstRequestedLOD;
 
 	Context.SetMinLOD(MinLOD);
 	Context.SetQualitySettingMinLODs(QualitySettingMinLODs);
-	Context.SetRequestedLODs(RequestedLODLevels);
+	Context.SetFirstRequestedLOD(FirstRequestedLOD);
 }
 
 
@@ -4156,11 +4156,6 @@ void UCustomizableInstancePrivate::BuildOrCopyElementData(const TSharedRef<FUpda
 			continue;
 		}
 
-		if (!LOD.bGenerated)
-		{
-			continue;
-		}
-
 		for (int32 SurfaceIndex = 0; SurfaceIndex < LOD.SurfaceCount; ++SurfaceIndex)
 		{
 			new(SkeletalMesh->GetResourceForRendering()->LODRenderData[LODIndex].RenderSections) FSkelMeshRenderSection();
@@ -4218,8 +4213,8 @@ void UCustomizableInstancePrivate::BuildOrCopyMorphTargetsData(const TSharedRef<
 			MorphLODModels.SetNum(ComponentMorphTargets->RealTimeMorphsLODData[MorphTargetIndex].Num());
 			
 			// Streamed LODs
-			const int32 FirstGeneratedLOD = OperationData->FirstLODAvailable[ComponentName];
-			for (int32 LODIndex = FirstGeneratedLOD; LODIndex < OperationData->FirstResidentLOD[ComponentName]; ++LODIndex)
+			const int32 FirstLODAvailable = OperationData->FirstLODAvailable[ComponentName];
+			for (int32 LODIndex = FirstLODAvailable; LODIndex < OperationData->FirstResidentLOD[ComponentName]; ++LODIndex)
 			{
 				// Copy data required for streaming
 				MorphLODModels[LODIndex].NumVertices = 1; // Trick the engine
@@ -4227,7 +4222,7 @@ void UCustomizableInstancePrivate::BuildOrCopyMorphTargetsData(const TSharedRef<
 			}
 			
 			// Residents LODs
-			for (int32 LODIndex = OperationData->FirstResidentLOD[ComponentName]; LODIndex < OperationData->NumLODsAvailable[ComponentName]; ++LODIndex)
+			for (int32 LODIndex = OperationData->GetFirstRequestedLOD()[ComponentName]; LODIndex < OperationData->NumLODsAvailable[ComponentName]; ++LODIndex)
 			{
 				MorphLODModels[LODIndex] = ComponentMorphTargets->RealTimeMorphsLODData[MorphTargetIndex][LODIndex];
 			}
@@ -4239,18 +4234,19 @@ void UCustomizableInstancePrivate::BuildOrCopyMorphTargetsData(const TSharedRef<
 
 		MorphTargets.Add(NewMorphTarget);
 	}
-	
+
+	// Mutable hacky LOD Streaming
 	if (!OperationData->bStreamMeshLODs)
 	{
 		// Copy MorphTargets from the FirstGeneratedLOD to the LODs below
-		const int32 FirstGeneratedLOD = FMath::Max(OperationData->GetRequestedLODs()[ComponentName], OperationData->GetMinLOD()[ComponentName]);
-		for (int32 LODIndex = OperationData->FirstLODAvailable[ComponentName]; LODIndex < FirstGeneratedLOD; ++LODIndex)
+		const int32 FirstRequestedLOD = OperationData->GetFirstRequestedLOD()[ComponentName];
+		for (int32 LODIndex = OperationData->FirstLODAvailable[ComponentName]; LODIndex < FirstRequestedLOD; ++LODIndex)
 		{
 			MUTABLE_CPUPROFILER_SCOPE(CopyMorphTargetsData);
 
 			for (int32 MorphTargetIndex = 0; MorphTargetIndex < MorphTargets.Num(); ++MorphTargetIndex)
 			{
-				MorphTargets[MorphTargetIndex]->GetMorphLODModels()[LODIndex] = MorphTargets[MorphTargetIndex]->GetMorphLODModels()[FirstGeneratedLOD];
+				MorphTargets[MorphTargetIndex]->GetMorphLODModels()[LODIndex] = MorphTargets[MorphTargetIndex]->GetMorphLODModels()[FirstRequestedLOD];
 			}
 		}
 	}
@@ -4330,15 +4326,13 @@ void UCustomizableInstancePrivate::BuildOrCopyClothingData(
 	{
 		MUTABLE_CPUPROFILER_SCOPE(DiscoverSectionsWithCloth);
 
-		for (int32 LODIndex = OperationData->FirstLODAvailable[ComponentName]; LODIndex < Component->LODCount; ++LODIndex)
+		// This should be FirstRequestedLOD but we because we do not share cloth data between LODs, we must also generate the hacky streaming ones.
+		const int32 FirstLOD = OperationData->FirstLODAvailable[ComponentName];
+		
+		for (int32 LODIndex = FirstLOD; LODIndex < Component->LODCount; ++LODIndex)
 		{
 			const FInstanceUpdateData::FLOD& LOD = OperationData->InstanceUpdateData.LODs[Component->FirstLOD+LODIndex];
 
-			if (!LOD.bGenerated)
-			{
-				continue;
-			}
-
 			if (TSharedPtr<const mu::FMesh> MutableMesh = LOD.Mesh)
 			{
 				const mu::FMeshBufferSet& MeshSet = MutableMesh->GetVertexBuffers();
@@ -4707,7 +4701,7 @@ void UCustomizableInstancePrivate::BuildOrCopyClothingData(
 	{
 		int32 NumSectionsWithClothProcessed = 0;
 
-		for (int32 LODIndex = OperationData->GetMinLOD()[ComponentName]; LODIndex < Component->LODCount; ++LODIndex)
+		for (int32 LODIndex = OperationData->GetFirstRequestedLOD()[ComponentName]; LODIndex < Component->LODCount; ++LODIndex)
 		{
 			const FInstanceUpdateData::FLOD& LOD = OperationData->InstanceUpdateData.LODs[Component->FirstLOD+LODIndex];
 
@@ -4756,7 +4750,7 @@ void UCustomizableInstancePrivate::BuildOrCopyClothingData(
 	{
 		MUTABLE_CPUPROFILER_SCOPE(InitClothRenderData)
 		// Based on FSkeletalMeshLODModel::GetClothMappingData().
-		for (int32 LODIndex = OperationData->GetMinLOD()[ComponentName]; LODIndex < OperationData->NumLODsAvailable[ComponentName]; ++LODIndex)
+		for (int32 LODIndex = OperationData->GetFirstRequestedLOD()[ComponentName]; LODIndex < OperationData->NumLODsAvailable[ComponentName]; ++LODIndex)
 		{
 			FSkeletalMeshLODRenderData& LODModel = RenderResource->LODRenderData[LODIndex];
 	
@@ -4808,17 +4802,18 @@ bool UCustomizableInstancePrivate::BuildOrCopyRenderData(const TSharedRef<FUpdat
 	const UModelResources& ModelResources = CustomizableObject->GetPrivate()->GetModelResourcesChecked();
 	const FName ComponentName = ModelResources.ComponentNamesPerObjectComponent[Component->Id.GetValue()];
 
-	const int32 FirstGeneratedLOD = FMath::Max(OperationData->GetRequestedLODs()[ComponentName], OperationData->GetMinLOD()[ComponentName]);
-
-
-	for (int32 LODIndex = FirstGeneratedLOD; LODIndex < Component->LODCount; ++LODIndex)
+	const int32 FirstLOD = OperationData->bStreamMeshLODs ?
+		OperationData->FirstLODAvailable[ComponentName] :
+		OperationData->GetFirstRequestedLOD()[ComponentName];
+	
+	for (int32 LODIndex = FirstLOD; LODIndex < Component->LODCount; ++LODIndex)
 	{
 		MUTABLE_CPUPROFILER_SCOPE(BuildRenderData);
-
+		
 		const FInstanceUpdateData::FLOD& LOD = OperationData->InstanceUpdateData.LODs[Component->FirstLOD+LODIndex];
 
 		// There could be components without a mesh in LODs
-		if (!LOD.bGenerated || !LOD.Mesh || LOD.SurfaceCount == 0)
+		if (!LOD.Mesh || LOD.SurfaceCount == 0)
 		{
 			UE_LOG(LogMutable, Warning, TEXT("Building instance: generated mesh [%s] has LOD [%d] of object component index [%d] with no mesh.")
 				, *SkeletalMesh->GetName()
@@ -5010,23 +5005,29 @@ bool UCustomizableInstancePrivate::BuildOrCopyRenderData(const TSharedRef<FUpdat
 #endif
 		}
 	}
+	
 
-	// Copy LODRenderData from the FirstGeneratedLOD to the LODs below
-	for (int32 LODIndex = OperationData->FirstLODAvailable[ComponentName]; LODIndex < FirstGeneratedLOD; ++LODIndex)
+	// Mutable hacky LOD Streaming
+	if (!OperationData->bStreamMeshLODs)
 	{
-		TRACE_CPUPROFILER_EVENT_SCOPE_TEXT(*FString::Printf(TEXT("CopyRenderData: From LOD %d to LOD %d"), FirstGeneratedLOD, LODIndex));
+		// Copy LODRenderData from the FirstRequestedLOD to the LODs below
+		const int32 FirstRequestedLOD = OperationData->GetFirstRequestedLOD()[ComponentName];
+		for (int32 LODIndex = OperationData->FirstLODAvailable[ComponentName]; LODIndex < FirstRequestedLOD; ++LODIndex)
+		{
+			TRACE_CPUPROFILER_EVENT_SCOPE_TEXT(*FString::Printf(TEXT("CopyRenderData: From LOD %d to LOD %d"), FirstRequestedLOD, LODIndex));
 
-		// Render Data will be reused from the previously generated component
-		FSkeletalMeshLODRenderData& SourceLODResource = RenderData->LODRenderData[FirstGeneratedLOD];
-		FSkeletalMeshLODRenderData& LODResource = RenderData->LODRenderData[LODIndex];
+			// Render Data will be reused from the previously generated component
+			FSkeletalMeshLODRenderData& SourceLODResource = RenderData->LODRenderData[FirstRequestedLOD];
+			FSkeletalMeshLODRenderData& LODResource = RenderData->LODRenderData[LODIndex];
 
-		UnrealConversionUtils::CopySkeletalMeshLODRenderData(
-			LODResource,
-			SourceLODResource,
-			*SkeletalMesh,
-			LODIndex,
-			SkeletalMesh->GetLODInfo(LODIndex)->bAllowCPUAccess
-		);
+			UnrealConversionUtils::CopySkeletalMeshLODRenderData(
+				LODResource,
+				SourceLODResource,
+				*SkeletalMesh,
+				LODIndex,
+				SkeletalMesh->GetLODInfo(LODIndex)->bAllowCPUAccess
+			);
+		}
 	}
 
 	return true;
@@ -5910,17 +5911,14 @@ void UCustomizableInstancePrivate::BuildMaterials(const TSharedRef<FUpdateContex
 
 		MUTABLE_CPUPROFILER_SCOPE(BuildMaterials_LODLoop);
 
-		const int32 FirstGeneratedLOD = FMath::Max(OperationData->GetRequestedLODs()[ComponentName], OperationData->GetMinLOD()[ComponentName]);
+		const int32 FirstLOD = OperationData->bStreamMeshLODs ?
+			OperationData->FirstLODAvailable[ComponentName] :
+			OperationData->GetFirstRequestedLOD()[ComponentName];
 
-		for (int32 LODIndex = FirstGeneratedLOD; LODIndex < Component.LODCount; LODIndex++)
+		for (int32 LODIndex = FirstLOD; LODIndex < Component.LODCount; LODIndex++)
 		{
 			const FInstanceUpdateData::FLOD& LOD = OperationData->InstanceUpdateData.LODs[Component.FirstLOD+LODIndex];
-
-			if (!LOD.bGenerated)
-			{
-				continue;
-			}
-
+			
 			if (!bUseOverrideMaterialsOnly && LODIndex < SkeletalMesh->GetLODNum())
 			{
 				SkeletalMesh->GetLODInfo(LODIndex)->LODMaterialMap.Reset();
@@ -6534,24 +6532,28 @@ void UCustomizableInstancePrivate::BuildMaterials(const TSharedRef<FUpdateContex
 			}
 		}
 
-		if(!bUseOverrideMaterialsOnly)
+		if (!bUseOverrideMaterialsOnly)
 		{
-			// Copy data from the FirstGeneratedLOD into the LODs below.
-			for (int32 LODIndex = OperationData->FirstLODAvailable[ComponentName]; LODIndex < FirstGeneratedLOD; ++LODIndex)
+			// Mutable hacky LOD Streaming
+			if (!OperationData->bStreamMeshLODs)
 			{
-				SkeletalMesh->GetLODInfo(LODIndex)->LODMaterialMap = SkeletalMesh->GetLODInfo(FirstGeneratedLOD)->LODMaterialMap;
+				// Copy data from the FirstLODAvailable into the LODs below.
+				for (int32 LODIndex = OperationData->FirstLODAvailable[ComponentName]; LODIndex < FirstLOD; ++LODIndex)
+				{
+					SkeletalMesh->GetLODInfo(LODIndex)->LODMaterialMap = SkeletalMesh->GetLODInfo(FirstLOD)->LODMaterialMap;
 
-				TIndirectArray<FSkeletalMeshLODRenderData>& LODRenderData = SkeletalMesh->GetResourceForRendering()->LODRenderData;
+					TIndirectArray<FSkeletalMeshLODRenderData>& LODRenderData = SkeletalMesh->GetResourceForRendering()->LODRenderData;
 
-				const int32 NumRenderSections = LODRenderData[LODIndex].RenderSections.Num();
-				check(NumRenderSections == LODRenderData[FirstGeneratedLOD].RenderSections.Num());
+					const int32 NumRenderSections = LODRenderData[LODIndex].RenderSections.Num();
+					check(NumRenderSections == LODRenderData[FirstLOD].RenderSections.Num());
 
-				if (NumRenderSections == LODRenderData[FirstGeneratedLOD].RenderSections.Num())
-				{
-					for (int32 RenderSectionIndex = 0; RenderSectionIndex < NumRenderSections; ++RenderSectionIndex)
+					if (NumRenderSections == LODRenderData[FirstLOD].RenderSections.Num())
 					{
-						const int32 MaterialIndex = LODRenderData[FirstGeneratedLOD].RenderSections[RenderSectionIndex].MaterialIndex;
-						LODRenderData[LODIndex].RenderSections[RenderSectionIndex].MaterialIndex = MaterialIndex;
+						for (int32 RenderSectionIndex = 0; RenderSectionIndex < NumRenderSections; ++RenderSectionIndex)
+						{
+							const int32 MaterialIndex = LODRenderData[FirstLOD].RenderSections[RenderSectionIndex].MaterialIndex;
+							LODRenderData[LODIndex].RenderSections[RenderSectionIndex].MaterialIndex = MaterialIndex;
+						}
 					}
 				}
 			}
@@ -6731,16 +6733,10 @@ int32 UCustomizableObjectInstance::GetNumComponents() const
 }
 
 
-void UCustomizableObjectInstance::SetRequestedLODs(const TMap<FName, uint8>& InMinLODs, const TMap<FName, uint8>& InRequestedLODs, FMutableInstanceUpdateMap& InOutRequestedUpdates)
+void UCustomizableObjectInstance::SetRequestedLODs(const TMap<FName, uint8>& InMinLODs, const TMap<FName, uint8>& InFirstRequestedLOD, FMutableInstanceUpdateMap& InOutRequestedUpdates)
 {
 	check(PrivateData);
-
-	// DOCUMENTATION 
-	// MinLod =					Hard limit set by the budgeter or the user (artifical limit).			Per Skeletal Mesh Component.
-	// QualitySettingMinLODs =	MinLOD based on the active quality settings.							COI Descriptor.
-	// First LOD Available =	First available lod per platform.										Per Skeletal Mesh.
-	// Requested LODs =			What the renderer wants to render.										Per Skeletal Mesh Component.
-
+	
 	if (!GetPrivate()->CanUpdateInstance())
 	{
 		return;
@@ -6758,6 +6754,11 @@ void UCustomizableObjectInstance::SetRequestedLODs(const TMap<FName, uint8>& InM
 		return;
 	}
 
+	if (IsStreamingEnabled(*CustomizableObject))
+	{
+		return;
+	}
+
 	if (CVarPreserveUserLODsOnFirstGeneration.GetValueOnGameThread() &&
 		CustomizableObject->bPreserveUserLODsOnFirstGeneration &&
 		GetPrivate()->SkeletalMeshStatus != ESkeletalMeshStatus::Success)
@@ -6778,12 +6779,12 @@ void UCustomizableObjectInstance::SetRequestedLODs(const TMap<FName, uint8>& InM
 	
 	// Save the new LODs
 	MutableUpdateCandidate.MinLOD = InMinLODs;
-	MutableUpdateCandidate.RequestedLODLevels = Descriptor.GetRequestedLODLevels();
+	MutableUpdateCandidate.FirstRequestedLOD = Descriptor.GetFirstRequestedLOD();
 	
-	const TMap<FName, uint8>& GeneratedLODsPerComponent = GetPrivate()->CommittedDescriptorHash.RequestedLODsPerComponent;
+	const TMap<FName, uint8>& FirstRequestedLOD = GetPrivate()->CommittedDescriptorHash.FirstRequestedLOD;
 	
 	for (const FName& ComponentName : ModelResources->ComponentNamesPerObjectComponent)
-	{
+	{	
 		uint8& InMinLOD = MutableUpdateCandidate.MinLOD.FindOrAdd(ComponentName);
 		if (const uint8* Result = InMinLODs.Find(ComponentName))
 		{
@@ -6822,28 +6823,30 @@ void UCustomizableObjectInstance::SetRequestedLODs(const TMap<FName, uint8>& InM
 			uint8 FirstNonStreamedLODIndex = 0;
 			if (const uint8* Found = ModelResources->NumLODsToStream.Find(ComponentName))
 			{
-				FirstNonStreamedLODIndex = FMath::Clamp(*Found, 0, MaxLODIdx);
+				FirstNonStreamedLODIndex = *Found;
 			}
 
-			MutableUpdateCandidate.RequestedLODLevels.Add(ComponentName, FirstNonStreamedLODIndex);
+			MutableUpdateCandidate.FirstRequestedLOD.Add(ComponentName, FirstNonStreamedLODIndex);
 
 			uint8 PredictedLOD = FirstNonStreamedLODIndex;
-			if (const uint8* Result = InRequestedLODs.Find(ComponentName))
+			if (const uint8* Result = InFirstRequestedLOD.Find(ComponentName))
 			{
 				PredictedLOD = FMath::Min(*Result, PredictedLOD);
 			}
 
-			if (const uint8* Result = GeneratedLODsPerComponent.Find(ComponentName))
+			if (const uint8* Result = FirstRequestedLOD.Find(ComponentName))
 			{
 				PredictedLOD = FMath::Min(*Result, PredictedLOD);
 			}
 
+			PredictedLOD = FMath::Clamp(PredictedLOD, MinLODIdx, MaxLODIdx);
+
 			// Save new RequestedLODs
-			MutableUpdateCandidate.RequestedLODLevels[ComponentName] = PredictedLOD;
+			MutableUpdateCandidate.FirstRequestedLOD[ComponentName] = PredictedLOD;
 		}
 	}
 	
-	if (bMinLODChanged || GeneratedLODsPerComponent != MutableUpdateCandidate.RequestedLODLevels)
+	if (bMinLODChanged || FirstRequestedLOD != MutableUpdateCandidate.FirstRequestedLOD)
 	{
 		// TODO: Remove this flag as it will become redundant with the new InOutRequestedUpdates system
 		PrivateData->SetCOInstanceFlags(PendingLODsUpdate);
diff --git a/Engine/Plugins/Mutable/Source/CustomizableObject/Private/MuCO/CustomizableObjectInstanceDescriptor.cpp b/Engine/Plugins/Mutable/Source/CustomizableObject/Private/MuCO/CustomizableObjectInstanceDescriptor.cpp
--- a/Engine/Plugins/Mutable/Source/CustomizableObject/Private/MuCO/CustomizableObjectInstanceDescriptor.cpp
+++ b/Engine/Plugins/Mutable/Source/CustomizableObject/Private/MuCO/CustomizableObjectInstanceDescriptor.cpp
@@ -827,7 +827,7 @@ void FCustomizableObjectInstanceDescriptor::ReloadParameters()
 
 	SetState(FMath::Clamp(GetState(), 0, CustomizableObject->GetStateCount() - 1));
 
-	RequestedLODLevels.Empty();
+	FirstRequestedLOD.Empty();
 
 	TArray<FCustomizableObjectBoolParameterValue> OldBoolParameters = BoolParameters;
 	TArray<FCustomizableObjectIntParameterValue> OldIntParameters = IntParameters;
@@ -1301,15 +1301,15 @@ void FCustomizableObjectInstanceDescriptor::ReloadParameters()
 }
 
 
-void FCustomizableObjectInstanceDescriptor::SetRequestedLODLevels(const TMap<FName, uint8>& InRequestedLODLevels)
+void FCustomizableObjectInstanceDescriptor::SetFirstRequestedLOD(const TMap<FName, uint8>& InFirstRequestedLOD)
 {
-	RequestedLODLevels = InRequestedLODLevels;
+	FirstRequestedLOD = InFirstRequestedLOD;
 }
 
 
-const TMap<FName, uint8>& FCustomizableObjectInstanceDescriptor::GetRequestedLODLevels() const
+const TMap<FName, uint8>& FCustomizableObjectInstanceDescriptor::GetFirstRequestedLOD() const
 {
-	return RequestedLODLevels;
+	return FirstRequestedLOD;
 }
 
 
diff --git a/Engine/Plugins/Mutable/Source/CustomizableObject/Private/MuCO/CustomizableObjectSkeletalMesh.cpp b/Engine/Plugins/Mutable/Source/CustomizableObject/Private/MuCO/CustomizableObjectSkeletalMesh.cpp
--- a/Engine/Plugins/Mutable/Source/CustomizableObject/Private/MuCO/CustomizableObjectSkeletalMesh.cpp
+++ b/Engine/Plugins/Mutable/Source/CustomizableObject/Private/MuCO/CustomizableObjectSkeletalMesh.cpp
@@ -14,7 +14,7 @@
 
 
 
-void UCustomizableObjectSkeletalMesh::InitMutableStreamingData(const TSharedRef<FUpdateContextPrivate>& InOperationData, const FName& ComponentName, const int32 FirstLOD, const int32 LODCount)
+void UCustomizableObjectSkeletalMesh::InitMutableStreamingData(const TSharedRef<FUpdateContextPrivate>& InOperationData, const FName& ComponentName, const int32 InstanceUpdateFirstLOD, const int32 LODCount)
 {
 	UCustomizableObject& CustomizableObject = *InOperationData->Object;
 
@@ -31,9 +31,13 @@ void UCustomizableObjectSkeletalMesh::InitMutableStreamingData(const TSharedRef<
 
 	MeshIDs.Init(MAX_uint64, MAX_MESH_LOD_COUNT);
 
-	for (int32 LODIndex = InOperationData->FirstLODAvailable[ComponentName]; LODIndex < LODCount; ++LODIndex)
+	const int32 FirstLOD = InOperationData->bStreamMeshLODs ?
+		InOperationData->FirstLODAvailable[ComponentName] :
+		InOperationData->GetFirstRequestedLOD()[ComponentName];
+
+	for (int32 LODIndex = FirstLOD; LODIndex < LODCount; ++LODIndex)
 	{
-		const FInstanceUpdateData::FLOD& LOD = InOperationData->InstanceUpdateData.LODs[FirstLOD + LODIndex];
+		const FInstanceUpdateData::FLOD& LOD = InOperationData->InstanceUpdateData.LODs[InstanceUpdateFirstLOD + LODIndex];
 		MeshIDs[LODIndex] = LOD.MeshID;
 	}
 
diff --git a/Engine/Plugins/Mutable/Source/CustomizableObject/Private/MuCO/CustomizableObjectSystem.cpp b/Engine/Plugins/Mutable/Source/CustomizableObject/Private/MuCO/CustomizableObjectSystem.cpp
--- a/Engine/Plugins/Mutable/Source/CustomizableObject/Private/MuCO/CustomizableObjectSystem.cpp
+++ b/Engine/Plugins/Mutable/Source/CustomizableObject/Private/MuCO/CustomizableObjectSystem.cpp
@@ -588,12 +588,6 @@ FUpdateContextPrivate::~FUpdateContextPrivate()
 }
 
 
-const TMap<FName, uint8>& FUpdateContextPrivate::GetMinLOD() const
-{
-	return CapturedDescriptor.MinLOD;
-}
-
-
 void FUpdateContextPrivate::SetMinLOD(const TMap<FName, uint8>& MinLOD)
 {
 	CapturedDescriptor.MinLOD = MinLOD;
@@ -608,16 +602,16 @@ void FUpdateContextPrivate::SetQualitySettingMinLODs(const TMap<FName, uint8>& Q
 }
 
 
-const TMap<FName, uint8>& FUpdateContextPrivate::GetRequestedLODs() const
+const TMap<FName, uint8>& FUpdateContextPrivate::GetFirstRequestedLOD() const
 {
-	return CapturedDescriptor.GetRequestedLODLevels();
+	return CapturedDescriptor.GetFirstRequestedLOD();
 }
 
 
-void FUpdateContextPrivate::SetRequestedLODs(const TMap<FName, uint8>& RequestedLODs)
+void FUpdateContextPrivate::SetFirstRequestedLOD(const TMap<FName, uint8>& FirstRequestedLOD)
 {
-	CapturedDescriptor.SetRequestedLODLevels(RequestedLODs);
-	CapturedDescriptorHash.RequestedLODsPerComponent = RequestedLODs;
+	CapturedDescriptor.SetFirstRequestedLOD(FirstRequestedLOD);
+	CapturedDescriptorHash.FirstRequestedLOD = FirstRequestedLOD;
 }
 
 
@@ -1134,7 +1128,7 @@ int32 UCustomizableObjectSystemPrivate::EnableOnlyGenerateRequestedLODs = 1;
 
 static FAutoConsoleVariableRef CVarEnableOnlyGenerateRequestedLODs(
 	TEXT("mutable.EnableOnlyGenerateRequestedLODs"), UCustomizableObjectSystemPrivate::EnableOnlyGenerateRequestedLODs,
-	TEXT("If 1 or greater, Only the RequestedLODLevels will be generated. If 0, all LODs will be build."),
+	TEXT("If 1 or greater, Only the RequestedLOD will be generated. If 0, all LODs will be build."),
 	ECVF_Default);
 
 int32 UCustomizableObjectSystemPrivate::EnableSkipGenerateResidentMips = 1;
@@ -2095,7 +2089,7 @@ namespace impl
 
 	struct FGetMeshData
 	{
-		int32 LODIndex;
+		int32 InstanceUpdateLODIndex;
 		mu::FResourceID MeshID;
 		mu::EMeshContentFlags ContentFilter;
 	};
@@ -2168,8 +2162,8 @@ namespace impl
 			return;
 		}
 
-		TMap<FName, uint8> MinLODs = Operation->bStreamMeshLODs ? TMap<FName, uint8>() : Operation->GetMinLOD();
-		TMap<FName, uint8> RequestedLODs = Operation->GetRequestedLODs();
+		TMap<FName, uint8> MinLODs = Operation->GetCapturedDescriptor().MinLOD;
+		TMap<FName, uint8> RequestedLODs = Operation->GetFirstRequestedLOD();
 		
 		for (int32 InstanceComponentIndex = 0; InstanceComponentIndex < Operation->NumInstanceComponents; ++InstanceComponentIndex)
 		{
@@ -2183,22 +2177,23 @@ namespace impl
 			const uint8& FirstLODAvailable = Operation->FirstLODAvailable.FindOrAdd(ComponentName, 0);
 
 			NumLODsAvailable = Operation->MutableInstance->GetLODCount(InstanceComponentIndex);
-			MinLOD = FMath::Clamp(MinLOD, FirstLODAvailable, NumLODsAvailable - 1);
 			
 			if (Operation->bStreamMeshLODs)
 			{
 				FirstResidentLOD = FMath::Clamp(FirstResidentLOD, FirstLODAvailable, NumLODsAvailable - 1);
-				RequestedLOD = MinLOD;
+				MinLOD = 0;
+				RequestedLOD = FirstResidentLOD;
 			}
 			else
 			{
 				FirstResidentLOD = FirstLODAvailable;
-				RequestedLOD = FMath::Min(RequestedLOD, static_cast<uint8>(NumLODsAvailable - 1));
+				MinLOD = FMath::Clamp(MinLOD, FirstLODAvailable, NumLODsAvailable - 1);
+				RequestedLOD = FMath::Clamp(RequestedLOD, MinLOD, static_cast<uint8>(NumLODsAvailable - 1));
 			}
 		}
 		
 		Operation->SetMinLOD(MinLODs);
-		Operation->SetRequestedLODs(RequestedLODs);
+		Operation->SetFirstRequestedLOD(RequestedLODs);
 	}
 	
 
@@ -2253,7 +2248,7 @@ namespace impl
 
 			for (int32 LODIndex = 0; LODIndex < Component.LODCount; ++LODIndex)
 			{
-				FInstanceUpdateData::FLOD& LOD = OperationData->InstanceUpdateData.LODs[Component.FirstLOD+LODIndex];
+				FInstanceUpdateData::FLOD& LOD = OperationData->InstanceUpdateData.LODs[Component.FirstLOD + LODIndex];
 
 				TSharedPtr<const mu::FMesh> Mesh = LOD.Mesh;
 				if (!Mesh || Mesh->IsReference())
@@ -2310,7 +2305,12 @@ namespace impl
 		{
 			return;
 		}
+
+		UCustomizableObject* Object = OperationData->Object.Get();
+		const UModelResources* ModelResources = Object->GetPrivate()->GetModelResources();
 		
+		const TSharedPtr<FModelStreamableBulkData>& ModelStreamableBulkData = Object->GetPrivate()->GetModelStreamableBulkData();
+
 		const int32 NumInstanceComponents = OperationData->InstanceUpdateData.Components.Num();
 		check(OperationData->InstanceUpdateData.RealTimeMorphTargets.IsEmpty());
 		OperationData->InstanceUpdateData.RealTimeMorphTargets.Reserve(NumInstanceComponents);
@@ -2347,19 +2347,17 @@ namespace impl
 			TArray<TArray<FMorphTargetLODModel>>& MorphsData = ComponentMorphTargetsData.RealTimeMorphsLODData;
 			const int32 NumMorphs = ComponentMorphTargetsData.RealTimeMorphTargetNames.Num();
 			
-			const int32 NumLODsInThisComponent = OperationData->NumLODsAvailable[ComponentName];
+			const int32 NumLODsAvailable = OperationData->NumLODsAvailable[ComponentName];
 
 			MorphsData.SetNum(ComponentMorphTargetsData.RealTimeMorphTargetNames.Num());
 			for (int32 MorphIndex = 0; MorphIndex < NumMorphs; ++MorphIndex)
 			{
-				MorphsData[MorphIndex].SetNum(NumLODsInThisComponent);
+				MorphsData[MorphIndex].SetNum(NumLODsAvailable);
 			}
 			
-			const int32 FirstGeneratedLOD = FMath::Max(OperationData->GetRequestedLODs()[ComponentName], OperationData->GetMinLOD()[ComponentName]);
-			for (int32 LODIndex = FirstGeneratedLOD; LODIndex < NumLODsInThisComponent; ++LODIndex)
+			for (int32 LODIndex = OperationData->GetFirstRequestedLOD()[ComponentName]; LODIndex < NumLODsAvailable; ++LODIndex)
 			{
 				const FInstanceUpdateData::FLOD& LOD = UpdateData.LODs[Component.FirstLOD+LODIndex];
-				check(LOD.bGenerated);
 
 				if (!LOD.Mesh)
 				{
@@ -2377,6 +2375,41 @@ namespace impl
 					}
 				}
 			}
+
+			const int32 FirstLOD = OperationData->bStreamMeshLODs ?
+				OperationData->FirstLODAvailable[ComponentName] :
+				OperationData->GetFirstRequestedLOD()[ComponentName];
+
+			// Find which Sections are being used in each LOD (Streamed and Residents).
+			for (int32 LODIndex = FirstLOD; LODIndex < NumLODsAvailable; ++LODIndex)
+			{
+				const FInstanceUpdateData::FLOD& LOD = UpdateData.LODs[Component.FirstLOD+LODIndex];
+				check(LOD.Mesh);
+				
+				for (int32 SectionIndex = 0; SectionIndex < LOD.Mesh->Surfaces.Num(); ++SectionIndex)
+				{
+					const mu::FMeshSurface& Surface = LOD.Mesh->Surfaces[SectionIndex];
+					for (const mu::FSurfaceSubMesh& SubMesh : Surface.SubMeshes)
+					{
+						const uint32 MorphMetadataId = ModelResources->MeshMetadata[SubMesh.ExternalId].MorphMetadataId;
+
+						FRealTimeMorphStreamable* Result = ModelStreamableBulkData->RealTimeMorphStreamables.Find(MorphMetadataId);
+						if (!Result)
+						{
+							continue;
+						}
+
+						for (const FName& MorphName : Result->NameResolutionMap)
+						{
+							const int32 MorphIndex = ComponentMorphTargetsData.RealTimeMorphTargetNames.Find(MorphName);
+							check(MorphIndex != INDEX_NONE);
+							
+							FMorphTargetLODModel& MorphTargetLODModel = ComponentMorphTargetsData.RealTimeMorphsLODData[MorphIndex][LODIndex];
+							MorphTargetLODModel.SectionIndices.Add(SectionIndex);
+						}
+					}
+				}
+			}
 			
 			// Remove empty morph targets
 			for (int32 MorphIndex = 0; MorphIndex < NumMorphs; ++MorphIndex)
@@ -2513,16 +2546,10 @@ namespace impl
 		for (int32 InstanceComponentIndex = 0; InstanceComponentIndex < OperationData->NumInstanceComponents; ++InstanceComponentIndex)
 		{
 			FInstanceUpdateData::FComponent& Component = OperationData->InstanceUpdateData.Components[InstanceComponentIndex];
-			for (int32 MutableLODIndex = 0; MutableLODIndex < Component.LODCount; ++MutableLODIndex)
-			{
-				const FName& ComponentName = ModelResources.ComponentNamesPerObjectComponent[Component.Id.GetValue()];
-				
-				// Skip LODs outside the range we want to generate
-				if (MutableLODIndex < OperationData->GetMinLOD()[ComponentName])
-				{
-					continue;
-				}
+			const FName& ComponentName = ModelResources.ComponentNamesPerObjectComponent[Component.Id.GetValue()];
 
+			for (int32 MutableLODIndex = OperationData->FirstLODAvailable[ComponentName]; MutableLODIndex < Component.LODCount; ++MutableLODIndex)
+			{
 				FInstanceUpdateData::FLOD& LOD = OperationData->InstanceUpdateData.LODs[Component.FirstLOD + MutableLODIndex];
 
 				LOD.FirstSurface = OperationData->InstanceUpdateData.Surfaces.Num();
@@ -2532,7 +2559,6 @@ namespace impl
 				{
 					continue;
 				}
-				
 
 				// This lambda does all the work to fill up the surface data
 				auto AddSurface = 
@@ -2700,7 +2726,7 @@ namespace impl
 	{
 		MUTABLE_CPUPROFILER_SCOPE(Task_MutableGetMeshes_GetMesh_Post)
 
-		const int32 LODIndex = (*GetMeshesData)[GetMeshIndex].LODIndex;
+		const int32 LODIndex = (*GetMeshesData)[GetMeshIndex].InstanceUpdateLODIndex;
 		FInstanceUpdateData::FLOD& LOD = OperationData->InstanceUpdateData.LODs[LODIndex];
 
 		LOD.Mesh = GetMeshTask.GetResult();
@@ -2826,45 +2852,33 @@ namespace impl
 				Component.OverlayMaterial = Instance->GetOverlayMaterial(InstanceComponentIndex);
 
 				const uint8 FirstResidentLOD = OperationData->FirstResidentLOD[ComponentName];
-				
-				for (int32 MutableLODIndex = 0; MutableLODIndex < Component.LODCount; ++MutableLODIndex)
-				{
-					// If the LOD is not generated we still add an empty one to keep indexes aligned.
-					const int32 UpdateDataLODIndex = OperationData->InstanceUpdateData.LODs.AddDefaulted();
-
-					// Skip LODs outside the range we want to generate
-					if (MutableLODIndex < OperationData->GetMinLOD()[ComponentName])
-					{
-						continue;
-					}
-
-					FInstanceUpdateData::FLOD& LOD = OperationData->InstanceUpdateData.LODs.Last();
 
-					const bool bGenerateLOD = OperationData->GetRequestedLODs()[ComponentName] <= MutableLODIndex;
+				// If the LOD is not generated we still add an empty one to keep indexes aligned.
+				OperationData->InstanceUpdateData.LODs.SetNum(Component.FirstLOD + Component.LODCount);
 
-					// Mesh
+				const int32 FirstLOD = OperationData->bStreamMeshLODs ?
+					OperationData->FirstLODAvailable[ComponentName] :
+					OperationData->GetFirstRequestedLOD()[ComponentName];
+				
+				for (int32 LODIndex = FirstLOD; LODIndex < Component.LODCount; ++LODIndex)
+				{
+					MUTABLE_CPUPROFILER_SCOPE(GetMesh);
+					
+					FInstanceUpdateData::FLOD& LOD = OperationData->InstanceUpdateData.LODs[Component.FirstLOD + LODIndex];
+					
+					LOD.MeshID = Instance->GetMeshId(InstanceComponentIndex, LODIndex);
+					
+					mu::EMeshContentFlags MeshContentFilter = mu::EMeshContentFlags::AllFlags; 
+					
+					if (!bForceGeometryGeneration)
 					{
-						MUTABLE_CPUPROFILER_SCOPE(GetMesh);
-
-						LOD.MeshID = Instance->GetMeshId(InstanceComponentIndex, MutableLODIndex);
-
-						if (bGenerateLOD)
+						if (LODIndex < FirstResidentLOD)
 						{
-							LOD.bGenerated = true;
-
-							mu::EMeshContentFlags MeshContentFilter = mu::EMeshContentFlags::AllFlags; 
-							
-							if (!bForceGeometryGeneration)
-							{
-								if (MutableLODIndex < FirstResidentLOD)
-								{
-									EnumRemoveFlags(MeshContentFilter, mu::EMeshContentFlags::GeometryData);
-								}
-							}
-
-							GetMeshesData->Add({ UpdateDataLODIndex, LOD.MeshID, MeshContentFilter });
+							EnumRemoveFlags(MeshContentFilter, mu::EMeshContentFlags::GeometryData);
 						}
 					}
+					
+					GetMeshesData->Add({ Component.FirstLOD + LODIndex, LOD.MeshID, MeshContentFilter });
 				}
 			}
 
@@ -3729,8 +3743,8 @@ if (CVarMutableLogObjectMemoryOnUpdate.GetValueOnAnyThread())
 		Operation->InstanceID = Operation->bLiveUpdateMode ? CandidateInstancePrivateData->LiveUpdateModeInstanceID : 0;
 		Operation->bUseMeshCache = CustomizableObject->bEnableMeshCache && !Operation->bLiveUpdateMode && UCustomizableObjectSystem::IsMeshCacheEnabled(true);
 
-		const bool bStreamingEnabled = (CustomizableObject->bEnableMeshStreaming || bForceStreamMeshLODs) && bStreamMeshLODs;
-		Operation->bStreamMeshLODs = bStreamingEnabled && IStreamingManager::Get().IsRenderAssetStreamingEnabled(EStreamableRenderAssetType::SkeletalMesh);
+		Operation->bStreamMeshLODs = IsStreamingEnabled(*CustomizableObject);
+
 #if WITH_EDITOR
 		Operation->PixelFormatOverride = SystemPrivateData->ImageFormatOverrideFunc;
 #endif
@@ -3780,13 +3794,13 @@ if (CVarMutableLogObjectMemoryOnUpdate.GetValueOnAnyThread())
 
 		if (bRequestAllLODs)
 		{
-			TMap<FName, uint8> RequestedLODs = Operation->GetRequestedLODs();
+			TMap<FName, uint8> RequestedLODs = Operation->GetFirstRequestedLOD();
 			for (const FName& ComponentName : ModelResources.ComponentNamesPerObjectComponent)
 			{
 				RequestedLODs.Add(ComponentName, 0);
 			}
 
-			Operation->SetRequestedLODs(RequestedLODs);
+			Operation->SetFirstRequestedLOD(RequestedLODs);
 		}
 
 		// CreateMutableInstance
@@ -3838,7 +3852,6 @@ if (CVarMutableLogObjectMemoryOnUpdate.GetValueOnAnyThread())
 		{
 			Operation->InitMeshDescriptors(Operation->NumObjectComponents);
 
-			const TMap<FName, uint8>& MinLOD = Operation->GetMinLOD();
 			const TMap<FName, uint8>& NumLODsAvailable = Operation->NumLODsAvailable;
 		
 			for (int32 InstanceComponentIndex = 0; InstanceComponentIndex < Operation->NumInstanceComponents; ++InstanceComponentIndex)
@@ -3853,14 +3866,14 @@ if (CVarMutableLogObjectMemoryOnUpdate.GetValueOnAnyThread())
 				MeshId->Init(MAX_uint64, MAX_MESH_LOD_COUNT);
 
 				const FName& ComponentName = Operation->ComponentNames[ObjectComponentIndex.GetValue()];
-			
-				for (int32 LODIndex = MinLOD[ComponentName]; LODIndex < NumLODsAvailable[ComponentName]; ++LODIndex)
+
+				const int32 FirstLOD = Operation->bStreamMeshLODs ?
+					Operation->FirstLODAvailable[ComponentName] :
+					Operation->GetFirstRequestedLOD()[ComponentName];
+	
+				for (int32 LODIndex = FirstLOD; LODIndex < NumLODsAvailable[ComponentName]; ++LODIndex)
 				{
-					const bool bGenerateLOD = Operation->GetRequestedLODs()[ComponentName] <= LODIndex;
-					if (bGenerateLOD)
-					{
-						(*MeshId)[LODIndex] = Operation->MutableInstance->GetMeshId(InstanceComponentIndex, LODIndex);
-					}
+					(*MeshId)[LODIndex] = Operation->MutableInstance->GetMeshId(InstanceComponentIndex, LODIndex);
 				}
 			}
 			
@@ -4134,7 +4147,7 @@ int32 UCustomizableObjectSystem::TickInternal(const bool bBlocking)
 						FDescriptorHash LODUpdateDescriptorHash = CustomizableInstancePrivate->CommittedDescriptorHash;
 						LODUpdateDescriptorHash.MinLODs = LODUpdateCandidate.MinLOD;
 						LODUpdateDescriptorHash.QualitySettingMinLODs = LODUpdateCandidate.QualitySettingMinLODs;
-						LODUpdateDescriptorHash.RequestedLODsPerComponent = LODUpdateCandidate.RequestedLODLevels;
+						LODUpdateDescriptorHash.FirstRequestedLOD = LODUpdateCandidate.FirstRequestedLOD;
 
 						if (CustomizableInstancePrivate->MinSquareDistFromComponentToPlayer < MaxSquareDistanceFound &&
 							!LODUpdateDescriptorHash.IsSubset(CustomizableInstancePrivate->CommittedDescriptorHash))
@@ -5052,9 +5065,16 @@ void ReconstructMorphTargets(const mu::FMesh& Mesh, const TArray<FName>& GlobalN
 			{
 				FMorphTargetLODModel& MorphTargetLodModel = OutMorphTargets[MorphIdx];
 
-				MorphTargetLodModel.SectionIndices.Add(Section);
 				MorphTargetLodModel.NumVertices += SectionMorphTargetVerticesCount[MorphIdx];
 			}
 		}
 	}
 }
+
+
+bool IsStreamingEnabled(const UCustomizableObject& Object)
+{
+	return (Object.bEnableMeshStreaming || bForceStreamMeshLODs) &&
+		bStreamMeshLODs &&
+		IStreamingManager::Get().IsRenderAssetStreamingEnabled(EStreamableRenderAssetType::SkeletalMesh);
+}
\ No newline at end of file
diff --git a/Engine/Plugins/Mutable/Source/CustomizableObject/Private/MuCO/DescriptorHash.cpp b/Engine/Plugins/Mutable/Source/CustomizableObject/Private/MuCO/DescriptorHash.cpp
--- a/Engine/Plugins/Mutable/Source/CustomizableObject/Private/MuCO/DescriptorHash.cpp
+++ b/Engine/Plugins/Mutable/Source/CustomizableObject/Private/MuCO/DescriptorHash.cpp
@@ -59,7 +59,7 @@ FDescriptorHash::FDescriptorHash(const FCustomizableObjectInstanceDescriptor& De
 	
 	MinLODs = Descriptor.MinLOD;
 	QualitySettingMinLODs = Descriptor.QualitySettingMinLODs;
-	RequestedLODsPerComponent = Descriptor.RequestedLODLevels;
+	FirstRequestedLOD = Descriptor.FirstRequestedLOD;
 }
 
 
@@ -90,14 +90,14 @@ bool FDescriptorHash::IsSubset(const FDescriptorHash& Other) const
 		return false;
 	}
 
-	if (RequestedLODsPerComponent.Num() != Other.RequestedLODsPerComponent.Num())
+	if (FirstRequestedLOD.Num() != Other.FirstRequestedLOD.Num())
 	{
 		return false;
 	}
 	
-	for (const TTuple<FName, uint8>& Pair : RequestedLODsPerComponent)
+	for (const TTuple<FName, uint8>& Pair : FirstRequestedLOD)
 	{
-		const uint8* PairOther = Other.RequestedLODsPerComponent.Find(Pair.Key);
+		const uint8* PairOther = Other.FirstRequestedLOD.Find(Pair.Key);
 		if (!PairOther)
 		{
 			return false;
@@ -133,8 +133,8 @@ FString FDescriptorHash::ToString() const
 	}
 	Builder.Appendf(TEXT("])"));
 
-	Builder.Appendf(TEXT("RequestedLODLevels=["));
-	for (const TTuple<FName, uint8>& RequestedLODs : RequestedLODsPerComponent)
+	Builder.Appendf(TEXT("FirstRequestedLOD=["));
+	for (const TTuple<FName, uint8>& RequestedLODs : FirstRequestedLOD)
 	{
 		Builder.Appendf(TEXT("(%s, %i),"), *RequestedLODs.Key.ToString(), RequestedLODs.Value);
 	}
diff --git a/Engine/Plugins/Mutable/Source/CustomizableObject/Private/MuCO/UnrealConversionUtils.cpp b/Engine/Plugins/Mutable/Source/CustomizableObject/Private/MuCO/UnrealConversionUtils.cpp
--- a/Engine/Plugins/Mutable/Source/CustomizableObject/Private/MuCO/UnrealConversionUtils.cpp
+++ b/Engine/Plugins/Mutable/Source/CustomizableObject/Private/MuCO/UnrealConversionUtils.cpp
@@ -821,7 +821,7 @@ namespace UnrealConversionUtils
 
 		// Reconstruct the final Morph Targets using the global names.
 		TArray<FMorphTargetLODModel> MorphTargetLODs;
-		ReconstructMorphTargets(MutableMesh, GlobalNames, GlobalMorphTargets, MorphTargetLODs);
+		ReconstructMorphTargets(MutableMesh, GlobalNames, GlobalMorphTargets, MorphTargetLODs); // FMorphTargetLODModel::SectionIndices not initialized here
 
 		TArray<const FMorphTargetLODModel*> MorphsToInit;
 		MorphsToInit.Reserve(MorphTargetLODs.Num());
diff --git a/Engine/Plugins/Mutable/Source/CustomizableObject/Public/MuCO/CustomizableObjectInstance.h b/Engine/Plugins/Mutable/Source/CustomizableObject/Public/MuCO/CustomizableObjectInstance.h
--- a/Engine/Plugins/Mutable/Source/CustomizableObject/Public/MuCO/CustomizableObjectInstance.h
+++ b/Engine/Plugins/Mutable/Source/CustomizableObject/Public/MuCO/CustomizableObjectInstance.h
@@ -679,9 +679,9 @@ public:
 	/** Sets an array of LODs to generate per component. Mutable will generate those plus the currently generated LODs (if any).
 	  * Requires mutable.EnableOnlyGenerateRequestedLODs and CurrentInstanceLODManagement->IsOnlyGenerateRequestedLODLevelsEnabled() to be true.
 	  * @param InMinLODs - MinLOD to generate per Component.
-	  * @param InRequestedLODs - Requested LODs per component to generate. Key is the Component name. Value represents the first LOD to generate. Zero means all LODs.
+	  * @param InFirstRequestedLOD - Requested LODs per component to generate. Key is the Component name. Value represents the first LOD to generate. Zero means all LODs.
 	  * @param InOutRequestedUpdates - Map from Instance to Update data that stores a request for the Instance to be updated, which will be either processed or discarded by priority (to be rerequested the next tick) */
-	void SetRequestedLODs(const TMap<FName, uint8>& InMinLODs, const TMap<FName, uint8>& InRequestedLODs, FMutableInstanceUpdateMap& InOutRequestedUpdates);
+	void SetRequestedLODs(const TMap<FName, uint8>& InMinLODs, const TMap<FName, uint8>& InFirstRequestedLOD, FMutableInstanceUpdateMap& InOutRequestedUpdates);
 
 #if WITH_EDITOR
 	/**
diff --git a/Engine/Plugins/Mutable/Source/CustomizableObject/Public/MuCO/CustomizableObjectInstanceDescriptor.h b/Engine/Plugins/Mutable/Source/CustomizableObject/Public/MuCO/CustomizableObjectInstanceDescriptor.h
--- a/Engine/Plugins/Mutable/Source/CustomizableObject/Public/MuCO/CustomizableObjectInstanceDescriptor.h
+++ b/Engine/Plugins/Mutable/Source/CustomizableObject/Public/MuCO/CustomizableObjectInstanceDescriptor.h
@@ -73,9 +73,9 @@ struct CUSTOMIZABLEOBJECT_API FCustomizableObjectInstanceDescriptor
 	/** Update all parameters to be up to date with the Mutable Core parameters. */
 	void ReloadParameters();
 	
-	void SetRequestedLODLevels(const TMap<FName, uint8>& InRequestedLODLevels);
+	void SetFirstRequestedLOD(const TMap<FName, uint8>& FirstRequestedLOD);
 
-	const TMap<FName, uint8>& GetRequestedLODLevels() const;
+	const TMap<FName, uint8>& GetFirstRequestedLOD() const;
 
 	// ------------------------------------------------------------
 	// FParameters
@@ -343,8 +343,7 @@ struct CUSTOMIZABLEOBJECT_API FCustomizableObjectInstanceDescriptor
 
 	/** RequestedLODs per component to generate.
 	 * Key is the Component name. Value represents the first LOD to generate. Zero means all LODs. */
-	TMap<FName, uint8> RequestedLODLevels;
-	
+	TMap<FName, uint8> FirstRequestedLOD;
 	
 	// Friends
 	friend FDescriptorHash;
diff --git a/Engine/Plugins/Mutable/Source/CustomizableObjectEditor/Private/MuCOE/CustomizableInstanceDetails.cpp b/Engine/Plugins/Mutable/Source/CustomizableObjectEditor/Private/MuCOE/CustomizableInstanceDetails.cpp
--- a/Engine/Plugins/Mutable/Source/CustomizableObjectEditor/Private/MuCOE/CustomizableInstanceDetails.cpp
+++ b/Engine/Plugins/Mutable/Source/CustomizableObjectEditor/Private/MuCOE/CustomizableInstanceDetails.cpp
@@ -2407,7 +2407,7 @@ FReply FCustomizableInstanceDetails::OnPasteAllParameters()
 	UCustomizableObject* CustomizableObject = Descriptor.GetCustomizableObject();
 	
 	const TMap<FName, uint8> MinLOD = Descriptor.MinLOD;
-	const TMap<FName, uint8> RequestedLODLevels = Descriptor.GetRequestedLODLevels();
+	const TMap<FName, uint8> FirstRequestedLOD = Descriptor.GetFirstRequestedLOD();
 	
 	FScopedTransaction LocalTransaction(LOCTEXT("OnPasteAllParameters", "Paste All Parameters"));
 	CustomInstance->Modify();
@@ -2418,7 +2418,7 @@ FReply FCustomizableInstanceDetails::OnPasteAllParameters()
 		
 		// Keep current LOD
 		Descriptor.MinLOD = MinLOD;
-		Descriptor.SetRequestedLODLevels(RequestedLODLevels);
+		Descriptor.SetFirstRequestedLOD(FirstRequestedLOD);
 		
 		UpdateInstance();
 	}
diff --git a/Engine/Plugins/Mutable/Source/CustomizableObjectEditor/Private/MuCOE/CustomizableObjectBenchmarkingUtils.cpp b/Engine/Plugins/Mutable/Source/CustomizableObjectEditor/Private/MuCOE/CustomizableObjectBenchmarkingUtils.cpp
--- a/Engine/Plugins/Mutable/Source/CustomizableObjectEditor/Private/MuCOE/CustomizableObjectBenchmarkingUtils.cpp
+++ b/Engine/Plugins/Mutable/Source/CustomizableObjectEditor/Private/MuCOE/CustomizableObjectBenchmarkingUtils.cpp
@@ -38,12 +38,12 @@ bool CustomizableObjectBenchmarkingUtils::GenerateDeterministicSetOfInstances(UC
 		if (UCustomizableObjectInstance* GeneratedInstance = TargetCustomizableObject.CreateInstance())
 		{
 			// Force generation of all LODS
-			TMap<FName, uint8> RequestedLodLevels;
+			TMap<FName, uint8> FirstRequestedLOD;
 			for (const FName ComponentName : ModelResources->ComponentNamesPerObjectComponent)
 			{
-				RequestedLodLevels.Add(ComponentName, 0);
+				FirstRequestedLOD.Add(ComponentName, 0);
 			}
-			GeneratedInstance->GetPrivate()->GetDescriptor().SetRequestedLODLevels(RequestedLodLevels);
+			GeneratedInstance->GetPrivate()->GetDescriptor().SetFirstRequestedLOD(FirstRequestedLOD);
 				
 			// Randomize instance values
 			GeneratedInstance->SetRandomValuesFromStream(RandomizationStream);
