[Mutable] Fix baking action getting stuck if the instance is not compiled.

#rnx

diff --git a/Engine/Plugins/Mutable/Source/CustomizableObjectEditor/Private/MuCOE/CustomizableObjectInstanceBakingUtils.cpp b/Engine/Plugins/Mutable/Source/CustomizableObjectEditor/Private/MuCOE/CustomizableObjectInstanceBakingUtils.cpp
--- a/Engine/Plugins/Mutable/Source/CustomizableObjectEditor/Private/MuCOE/CustomizableObjectInstanceBakingUtils.cpp
+++ b/Engine/Plugins/Mutable/Source/CustomizableObjectEditor/Private/MuCOE/CustomizableObjectInstanceBakingUtils.cpp
@@ -330,6 +330,14 @@ void ScheduleInstanceUpdateForBaking(UCustomizableObjectInstance& InInstance, FI
 		return;
 	}
 
+	// TODO UE-262123: Compile Instance with correct settings instead
+	if (!InInstance.GetCustomizableObject() || !InInstance.GetCustomizableObject()->IsCompiled())
+	{
+		UE_LOG(LogMutable, Error, TEXT("The COInstance update for baking could not be scheduled. Instance not compiled."));
+		InInstanceUpdateDelegate.Broadcast({ EUpdateResult::Error });
+		return;
+	}
+
 	InInstanceUpdateDelegate.AddStatic(&OnInstanceUpdateFinish);
 	
 	// Desired System settings for the update part of a baking operation
