[mutable] Added some missing constant types to be able to look for their users in the CodeViewer
- Improved handling of constant resources which user operantions can not be found.
- Improved handling of clearing of the CodeViewer found operation types.
- Improved handling of duplicated states. Should never happen but for the debugger it is better to be as defensive as possible.

diff --git a/Engine/Plugins/Mutable/Source/CustomizableObjectEditor/Private/MuCOE/SMutableCodeViewer.cpp b/Engine/Plugins/Mutable/Source/CustomizableObjectEditor/Private/MuCOE/SMutableCodeViewer.cpp
--- a/Engine/Plugins/Mutable/Source/CustomizableObjectEditor/Private/MuCOE/SMutableCodeViewer.cpp
+++ b/Engine/Plugins/Mutable/Source/CustomizableObjectEditor/Private/MuCOE/SMutableCodeViewer.cpp
@@ -996,12 +996,19 @@ TSharedRef<SWidget> SMutableCodeViewer::OnGenerateOpNavigationDropDownWidget(
 void SMutableCodeViewer::OnNavigationSelectedOperationChanged(
 	TSharedPtr<const FMutableOperationElement, ESPMode::ThreadSafe> MutableOperationElement, ESelectInfo::Type Arg)
 {
-	check (MutableOperationElement.IsValid());
+	// Handle the case where we do not want an option selected, for example, when clearing the selected option.
+	TSharedPtr<const FMutableOperationElement, ESPMode::ThreadSafe> NewSelectedElement = MutableOperationElement;
+	if (!NewSelectedElement.IsValid())
+	{
+		NewSelectedElement = NoneOperationEntry;
+	}
+	
+	check (NewSelectedElement.IsValid());
 
 	// Cache the currently selected operation set on the UI by the user
-	const mu::EOpType NewOperationType = MutableOperationElement->OperationType;
+	const mu::EOpType NewOperationType = NewSelectedElement->OperationType;
 	OperationTypeToSearch = NewOperationType;
-	CurrentlySelectedOperationTypeElement = MutableOperationElement;
+	CurrentlySelectedOperationTypeElement = NewSelectedElement;
 	
 	// Only do the internal work if the type is one that makes sense searching
 	if (OperationTypeToSearch != mu::EOpType::NONE)
@@ -1336,26 +1343,38 @@ void SMutableCodeViewer::GenerateAllTreeElements()
 
 	// Generate all root nodes
 	const mu::FModel::Private* ModelPrivate = MutableModel->GetPrivate();
+	check(ModelPrivate);
 	const mu::FProgram& Program = ModelPrivate->Program;
-	const uint32 StateCount = ModelPrivate->Program.States.Num();
-	for ( uint32 StateIndex=0; StateIndex<StateCount; ++StateIndex )
+	const uint32 StateCount = Program.States.Num();
+	for ( uint32 StateIndex = 0; StateIndex < StateCount; ++StateIndex )
 	{
-		const mu::FProgram::FState& State = ModelPrivate->Program.States[StateIndex];
-		FString Caption = FString::Printf( TEXT("state [%s]"), *State.Name );
+		const mu::FProgram::FState& State = Program.States[StateIndex];
+		const FString Caption = FString::Printf( TEXT("state [%s]"), *State.Name );
 
 		const FSlateColor LabelColor = ColorPerComputationalCost[StaticCast<uint8>(GetOperationTypeComputationalCost(
-			ModelPrivate->Program.GetOpType(State.Root)))];
-
-		RootNodes.Add(MakeShareable(new FMutableCodeTreeElement(ItemCache.Num(),StateIndex, MutableModel, State.Root, Caption,LabelColor)));
+			Program.GetOpType(State.Root)))];
 
-		// Iterate over each root node and generate all the elements in a human readable pattern (Z Pattern)
+		// Locate the "original" tree element :
+		// This may happen if for some reason the state is duplicated (should never happen)
+		const TSharedPtr<FMutableCodeTreeElement>* MainItemPtr = MainItemPerOp.Find(State.Root);
 
-		constexpr  mu::OP::ADDRESS CommonParent = 0;
+		// Create a new root element and add it to the collection of root nodes
+		TSharedPtr<FMutableCodeTreeElement> RootNodeElement = MakeShareable(new FMutableCodeTreeElement(ItemCache.Num(),StateIndex, MutableModel, State.Root, Caption,LabelColor, MainItemPtr));
+		RootNodes.Add(RootNodeElement);
+		
+		// Add the element to the cache so we keep the indices straight.
+		constexpr mu::OP::ADDRESS CommonParent = 0;
 		const FItemCacheKey Key = { CommonParent, State.Root, StateIndex };
+		ItemCache.Add(Key, RootNodeElement);
+		
+		if (!MainItemPtr)
+		{
+			// Cache this node as it may be duplicated of another state. Check the "MainItemPtr" initialization for more info
+			MainItemPerOp.Add(State.Root, RootNodeElement);
 
-		// Add the element to the cache so we keep the indices straight.
-		ItemCache.Add(Key, RootNodes.Last());
-		GenerateElementRecursive(StateIndex,State.Root,Program);
+			// Iterate over each root node and generate all the elements in a human-readable pattern (Z Pattern)
+			GenerateElementRecursive(StateIndex,State.Root,Program);
+		}
 	}
 	
 }
@@ -1369,9 +1388,9 @@ void SMutableCodeViewer::GenerateElementRecursive(const int32& InStateIndex, mu:
 			{
 				const FItemCacheKey Key = { InParentAddress, ChildAddress, ChildIndex };
 				const TSharedPtr<FMutableCodeTreeElement>* CachedItem = ItemCache.Find(Key);
-				check(!CachedItem);
 
-				if (!CachedItem)
+				// If not already cached then process it
+				if (ensure(!CachedItem))
 				{
 					// Locate the "original" tree element
 					const TSharedPtr<FMutableCodeTreeElement>* MainItemPtr = MainItemPerOp.Find(ChildAddress);
@@ -1392,6 +1411,12 @@ void SMutableCodeViewer::GenerateElementRecursive(const int32& InStateIndex, mu:
 						GenerateElementRecursive(InStateIndex,ChildAddress, InProgram);
 					}
 				}
+				else
+				{
+					UE_LOG(LogMutable, Error, TEXT("An already processed operation is being re-processed in order to generate a tree row. "
+									"ParentAddress : %u , ChildAddress : %u , ChildIndex : %u "), InParentAddress, ChildAddress, ChildIndex);
+				}
+				
 			}
 			++ChildIndex;
 		};
@@ -2412,14 +2437,16 @@ void SMutableCodeViewer::CacheRootNodeAddresses()
 		// Dump the located resources array onto the navigation array since we have content to navigate over
 		NavigationElements = MoveTemp(SearchPayload.FoundElements);
 		SortElementsByTreeIndex(NavigationElements);
-		
-		// Reset the navigation index
-		NavigationIndex = -1;
 	}
 	else
 	{
+		TargetedTypeSelector->SetSelectedItem(NoneOperationEntry);
+		NavigationElements.Reset();
 		UE_LOG(LogTemp,Error,TEXT("The provided constant index does not seem to be used anywhere : Make sure the index is valid and that IsConstantResourceUsedByOperation() switch is up to date"));
 	}
+
+	// Reset the navigation index
+	NavigationIndex = -1;
 }
 
 
@@ -2517,6 +2544,32 @@ bool SMutableCodeViewer::IsConstantResourceUsedByOperation(const int32 IndexOnCo
 				{
 					bResourceLocated = IndexOnConstantsArray == InProgram.GetOpArgs<mu::OP::InstanceAddArgs>(OperationAddress).name;
 				}
+				else if (OperationType == mu::EOpType::IN_CONDITIONAL)
+				{
+					mu::OP::ConditionalArgs Arguments = InProgram.GetOpArgs<mu::OP::ConditionalArgs>(OperationAddress);
+					
+					if (IndexOnConstantsArray == Arguments.condition)
+					{
+						bResourceLocated = true;
+						break;
+					}
+
+					if (IndexOnConstantsArray == Arguments.yes)
+					{
+						bResourceLocated = true;
+						break;
+					}
+
+					if (IndexOnConstantsArray == Arguments.no)
+					{
+						bResourceLocated = true;
+						break;
+					}
+				}
+				else if (OperationType == mu::EOpType::IN_ADDEXTENSIONDATA)
+				{
+					bResourceLocated = IndexOnConstantsArray == InProgram.GetOpArgs<mu::OP::InstanceAddExtensionDataArgs>(OperationAddress).ExtensionDataName;
+				}
 				else if (OperationType == mu::EOpType::ME_BINDSHAPE)
 				{
 					mu::OP::MeshBindShapeArgs Arguments = InProgram.GetOpArgs<mu::OP::MeshBindShapeArgs>(OperationAddress);
@@ -2536,35 +2589,37 @@ bool SMutableCodeViewer::IsConstantResourceUsedByOperation(const int32 IndexOnCo
 						if (*Data == IndexOnConstantsArray)
 						{
 							bResourceLocated = true;
-							break;
+							continue;
 						}
 						
 						Data += sizeof(int32);
 					}
 				}
-				else if (OperationType == mu::EOpType::IN_CONDITIONAL)
+				else if (OperationType == mu::EOpType::ME_ADDTAGS)
 				{
-					mu::OP::ConditionalArgs Arguments = InProgram.GetOpArgs<mu::OP::ConditionalArgs>(OperationAddress);
+					 const uint8* OpData = InProgram.GetOpArgsPointer(OperationAddress);
 					
-					if (IndexOnConstantsArray == Arguments.condition)
-					{
-						bResourceLocated = true;
-						break;
-					}
-
-					if (IndexOnConstantsArray == Arguments.yes)
-					{
-						bResourceLocated = true;
-						break;
-					}
-
-					if (IndexOnConstantsArray == Arguments.no)
+					 mu::OP::ADDRESS SourceAddress;
+					 FMemory::Memcpy(&SourceAddress, OpData, sizeof(mu::OP::ADDRESS));
+					 OpData += sizeof(mu::OP::ADDRESS);
+					
+					 uint16 TagCount;
+					 FMemory::Memcpy(&TagCount, OpData, sizeof(uint16));
+					 OpData += sizeof(uint16);
+					
+					for (int32 TagIndex = 0; TagIndex < TagCount; ++TagIndex)
 					{
-						bResourceLocated = true;
-						break;
+						// Exit once we know that the data is pointing to the index provided
+						if (*OpData == IndexOnConstantsArray)
+						{
+							bResourceLocated = true;
+							continue;
+						}
+						
+						OpData += sizeof(uint16);
 					}
 				}
-				
+
 				break;
 			}
 
@@ -2634,6 +2689,10 @@ bool SMutableCodeViewer::IsConstantResourceUsedByOperation(const int32 IndexOnCo
 				{
 					bResourceLocated = IndexOnConstantsArray == InProgram.GetOpArgs<mu::OP::MeshTransformWithinMeshArgs>(OperationAddress).matrix;
 				}
+				else if (OperationType == mu::EOpType::MA_CONSTANT)
+				{
+					bResourceLocated = IndexOnConstantsArray == InProgram.GetOpArgs<mu::OP::MatrixConstantArgs>(OperationAddress).value;
+				}
 				break;
 			}
 
@@ -2674,6 +2733,10 @@ bool SMutableCodeViewer::IsConstantResourceUsedByOperation(const int32 IndexOnCo
 				{
 					bResourceLocated = IndexOnConstantsArray == InProgram.GetOpArgs<mu::OP::MeshConstantArgs>(OperationAddress).Skeleton;
 				}
+				else if (OperationType == mu::EOpType::ME_SETSKELETON)
+				{
+					bResourceLocated = IndexOnConstantsArray == InProgram.GetOpArgs<mu::OP::MeshSetSkeletonArgs>(OperationAddress).Skeleton;
+				}
 				break;
 			}
 		case mu::EDataType::PhysicsAsset:
diff --git a/Engine/Plugins/Mutable/Source/CustomizableObjectEditor/Private/MuCOE/SMutableConstantsWidget.cpp b/Engine/Plugins/Mutable/Source/CustomizableObjectEditor/Private/MuCOE/SMutableConstantsWidget.cpp
--- a/Engine/Plugins/Mutable/Source/CustomizableObjectEditor/Private/MuCOE/SMutableConstantsWidget.cpp
+++ b/Engine/Plugins/Mutable/Source/CustomizableObjectEditor/Private/MuCOE/SMutableConstantsWidget.cpp
@@ -496,7 +496,7 @@ void SMutableConstantsWidget::Construct(const FArguments& InArgs,const mu::FProg
 	
 		+ SHeaderRow::Column(ImageConstantTitles::ImageSize)
 		.DefaultLabel(FText(LOCTEXT("ImageResolution","Resolution")))
-		.DefaultTooltip(FText(LOCTEXT("ImageResolutionolumnToolTip","Pixel resolution")))
+		.DefaultTooltip(FText(LOCTEXT("ImageResolutionColumnToolTip","Pixel resolution")))
 		.OnSort(this, &SMutableConstantsWidget::OnImageTableSortRequested)
 		.SortMode(this, &SMutableConstantsWidget::GetImageListColumnSortMode, ImageConstantTitles::ImageSize )
 		
