[mutable] Fix bug with incorrect reusal of meshes with different layouts.
- Meshes are now cached before layouts are applied to them.
- Layouts are no longer compared and used in the cache check.
- the operation that prepares meshes for layouts is no longer reduced in constant optimization.
- additional mutable code optimizations around the "prepare layout" and "make layout" operation.
- an optimization pass to the "prepare layout" function.
#rnx

diff --git a/Engine/Plugins/Mutable/Source/CustomizableObject/Internal/MuCO/CustomizableObjectPrivate.h b/Engine/Plugins/Mutable/Source/CustomizableObject/Internal/MuCO/CustomizableObjectPrivate.h
--- a/Engine/Plugins/Mutable/Source/CustomizableObject/Internal/MuCO/CustomizableObjectPrivate.h
+++ b/Engine/Plugins/Mutable/Source/CustomizableObject/Internal/MuCO/CustomizableObjectPrivate.h
@@ -1606,6 +1606,8 @@ enum class ECustomizableObjectVersions : int32
 
 	AddConnectedChildObjectComponentsToPrepass,
 
+	FixMeshReusalDueToLayouts,
+
 	// -----<new versions can be added above this line>--------
 	LastCustomizableObjectVersion
 };
@@ -1883,6 +1885,6 @@ public:
 	static constexpr int32 CurrentSupportedVersion = static_cast<int32>(ECustomizableObjectVersions::LastCustomizableObjectVersion);
 #endif
 	
-	static constexpr int32 DerivedDataVersion = 0x49799a1c;
+	static constexpr int32 DerivedDataVersion = 0x46799a1c;
 };
 
diff --git a/Engine/Plugins/Mutable/Source/MutableRuntime/Private/MuR/Mesh.cpp b/Engine/Plugins/Mutable/Source/MutableRuntime/Private/MuR/Mesh.cpp
--- a/Engine/Plugins/Mutable/Source/MutableRuntime/Private/MuR/Mesh.cpp
+++ b/Engine/Plugins/Mutable/Source/MutableRuntime/Private/MuR/Mesh.cpp
@@ -974,7 +974,7 @@ void FMesh::Unserialise(FInputArchive& Arch)
 }
 
 
-bool FMesh::IsSimilar(const FMesh& Other, bool bCompareLayouts) const
+bool FMesh::IsSimilar(const FMesh& Other) const
 {
 	// Some meshes are just vertex indices (masks) we don't consider them for similarity,
 	// because the kind of vertex channel data they store is the kind that is ignored.
@@ -985,11 +985,6 @@ bool FMesh::IsSimilar(const FMesh& Other, bool bCompareLayouts) const
 
 	bool bEqual = IndexBuffers == Other.IndexBuffers;
 	bEqual = bEqual && (ReferenceID == Other.ReferenceID);
-	
-	if (bEqual && bCompareLayouts) 
-	{
-		bEqual = (Layouts.Num() == Other.Layouts.Num());
-	}
 
 	if (bEqual && Skeleton != Other.Skeleton)
 	{
@@ -1018,24 +1013,13 @@ bool FMesh::IsSimilar(const FMesh& Other, bool bCompareLayouts) const
 	bEqual = bEqual && (Surfaces == Other.Surfaces);
 	bEqual = bEqual && (Tags == Other.Tags);
 
-	// Special comparison for layouts
-	if (bCompareLayouts)
-	{
-		const int32 NumLayouts = Layouts.Num();
-		for (int32 LayoutIndex = 0; bEqual && LayoutIndex < NumLayouts; ++LayoutIndex)
-		{
-			bEqual &= Layouts[LayoutIndex]->IsSimilar(*Other.Layouts[LayoutIndex]);
-		}
-	}
-
 	// Special comparison for vertex buffers
 	if (bEqual)
 	{
-		bEqual = VertexBuffers.IsSimilarRobust(Other.VertexBuffers, bCompareLayouts);
+		bEqual = VertexBuffers.IsSimilarRobust(Other.VertexBuffers, false);
 	}
 
 	return bEqual;
-
 }
 
 
@@ -1258,21 +1242,28 @@ static bool StaticMeshFormatIdentify_Project(const FMesh* InMesh)
     }
 
     // The first index buffer must be just index buffers u32
-    if (bResult)
-    {
-        bResult &= InMesh->IndexBuffers.Buffers[0].Channels.Num() >= 1;
-    }
+	if (!InMesh->IndexBuffers.Buffers.IsEmpty())
+	{
+		if (bResult)
+		{
+			bResult &= InMesh->IndexBuffers.Buffers[0].Channels.Num() >= 1;
+		}
 
-    if (bResult)
-    {
-        const FMeshBufferChannel& Channel = InMesh->IndexBuffers.Buffers[0].Channels[0];
+		if (bResult)
+		{
+			const FMeshBufferChannel& Channel = InMesh->IndexBuffers.Buffers[0].Channels[0];
 
-        bResult &= Channel.Semantic == EMeshBufferSemantic::VertexIndex;
-        bResult &= Channel.Format == EMeshBufferFormat::UInt32;
-        bResult &= Channel.ComponentCount == 1;
-        bResult &= Channel.SemanticIndex == 0;
-        bResult &= Channel.Offset == 0;
-    }
+			bResult &= Channel.Semantic == EMeshBufferSemantic::VertexIndex;
+			bResult &= Channel.Format == EMeshBufferFormat::UInt32;
+			bResult &= Channel.ComponentCount == 1;
+			bResult &= Channel.SemanticIndex == 0;
+			bResult &= Channel.Offset == 0;
+		}
+	}
+	else
+	{
+		bResult = false;
+	}
 
     return bResult;
 }
@@ -1389,6 +1380,11 @@ namespace
             Out += FString::Printf(TEXT("    Buffer with %d channels and %d elementsize\n"), Buffer.Channels.Num(), Buffer.ElementSize);
             
 			const uint8* DataPtr = Buffer.Data.GetData();
+			if (!DataPtr)
+			{
+				continue;
+			}
+
             for (const FMeshBufferChannel& Channel : Buffer.Channels)
             {
 				Out += FString::Printf(TEXT("      Channel with format: %d semantic: %d %d, components: %d, offset: %d\n"),
diff --git a/Engine/Plugins/Mutable/Source/MutableRuntime/Private/MuR/OpMeshPrepareLayout.cpp b/Engine/Plugins/Mutable/Source/MutableRuntime/Private/MuR/OpMeshPrepareLayout.cpp
--- a/Engine/Plugins/Mutable/Source/MutableRuntime/Private/MuR/OpMeshPrepareLayout.cpp
+++ b/Engine/Plugins/Mutable/Source/MutableRuntime/Private/MuR/OpMeshPrepareLayout.cpp
@@ -91,43 +91,46 @@ namespace mu
 
 		// Create an array of block index per cell
 		TArray<int32> OverlappingBlocks;
-		for (int32 BlockIndex = 0; BlockIndex < NumBlocks; ++BlockIndex)
 		{
-			// TODO
-			bool bBlockHasMask = false; // GeneratedLayout.Source->Blocks[BlockIndex].Mask != nullptr;
-
-			// Fill the block rect
-			FIntVector2 Min = Layout->Blocks[BlockIndex].Min;
-			FIntVector2 Size = Layout->Blocks[BlockIndex].Size;
-
-			box<FVector2f>& BlockRect = BlockRects[BlockIndex];
-			BlockRect.min[0] = float(Min.X) / float(LayoutGrid.X);
-			BlockRect.min[1] = float(Min.Y) / float(LayoutGrid.Y);
-			BlockRect.size[0] = float(Size.X) / float(LayoutGrid.X);
-			BlockRect.size[1] = float(Size.Y) / float(LayoutGrid.Y);
-
-			// Fill the block index per cell array
-			// Ignore the block in this stage if it has a mask, because blocks with masks will very likely overlap other blocks
-			if (!bBlockHasMask)
+			MUTABLE_CPUPROFILER_SCOPE(CreateGrid);
+
+			for (int32 BlockIndex = 0; BlockIndex < NumBlocks; ++BlockIndex)
 			{
-				for (uint16 Y = Min.Y; Y < Min.Y + Size.Y; ++Y)
+				// TODO
+				bool bBlockHasMask = false; // GeneratedLayout.Source->Blocks[BlockIndex].Mask != nullptr;
+
+				// Fill the block rect
+				FIntVector2 Min = Layout->Blocks[BlockIndex].Min;
+				FIntVector2 Size = Layout->Blocks[BlockIndex].Size;
+
+				box<FVector2f>& BlockRect = BlockRects[BlockIndex];
+				BlockRect.min[0] = float(Min.X) / float(LayoutGrid.X);
+				BlockRect.min[1] = float(Min.Y) / float(LayoutGrid.Y);
+				BlockRect.size[0] = float(Size.X) / float(LayoutGrid.X);
+				BlockRect.size[1] = float(Size.Y) / float(LayoutGrid.Y);
+
+				// Fill the block index per cell array
+				// Ignore the block in this stage if it has a mask, because blocks with masks will very likely overlap other blocks
+				if (!bBlockHasMask)
 				{
-					for (uint16 X = Min.X; X < Min.X + Size.X; ++X)
+					for (uint16 Y = Min.Y; Y < Min.Y + Size.Y; ++Y)
 					{
-						if (GridBlockBlockId.Get(X, Y) == MAX_uint16)
-						{
-							GridBlockBlockId.Set(X, Y, BlockIndex);
-						}
-						else
+						for (uint16 X = Min.X; X < Min.X + Size.X; ++X)
 						{
-							OverlappingBlocks.AddUnique(BlockIndex);
+							if (GridBlockBlockId.Get(X, Y) == MAX_uint16)
+							{
+								GridBlockBlockId.Set(X, Y, BlockIndex);
+							}
+							else
+							{
+								OverlappingBlocks.AddUnique(BlockIndex);
+							}
 						}
 					}
 				}
 			}
 		}
 
-
 		// Notify Overlapping layout blocks
 		if (!OverlappingBlocks.IsEmpty())
 		{
@@ -154,13 +157,15 @@ namespace mu
 		check(TexCoordsChannel.Semantic == EMeshBufferSemantic::TexCoords);
 
 		uint8* TexCoordData = Mesh.GetVertexBuffers().GetBufferData(TexCoordsBufferIndex);
-		int32 elemSize = Mesh.GetVertexBuffers().GetElementSize(TexCoordsBufferIndex);
+		int32 TexElemSize = Mesh.GetVertexBuffers().GetElementSize(TexCoordsBufferIndex);
 		int32 channelOffset = TexCoordsChannel.Offset;
 		TexCoordData += channelOffset;
 
 		// Get a copy of the UVs as FVector2f to work with them. 
 		TArray<FVector2f> TexCoords;
 		{
+			MUTABLE_CPUPROFILER_SCOPE(CopyUVs);
+
 			TexCoords.SetNumUninitialized(NumVertices);
 
 			bool bNonNormalizedUVs = false;
@@ -189,7 +194,7 @@ namespace mu
 					bNonNormalizedUVs = true;
 				}
 
-				pVertices += elemSize;
+				pVertices += TexElemSize;
 			}
 
 			// Mutable does not support non-normalized UVs
@@ -200,210 +205,217 @@ namespace mu
 		}
 
 
-		const int32 NumTriangles = Mesh.GetIndexCount() / 3;
-		TArray<FTriangleInfo> Triangles;
-
-		// Vertices mapped to unique vertex index
-		TArray<int32> CollapsedVertices;
-
-		// Vertex to face map used to speed up connectivity building
-		TMultiMap<int32, uint32> VertexToFaceMap;
-
-		// Find Unique Vertices
-		if (bClampUVIslands)
-		{
-			VertexToFaceMap.Reserve(NumVertices);
-			Triangles.SetNumUninitialized(NumTriangles);
-
-			MeshCreateCollapsedVertexMap(&Mesh, CollapsedVertices);
-		}
-
-		TArray<int32> ConflictiveTriangles;
-
 		const uint32 MaxGridX = bNormalizeUVs ? MAX_uint32 : WorkingGrid.X - 1;
 		const uint32 MaxGridY = bNormalizeUVs ? MAX_uint32 : WorkingGrid.Y - 1;
 
 		// Allocate the per-vertex layout block data
 		TArray<uint16> LayoutData;
 		constexpr uint16 NullBlockId = MAX_uint16 - 1;
-		LayoutData.Init(NullBlockId, NumVertices);
+		LayoutData.SetNumUninitialized(NumVertices);
 
-		UntypedMeshBufferIteratorConst ItIndices(Mesh.GetIndexBuffers(), EMeshBufferSemantic::VertexIndex);
-		for (int32 TriangleIndex = 0; TriangleIndex < NumTriangles; ++TriangleIndex)
+		// Assign a block to each vertex
 		{
-			uint32 Index0 = ItIndices.GetAsUINT32();
-			++ItIndices;
-			uint32 Index1 = ItIndices.GetAsUINT32();
-			++ItIndices;
-			uint32 Index2 = ItIndices.GetAsUINT32();
-			++ItIndices;
-
-			auto AssignOneVertex =
-				[LayoutGrid, MaxGridX, MaxGridY, NumBlocks, Layout, &GridBlockBlockId, &TexCoords, &LayoutData]
-				(int32 VertexIndex)
-				{
-					uint16& BlockIndex = LayoutData[VertexIndex];
+			MUTABLE_CPUPROFILER_SCOPE(Assign);
 
-					// Was it previously assigned?
-					if (BlockIndex != NullBlockId)
-					{
-						return BlockIndex;
-					}
+			for (int32 VertexIndex = 0; VertexIndex < NumVertices; ++VertexIndex)
+			{
+				uint16 BlockIndex = NullBlockId;
 
-					FVector2f UV = TexCoords[VertexIndex];
+				FVector2f UV = TexCoords[VertexIndex];
 
-					int32 VertexWorkingGridX = FMath::Clamp(LayoutGrid.X * UV[0], 0, LayoutGrid.X - 1);
-					int32 VertexWorkingGridY = FMath::Clamp(LayoutGrid.Y * UV[1], 0, LayoutGrid.Y - 1);
+				int32 VertexWorkingGridX = FMath::Clamp(LayoutGrid.X * UV[0], 0, LayoutGrid.X - 1);
+				int32 VertexWorkingGridY = FMath::Clamp(LayoutGrid.Y * UV[1], 0, LayoutGrid.Y - 1);
 
-					// First: Assign the vertices to masked blocks in order
-					for (int32 CandidateBlockIndex = 0; CandidateBlockIndex < NumBlocks; ++CandidateBlockIndex)
+				// First: Assign the vertices to masked blocks in order
+				for (int32 CandidateBlockIndex = 0; CandidateBlockIndex < NumBlocks; ++CandidateBlockIndex)
+				{
+					// TODO
+					const mu::FImage* Mask = nullptr; // GeneratedLayout.Source->Blocks[CandidateBlockIndex].Mask.Get();
+					if (Mask)
 					{
-						// TODO
-						const mu::FImage* Mask = nullptr; // GeneratedLayout.Source->Blocks[CandidateBlockIndex].Mask.Get();
-						if (Mask)
-						{
-							// First discard with block limits.
-							FIntVector2 Min = Layout->Blocks[CandidateBlockIndex].Min;
-							FIntVector2 Size = Layout->Blocks[CandidateBlockIndex].Size;
+						// First discard with block limits.
+						FIntVector2 Min = Layout->Blocks[CandidateBlockIndex].Min;
+						FIntVector2 Size = Layout->Blocks[CandidateBlockIndex].Size;
 
-							bool bInBlock =
-								(VertexWorkingGridX >= Min.X && VertexWorkingGridX < Min.X + Size.X)
-								&&
-								(VertexWorkingGridY >= Min.Y && VertexWorkingGridY < Min.Y + Size.Y);
+						bool bInBlock =
+							(VertexWorkingGridX >= Min.X && VertexWorkingGridX < Min.X + Size.X)
+							&&
+							(VertexWorkingGridY >= Min.Y && VertexWorkingGridY < Min.Y + Size.Y);
+
+						if (bInBlock)
+						{
+							// TODO: This always clamps the UVs
+							FVector2f SampleUV;
+							SampleUV.X = FMath::Fmod(UV.X, 1.0);
+							SampleUV.Y = FMath::Fmod(UV.Y, 1.0);
 
-							if (bInBlock)
+							FVector4f MaskValue = Mask->Sample(SampleUV);
+							if (MaskValue.X > 0.5f)
 							{
-								// TODO: This always clamps the UVs
-								FVector2f SampleUV;
-								SampleUV.X = FMath::Fmod(UV.X, 1.0);
-								SampleUV.Y = FMath::Fmod(UV.Y, 1.0);
-
-								FVector4f MaskValue = Mask->Sample(SampleUV);
-								if (MaskValue.X > 0.5f)
-								{
-									BlockIndex = CandidateBlockIndex;
-									break;
-								}
+								BlockIndex = CandidateBlockIndex;
+								break;
 							}
 						}
 					}
+				}
 
-					// Second: Assign to non-masked blocks if not assigned yet
-					if (BlockIndex == NullBlockId)
-					{
-						uint32 ClampedX = FMath::Min<uint32>(MaxGridX, FMath::Max<uint32>(0, VertexWorkingGridX));
-						uint32 ClampedY = FMath::Min<uint32>(MaxGridY, FMath::Max<uint32>(0, VertexWorkingGridY));
-						BlockIndex = GridBlockBlockId.Get(ClampedX, ClampedY);
-					}
-					return BlockIndex;
-				};
-
-			const uint16 BlockIndexV0 = AssignOneVertex(Index0);
-			const uint16 BlockIndexV1 = AssignOneVertex(Index1);
-			const uint16 BlockIndexV2 = AssignOneVertex(Index2);
-
-			if (bClampUVIslands)
-			{
-				if (BlockIndexV0 != BlockIndexV1 || BlockIndexV0 != BlockIndexV2)
+				// Second: Assign to non-masked blocks if not assigned yet
+				if (BlockIndex == NullBlockId)
 				{
-					ConflictiveTriangles.Add(TriangleIndex);
+					uint32 ClampedX = FMath::Min<uint32>(MaxGridX, FMath::Max<uint32>(0, VertexWorkingGridX));
+					uint32 ClampedY = FMath::Min<uint32>(MaxGridY, FMath::Max<uint32>(0, VertexWorkingGridY));
+					BlockIndex = GridBlockBlockId.Get(ClampedX, ClampedY);
 				}
 
-				FTriangleInfo& Triangle = Triangles[TriangleIndex];
-
-				Triangle.Indices[0] = Index0;
-				Triangle.Indices[1] = Index1;
-				Triangle.Indices[2] = Index2;
-				Triangle.CollapsedIndices[0] = CollapsedVertices[Index0];
-				Triangle.CollapsedIndices[1] = CollapsedVertices[Index1];
-				Triangle.CollapsedIndices[2] = CollapsedVertices[Index2];
-
-				Triangle.BlockIndices[0] = BlockIndexV0;
-				Triangle.BlockIndices[1] = BlockIndexV1;
-				Triangle.BlockIndices[2] = BlockIndexV2;
-				Triangle.bUVsFixed = false;
-
-				VertexToFaceMap.Add(Triangle.CollapsedIndices[0], TriangleIndex);
-				VertexToFaceMap.Add(Triangle.CollapsedIndices[1], TriangleIndex);
-				VertexToFaceMap.Add(Triangle.CollapsedIndices[2], TriangleIndex);
+				LayoutData[VertexIndex] = BlockIndex;
 			}
 		}
 
-		// Clamp UV islands to the predominant block of each island. Will only happen if bClampUVIslands is true.
-		for (int32 ConflictiveTriangleIndex : ConflictiveTriangles)
+
+		if (bClampUVIslands)
 		{
-			FTriangleInfo& Triangle = Triangles[ConflictiveTriangleIndex];
+			const int32 NumTriangles = Mesh.GetIndexCount() / 3;
+			TArray<FTriangleInfo> Triangles;
+
+			// Vertices mapped to unique vertex index
+			TArray<int32> CollapsedVertices;
 
-			// Skip the ones that have been fixed already
-			if (Triangle.bUVsFixed)
+			// Vertex to face map used to speed up connectivity building
+			TMultiMap<int32, uint32> VertexToFaceMap;
+
+			// Find Unique Vertices
+			if (bClampUVIslands)
 			{
-				continue;
-			}
+				VertexToFaceMap.Reserve(NumTriangles * 3);
+				Triangles.SetNumUninitialized(NumTriangles);
 
-			// Find triangles from the same UV Island
-			TArray<uint32> TriangleIndices;
-			GetUVIsland(Triangles, ConflictiveTriangleIndex, TriangleIndices, TexCoords, VertexToFaceMap);
+				MeshCreateCollapsedVertexMap(&Mesh, CollapsedVertices);
+			}
 
-			// Get predominant BlockId != MAX_uint16
-			TArray<uint32> NumVerticesPerBlock;
-			NumVerticesPerBlock.SetNumZeroed(NumBlocks);
+			TArray<int32> ConflictiveTriangles;
 
-			for (int32 TriangleIndex : TriangleIndices)
+			UntypedMeshBufferIteratorConst ItIndices(Mesh.GetIndexBuffers(), EMeshBufferSemantic::VertexIndex);
 			{
-				FTriangleInfo& OtherTriangle = Triangles[TriangleIndex];
-				for (int32 VertexIndex = 0; VertexIndex < 3; ++VertexIndex)
+				MUTABLE_CPUPROFILER_SCOPE(CreateTriangles);
+
+				for (int32 TriangleIndex = 0; TriangleIndex < NumTriangles; ++TriangleIndex)
 				{
-					const uint16& BlockIndex = OtherTriangle.BlockIndices[VertexIndex];
-					if (BlockIndex != MAX_uint16)
+					uint32 Index0 = ItIndices.GetAsUINT32();
+					++ItIndices;
+					uint32 Index1 = ItIndices.GetAsUINT32();
+					++ItIndices;
+					uint32 Index2 = ItIndices.GetAsUINT32();
+					++ItIndices;
+
+					const uint16 BlockIndexV0 = LayoutData[Index0];
+					const uint16 BlockIndexV1 = LayoutData[Index1];
+					const uint16 BlockIndexV2 = LayoutData[Index2];
+
+					if (BlockIndexV0 != BlockIndexV1 || BlockIndexV0 != BlockIndexV2)
 					{
-						NumVerticesPerBlock[BlockIndex]++;
+						ConflictiveTriangles.Add(TriangleIndex);
 					}
+
+					FTriangleInfo& Triangle = Triangles[TriangleIndex];
+
+					Triangle.Indices[0] = Index0;
+					Triangle.Indices[1] = Index1;
+					Triangle.Indices[2] = Index2;
+					Triangle.CollapsedIndices[0] = CollapsedVertices[Index0];
+					Triangle.CollapsedIndices[1] = CollapsedVertices[Index1];
+					Triangle.CollapsedIndices[2] = CollapsedVertices[Index2];
+
+					Triangle.BlockIndices[0] = BlockIndexV0;
+					Triangle.BlockIndices[1] = BlockIndexV1;
+					Triangle.BlockIndices[2] = BlockIndexV2;
+					Triangle.bUVsFixed = false;
+
+					VertexToFaceMap.Add(Triangle.CollapsedIndices[0], TriangleIndex);
+					VertexToFaceMap.Add(Triangle.CollapsedIndices[1], TriangleIndex);
+					VertexToFaceMap.Add(Triangle.CollapsedIndices[2], TriangleIndex);
 				}
 			}
 
-			uint16 BlockIndex = 0;
-			uint32 CurrentMaxVertices = 0;
-			for (int32 Index = 0; Index < NumBlocks; ++Index)
+			// Clamp UV islands to the predominant block of each island. Will only happen if bClampUVIslands is true.
 			{
-				if (NumVerticesPerBlock[Index] > CurrentMaxVertices)
+				MUTABLE_CPUPROFILER_SCOPE(ClampUVs);
+
+				for (int32 ConflictiveTriangleIndex : ConflictiveTriangles)
 				{
-					BlockIndex = Index;
-					CurrentMaxVertices = NumVerticesPerBlock[Index];
-				}
-			}
+					FTriangleInfo& Triangle = Triangles[ConflictiveTriangleIndex];
 
-			// Get the limits of the predominant block rect
-			const FLayoutBlock& LayoutBlock = Layout->Blocks[BlockIndex];
+					// Skip the ones that have been fixed already
+					if (Triangle.bUVsFixed)
+					{
+						continue;
+					}
 
-			const float SmallNumber = 0.000001;
-			const float MinX = ((float)LayoutBlock.Min.X) / (float)LayoutGrid.X + SmallNumber;
-			const float MinY = ((float)LayoutBlock.Min.Y) / (float)LayoutGrid.Y + SmallNumber;
-			const float MaxX = (((float)LayoutBlock.Size.X + LayoutBlock.Min.X) / (float)LayoutGrid.X) - 2 * SmallNumber;
-			const float MaxY = (((float)LayoutBlock.Size.Y + LayoutBlock.Min.Y) / (float)LayoutGrid.Y) - 2 * SmallNumber;
+					// Find triangles from the same UV Island
+					TArray<uint32> TriangleIndices;
+					GetUVIsland(Triangles, ConflictiveTriangleIndex, TriangleIndices, TexCoords, VertexToFaceMap);
 
-			// Iterate triangles and clamp the UVs
-			for (int32 TriangleIndex : TriangleIndices)
-			{
-				FTriangleInfo& OtherTriangle = Triangles[TriangleIndex];
+					// Get predominant BlockId != MAX_uint16
+					TArray<uint32> NumVerticesPerBlock;
+					NumVerticesPerBlock.SetNumZeroed(NumBlocks);
 
-				for (int8 VertexIndex = 0; VertexIndex < 3; ++VertexIndex)
-				{
-					if (OtherTriangle.BlockIndices[VertexIndex] == BlockIndex)
+					for (int32 TriangleIndex : TriangleIndices)
 					{
-						continue;
+						FTriangleInfo& OtherTriangle = Triangles[TriangleIndex];
+						for (int32 VertexIndex = 0; VertexIndex < 3; ++VertexIndex)
+						{
+							const uint16& BlockIndex = OtherTriangle.BlockIndices[VertexIndex];
+							if (BlockIndex != MAX_uint16)
+							{
+								NumVerticesPerBlock[BlockIndex]++;
+							}
+						}
+					}
+
+					uint16 BlockIndex = 0;
+					uint32 CurrentMaxVertices = 0;
+					for (int32 Index = 0; Index < NumBlocks; ++Index)
+					{
+						if (NumVerticesPerBlock[Index] > CurrentMaxVertices)
+						{
+							BlockIndex = Index;
+							CurrentMaxVertices = NumVerticesPerBlock[Index];
+						}
 					}
 
-					OtherTriangle.BlockIndices[VertexIndex] = BlockIndex;
+					// Get the limits of the predominant block rect
+					const FLayoutBlock& LayoutBlock = Layout->Blocks[BlockIndex];
 
-					// Clamp UVs to the block they are assigned to
-					const int32 UVIndex = OtherTriangle.Indices[VertexIndex];
-					FVector2f& UV = TexCoords[UVIndex];
-					UV[0] = FMath::Clamp(UV[0], MinX, MaxX);
-					UV[1] = FMath::Clamp(UV[1], MinY, MaxY);
-					LayoutData[UVIndex] = BlockIndex;
-				}
+					const float SmallNumber = 0.000001;
+					const float MinX = ((float)LayoutBlock.Min.X) / (float)LayoutGrid.X + SmallNumber;
+					const float MinY = ((float)LayoutBlock.Min.Y) / (float)LayoutGrid.Y + SmallNumber;
+					const float MaxX = (((float)LayoutBlock.Size.X + LayoutBlock.Min.X) / (float)LayoutGrid.X) - 2 * SmallNumber;
+					const float MaxY = (((float)LayoutBlock.Size.Y + LayoutBlock.Min.Y) / (float)LayoutGrid.Y) - 2 * SmallNumber;
+
+					// Iterate triangles and clamp the UVs
+					for (int32 TriangleIndex : TriangleIndices)
+					{
+						FTriangleInfo& OtherTriangle = Triangles[TriangleIndex];
+
+						for (int8 VertexIndex = 0; VertexIndex < 3; ++VertexIndex)
+						{
+							if (OtherTriangle.BlockIndices[VertexIndex] == BlockIndex)
+							{
+								continue;
+							}
+
+							OtherTriangle.BlockIndices[VertexIndex] = BlockIndex;
+
+							// Clamp UVs to the block they are assigned to
+							const int32 UVIndex = OtherTriangle.Indices[VertexIndex];
+							FVector2f& UV = TexCoords[UVIndex];
+							UV[0] = FMath::Clamp(UV[0], MinX, MaxX);
+							UV[1] = FMath::Clamp(UV[1], MinY, MaxY);
+							LayoutData[UVIndex] = BlockIndex;
+						}
 
-				OtherTriangle.bUVsFixed = true;
+						OtherTriangle.bUVsFixed = true;
+					}
+				}
 			}
 		}
 
@@ -431,89 +443,109 @@ namespace mu
 		//	}
 		//}
 
-		// Create the layout block vertex buffer
-		uint8* LayoutBufferPtr = nullptr;
 		{
-			const int32 LayoutBufferIndex = Mesh.GetVertexBuffers().GetBufferCount();
-			Mesh.GetVertexBuffers().SetBufferCount(LayoutBufferIndex + 1);
-
-			// TODO
-			check(Layout->GetBlockCount() < MAX_uint16);
-			const EMeshBufferSemantic LayoutSemantic = EMeshBufferSemantic::LayoutBlock;
-			const int32 LayoutSemanticIndex = int32(LayoutChannel);
-			const EMeshBufferFormat LayoutFormat = bUseAbsoluteBlockIds ? EMeshBufferFormat::UInt64 : EMeshBufferFormat::UInt16;
-			const int32 LayoutComponents = 1;
-			const int32 LayoutOffset = 0;
-			int32 ElementSize = bUseAbsoluteBlockIds ? sizeof(uint64) : sizeof(uint16);
-			Mesh.GetVertexBuffers().SetBuffer
-			(
-				LayoutBufferIndex,
-				ElementSize,
-				1,
-				&LayoutSemantic, &LayoutSemanticIndex,
-				&LayoutFormat, &LayoutComponents,
-				&LayoutOffset
-			);
-			LayoutBufferPtr = Mesh.GetVertexBuffers().GetBufferData(LayoutBufferIndex);
-		}
+			MUTABLE_CPUPROFILER_SCOPE(CreateBuffers);
 
-		// Copy UVs back to the mesh
-		{
-			uint8* pVertices = TexCoordData;
-			FVector2f* UVs = TexCoords.GetData();
+			// Create the layout block vertex buffer
+			uint8* LayoutBufferPtr = nullptr;
+			{
+				const int32 LayoutBufferIndex = Mesh.GetVertexBuffers().GetBufferCount();
+				Mesh.GetVertexBuffers().SetBufferCount(LayoutBufferIndex + 1);
+
+				// TODO
+				check(Layout->GetBlockCount() < MAX_uint16);
+				const EMeshBufferSemantic LayoutSemantic = EMeshBufferSemantic::LayoutBlock;
+				const int32 LayoutSemanticIndex = int32(LayoutChannel);
+				const EMeshBufferFormat LayoutFormat = bUseAbsoluteBlockIds ? EMeshBufferFormat::UInt64 : EMeshBufferFormat::UInt16;
+				const int32 LayoutComponents = 1;
+				const int32 LayoutOffset = 0;
+				int32 ElementSize = bUseAbsoluteBlockIds ? sizeof(uint64) : sizeof(uint16);
+				Mesh.GetVertexBuffers().SetBuffer
+				(
+					LayoutBufferIndex,
+					ElementSize,
+					1,
+					&LayoutSemantic, &LayoutSemanticIndex,
+					&LayoutFormat, &LayoutComponents,
+					&LayoutOffset
+				);
+				LayoutBufferPtr = Mesh.GetVertexBuffers().GetBufferData(LayoutBufferIndex);
+			}
 
-			for (int32 VertexIndex = 0; VertexIndex < NumVertices; ++VertexIndex)
 			{
-				FVector2f& UV = TexCoords[VertexIndex];
+				MUTABLE_CPUPROFILER_SCOPE(LayoutAndHomogenize);
 
-				uint16 LayoutBlockIndex = LayoutData[VertexIndex];
-				if (Layout->Blocks.IsValidIndex(LayoutBlockIndex))
+				for (int32 VertexIndex = 0; VertexIndex < NumVertices; ++VertexIndex)
 				{
-					uint64 LayoutBlockId = Layout->Blocks[LayoutBlockIndex].Id;
+					FVector2f& UV = TexCoords[VertexIndex];
 
-					UV = BlockRects[LayoutBlockIndex].Homogenize(UV);
-
-					// Replace block index by the actual id of the block
-					if (bUseAbsoluteBlockIds)
-					{
-						uint64* Ptr = reinterpret_cast<uint64*>(LayoutBufferPtr) + VertexIndex;
-						*Ptr = LayoutBlockId;
-					}
-					else
+					uint16 LayoutBlockIndex = LayoutData[VertexIndex];
+					if (Layout->Blocks.IsValidIndex(LayoutBlockIndex))
 					{
-						uint16* Ptr = reinterpret_cast<uint16*>(LayoutBufferPtr) + VertexIndex;
-						*Ptr = uint16(LayoutBlockId & 0xffff);
-					}
-				}
-				else
-				{
-					// Map vertices without block
-					if (bUseAbsoluteBlockIds)
-					{
-						uint64* Ptr = reinterpret_cast<uint64*>(LayoutBufferPtr) + VertexIndex;
-						*Ptr = bEnsureAllVerticesHaveLayoutBlock ? 0 : std::numeric_limits<uint64>::max();
+						uint64 LayoutBlockId = Layout->Blocks[LayoutBlockIndex].Id;
+
+						UV = BlockRects[LayoutBlockIndex].Homogenize(UV);
+
+						// Replace block index by the actual id of the block
+						if (bUseAbsoluteBlockIds)
+						{
+							uint64* Ptr = reinterpret_cast<uint64*>(LayoutBufferPtr) + VertexIndex;
+							*Ptr = LayoutBlockId;
+						}
+						else
+						{
+							uint16* Ptr = reinterpret_cast<uint16*>(LayoutBufferPtr) + VertexIndex;
+							*Ptr = uint16(LayoutBlockId & 0xffff);
+						}
 					}
 					else
 					{
-						uint16* Ptr = reinterpret_cast<uint16*>(LayoutBufferPtr) + VertexIndex;
-						*Ptr = bEnsureAllVerticesHaveLayoutBlock ? 0 : std::numeric_limits<uint16>::max();
+						// Map vertices without block
+						if (bUseAbsoluteBlockIds)
+						{
+							uint64* Ptr = reinterpret_cast<uint64*>(LayoutBufferPtr) + VertexIndex;
+							*Ptr = bEnsureAllVerticesHaveLayoutBlock ? 0 : std::numeric_limits<uint64>::max();
+						}
+						else
+						{
+							uint16* Ptr = reinterpret_cast<uint16*>(LayoutBufferPtr) + VertexIndex;
+							*Ptr = bEnsureAllVerticesHaveLayoutBlock ? 0 : std::numeric_limits<uint16>::max();
+						}
 					}
 				}
+			}
+
+			{
+				MUTABLE_CPUPROFILER_SCOPE(CopyUVs);
 
 				// Copy UVs
-				if (TexCoordsChannel.Format == EMeshBufferFormat::Float32)
+				if (TexCoordsChannel.Format == EMeshBufferFormat::Float32 && TexElemSize==sizeof(float)*2)
 				{
-					FVector2f* pUV = reinterpret_cast<FVector2f*>(pVertices);
-					*pUV = UV;
+					FMemory::Memcpy(TexCoordData, TexCoords.GetData(), sizeof(float) * 2 * NumVertices);
 				}
-				else if (TexCoordsChannel.Format == EMeshBufferFormat::Float16)
+				else
 				{
-					FFloat16* pUV = reinterpret_cast<FFloat16*>(pVertices);
-					pUV[0] = FFloat16(UV[0]);
-					pUV[1] = FFloat16(UV[1]);
-				}
+					MUTABLE_CPUPROFILER_SCOPE(SlowPath);
+
+					for (int32 VertexIndex = 0; VertexIndex < NumVertices; ++VertexIndex)
+					{
+						FVector2f& UV = TexCoords[VertexIndex];
 
-				pVertices += elemSize;
+						if (TexCoordsChannel.Format == EMeshBufferFormat::Float32)
+						{
+							FVector2f* pUV = reinterpret_cast<FVector2f*>(TexCoordData);
+							*pUV = UV;
+						}
+						else if (TexCoordsChannel.Format == EMeshBufferFormat::Float16)
+						{
+							FFloat16* pUV = reinterpret_cast<FFloat16*>(TexCoordData);
+							pUV[0] = FFloat16(UV[0]);
+							pUV[1] = FFloat16(UV[1]);
+						}
+
+						TexCoordData += TexElemSize;
+					}
+				}
 			}
 		}
 	}
diff --git a/Engine/Plugins/Mutable/Source/MutableRuntime/Public/MuR/Mesh.h b/Engine/Plugins/Mutable/Source/MutableRuntime/Public/MuR/Mesh.h
--- a/Engine/Plugins/Mutable/Source/MutableRuntime/Public/MuR/Mesh.h
+++ b/Engine/Plugins/Mutable/Source/MutableRuntime/Public/MuR/Mesh.h
@@ -551,11 +551,8 @@ namespace mu
 			return bEqual;
 		}
 
-		/** 
-		 * Compare the mesh with another one, but ignore internal data like generated vertex
-		 * indices.
-		 */
-		bool IsSimilar(const FMesh& Other, bool bCompareLayouts) const;
+		/** Compare the mesh with another one, but ignore internal data like generated vertex indices. */
+		bool IsSimilar(const FMesh& Other) const;
 
 
 		/**
diff --git a/Engine/Plugins/Mutable/Source/MutableTools/Private/MuT/ASTOpMeshExtractLayoutBlocks.cpp b/Engine/Plugins/Mutable/Source/MutableTools/Private/MuT/ASTOpMeshExtractLayoutBlocks.cpp
--- a/Engine/Plugins/Mutable/Source/MutableTools/Private/MuT/ASTOpMeshExtractLayoutBlocks.cpp
+++ b/Engine/Plugins/Mutable/Source/MutableTools/Private/MuT/ASTOpMeshExtractLayoutBlocks.cpp
@@ -12,6 +12,8 @@
 #include "MuT/ASTOpMeshApplyLayout.h"
 #include "MuT/ASTOpMeshMerge.h"
 #include "MuT/ASTOpMeshSetSkeleton.h"
+#include "MuT/ASTOpMeshPrepareLayout.h"
+#include "MuT/ASTOpConstantResource.h"
 #include "MuR/ModelPrivate.h"
 #include "MuR/RefCounted.h"
 #include "MuR/Types.h"
@@ -159,6 +161,50 @@ namespace mu
 			break;
 		}
 
+		case EOpType::ME_PREPARELAYOUT:
+		{
+			ASTOpMeshPrepareLayout* SourcePrepareOp = static_cast<ASTOpMeshPrepareLayout*>(at.get());
+
+			bool bHandled = false;
+
+			// If the "prepare" operation is for an unrelated layout, sink the "extract" down the "prepare" mesh.
+			if (SourcePrepareOp->LayoutChannel!= currentSinkOp->LayoutIndex)
+			{
+				newAt = Visit(SourcePrepareOp->Mesh.child(), currentSinkOp);
+				bHandled = true;
+			}
+			// If the "prepare" operation is for a this layout channel...
+			else if (SourcePrepareOp->LayoutChannel == currentSinkOp->LayoutIndex)
+			{
+				// ...and the layout is constant...
+				if (SourcePrepareOp->Layout && SourcePrepareOp->Layout->GetOpType() == EOpType::LA_CONSTANT)
+				{
+					ASTOpConstantResource* LayoutOp = static_cast<ASTOpConstantResource*>(SourcePrepareOp->Layout.child().get());
+					const FLayout* Layout = static_cast<const FLayout*>(LayoutOp->GetValue().Get());
+
+					// ...and it is a single block covering all the layout.
+					if (Layout && Layout->IsSingleBlockAndFull())
+					{
+						// We can omit the "extract" operation because it will get the whole mesh anyway
+						newAt = SourcePrepareOp;
+						bHandled = true;
+					}
+				}
+			}
+			
+			if (!bHandled)
+			{
+				if (at != InitialSource)
+				{
+					// We didn't optimize it, so create the end operation as in the switch "default" case.
+					mu::Ptr<ASTOpMeshExtractLayoutBlocks> newOp = mu::Clone<ASTOpMeshExtractLayoutBlocks>(currentSinkOp);
+					newOp->Source = at;
+					newAt = newOp;
+				}
+			}
+			break;
+		}
+
 		case EOpType::ME_SETSKELETON:
 		{
 			Ptr<ASTOpMeshSetSkeleton> NewOp = mu::Clone<ASTOpMeshSetSkeleton>(at);
diff --git a/Engine/Plugins/Mutable/Source/MutableTools/Private/MuT/CodeGenerator.h b/Engine/Plugins/Mutable/Source/MutableTools/Private/MuT/CodeGenerator.h
--- a/Engine/Plugins/Mutable/Source/MutableTools/Private/MuT/CodeGenerator.h
+++ b/Engine/Plugins/Mutable/Source/MutableTools/Private/MuT/CodeGenerator.h
@@ -318,7 +318,6 @@ namespace mu
 			{
 				TSharedPtr<FMesh> Mesh;
 				Ptr<ASTOp> LastMeshOp;
-				TArray<FGeneratedLayout> GeneratedLayouts;
 			};
 			TMap<uint64, TArray<FGeneratedConstantMesh>> GeneratedConstantMeshes;
 		};
diff --git a/Engine/Plugins/Mutable/Source/MutableTools/Private/MuT/CodeGenerator_Mesh.cpp b/Engine/Plugins/Mutable/Source/MutableTools/Private/MuT/CodeGenerator_Mesh.cpp
--- a/Engine/Plugins/Mutable/Source/MutableTools/Private/MuT/CodeGenerator_Mesh.cpp
+++ b/Engine/Plugins/Mutable/Source/MutableTools/Private/MuT/CodeGenerator_Mesh.cpp
@@ -636,7 +636,9 @@ namespace mu
 			return LastMeshOp;
 		}
 
-		bool bUseAbsoluteBlockIds = false;
+		// Always absolute to ease mesh reusal.
+		bool bUseAbsoluteBlockIds = true;
+
 		bool bIsOverridingLayouts = !Options.OverrideLayouts.IsEmpty();
 		if (!bIsOverridingLayouts)
 		{
@@ -775,10 +777,6 @@ namespace mu
 					return FMeshGenerationResult();
 				}
 
-				Ptr<ASTOpConstantResource> ConstantOp = new ASTOpConstantResource();
-				ConstantOp->Type = EOpType::ME_CONSTANT;
-				ConstantOp->SourceDataDescriptor = NodeCopy->SourceDataDescriptor;
-
 				// Separate the tags from the mesh
 				TArray<FString> Tags = ReferenceMesh->Tags;
 				if (!Mesh->Tags.IsEmpty())
@@ -804,85 +802,53 @@ namespace mu
 				TArray<FGenerateMeshConstantState::FGeneratedConstantMesh>& CachedCandidates = GenerateMeshConstantState.GeneratedConstantMeshes.FindOrAdd(ThisMeshHash, {});
 				for (const FGenerateMeshConstantState::FGeneratedConstantMesh& Candidate : CachedCandidates)
 				{
-					bool bCompareLayouts = Options.bLayouts && !bIsOverridingLayouts;
-
-					if (Candidate.Mesh->IsSimilar(*Mesh, bCompareLayouts))
+					if (Candidate.Mesh->IsSimilar(*Mesh))
 					{
-						// If it is similar and we are overriding the layouts, we must compare the layouts of the candidate with the ones
-						// we are using to override.
-						if (bIsOverridingLayouts)
-						{
-							if (Candidate.Mesh->GetLayoutCount() != Options.OverrideLayouts.Num())
-							{
-								continue;
-							}
-
-							bool bLayoutsAreEqual = true;
-							for (int32 l = 0; l < Candidate.Mesh->GetLayoutCount(); ++l)
-							{
-								bLayoutsAreEqual = (*Candidate.Mesh->GetLayout(l) == *Options.OverrideLayouts[l].Layout);
-								if (!bLayoutsAreEqual)
-								{
-									break;
-								}
-							}
-
-							if (!bLayoutsAreEqual)
-							{
-								continue;
-							}
-						}
-
 						DuplicateOf = Candidate;
 						break;
 					}
 				}
 
-				FMeshGenerationResult Result;
-
-				Ptr<ASTOp> LastMeshOp = ConstantOp;
+				Ptr<ASTOp> LastMeshOp;
+				uint32 MeshIDPrefix = 0;
 
-				TSharedPtr<FMesh> ClonedMesh;
 				if (!DuplicateOf.Mesh)
 				{
+					Ptr<ASTOpConstantResource> ConstantOp = new ASTOpConstantResource();
+					ConstantOp->Type = EOpType::ME_CONSTANT;
+					ConstantOp->SourceDataDescriptor = NodeCopy->SourceDataDescriptor;
+
+					LastMeshOp = ConstantOp;
+
 					// We need to clone the mesh in the node because we will modify it.
-					ClonedMesh = Mesh->Clone();
+					TSharedPtr<FMesh> ClonedMesh = Mesh->Clone();
 					ClonedMesh->EnsureSurfaceData();
 
 					ConstantOp->SetValue(ClonedMesh, CompilerOptions->OptimisationOptions.DiskCacheContext);
 
 					// Add the unique vertex ID prefix in all cases, since it is free memory-wise
-					uint32 MeshIDPrefix = uint32(ConstantOp->GetValueHash());
+					MeshIDPrefix = uint32(ConstantOp->GetValueHash());
 					MeshIDPrefix = UniqueMeshIds.EnsureUnique(MeshIDPrefix);
 					ClonedMesh->MeshIDPrefix = MeshIDPrefix;
 
-					LastMeshOp = GenerateLayoutOpsAndResult(Options, ConstantOp, NodeCopy->Layouts, ClonedMesh->MeshIDPrefix, Result);
-
 					// Add the constant data
 					FGenerateMeshConstantState::FGeneratedConstantMesh MeshEntry;
 					MeshEntry.Mesh = ClonedMesh;
 					MeshEntry.LastMeshOp = LastMeshOp;
-					MeshEntry.GeneratedLayouts = Result.GeneratedLayouts;
 					CachedCandidates.Add(MeshEntry);
 				}
 
 				if (DuplicateOf.Mesh)
 				{
-					// Make sure the source layouts of the mesh are mapped to the layouts of the duplicated mesh.
-					if (Options.bLayouts)
-					{
-						if (bIsOverridingLayouts)
-						{
-							Result.GeneratedLayouts = Options.OverrideLayouts;
-						}
-						else
-						{
-							Result.GeneratedLayouts = DuplicateOf.GeneratedLayouts;
-						}
-					}
-
 					LastMeshOp = DuplicateOf.LastMeshOp;
-					ConstantOp = nullptr;
+					MeshIDPrefix = DuplicateOf.Mesh->MeshIDPrefix;
+				}
+
+				FMeshGenerationResult Result;
+
+				if (Options.bLayouts)
+				{
+					LastMeshOp = GenerateLayoutOpsAndResult(Options, LastMeshOp, NodeCopy->Layouts, MeshIDPrefix, Result);
 				}
 
 				Result.BaseMeshOp = LastMeshOp;
diff --git a/Engine/Plugins/Mutable/Source/MutableTools/Private/MuT/CodeOptimiser.cpp b/Engine/Plugins/Mutable/Source/MutableTools/Private/MuT/CodeOptimiser.cpp
--- a/Engine/Plugins/Mutable/Source/MutableTools/Private/MuT/CodeOptimiser.cpp
+++ b/Engine/Plugins/Mutable/Source/MutableTools/Private/MuT/CodeOptimiser.cpp
@@ -14,6 +14,8 @@
 #include "MuT/ASTOpMeshBindShape.h"
 #include "MuT/ASTOpMeshClipMorphPlane.h"
 #include "MuT/ASTOpMeshApplyPose.h"
+#include "MuT/ASTOpMeshPrepareLayout.h"
+#include "MuT/ASTOpMeshAddTags.h"
 #include "MuT/ASTOpImageRasterMesh.h"
 #include "MuT/ASTOpReferenceResource.h"
 
@@ -632,6 +634,8 @@ namespace mu
 				{
 					EOpType SubgraphType = SubgraphRoot->GetOpType();
 
+					bool bGetFromChildren = false;
+
 					bool bIsConstantSubgraph = true;
 					switch (SubgraphType)
 					{
@@ -671,11 +675,27 @@ namespace mu
 					case EOpType::ME_CLIPMORPHPLANE:
 					case EOpType::ME_APPLYPOSE:
 					case EOpType::ME_REMOVEMASK:
-					case EOpType::ME_ADDTAGS:
+					case EOpType::ME_PREPARELAYOUT:
 					case EOpType::IM_PLAINCOLOUR:
 						bHasSpecialOpInSubgraph = true;
 						break;
 
+					case EOpType::ME_ADDTAGS:
+					{
+						const ASTOpMeshAddTags* AddTags = static_cast<const ASTOpMeshAddTags*>(SubgraphRoot.get());
+						if (AddTags->Source->GetParentCount()==1)
+						{
+							// We could resolve the constant if there are no other constraints since this op's mesh op is only referenced once.
+							bGetFromChildren = true;
+						}
+						else
+						{
+							// Don't reduce constant to avoid duplicating mesh data.
+							bHasSpecialOpInSubgraph = true;
+						}
+						break;
+					}
+
 					case EOpType::IM_RASTERMESH:
 					{
 						const ASTOpImageRasterMesh* Raster = static_cast<const ASTOpImageRasterMesh*>(SubgraphRoot.get());
@@ -686,6 +706,20 @@ namespace mu
 						break;
 					}
 
+					case EOpType::LA_FROMMESH:
+					{
+						// We want to reduce this type of operation regardless of it having special ops below.
+						bHasSpecialOpInSubgraph = false;
+						break;
+					}
+
+					case EOpType::ME_APPLYLAYOUT:
+					{
+						// We want to reduce this type of operation regardless of it having special ops below.
+						bHasSpecialOpInSubgraph = false;
+						break;
+					}
+
 					case EOpType::ME_REFERENCE:
 					case EOpType::IM_REFERENCE:
 						// If we are in a reference-resolution optimization phase, then the ops are not special.
@@ -701,6 +735,12 @@ namespace mu
 						break;
 
 					default:
+						bGetFromChildren = true;
+						break;
+					}
+
+					if (bGetFromChildren)
+					{
 						// Propagate from children
 						SubgraphRoot->ForEachChild([&](ASTChild& c)
 							{
@@ -709,8 +749,8 @@ namespace mu
 									bHasSpecialOpInSubgraph = bHasSpecialOpInSubgraph || c->bHasSpecialOpInSubgraph;
 								}
 							});
-						break;
 					}
+
 					SubgraphRoot->bHasSpecialOpInSubgraph = bHasSpecialOpInSubgraph;
 
 					bool bIsDataTypeThanCanTurnIntoConst = false;
