[Mutable] Remove Alternative resolution per state on group projectors

diff --git a/Engine/Plugins/Mutable/Source/CustomizableObjectEditor/Private/MuCOE/GenerateMutableSource/GenerateMutableSource.h b/Engine/Plugins/Mutable/Source/CustomizableObjectEditor/Private/MuCOE/GenerateMutableSource/GenerateMutableSource.h
--- a/Engine/Plugins/Mutable/Source/CustomizableObjectEditor/Private/MuCOE/GenerateMutableSource/GenerateMutableSource.h
+++ b/Engine/Plugins/Mutable/Source/CustomizableObjectEditor/Private/MuCOE/GenerateMutableSource/GenerateMutableSource.h
@@ -327,15 +327,10 @@ struct FGroupProjectorImageInfo
 	UCustomizableObjectNodeMaterialBase* TypedNodeMat;
 	FString TextureName;
 	FString RealTextureName;
-	FString AlternateResStateName;
-	float AlternateProjectionResolutionFactor;
-	bool bIsAlternateResolutionResized = false;
 	int32 UVLayout = 0;
 
-	FGroupProjectorImageInfo(mu::NodeImagePtr InImageNode, const FString& InTextureName, const FString& InRealTextureName, UCustomizableObjectNodeMaterialBase* InTypedNodeMat,
-		float InAlternateProjectionResolutionFactor, const FString& InAlternateResStateName, mu::Ptr<mu::NodeSurfaceNew> InSurfNode, int32 InUVLayout)
+	FGroupProjectorImageInfo(mu::NodeImagePtr InImageNode, const FString& InTextureName, const FString& InRealTextureName, UCustomizableObjectNodeMaterialBase* InTypedNodeMat, mu::Ptr<mu::NodeSurfaceNew> InSurfNode, int32 InUVLayout)
 		: TypedNodeMat(InTypedNodeMat), TextureName(InTextureName), RealTextureName(InRealTextureName),
-		AlternateResStateName(InAlternateResStateName), AlternateProjectionResolutionFactor(InAlternateProjectionResolutionFactor), 
 		UVLayout(InUVLayout)
 	{
 		ImageNode = InImageNode;
diff --git a/Engine/Plugins/Mutable/Source/CustomizableObjectEditor/Private/MuCOE/GenerateMutableSource/GenerateMutableSourceGroupProjector.cpp b/Engine/Plugins/Mutable/Source/CustomizableObjectEditor/Private/MuCOE/GenerateMutableSource/GenerateMutableSourceGroupProjector.cpp
--- a/Engine/Plugins/Mutable/Source/CustomizableObjectEditor/Private/MuCOE/GenerateMutableSource/GenerateMutableSourceGroupProjector.cpp
+++ b/Engine/Plugins/Mutable/Source/CustomizableObjectEditor/Private/MuCOE/GenerateMutableSource/GenerateMutableSourceGroupProjector.cpp
@@ -37,7 +37,7 @@ class UPoseAsset;
 
 mu::NodeImagePtr GenerateMutableSourceGroupProjector(const int32 NodeLOD, const int32 ImageIndex, mu::NodeMeshPtr MeshNode, FMutableGraphGenerationContext& GenerationContext,
 	UCustomizableObjectNodeMaterialBase* TypedNodeMat, UCustomizableObjectNodeModifierExtendMeshSection* TypedNodeExt, bool& bShareProjectionTexturesBetweenLODs, bool& bIsGroupProjectorImage,
-	UTexture2D*& GroupProjectionReferenceTexture, TMap<FString, float>& TextureNameToProjectionResFactor, FString& AlternateResStateName)
+	UTexture2D*& GroupProjectionReferenceTexture)
 {
 	check(static_cast<bool>(TypedNodeMat) != static_cast<bool>(TypedNodeExt)); // XOr. TypedNodeMat valid or TypedNodeExt valid. At least one valid.
 
@@ -131,28 +131,6 @@ mu::NodeImagePtr GenerateMutableSourceGroupProjector(const int32 NodeLOD, const
 				bIsGroupProjectorImage = true;
 				ImageNode->Layout = ProjectorTempData.CustomizableObjectNodeGroupProjectorParameter->UVLayout;
 
-				{
-					if (TypedNodeMat &&
-						!ProjectorTempData.CustomizableObjectNodeGroupProjectorParameter->AlternateProjectionResolutionStateName.IsEmpty()
-						&& ProjectorTempData.CustomizableObjectNodeGroupProjectorParameter->AlternateProjectionResolutionFactor > 0)
-					{
-						TextureNameToProjectionResFactor.Add(
-							TypedNodeMat->GetParameterName(EMaterialParameterType::Texture, ImageIndex).ToString(),
-							ProjectorTempData.CustomizableObjectNodeGroupProjectorParameter->AlternateProjectionResolutionFactor);
-
-						if (!AlternateResStateName.IsEmpty() &&
-							AlternateResStateName != ProjectorTempData.CustomizableObjectNodeGroupProjectorParameter->AlternateProjectionResolutionStateName &&
-							!ProjectorTempData.bAlternateResStateNameWarningDisplayed)
-						{
-							FString msg = FString::Printf(TEXT("All 'Alternate Projection Resolution State Name' properties in Group Projector Parameter nodes connected to same Group node must have the same value or be blank. Only the value of the last connected node will be used."));
-							GenerationContext.Log(FText::FromString(msg), ProjectorTempData.CustomizableObjectNodeGroupProjectorParameter);
-							ProjectorTempData.bAlternateResStateNameWarningDisplayed = true;
-						}
-
-						AlternateResStateName = ProjectorTempData.CustomizableObjectNodeGroupProjectorParameter->AlternateProjectionResolutionStateName;
-					}
-				}
-
 				{
 					mu::Ptr<mu::NodeScalarConstant> ScalarNode = new mu::NodeScalarConstant;
 					ScalarNode->Value = 120.f;
diff --git a/Engine/Plugins/Mutable/Source/CustomizableObjectEditor/Private/MuCOE/GenerateMutableSource/GenerateMutableSourceGroupProjector.h b/Engine/Plugins/Mutable/Source/CustomizableObjectEditor/Private/MuCOE/GenerateMutableSource/GenerateMutableSourceGroupProjector.h
--- a/Engine/Plugins/Mutable/Source/CustomizableObjectEditor/Private/MuCOE/GenerateMutableSource/GenerateMutableSourceGroupProjector.h
+++ b/Engine/Plugins/Mutable/Source/CustomizableObjectEditor/Private/MuCOE/GenerateMutableSource/GenerateMutableSourceGroupProjector.h
@@ -19,7 +19,7 @@ struct FMutableGraphGenerationContext;
 
 mu::NodeImagePtr GenerateMutableSourceGroupProjector(const int32 NodeLOD, const int32 ImageIndex, mu::NodeMeshPtr MeshNode, FMutableGraphGenerationContext& GenerationContext,
 	UCustomizableObjectNodeMaterialBase* TypedNodeMat, UCustomizableObjectNodeModifierExtendMeshSection* TypedNodeExt, bool& bShareProjectionTexturesBetweenLODs, bool& bIsGroupProjectorImage,
-	UTexture2D*& GroupProjectionReferenceTexture, TMap<FString, float>& TextureNameToProjectionResFactor, FString& AlternateResStateName);
+	UTexture2D*& GroupProjectionReferenceTexture);
 
 
 TOptional<FGroupProjectorTempData> GenerateMutableGroupProjector(UCustomizableObjectNodeGroupProjectorParameter* ProjParamNode, FMutableGraphGenerationContext& GenerationContext);
diff --git a/Engine/Plugins/Mutable/Source/CustomizableObjectEditor/Private/MuCOE/GenerateMutableSource/GenerateMutableSourceModifier.cpp b/Engine/Plugins/Mutable/Source/CustomizableObjectEditor/Private/MuCOE/GenerateMutableSource/GenerateMutableSourceModifier.cpp
--- a/Engine/Plugins/Mutable/Source/CustomizableObjectEditor/Private/MuCOE/GenerateMutableSource/GenerateMutableSourceModifier.cpp
+++ b/Engine/Plugins/Mutable/Source/CustomizableObjectEditor/Private/MuCOE/GenerateMutableSource/GenerateMutableSourceModifier.cpp
@@ -349,9 +349,6 @@ mu::Ptr<mu::NodeModifier> GenerateMutableSourceModifier(const UEdGraphPin * Pin,
 						ensure(LOD > GenerationContext.FirstLODAvailable[GenerationContext.CurrentMeshComponent]);
 						check(ProjectorInfo->SurfNode->Images[ImageIndex].Image == ProjectorInfo->ImageNode);
 						ImageNode = ProjectorInfo->ImageNode;
-
-						//TextureNameToProjectionResFactor.Add(ProjectorInfo->RealTextureName, ProjectorInfo->AlternateProjectionResolutionFactor);
-						//AlternateResStateName = ProjectorInfo->AlternateResStateName;
 					}
 				}
 
@@ -360,12 +357,10 @@ mu::Ptr<mu::NodeModifier> GenerateMutableSourceModifier(const UEdGraphPin * Pin,
 					bool bShareProjectionTexturesBetweenLODs = false;
 					bool bIsGroupProjectorImage = false;
 					UTexture2D* GroupProjectionReferenceTexture = nullptr;
-					TMap<FString, float> TextureNameToProjectionResFactor;
-					FString AlternateResStateName;
 
 					ImageNode = GenerateMutableSourceGroupProjector(LOD, ImageIndex, AddMeshNode, GenerationContext,
 						nullptr, TypedNodeExt, bShareProjectionTexturesBetweenLODs, bIsGroupProjectorImage,
-						GroupProjectionReferenceTexture, TextureNameToProjectionResFactor, AlternateResStateName);
+						GroupProjectionReferenceTexture);
 				}
 
 				if (!ImageNode) // Else if
diff --git a/Engine/Plugins/Mutable/Source/CustomizableObjectEditor/Private/MuCOE/GenerateMutableSource/GenerateMutableSourceSurface.cpp b/Engine/Plugins/Mutable/Source/CustomizableObjectEditor/Private/MuCOE/GenerateMutableSource/GenerateMutableSourceSurface.cpp
--- a/Engine/Plugins/Mutable/Source/CustomizableObjectEditor/Private/MuCOE/GenerateMutableSource/GenerateMutableSourceSurface.cpp
+++ b/Engine/Plugins/Mutable/Source/CustomizableObjectEditor/Private/MuCOE/GenerateMutableSource/GenerateMutableSourceSurface.cpp
@@ -208,9 +208,6 @@ mu::Ptr<mu::NodeSurface> GenerateMutableSourceSurface(const UEdGraphPin * Pin, F
 				}
 			}
 		}
-		
-		TMap<FString, float> TextureNameToProjectionResFactor;
-		FString AlternateResStateName;
 
 		bool bTableMaterialPinLinked = TypedNodeMat->GetMaterialAssetPin() && FollowInputPin(*TypedNodeMat->GetMaterialAssetPin()) != nullptr;
 
@@ -300,7 +297,7 @@ mu::Ptr<mu::NodeSurface> GenerateMutableSourceSurface(const UEdGraphPin * Pin, F
 
 					GroupProjectionImg = GenerateMutableSourceGroupProjector(LOD, ImageIndex, MeshNode, GenerationContext,
 						TypedNodeMat, nullptr, bShareProjectionTexturesBetweenLODs, bIsGroupProjectorImage,
-						GroupProjectionReferenceTexture, TextureNameToProjectionResFactor, AlternateResStateName);
+						GroupProjectionReferenceTexture);
 
 					if (GroupProjectionImg.get() || TypedNodeMat->IsImageMutableMode(ImageIndex))
 					{
@@ -691,10 +688,8 @@ mu::Ptr<mu::NodeSurface> GenerateMutableSourceSurface(const UEdGraphPin * Pin, F
 						{
 							// Add to the GroupProjectorLODCache to potentially reuse this projection texture in higher LODs
 							ensure(LOD == GenerationContext.FirstLODAvailable[GenerationContext.CurrentMeshComponent]);
-							float* AlternateProjectionResFactor = TextureNameToProjectionResFactor.Find(ImageName);
 							GenerationContext.GroupProjectorLODCache.Add(MaterialImageId,
-								FGroupProjectorImageInfo(ImageNode, ImageName, ImageName, TypedNodeMat,
-									AlternateProjectionResFactor ? *AlternateProjectionResFactor : 0.f, AlternateResStateName, SurfNode, UVLayout));
+								FGroupProjectorImageInfo(ImageNode, ImageName, ImageName, TypedNodeMat, SurfNode, UVLayout));
 						}
 					}
 				}
@@ -705,9 +700,6 @@ mu::Ptr<mu::NodeSurface> GenerateMutableSourceSurface(const UEdGraphPin * Pin, F
 					SurfNode->Images[ImageIndex].Image = ProjectorInfo->ImageNode;
 					SurfNode->Images[ImageIndex].Name = ProjectorInfo->TextureName;
 					SurfNode->Images[ImageIndex].LayoutIndex = ProjectorInfo->UVLayout;
-
-					TextureNameToProjectionResFactor.Add(ProjectorInfo->RealTextureName, ProjectorInfo->AlternateProjectionResolutionFactor);
-					AlternateResStateName = ProjectorInfo->AlternateResStateName;
 				}
 			}
 		}
@@ -848,65 +840,6 @@ mu::Ptr<mu::NodeSurface> GenerateMutableSourceSurface(const UEdGraphPin * Pin, F
 				}
 			}
 		}
-
-		// If an alternate resolution for a particular state is present, clone the surface node, add the image resizing and inject the surface variation node
-		if (TextureNameToProjectionResFactor.Num() > 0 && !AlternateResStateName.IsEmpty())
-		{
-			mu::Ptr<mu::NodeSurfaceNew> SurfNode2 = new mu::NodeSurfaceNew;
-
-			SurfNode2->ExternalId = SurfaceMetadataUniqueHash;
-
-			SurfNode2->Mesh = SurfNode->Mesh;
-			SurfNode2->Tags = SurfNode->Tags;
-			SurfNode2->Vectors = SurfNode->Vectors;
-			SurfNode2->Scalars = SurfNode->Scalars;
-			SurfNode2->Strings = SurfNode->Strings;
-			SurfNode2->Images = SurfNode->Images;
-
-			for (int32 ImageIndex = 0; ImageIndex < SurfNode2->Images.Num(); ++ImageIndex)
-			{
-				const FString ImageName = TypedNodeMat->GetParameterName(EMaterialParameterType::Texture, ImageIndex).ToString();
-				if (float* ResolutionFactor = TextureNameToProjectionResFactor.Find(ImageName))
-				{
-					FString MaterialImageId = FGroupProjectorImageInfo::GenerateId(TypedNodeMat, ImageIndex);
-					FGroupProjectorImageInfo* ProjectorInfo = GenerationContext.GroupProjectorLODCache.Find(MaterialImageId);
-
-					if (!ProjectorInfo || !ProjectorInfo->bIsAlternateResolutionResized)
-					{
-						mu::Ptr<mu::NodeImageResize> NodeImageResize = new mu::NodeImageResize;
-						NodeImageResize->bRelative = true;
-						NodeImageResize->SizeX = *ResolutionFactor;
-						NodeImageResize->SizeY = *ResolutionFactor;
-						NodeImageResize->Base = SurfNode2->Images[ImageIndex].Image;
-
-						SurfNode2->Images[ImageIndex].Image = NodeImageResize;
-
-						if (ProjectorInfo)
-						{
-							ensure(LOD == GenerationContext.FirstLODAvailable[GenerationContext.CurrentMeshComponent]);
-							ProjectorInfo->ImageResizeNode = NodeImageResize;
-							ProjectorInfo->bIsAlternateResolutionResized = true;
-						}
-					}
-					else
-					{
-						ensure(LOD > GenerationContext.FirstLODAvailable[GenerationContext.CurrentMeshComponent]);
-						check(ProjectorInfo->bIsAlternateResolutionResized);
-						SurfNode2->Images[ImageIndex].Image = ProjectorInfo->ImageResizeNode;
-					}
-				}
-			}
-
-			mu::Ptr<mu::NodeSurfaceVariation> SurfaceVariation = new mu::NodeSurfaceVariation;
-			SurfaceVariation->Type = mu::NodeSurfaceVariation::VariationType::State;
-			SurfaceVariation->Variations.SetNum(1);
-			SurfaceVariation->Variations[0].Tag = AlternateResStateName;
-
-			SurfaceVariation->DefaultSurfaces.Add(SurfNode);
-			SurfaceVariation->Variations[0].Surfaces.Add(SurfNode2);
-
-			Result = SurfaceVariation;
-		}
 	}
 
 	else if (const UCustomizableObjectNodeMaterialVariation* TypedNodeVar = Cast<UCustomizableObjectNodeMaterialVariation>(Node))
diff --git a/Engine/Plugins/Mutable/Source/CustomizableObjectEditor/Private/MuCOE/Nodes/CustomizableObjectNodeGroupProjectorParameter.h b/Engine/Plugins/Mutable/Source/CustomizableObjectEditor/Private/MuCOE/Nodes/CustomizableObjectNodeGroupProjectorParameter.h
--- a/Engine/Plugins/Mutable/Source/CustomizableObjectEditor/Private/MuCOE/Nodes/CustomizableObjectNodeGroupProjectorParameter.h
+++ b/Engine/Plugins/Mutable/Source/CustomizableObjectEditor/Private/MuCOE/Nodes/CustomizableObjectNodeGroupProjectorParameter.h
@@ -87,12 +87,6 @@ public:
 	UPROPERTY(EditAnywhere, Category = ProjectorGroup)
 	TObjectPtr<UDataTable> OptionTexturesDataTable = nullptr;
 
-	UPROPERTY(EditAnywhere, Category = ProjectorGroup)
-	FString AlternateProjectionResolutionStateName;
-
-	UPROPERTY(EditAnywhere, Category = ProjectorGroup)
-	float AlternateProjectionResolutionFactor = 0.0f;
-
 	UPROPERTY(EditAnywhere, Category = ProjectorGroup, Meta = (ClampMin = "0"))
 	int32 UVLayout = 0;
 
