[Mutable] Fix all static parameters being enabled on baked materials.

#rnx

diff --git a/Engine/Plugins/Mutable/Source/CustomizableObjectEditor/Private/MuCOE/UnrealBakeHelpers.h b/Engine/Plugins/Mutable/Source/CustomizableObjectEditor/Private/MuCOE/UnrealBakeHelpers.h
--- a/Engine/Plugins/Mutable/Source/CustomizableObjectEditor/Private/MuCOE/UnrealBakeHelpers.h
+++ b/Engine/Plugins/Mutable/Source/CustomizableObjectEditor/Private/MuCOE/UnrealBakeHelpers.h
@@ -44,7 +44,7 @@ public:
 			for (const FMaterialParameterInfo& Param : ScalarParameterInfoArray)
 			{
 				float Value = 0.f;
-				if (OriginMaterial.GetScalarParameterValue(Param, Value))
+				if (OriginMaterial.GetScalarParameterValue(Param, Value, true))
 				{
 					DestMaterial.SetScalarParameterValueEditorOnly(Param.Name, Value);
 				}
@@ -60,7 +60,7 @@ public:
 			for (const FMaterialParameterInfo& Param : VectorParameterInfoArray)
 			{
 				FLinearColor Value;
-				if (OriginMaterial.GetVectorParameterValue(Param, Value))
+				if (OriginMaterial.GetVectorParameterValue(Param, Value, true))
 				{
 					DestMaterial.SetVectorParameterValueEditorOnly(Param.Name, Value);
 				}
@@ -76,7 +76,7 @@ public:
 			{
 				bool Value = false;
 				FGuid ExpressionsGuid;
-				if (OriginMaterial.GetStaticSwitchParameterValue(StaticSwitchParameterInfoArray[i].Name, Value, ExpressionsGuid))
+				if (OriginMaterial.GetStaticSwitchParameterValue(StaticSwitchParameterInfoArray[i].Name, Value, ExpressionsGuid, true))
 				{
 					// For some reason UMaterialInstance::SetStaticSwitchParameterValueEditorOnly signature is different than UMaterial::SetStaticSwitchParameterValueEditorOnly
 					if constexpr (std::is_same_v<MaterialType, UMaterial>)
