[Mutable] Start update cleanup.

- Added scope with the name of the CO.
- Moved ClearResourceCacheProtected to private.
- Ensured all codepaths call ClearCurrentMutableOperation.

#rnx

diff --git a/Engine/Plugins/Mutable/Source/CustomizableObject/Internal/MuCO/CustomizableObjectMipDataProvider.h b/Engine/Plugins/Mutable/Source/CustomizableObject/Internal/MuCO/CustomizableObjectMipDataProvider.h
--- a/Engine/Plugins/Mutable/Source/CustomizableObject/Internal/MuCO/CustomizableObjectMipDataProvider.h
+++ b/Engine/Plugins/Mutable/Source/CustomizableObject/Internal/MuCO/CustomizableObjectMipDataProvider.h
@@ -62,6 +62,8 @@ private:
 	FString CustomizableObjectPathName;
 	FString InstancePathName;
 
+	FName CustomizableObjectName;
+
 	TSharedPtr<mu::FSystem> System;
 	TSharedPtr<mu::FModel, ESPMode::ThreadSafe> Model;
 	TSharedPtr<FModelStreamableBulkData> ModelStreamableBulkData;
diff --git a/Engine/Plugins/Mutable/Source/CustomizableObject/Internal/MuCO/CustomizableObjectSystemPrivate.h b/Engine/Plugins/Mutable/Source/CustomizableObject/Internal/MuCO/CustomizableObjectSystemPrivate.h
--- a/Engine/Plugins/Mutable/Source/CustomizableObject/Internal/MuCO/CustomizableObjectSystemPrivate.h
+++ b/Engine/Plugins/Mutable/Source/CustomizableObject/Internal/MuCO/CustomizableObjectSystemPrivate.h
@@ -936,6 +936,10 @@ public:
 	void AddPendingLoad(UCustomizableObject*);
 
 #endif
+
+	// Unprotect the resources in the instances of this object of being garbage collector
+	// while an instance is being built or updated, so that they can be reused.
+	void ClearResourceCacheProtected();
 	
 	/** Mutable TaskGraph system (Mutable Thread). */
 	FMutableTaskGraph MutableTaskGraph;
diff --git a/Engine/Plugins/Mutable/Source/CustomizableObject/Private/MuCO/CustomizableObjectMeshUpdate.cpp b/Engine/Plugins/Mutable/Source/CustomizableObject/Private/MuCO/CustomizableObjectMeshUpdate.cpp
--- a/Engine/Plugins/Mutable/Source/CustomizableObject/Private/MuCO/CustomizableObjectMeshUpdate.cpp
+++ b/Engine/Plugins/Mutable/Source/CustomizableObject/Private/MuCO/CustomizableObjectMeshUpdate.cpp
@@ -105,6 +105,8 @@ void FCustomizableObjectMeshStreamIn::DoInitiate(const FContext& Context)
 {
 	check(Context.CurrentThread == TT_Async);
 
+	MUTABLE_CPUPROFILER_SCOPE(FCustomizableObjectMeshStreamIn::DoInitiate)
+	
 	// Launch MutableTask
 	RequestMeshUpdate(Context);
 
@@ -121,6 +123,9 @@ void FCustomizableObjectMeshStreamIn::DoInitiate(const FContext& Context)
 void FCustomizableObjectMeshStreamIn::DoConvertResources(const FContext& Context)
 {
 	check(Context.CurrentThread == TT_Async);
+
+	MUTABLE_CPUPROFILER_SCOPE(FCustomizableObjectMeshStreamIn::DoConvertResources)
+
 	ConvertMesh(Context);
 
 	PushTask(Context
@@ -130,6 +135,8 @@ void FCustomizableObjectMeshStreamIn::DoConvertResources(const FContext& Context
 
 void FCustomizableObjectMeshStreamIn::DoCreateBuffers(const FContext& Context)
 {
+	MUTABLE_CPUPROFILER_SCOPE(FCustomizableObjectMeshStreamIn::DoCreateBuffers)
+
 	CreateBuffers(Context);
 
 	check(!TaskSynchronization.GetValue());
@@ -143,6 +150,8 @@ void FCustomizableObjectMeshStreamIn::DoCreateBuffers(const FContext& Context)
 
 void FCustomizableObjectMeshStreamIn::DoCancelMeshUpdate(const FContext& Context)
 {
+	MUTABLE_CPUPROFILER_SCOPE(FCustomizableObjectMeshStreamIn::DoCancelMeshUpdate)
+
 	CancelMeshUpdate(Context);
 	PushTask(Context, TT_None, nullptr, (EThreadType)Context.CurrentThread, SRA_UPDATE_CALLBACK(DoCancel));
 }
@@ -240,6 +249,11 @@ namespace impl
 
 void FCustomizableObjectMeshStreamIn::RequestMeshUpdate(const FContext& Context)
 {
+	MUTABLE_CPUPROFILER_SCOPE(FCustomizableObjectMeshStreamIn::RequestMeshUpdate)
+	
+	FSoftObjectPath Path(Cast<UCustomizableObjectSkeletalMesh>(Context.Mesh)->CustomizableObjectPathName);
+	TRACE_CPUPROFILER_EVENT_SCOPE_TEXT(*Path.GetAssetName())
+
 	if (IsCancelled())
 	{
 		return;
@@ -320,6 +334,8 @@ void FCustomizableObjectMeshStreamIn::RequestMeshUpdate(const FContext& Context)
 
 void FCustomizableObjectMeshStreamIn::CancelMeshUpdate(const FContext& Context)
 {
+	MUTABLE_CPUPROFILER_SCOPE(FCustomizableObjectMeshStreamIn::CancelMeshUpdate)
+	
 	UCustomizableObjectSystemPrivate* CustomizableObjectSystem = UCustomizableObjectSystem::GetInstance()->GetPrivate();
 	if (CustomizableObjectSystem && MutableTaskId != FMutableTaskGraph::INVALID_ID)
 	{
@@ -348,7 +364,7 @@ void FCustomizableObjectMeshStreamIn::CancelMeshUpdate(const FContext& Context)
 
 void FCustomizableObjectMeshStreamIn::ConvertMesh(const FContext& Context)
 {
-	MUTABLE_CPUPROFILER_SCOPE(ConvertMesh);
+	MUTABLE_CPUPROFILER_SCOPE(FCustomizableObjectMeshStreamIn::ConvertMesh);
 
 	check(!TaskSynchronization.GetValue());
 
diff --git a/Engine/Plugins/Mutable/Source/CustomizableObject/Private/MuCO/CustomizableObjectMipDataProvider.cpp b/Engine/Plugins/Mutable/Source/CustomizableObject/Private/MuCO/CustomizableObjectMipDataProvider.cpp
--- a/Engine/Plugins/Mutable/Source/CustomizableObject/Private/MuCO/CustomizableObjectMipDataProvider.cpp
+++ b/Engine/Plugins/Mutable/Source/CustomizableObject/Private/MuCO/CustomizableObjectMipDataProvider.cpp
@@ -330,6 +330,9 @@ int32 FMutableTextureMipDataProvider::GetMips(const FTextureUpdateContext& Conte
 {
 	MUTABLE_CPUPROFILER_SCOPE(FMutableTextureMipDataProvider::GetMips);
 
+	FSoftObjectPath Path(UpdateContext->GetCustomizableObjectPathName());
+	TRACE_CPUPROFILER_EVENT_SCOPE_TEXT(*Path.GetAssetName())
+
 	if (!UCustomizableObjectSystem::IsActive())
 	{
 		// Mutable is disabled. Skip all mip operations and mark the update task as completed.
diff --git a/Engine/Plugins/Mutable/Source/CustomizableObject/Private/MuCO/CustomizableObjectSystem.cpp b/Engine/Plugins/Mutable/Source/CustomizableObject/Private/MuCO/CustomizableObjectSystem.cpp
--- a/Engine/Plugins/Mutable/Source/CustomizableObject/Private/MuCO/CustomizableObjectSystem.cpp
+++ b/Engine/Plugins/Mutable/Source/CustomizableObject/Private/MuCO/CustomizableObjectSystem.cpp
@@ -1165,6 +1165,14 @@ void FinishUpdateGlobal(const TSharedRef<FUpdateContextPrivate>& Context)
 	UCustomizableObjectSystem* System = UCustomizableObjectSystem::GetInstance();
 	UCustomizableObjectSystemPrivate* SystemPrivate = System ? System->GetPrivate() : nullptr;
 
+	if (System &&
+		Context->UpdateStarted)
+	{
+		SystemPrivate->CurrentInstanceBeingUpdated = nullptr;
+		SystemPrivate->CurrentMutableOperation = nullptr;
+		SystemPrivate->ClearResourceCacheProtected();
+	}
+	
 	if (Instance)
 	{
 		UCustomizableInstancePrivate* PrivateInstance = Instance->GetPrivate();
@@ -1854,13 +1862,12 @@ void UCustomizableObjectSystem::SetReplaceDiscardedWithReferenceMeshEnabled(bool
 }
 
 
-void UCustomizableObjectSystem::ClearResourceCacheProtected()
+void UCustomizableObjectSystemPrivate::ClearResourceCacheProtected()
 {
 	check(IsInGameThread());
 
-	GetPrivate()->ProtectedCachedTextures.Reset(0);
-	check(GetPrivate() != nullptr);
-	GetPrivate()->ProtectedObjectCachedImages.Reset(0);
+	ProtectedCachedTextures.Reset(0);
+	ProtectedObjectCachedImages.Reset(0);
 }
 
 
@@ -2032,16 +2039,6 @@ bool UCustomizableObjectSystem::IsAutoCompilationSync() const
 
 #endif
 
-
-void UCustomizableObjectSystem::ClearCurrentMutableOperation()
-{
-	check(Private != nullptr);
-	Private->CurrentInstanceBeingUpdated = nullptr;
-	Private->CurrentMutableOperation = nullptr;
-	ClearResourceCacheProtected();
-}
-
-
 void UCustomizableObjectSystemPrivate::UpdateMemoryLimit()
 {
 	// This must run on game thread, and when the mutable thread is not running
@@ -3011,10 +3008,10 @@ namespace impl
 		FInstanceUpdateData::FImage& Image = OperationData->InstanceUpdateData.Images[ImageIndex];
 		Image.ConstantImagesNeededToGenerate = MoveTemp(ImageDesc.ConstantImagesNeededToGenerate);
 
-		const UCustomizableObjectSystemPrivate* CustomizableObjectSystemPrivateData = UCustomizableObjectSystem::GetInstanceChecked()->GetPrivate();
+		const UCustomizableObjectSystemPrivate* CustomizableObjectSystemPrivate = UCustomizableObjectSystem::GetInstanceChecked()->GetPrivate();
 		
 		{
-			const uint16 MaxTextureSizeToGenerate = static_cast<uint16>(CustomizableObjectSystemPrivateData->MaxTextureSizeToGenerate);
+			const uint16 MaxTextureSizeToGenerate = static_cast<uint16>(CustomizableObjectSystemPrivate->MaxTextureSizeToGenerate);
 			const uint16 MaxSize = FMath::Max(ImageDesc.m_size[0], ImageDesc.m_size[1]);
 
 			Image.BaseMip = 0;
@@ -3036,7 +3033,7 @@ namespace impl
 		}
 
 		const bool bCached = ImagesInThisInstance->Contains(Image.ImageID) || // See if it is cached from this same instance (can happen with LODs)
-			(UCustomizableObjectSystem::ShouldReuseTexturesBetweenInstances() && CustomizableObjectSystemPrivateData->ProtectedObjectCachedImages.Contains(Image.ImageID)); // See if it is cached from another instance
+			(UCustomizableObjectSystem::ShouldReuseTexturesBetweenInstances() && CustomizableObjectSystemPrivate->ProtectedObjectCachedImages.Contains(Image.ImageID)); // See if it is cached from another instance
 
 		if (bCached)
 		{
@@ -3060,7 +3057,7 @@ namespace impl
 
 		const int32 MipSizeX = FMath::Max(Image.FullImageSizeX >> MipsToSkip, 1);
 		const int32 MipSizeY = FMath::Max(Image.FullImageSizeY >> MipsToSkip, 1);
-		if (MipsToSkip > 0 && CustomizableObjectSystemPrivateData->EnableSkipGenerateResidentMips != 0 && OperationData->LowPriorityTextures.Find(Image.Name.ToString()) != INDEX_NONE)
+		if (MipsToSkip > 0 && CustomizableObjectSystemPrivate->EnableSkipGenerateResidentMips != 0 && OperationData->LowPriorityTextures.Find(Image.Name.ToString()) != INDEX_NONE)
 		{
 			TSharedPtr<const mu::FImage> NewImage = MakeShared<mu::FImage>(MipSizeX, MipSizeY, FullLODCount - MipsToSkip, ImageDesc.m_format, mu::EInitializationType::Black);
 
@@ -3209,7 +3206,6 @@ namespace impl
 		// TODO: Review checks.
 		if (!CustomizableObjectInstance || !CustomizableObjectInstance->IsValidLowLevel() )
 		{
-			System->ClearCurrentMutableOperation();
 			OperationData->UpdateResult = EUpdateResult::Error;
 			FinishUpdateGlobal(OperationData);
 			return;
@@ -3218,22 +3214,21 @@ namespace impl
 		UCustomizableObject* CustomizableObject = CustomizableObjectInstance->GetCustomizableObject();
 		if (!CustomizableObject)
 		{
-			System->ClearCurrentMutableOperation();
 			OperationData->UpdateResult = EUpdateResult::Error;
 			FinishUpdateGlobal(OperationData);
 			return;
 		}
 		
-		UCustomizableObjectSystemPrivate * CustomizableObjectSystemPrivateData = System->GetPrivate();
+		UCustomizableObjectSystemPrivate * CustomizableObjectSystemPrivate = System->GetPrivate();
 
 		// Actual work
 		// TODO MTBL-391: Review This hotfix
 		UpdateSkeletalMesh(OperationData);
 
 		// All work is done, release unused textures.
-		if (CustomizableObjectSystemPrivateData->bReleaseTexturesImmediately)
+		if (CustomizableObjectSystemPrivate->bReleaseTexturesImmediately)
 		{
-			FMutableResourceCache& Cache = CustomizableObjectSystemPrivateData->GetObjectCache(CustomizableObjectInstance->GetCustomizableObject());
+			FMutableResourceCache& Cache = CustomizableObjectSystemPrivate->GetObjectCache(CustomizableObjectInstance->GetCustomizableObject());
 
 			UCustomizableInstancePrivate* CustomizableObjectInstancePrivateData = CustomizableObjectInstance->GetPrivate();
 			for (FGeneratedTexture& GeneratedTexture : CustomizableObjectInstancePrivateData->TexturesToRelease)
@@ -3245,8 +3240,6 @@ namespace impl
 		}
 
 		// End Update
-		System->ClearCurrentMutableOperation();
-
 		FinishUpdateGlobal(OperationData);
 	}
 
@@ -3265,7 +3258,6 @@ namespace impl
 			FinishUpdateGlobal(OperationData);
 			return;
 		}
-
 		
 		if (CVarEnableRealTimeMorphTargets.GetValueOnAnyThread())
 		{
@@ -3317,16 +3309,16 @@ namespace impl
 			OperationData->UpdateEndRealPeakBytes = OperationData->UpdateEndPeakBytes + OperationData->UpdateStartBytes;
 		}
 		
-		UCustomizableObjectSystemPrivate* CustomizableObjectSystemPrivateData = System->GetPrivate();
+		UCustomizableObjectSystemPrivate* CustomizableObjectSystemPrivate = System->GetPrivate();
 
 		// Unregister the current CO from the external resource provider.
 		check(IsInGameThread());
-		CustomizableObjectSystemPrivateData->ResourceProvider->SetCurrentObject(nullptr);
+		CustomizableObjectSystemPrivate->ResourceProvider->SetCurrentObject(nullptr);
 
 		// Next Task: Release Mutable. We need this regardless if we cancel or not
 		//-------------------------------------------------------------		
-		const  TSharedPtr<mu::FSystem> MutableSystem = CustomizableObjectSystemPrivateData->MutableSystem;
-		CustomizableObjectSystemPrivateData->LastUpdateMutableTask = CustomizableObjectSystemPrivateData->MutableTaskGraph.AddMutableThreadTask(
+		const  TSharedPtr<mu::FSystem> MutableSystem = CustomizableObjectSystemPrivate->MutableSystem;
+		CustomizableObjectSystemPrivate->LastUpdateMutableTask = CustomizableObjectSystemPrivate->MutableTaskGraph.AddMutableThreadTask(
 			TEXT("Task_Mutable_ReleaseInstance"),
 			[InstanceID = OperationData->InstanceID, MutableSystem, bLiveUpdateMode = OperationData->bLiveUpdateMode]()
 			{
@@ -3349,7 +3341,7 @@ namespace impl
 					"Cannot move FMutableReleasePlatformOperationData::ImageToPlatformDataMap to FUpdateContextPrivate::ImageToPlatformDataMap, types do not match.");
 
 			ReleaseOperationData->ImageToPlatformDataMap = MoveTemp(OperationData->ImageToPlatformDataMap);
-			CustomizableObjectSystemPrivateData->MutableTaskGraph.AddAnyThreadTask(
+			CustomizableObjectSystemPrivate->MutableTaskGraph.AddAnyThreadTask(
 				TEXT("Mutable_ReleasePlatformData"),
 				[ReleaseOperationData]()
 				{
@@ -3364,12 +3356,12 @@ namespace impl
 			{
 				// Unlock the resource cache for the object used by this instance to avoid
 				// the destruction of resources that we may want to reuse.
-				System->ClearResourceCacheProtected();
+				System->GetPrivate()->ClearResourceCacheProtected();
 			}
 
 			// Next Task: Callbacks
 			//-------------------------------------------------------------
-			CustomizableObjectSystemPrivateData->MutableTaskGraph.AddGameThreadTask(
+			CustomizableObjectSystemPrivate->MutableTaskGraph.AddGameThreadTask(
 				TEXT("Task_Game_Callbacks"),
 				[OperationData]()
 				{
@@ -3396,21 +3388,30 @@ namespace impl
 		UCustomizableObjectSystem* System = UCustomizableObjectSystem::GetInstance();
 		if (!System)
 		{
+			OperationData->UpdateResult = EUpdateResult::Error;
+			FinishUpdateGlobal(OperationData);
 			return;
 		}
+		
+		UCustomizableObjectSystemPrivate* SystemPrivate = System->GetPrivate();
 
+		const UCustomizableObject* CustomizableObject = OperationData->Object.Get(); 
+		if (!CustomizableObject)
+		{
+			OperationData->UpdateResult = EUpdateResult::Error;
+			FinishUpdateGlobal(OperationData);
+			return;
+		}
+		
 		UCustomizableObjectInstance* ObjectInstance = OperationData->Instance.Get();
 		if (!ObjectInstance)
 		{
-			System->ClearCurrentMutableOperation();
-			
 			OperationData->UpdateResult = EUpdateResult::Error;
 			FinishUpdateGlobal(OperationData);
 			return;
 		}
 
 		UCustomizableInstancePrivate* ObjectInstancePrivateData = ObjectInstance->GetPrivate();
-		check(ObjectInstancePrivateData != nullptr);
 
 		if (OperationData->bLiveUpdateMode)
 		{
@@ -3423,52 +3424,40 @@ namespace impl
 			}
 		}
 
-		const UCustomizableObject* CustomizableObject = ObjectInstance->GetCustomizableObject(); 
-		if (!CustomizableObject)
-		{
-			System->ClearCurrentMutableOperation();
-			
-			OperationData->UpdateResult = EUpdateResult::Error;
-			FinishUpdateGlobal(OperationData);
-			return;
-		}
 		
 		if (OperationData->GetCapturedDescriptor().GetBuildParameterRelevancy())
 		{
 			// Relevancy
 			ObjectInstancePrivateData->RelevantParameters = OperationData->RelevantParametersInProgress;
 		}
-
 		
 		// Selectively lock the resource cache for the object used by this instance to avoid the destruction of resources that we may want to reuse.
 		// When protecting textures there mustn't be any left from a previous update
-		check(System->GetPrivate()->ProtectedCachedTextures.Num() == 0);
-
-		UCustomizableObjectSystemPrivate* SystemPrivateData = System->GetPrivate();
-
+		check(SystemPrivate->ProtectedCachedTextures.Num() == 0);
+		
 		// TODO: If this is the first code that runs after the CO program has finished AND if it's
 		// guaranteed that the next CO program hasn't started yet, we need to call ClearActiveObject
-		// and CancelPendingLoads on SystemPrivateData->ExtensionDataStreamer.
+		// and CancelPendingLoads on SystemPrivate->ExtensionDataStreamer.
 		//
 		// FExtensionDataStreamer->AreAnyLoadsPending should return false if the program succeeded.
 		//
 		// If the program aborted, AreAnyLoadsPending may return true, as the program doesn't cancel
 		// its own loads on exit (maybe it should?)
 
-		FMutableResourceCache& Cache = SystemPrivateData->GetObjectCache(CustomizableObject);
+		FMutableResourceCache& Cache = SystemPrivate->GetObjectCache(CustomizableObject);
 
-		System->GetPrivate()->ProtectedCachedTextures.Reset(Cache.Images.Num());
-		SystemPrivateData->ProtectedObjectCachedImages.Reset(Cache.Images.Num());
+		SystemPrivate->ProtectedCachedTextures.Reset(Cache.Images.Num());
+		SystemPrivate->ProtectedObjectCachedImages.Reset(Cache.Images.Num());
 
 		for (const FInstanceUpdateData::FImage& Image : OperationData->InstanceUpdateData.Images)
 		{
 			FMutableImageCacheKey Key(Image.ImageID, OperationData->MipsToSkip);
 			const TWeakObjectPtr<UTexture2D>* TexturePtr = Cache.Images.Find(Key);
 
-			if (TexturePtr && TexturePtr->Get() && SystemPrivateData->TextureHasReferences(Key))
+			if (TexturePtr && TexturePtr->Get() && SystemPrivate->TextureHasReferences(Key))
 			{
-				System->GetPrivate()->ProtectedCachedTextures.Add(TexturePtr->Get());
-				SystemPrivateData->ProtectedObjectCachedImages.Add(Image.ImageID);
+				SystemPrivate->ProtectedCachedTextures.Add(TexturePtr->Get());
+				SystemPrivate->ProtectedObjectCachedImages.Add(Image.ImageID);
 			}
 		}
 
@@ -3483,7 +3472,7 @@ namespace impl
 		UE::Tasks::FTask Mutable_GetImagesTask;
 		{
 			// Task inputs
-			Mutable_GetImagesTask = SystemPrivateData->MutableTaskGraph.AddMutableThreadTask(
+			Mutable_GetImagesTask = SystemPrivate->MutableTaskGraph.AddMutableThreadTask(
 				TEXT("Task_Mutable_GetImages"),
 				[OperationData]()
 				{
@@ -3498,7 +3487,7 @@ namespace impl
 
 		// Next-next Task: Convert Resources
 		//-------------------------------------------------------------
-		SystemPrivateData->MutableTaskGraph.AddGameThreadTask(
+		SystemPrivate->MutableTaskGraph.AddGameThreadTask(
 			TEXT("Task_Game_ConvertResources"),
 			[OperationData]()
 			{
@@ -3515,15 +3504,15 @@ namespace impl
 		MUTABLE_CPUPROFILER_SCOPE(Task_Game_ReleaseInstanceID)
 
 		UCustomizableObjectSystem* System = UCustomizableObjectSystem::GetInstanceChecked();
-		UCustomizableObjectSystemPrivate* SystemPrivateData = System->GetPrivate();
+		UCustomizableObjectSystemPrivate* SystemPrivate = System->GetPrivate();
 
-		const TSharedPtr<mu::FSystem> MutableSystem = SystemPrivateData->MutableSystem;
+		const TSharedPtr<mu::FSystem> MutableSystem = SystemPrivate->MutableSystem;
 
 		// Task: Release Instance ID
 		//-------------------------------------------------------------
 		{
 			// Task inputs
-			SystemPrivateData->MutableTaskGraph.AddMutableThreadTask(
+			SystemPrivate->MutableTaskGraph.AddMutableThreadTask(
 				TEXT("Task_Mutable_ReleaseInstanceID"),
 				[IDToRelease, MutableSystem]()
 				{
@@ -3540,13 +3529,8 @@ namespace impl
 
 		check(IsInGameThread());
 
-#if WITH_EDITOR
-if (CVarMutableLogObjectMemoryOnUpdate.GetValueOnAnyThread())
-{
-		UCustomizableObject* CO = Operation->Instance->GetCustomizableObject();
-		CO->GetPrivate()->LogMemory();
-}
-#endif
+		Operation->UpdateStarted = true;
+		TRACE_BEGIN_REGION(UE_MUTABLE_UPDATE_REGION)
 		
 		// Check if a level has been loaded
 		if (FLogBenchmarkUtil::IsBenchmarkingReportingEnabled() && GWorld)
@@ -3555,102 +3539,101 @@ if (CVarMutableLogObjectMemoryOnUpdate.GetValueOnAnyThread())
 		}
 
 		Operation->StartUpdateTime = FPlatformTime::Seconds();
-
-		Operation->bLowPriorityTasksBlocked = true;
-		UCustomizableObjectSystem::GetInstance()->GetPrivate()->MutableTaskGraph.AllowLaunchingMutableTaskLowPriority(false, false);
 		
 		UCustomizableObjectSystem* System = UCustomizableObjectSystem::GetInstance();
-		check(System != nullptr);
-
-		if (!Operation->Instance.IsValid() || !Operation->Instance->IsValidLowLevel()) // Only start if it hasn't been already destroyed (i.e. GC after finish PIE)
+		if (!System)
 		{
-			System->ClearCurrentMutableOperation();
+			Operation->UpdateResult = EUpdateResult::Error;
+			FinishUpdateGlobal(Operation);
+			return;
+		}
 
+		UCustomizableObjectSystemPrivate* SystemPrivate = System->GetPrivate();
+		
+		UCustomizableObject* Object = Operation->Object.Get();
+		if (!Object)
+		{
 			Operation->UpdateResult = EUpdateResult::Error;
 			FinishUpdateGlobal(Operation);
 			return;
 		}
 
-		const UCustomizableObjectSystemPrivate* SystemPrivate = System->GetPrivate();
+		UCustomizableObjectPrivate* ObjectPrivate = Object->GetPrivate();
 
-		UCustomizableObjectInstance* CandidateInstance = Operation->Instance.Get();
-		
-		UCustomizableInstancePrivate* CandidateInstancePrivateData = CandidateInstance->GetPrivate();
-		if (!CandidateInstancePrivateData)
+		TRACE_CPUPROFILER_EVENT_SCOPE_TEXT(*Object->GetName()); // So that we know what we are currently updating.
+
+		UCustomizableObjectInstance* Instance = Operation->Instance.Get();
+		if (!Instance) // Only start if it hasn't been already destroyed (i.e. GC after finish PIE)
 		{
-			System->ClearCurrentMutableOperation();
-			
 			Operation->UpdateResult = EUpdateResult::Error;
 			FinishUpdateGlobal(Operation);
 			return;
 		}
+		
+		UCustomizableInstancePrivate* InstancePrivate = Instance->GetPrivate();
 
-		if (CandidateInstancePrivateData->HasCOInstanceFlags(PendingLODsUpdate))
+		if (InstancePrivate->HasCOInstanceFlags(PendingLODsUpdate))
 		{
-			CandidateInstancePrivateData->ClearCOInstanceFlags(PendingLODsUpdate);
+			InstancePrivate->ClearCOInstanceFlags(PendingLODsUpdate);
 			// TODO: Is anything needed for this now?
 			//Operation->CustomizableObjectInstance->ReleaseMutableInstanceId(); // To make mutable regenerate the LODs even if the instance parameters have not changed
 		}
 
 		// Skip update, the requested update is equal to the running update.
-		if (Operation->GetCapturedDescriptorHash().IsSubset(CandidateInstancePrivateData->CommittedDescriptorHash))
+		if (Operation->GetCapturedDescriptorHash().IsSubset(InstancePrivate->CommittedDescriptorHash))
 		{
-			System->ClearCurrentMutableOperation();
-
 			Operation->UpdateResult = EUpdateResult::Success;
 			UpdateSkeletalMesh(Operation);
 			FinishUpdateGlobal(Operation);
 			return;
 		}
-
-		bool bCancel = false;
-
-		TObjectPtr<UCustomizableObject> CustomizableObject = CandidateInstance->GetCustomizableObject();
-
+		
 		// If the object is locked (for instance, compiling) we skip any instance update.
-		if (!CustomizableObject || CustomizableObject->GetPrivate()->bLocked)
+		if (ObjectPrivate->bLocked)
 		{
-			bCancel = true;
+			Operation->UpdateResult = EUpdateResult::Error;
+			FinishUpdateGlobal(Operation);
+			return;
 		}
 
 		// Only update resources if the instance is in range (it could have got far from the player since the task was queued)
-		check(System->GetPrivate()->CurrentInstanceLODManagement != nullptr);
-		if (System->GetPrivate()->CurrentInstanceLODManagement->IsOnlyUpdateCloseCustomizableObjectsEnabled()
-			&& CandidateInstancePrivateData
-			&& CandidateInstancePrivateData->LastMinSquareDistFromComponentToPlayer > FMath::Square(System->GetPrivate()->CurrentInstanceLODManagement->GetOnlyUpdateCloseCustomizableObjectsDist())
-			&& CandidateInstancePrivateData->LastMinSquareDistFromComponentToPlayer != FLT_MAX // This means it is the first frame so it has to be updated
+		check(SystemPrivate->CurrentInstanceLODManagement != nullptr);
+		if (SystemPrivate->CurrentInstanceLODManagement->IsOnlyUpdateCloseCustomizableObjectsEnabled()
+			&& InstancePrivate
+			&& InstancePrivate->LastMinSquareDistFromComponentToPlayer > FMath::Square(SystemPrivate->CurrentInstanceLODManagement->GetOnlyUpdateCloseCustomizableObjectsDist())
+			&& InstancePrivate->LastMinSquareDistFromComponentToPlayer != FLT_MAX // This means it is the first frame so it has to be updated
 		   )
 		{
-			bCancel = true;
+			Operation->UpdateResult = EUpdateResult::Error;
+			FinishUpdateGlobal(Operation);
+			return;
 		}
 
 		Operation->Parameters = Operation->GetCapturedDescriptor().GetParameters();
 		if (!Operation->Parameters)
 		{
-			bCancel = true;
-		}
-
-		if (bCancel)
-		{
-			System->ClearCurrentMutableOperation();
-
 			Operation->UpdateResult = EUpdateResult::Error;
 			FinishUpdateGlobal(Operation);
 			return;
 		}
 
-		UCustomizableObjectSystemPrivate* SystemPrivateData = System->GetPrivate();
-
-		SystemPrivateData->CurrentInstanceBeingUpdated = CandidateInstance;
+#if WITH_EDITOR
+		if (CVarMutableLogObjectMemoryOnUpdate.GetValueOnAnyThread())
+		{
+			ObjectPrivate->LogMemory();
+		}
+#endif
+		
+		SystemPrivate->CurrentInstanceBeingUpdated = Instance;
 
-		const UModelResources& ModelResources = CandidateInstance->GetCustomizableObject()->GetPrivate()->GetModelResourcesChecked();
+		const UModelResources& ModelResources = Instance->GetCustomizableObject()->GetPrivate()->GetModelResourcesChecked();
 
-		FString StateName = CandidateInstance->GetCustomizableObject()->GetPrivate()->GetStateName(CandidateInstance->GetPrivate()->GetState());
+		const FString& StateName = ObjectPrivate->GetStateName(InstancePrivate->GetState());
 		const FMutableStateData* StateData = ModelResources.StateUIDataMap.Find(StateName);
 
 		Operation->bLiveUpdateMode = false;
 
-		if (SystemPrivateData->EnableMutableLiveUpdate)
+		if (SystemPrivate->EnableMutableLiveUpdate)
 		{
 			Operation->bLiveUpdateMode = StateData ? StateData->bLiveUpdateMode : false;
 		}
@@ -3658,7 +3641,7 @@ if (CVarMutableLogObjectMemoryOnUpdate.GetValueOnAnyThread())
 		Operation->bNeverStream = false;
 		Operation->MipsToSkip = 0;
 
-		SystemPrivateData->GetMipStreamingConfig(*CandidateInstance, Operation->bNeverStream, Operation->MipsToSkip);
+		SystemPrivate->GetMipStreamingConfig(*Instance, Operation->bNeverStream, Operation->MipsToSkip);
 		
 		if (Operation->bLiveUpdateMode && (!Operation->bNeverStream || Operation->MipsToSkip > 0))
 		{
@@ -3668,10 +3651,10 @@ if (CVarMutableLogObjectMemoryOnUpdate.GetValueOnAnyThread())
 
 		Operation->bReuseInstanceTextures = false;
 
-		if (SystemPrivateData->EnableReuseInstanceTextures)
+		if (SystemPrivate->EnableReuseInstanceTextures)
 		{
 			Operation->bReuseInstanceTextures = StateData ? StateData->bReuseInstanceTextures : false;
-			Operation->bReuseInstanceTextures |= CandidateInstancePrivateData->HasCOInstanceFlags(ReuseTextures);
+			Operation->bReuseInstanceTextures |= InstancePrivate->HasCOInstanceFlags(ReuseTextures);
 			
 			if (Operation->bReuseInstanceTextures && !Operation->bNeverStream)
 			{
@@ -3680,22 +3663,21 @@ if (CVarMutableLogObjectMemoryOnUpdate.GetValueOnAnyThread())
 			}
 		}
 
-		if (!Operation->bLiveUpdateMode && CandidateInstancePrivateData->LiveUpdateModeInstanceID != 0)
+		if (!Operation->bLiveUpdateMode && InstancePrivate->LiveUpdateModeInstanceID != 0)
 		{
 			// The instance was in live update mode last update, but now it's not. So the Id and resources have to be released.
 			// Enqueue a new mutable task to release them
-			Task_Game_ReleaseInstanceID(CandidateInstancePrivateData->LiveUpdateModeInstanceID);
-			CandidateInstancePrivateData->LiveUpdateModeInstanceID = 0;
+			Task_Game_ReleaseInstanceID(InstancePrivate->LiveUpdateModeInstanceID);
+			InstancePrivate->LiveUpdateModeInstanceID = 0;
 		}
-
-
-		Operation->Model = CustomizableObject->GetPrivate()->GetModel().ToSharedRef();
+		
+		Operation->Model = ObjectPrivate->GetModel().ToSharedRef();
 
 		// Register the current CO into the external resource provider for runtime meshes and images.
 		// It is needed to get access to the bone map when converting meshes at runtime, for example.
-		if (SystemPrivateData->ResourceProvider)
+		if (SystemPrivate->ResourceProvider)
 		{
-			SystemPrivateData->ResourceProvider->SetCurrentObject(CustomizableObject);
+			SystemPrivate->ResourceProvider->SetCurrentObject(Object);
 		}
 
 		UE::Tasks::FTask CacheRuntimeTexturesEvent = UE::Tasks::MakeCompletedTask<void>();
@@ -3715,7 +3697,7 @@ if (CVarMutableLogObjectMemoryOnUpdate.GetValueOnAnyThread())
 				Textures.Add(Texture.ToSoftObjectPath());
 			}
 		
-			SystemPrivateData->StreamableManager->RequestAsyncLoad(Textures, FStreamableDelegate::CreateLambda([Operation, Event]() mutable
+			SystemPrivate->StreamableManager->RequestAsyncLoad(Textures, FStreamableDelegate::CreateLambda([Operation, Event]() mutable
 			{
 				UCustomizableObjectSystem* System = UCustomizableObjectSystem::GetInstance();
 				if (!System)
@@ -3724,14 +3706,14 @@ if (CVarMutableLogObjectMemoryOnUpdate.GetValueOnAnyThread())
 					return;
 				}
 				
-				UCustomizableObject* CustomizableObject = Operation->GetCapturedDescriptor().GetCustomizableObject();
-				if (!CustomizableObject)
+				UCustomizableObject* Object = Operation->GetCapturedDescriptor().GetCustomizableObject();
+				if (!Object)
 				{
 					Event.Trigger();
 					return;
 				}
 			
-				const UModelResources& ModelResources = CustomizableObject->GetPrivate()->GetModelResourcesChecked();
+				const UModelResources& ModelResources = Object->GetPrivate()->GetModelResourcesChecked();
 				System->GetPrivate()->GetResourceProviderChecked()->CacheRuntimeReferencedImages(Operation->Model.ToSharedRef(), ModelResources.RuntimeReferencedTextures);
 				Event.Trigger();
 			}));
@@ -3740,22 +3722,22 @@ if (CVarMutableLogObjectMemoryOnUpdate.GetValueOnAnyThread())
 		
 		// Task: Mutable Update and GetMesh
 		//-------------------------------------------------------------
-		Operation->InstanceID = Operation->bLiveUpdateMode ? CandidateInstancePrivateData->LiveUpdateModeInstanceID : 0;
-		Operation->bUseMeshCache = CustomizableObject->bEnableMeshCache && !Operation->bLiveUpdateMode && UCustomizableObjectSystem::IsMeshCacheEnabled(true);
+		Operation->InstanceID = Operation->bLiveUpdateMode ? InstancePrivate->LiveUpdateModeInstanceID : 0;
+		Operation->bUseMeshCache = Object->bEnableMeshCache && !Operation->bLiveUpdateMode && UCustomizableObjectSystem::IsMeshCacheEnabled(true);
 
-		Operation->bStreamMeshLODs = IsStreamingEnabled(*CustomizableObject);
+		Operation->bStreamMeshLODs = IsStreamingEnabled(*Object);
 
 #if WITH_EDITOR
-		Operation->PixelFormatOverride = SystemPrivateData->ImageFormatOverrideFunc;
+		Operation->PixelFormatOverride = SystemPrivate->ImageFormatOverrideFunc;
 #endif
 
-		if (!CandidateInstancePrivateData->HasCOInstanceFlags(ForceGenerateMipTail))
+		if (!InstancePrivate->HasCOInstanceFlags(ForceGenerateMipTail))
 		{
-			CustomizableObject->GetPrivate()->GetLowPriorityTextureNames(Operation->LowPriorityTextures);
+			ObjectPrivate->GetLowPriorityTextureNames(Operation->LowPriorityTextures);
 		}
 
 		bool bRequestAllLODs = !System->IsOnlyGenerateRequestedLODsEnabled() ||
-			!System->GetPrivate()->CurrentInstanceLODManagement->IsOnlyGenerateRequestedLODLevelsEnabled();
+			!SystemPrivate->CurrentInstanceLODManagement->IsOnlyGenerateRequestedLODLevelsEnabled();
 
 #if WITH_EDITOR
 		// In the editor LOD Management is disabled by default. Overwrite requested LODs when disabled.
@@ -3769,7 +3751,7 @@ if (CVarMutableLogObjectMemoryOnUpdate.GetValueOnAnyThread())
 			}
 
 			if (IsValid(*CustomizableObjectInstanceUsage) &&
-				CustomizableObjectInstanceUsage->GetCustomizableObjectInstance() == CandidateInstance)
+				CustomizableObjectInstanceUsage->GetCustomizableObjectInstance() == Instance)
 			{
 				EWorldType::Type WorldType = EWorldType::Type::None;
 
@@ -3817,6 +3799,9 @@ if (CVarMutableLogObjectMemoryOnUpdate.GetValueOnAnyThread())
 			check(SystemPrivate->Streamer != nullptr);
 			SystemPrivate->Streamer->PrepareStreamingForObject(Operation->Instance->GetCustomizableObject());			
 
+			Operation->bLowPriorityTasksBlocked = true;
+			SystemPrivate->MutableTaskGraph.AllowLaunchingMutableTaskLowPriority(false, false);
+
 			const TSharedPtr<mu::FSystem> MutableSystem = SystemPrivate->MutableSystem;
 			
 			if (Operation->bLiveUpdateMode)
@@ -3879,14 +3864,14 @@ if (CVarMutableLogObjectMemoryOnUpdate.GetValueOnAnyThread())
 			
 			for (const TArray<mu::FResourceID>& MeshId : Operation->GetMeshDescriptors())
 			{
-				if (USkeletalMesh* CachedMesh = CustomizableObject->GetPrivate()->MeshCache.Get(MeshId))
+				if (USkeletalMesh* CachedMesh = ObjectPrivate->MeshCache.Get(MeshId))
 				{
 					Operation->Objects.Emplace(CachedMesh);
 				}
 			}
 		}
 		
-		UE::Tasks::FTask Mutable_GetMeshTask = SystemPrivateData->MutableTaskGraph.AddMutableThreadTask(
+		UE::Tasks::FTask Mutable_GetMeshTask = SystemPrivate->MutableTaskGraph.AddMutableThreadTask(
 		TEXT("Task_Mutable_GetMeshes"),
 		[Operation]()
 		{
@@ -3894,7 +3879,7 @@ if (CVarMutableLogObjectMemoryOnUpdate.GetValueOnAnyThread())
 		},
 		{ CacheRuntimeTexturesEvent });
 
-		SystemPrivateData->MutableTaskGraph.AddGameThreadTask(
+		SystemPrivate->MutableTaskGraph.AddGameThreadTask(
 			TEXT("Task_Game_LockCache"),
 			[Operation]()
 			{
@@ -4751,9 +4736,6 @@ void UCustomizableObjectSystemPrivate::OnMutableEnabledChanged(IConsoleVariable*
 
 void UCustomizableObjectSystemPrivate::StartUpdateSkeletalMesh(const TSharedRef<FUpdateContextPrivate>& Context)
 {
-	Context->UpdateStarted = true;
-	TRACE_BEGIN_REGION(UE_MUTABLE_UPDATE_REGION);
-
 	check(!CurrentMutableOperation); // Can not start an update if there is already another in progress
 	check(Context->bValid);
 	check(Context->Instance.IsValid()) // The COI has to be alive to start the update
diff --git a/Engine/Plugins/Mutable/Source/CustomizableObject/Public/MuCO/CustomizableObjectSystem.h b/Engine/Plugins/Mutable/Source/CustomizableObject/Public/MuCO/CustomizableObjectSystem.h
--- a/Engine/Plugins/Mutable/Source/CustomizableObject/Public/MuCO/CustomizableObjectSystem.h
+++ b/Engine/Plugins/Mutable/Source/CustomizableObject/Public/MuCO/CustomizableObjectSystem.h
@@ -240,10 +240,6 @@ public:
 	bool IsReplaceDiscardedWithReferenceMeshEnabled() const;
 	void SetReplaceDiscardedWithReferenceMeshEnabled(bool bIsEnabled);
 
-	// Unprotect the resources in the instances of this object of being garbage collector
-	// while an instance is being built or updated, so that they can be reused.
-	void ClearResourceCacheProtected();
-
 #if WITH_EDITOR
 	// Lock a CustomizableObjects, preventing the generation or update of any of its instances
 	// Will return true if successful, false if it fails to lock because an update is already underway
@@ -359,8 +355,6 @@ public:
 	uint64 GetMaxChunkSizeForPlatform(const ITargetPlatform* InTargetPlatform);
 #endif
 
-	void ClearCurrentMutableOperation();
-
 private:
 	// Most of the work in this plugin happens here.
 	bool Tick(float DeltaTime);
