[Mutable] Fix static analysis warning.

diff --git a/Engine/Plugins/Mutable/Source/CustomizableObject/Private/MuCO/CustomizableObjectSystem.cpp b/Engine/Plugins/Mutable/Source/CustomizableObject/Private/MuCO/CustomizableObjectSystem.cpp
--- a/Engine/Plugins/Mutable/Source/CustomizableObject/Private/MuCO/CustomizableObjectSystem.cpp
+++ b/Engine/Plugins/Mutable/Source/CustomizableObject/Private/MuCO/CustomizableObjectSystem.cpp
@@ -3571,7 +3571,8 @@ namespace impl
 		}
 		
 		UCustomizableInstancePrivate* InstancePrivate = Instance->GetPrivate();
-
+		check(InstancePrivate); // Already checked in GetPrivate. Make static analyzer happy.
+		
 		if (InstancePrivate->HasCOInstanceFlags(PendingLODsUpdate))
 		{
 			InstancePrivate->ClearCOInstanceFlags(PendingLODsUpdate);
