[mutable] Some optimizations around the mesh-prepare-layout operation.
- some optimization around the mesh-add-tags operation.
- cleanup some old naming convention.
#rnx

diff --git a/Engine/Plugins/Mutable/Source/MutableTools/Private/MuT/ASTOpImageCrop.cpp b/Engine/Plugins/Mutable/Source/MutableTools/Private/MuT/ASTOpImageCrop.cpp
--- a/Engine/Plugins/Mutable/Source/MutableTools/Private/MuT/ASTOpImageCrop.cpp
+++ b/Engine/Plugins/Mutable/Source/MutableTools/Private/MuT/ASTOpImageCrop.cpp
@@ -280,10 +280,10 @@ namespace mu
 		{
 			// We move down all the paths
 			Ptr<ASTOpSwitch> newOp = mu::Clone<ASTOpSwitch>(at);
-			newOp->def = Visit(newOp->def.child(), currentCropOp);
-			for (auto& c : newOp->cases)
+			newOp->Default = Visit(newOp->Default.child(), currentCropOp);
+			for (auto& c : newOp->Cases)
 			{
-				c.branch = Visit(c.branch.child(), currentCropOp);
+				c.Branch = Visit(c.Branch.child(), currentCropOp);
 			}
 			newAt = newOp;
 			break;
diff --git a/Engine/Plugins/Mutable/Source/MutableTools/Private/MuT/ASTOpImageDisplace.cpp b/Engine/Plugins/Mutable/Source/MutableTools/Private/MuT/ASTOpImageDisplace.cpp
--- a/Engine/Plugins/Mutable/Source/MutableTools/Private/MuT/ASTOpImageDisplace.cpp
+++ b/Engine/Plugins/Mutable/Source/MutableTools/Private/MuT/ASTOpImageDisplace.cpp
@@ -208,23 +208,23 @@ namespace mu
 					// Move the format down all the paths
 					auto nop = mu::Clone<ASTOpSwitch>(sourceAt);
 
-					if (nop->def)
+					if (nop->Default)
 					{
 						Ptr<ASTOpImageDisplace> defOp = mu::Clone<ASTOpImageDisplace>(this);
-						defOp->Source = typedSource->def.child();
-						defOp->DisplacementMap = typedDisplacementMap->def.child();
-						nop->def = defOp;
+						defOp->Source = typedSource->Default.child();
+						defOp->DisplacementMap = typedDisplacementMap->Default.child();
+						nop->Default = defOp;
 					}
 
 					// We need to copy the options because we change them
-					for (size_t v = 0; v < nop->cases.Num(); ++v)
+					for (size_t v = 0; v < nop->Cases.Num(); ++v)
 					{
-						if (nop->cases[v].branch)
+						if (nop->Cases[v].Branch)
 						{
 							Ptr<ASTOpImageDisplace> bOp = mu::Clone<ASTOpImageDisplace>(this);
-							bOp->Source = typedSource->cases[v].branch.child();
-							bOp->DisplacementMap = typedDisplacementMap->FindBranch(typedSource->cases[v].condition);
-							nop->cases[v].branch = bOp;
+							bOp->Source = typedSource->Cases[v].Branch.child();
+							bOp->DisplacementMap = typedDisplacementMap->FindBranch(typedSource->Cases[v].Condition);
+							nop->Cases[v].Branch = bOp;
 						}
 					}
 
@@ -238,21 +238,21 @@ namespace mu
 				// Make a project for every path
 				Ptr<ASTOpSwitch> nop = mu::Clone<ASTOpSwitch>(sourceAt.get());
 
-				if (nop->def)
+				if (nop->Default)
 				{
 					Ptr<ASTOpImageDisplace> defOp = mu::Clone<ASTOpImageDisplace>(this);
-					defOp->Source = nop->def.child();
-					nop->def = defOp;
+					defOp->Source = nop->Default.child();
+					nop->Default = defOp;
 				}
 
 				// We need to copy the options because we change them
-				for (size_t o = 0; o < nop->cases.Num(); ++o)
+				for (size_t o = 0; o < nop->Cases.Num(); ++o)
 				{
-					if (nop->cases[o].branch)
+					if (nop->Cases[o].Branch)
 					{
 						Ptr<ASTOpImageDisplace> bOp = mu::Clone<ASTOpImageDisplace>(this);
-						bOp->Source = nop->cases[o].branch.child();
-						nop->cases[o].branch = bOp;
+						bOp->Source = nop->Cases[o].Branch.child();
+						nop->Cases[o].Branch = bOp;
 					}
 				}
 
diff --git a/Engine/Plugins/Mutable/Source/MutableTools/Private/MuT/ASTOpImageLayer.cpp b/Engine/Plugins/Mutable/Source/MutableTools/Private/MuT/ASTOpImageLayer.cpp
--- a/Engine/Plugins/Mutable/Source/MutableTools/Private/MuT/ASTOpImageLayer.cpp
+++ b/Engine/Plugins/Mutable/Source/MutableTools/Private/MuT/ASTOpImageLayer.cpp
@@ -548,12 +548,12 @@ namespace mu
 
 				// If at least a switch option is a plain colour, sink the layer into the switch
 				bool bWorthSinking = false;
-				for (int32 v = 0; v < BlendSwitch->cases.Num(); ++v)
+				for (int32 v = 0; v < BlendSwitch->Cases.Num(); ++v)
 				{
-					if (BlendSwitch->cases[v].branch)
+					if (BlendSwitch->Cases[v].Branch)
 					{
 						// \TODO: Use the smarter query function to detect plain images?
-						if (BlendSwitch->cases[v].branch->GetOpType() == EOpType::IM_PLAINCOLOUR)
+						if (BlendSwitch->Cases[v].Branch->GetOpType() == EOpType::IM_PLAINCOLOUR)
 						{
 							bWorthSinking = true;
 							break;
@@ -573,28 +573,28 @@ namespace mu
 
 					Ptr<ASTOpSwitch> NewSwitch = mu::Clone<ASTOpSwitch>(BlendSwitch);
 
-					if (NewSwitch->def)
+					if (NewSwitch->Default)
 					{
 						Ptr<ASTOpImageLayer> defOp = mu::Clone<ASTOpImageLayer>(this);
-						defOp->blend = BlendSwitch->def.child();
+						defOp->blend = BlendSwitch->Default.child();
 						if (bMaskIsCompatibleSwitch)
 						{
-							defOp->mask = MaskSwitch->def.child();
+							defOp->mask = MaskSwitch->Default.child();
 						}
-						NewSwitch->def = defOp;
+						NewSwitch->Default = defOp;
 					}
 
-					for (int32 v = 0; v < NewSwitch->cases.Num(); ++v)
+					for (int32 v = 0; v < NewSwitch->Cases.Num(); ++v)
 					{
-						if (NewSwitch->cases[v].branch)
+						if (NewSwitch->Cases[v].Branch)
 						{
 							Ptr<ASTOpImageLayer> BranchOp = mu::Clone<ASTOpImageLayer>(this);
-							BranchOp->blend = BlendSwitch->cases[v].branch.child();
+							BranchOp->blend = BlendSwitch->Cases[v].Branch.child();
 							if (bMaskIsCompatibleSwitch)
 							{
-								BranchOp->mask = MaskSwitch->cases[v].branch.child();
+								BranchOp->mask = MaskSwitch->Cases[v].Branch.child();
 							}
-							NewSwitch->cases[v].branch = BranchOp;
+							NewSwitch->Cases[v].Branch = BranchOp;
 						}
 					}
 
@@ -623,12 +623,12 @@ namespace mu
 
 				// If at least a switch option is a plain colour, sink the layer into the switch
 				bool bWorthSinking = false;
-				for (int32 v = 0; v < MaskSwitch->cases.Num(); ++v)
+				for (int32 v = 0; v < MaskSwitch->Cases.Num(); ++v)
 				{
-					if (MaskSwitch->cases[v].branch)
+					if (MaskSwitch->Cases[v].Branch)
 					{
 						// \TODO: Use the smarter query function to detect plain images?
-						if (MaskSwitch->cases[v].branch->GetOpType() == EOpType::IM_PLAINCOLOUR)
+						if (MaskSwitch->Cases[v].Branch->GetOpType() == EOpType::IM_PLAINCOLOUR)
 						{
 							bWorthSinking = true;
 							break;
@@ -640,20 +640,20 @@ namespace mu
 				{
 					Ptr<ASTOpSwitch> NewSwitch = mu::Clone<ASTOpSwitch>(MaskSwitch);
 
-					if (NewSwitch->def)
+					if (NewSwitch->Default)
 					{
 						Ptr<ASTOpImageLayer> defOp = mu::Clone<ASTOpImageLayer>(this);
-						defOp->mask = MaskSwitch->def.child();
-						NewSwitch->def = defOp;
+						defOp->mask = MaskSwitch->Default.child();
+						NewSwitch->Default = defOp;
 					}
 
-					for (int32 v = 0; v < NewSwitch->cases.Num(); ++v)
+					for (int32 v = 0; v < NewSwitch->Cases.Num(); ++v)
 					{
-						if (NewSwitch->cases[v].branch)
+						if (NewSwitch->Cases[v].Branch)
 						{
 							Ptr<ASTOpImageLayer> BranchOp = mu::Clone<ASTOpImageLayer>(this);
-							BranchOp->mask = MaskSwitch->cases[v].branch.child();
-							NewSwitch->cases[v].branch = BranchOp;
+							BranchOp->mask = MaskSwitch->Cases[v].Branch.child();
+							NewSwitch->Cases[v].Branch = BranchOp;
 						}
 					}
 
diff --git a/Engine/Plugins/Mutable/Source/MutableTools/Private/MuT/ASTOpImageLayerColor.cpp b/Engine/Plugins/Mutable/Source/MutableTools/Private/MuT/ASTOpImageLayerColor.cpp
--- a/Engine/Plugins/Mutable/Source/MutableTools/Private/MuT/ASTOpImageLayerColor.cpp
+++ b/Engine/Plugins/Mutable/Source/MutableTools/Private/MuT/ASTOpImageLayerColor.cpp
@@ -370,20 +370,20 @@ namespace mu
 			// Move the layer operation down base paths
 			Ptr<ASTOpSwitch> nop = mu::Clone<ASTOpSwitch>(baseSwitch);
 
-			if (nop->def)
+			if (nop->Default)
 			{
 				Ptr<ASTOpImageLayerColor> defOp = mu::Clone<ASTOpImageLayerColor>(this);
-				defOp->base = nop->def.child();
-				nop->def = defOp;
+				defOp->base = nop->Default.child();
+				nop->Default = defOp;
 			}
 
-			for (int32 v = 0; v < nop->cases.Num(); ++v)
+			for (int32 v = 0; v < nop->Cases.Num(); ++v)
 			{
-				if (nop->cases[v].branch)
+				if (nop->Cases[v].Branch)
 				{
 					Ptr<ASTOpImageLayerColor> bOp = mu::Clone<ASTOpImageLayerColor>(this);
-					bOp->base = nop->cases[v].branch.child();
-					nop->cases[v].branch = bOp;
+					bOp->base = nop->Cases[v].Branch.child();
+					nop->Cases[v].Branch = bOp;
 				}
 			}
 
diff --git a/Engine/Plugins/Mutable/Source/MutableTools/Private/MuT/ASTOpImageMakeGrowMap.cpp b/Engine/Plugins/Mutable/Source/MutableTools/Private/MuT/ASTOpImageMakeGrowMap.cpp
--- a/Engine/Plugins/Mutable/Source/MutableTools/Private/MuT/ASTOpImageMakeGrowMap.cpp
+++ b/Engine/Plugins/Mutable/Source/MutableTools/Private/MuT/ASTOpImageMakeGrowMap.cpp
@@ -118,21 +118,21 @@ namespace mu
 			// Move the format down all the paths
 			Ptr<ASTOpSwitch> nop = mu::Clone<ASTOpSwitch>(Mask.child());
 
-			if (nop->def)
+			if (nop->Default)
 			{
 				Ptr<ASTOpImageMakeGrowMap> defOp = mu::Clone<ASTOpImageMakeGrowMap>(this);
-				defOp->Mask = nop->def.child();
-				nop->def = defOp;
+				defOp->Mask = nop->Default.child();
+				nop->Default = defOp;
 			}
 
 			// We need to copy the options because we change them
-			for (size_t v = 0; v < nop->cases.Num(); ++v)
+			for (size_t v = 0; v < nop->Cases.Num(); ++v)
 			{
-				if (nop->cases[v].branch)
+				if (nop->Cases[v].Branch)
 				{
 					Ptr<ASTOpImageMakeGrowMap> bOp = mu::Clone<ASTOpImageMakeGrowMap>(this);
-					bOp->Mask = nop->cases[v].branch.child();
-					nop->cases[v].branch = bOp;
+					bOp->Mask = nop->Cases[v].Branch.child();
+					nop->Cases[v].Branch = bOp;
 				}
 			}
 
diff --git a/Engine/Plugins/Mutable/Source/MutableTools/Private/MuT/ASTOpImageMipmap.cpp b/Engine/Plugins/Mutable/Source/MutableTools/Private/MuT/ASTOpImageMipmap.cpp
--- a/Engine/Plugins/Mutable/Source/MutableTools/Private/MuT/ASTOpImageMipmap.cpp
+++ b/Engine/Plugins/Mutable/Source/MutableTools/Private/MuT/ASTOpImageMipmap.cpp
@@ -370,10 +370,10 @@ namespace mu
 		{
 			// We move the op down all the paths
 			Ptr<ASTOpSwitch> newOp = mu::Clone<ASTOpSwitch>(at);
-			newOp->def = Visit(newOp->def.child(), currentMipmapOp);
-			for (ASTOpSwitch::FCase& c : newOp->cases)
+			newOp->Default = Visit(newOp->Default.child(), currentMipmapOp);
+			for (ASTOpSwitch::FCase& c : newOp->Cases)
 			{
-				c.branch = Visit(c.branch.child(), currentMipmapOp);
+				c.Branch = Visit(c.Branch.child(), currentMipmapOp);
 			}
 			newAt = newOp;
 			break;
diff --git a/Engine/Plugins/Mutable/Source/MutableTools/Private/MuT/ASTOpImagePixelFormat.cpp b/Engine/Plugins/Mutable/Source/MutableTools/Private/MuT/ASTOpImagePixelFormat.cpp
--- a/Engine/Plugins/Mutable/Source/MutableTools/Private/MuT/ASTOpImagePixelFormat.cpp
+++ b/Engine/Plugins/Mutable/Source/MutableTools/Private/MuT/ASTOpImagePixelFormat.cpp
@@ -401,10 +401,10 @@ namespace mu
 		{
 			// We move the op down all the paths
 			Ptr<ASTOpSwitch> newOp = mu::Clone<ASTOpSwitch>(at);
-			newOp->def = Visit(newOp->def.child(), currentFormatOp);
-			for (auto& c : newOp->cases)
+			newOp->Default = Visit(newOp->Default.child(), currentFormatOp);
+			for (ASTOpSwitch::FCase& c : newOp->Cases)
 			{
-				c.branch = Visit(c.branch.child(), currentFormatOp);
+				c.Branch = Visit(c.Branch.child(), currentFormatOp);
 			}
 			newAt = newOp;
 			break;
diff --git a/Engine/Plugins/Mutable/Source/MutableTools/Private/MuT/ASTOpImageRasterMesh.cpp b/Engine/Plugins/Mutable/Source/MutableTools/Private/MuT/ASTOpImageRasterMesh.cpp
--- a/Engine/Plugins/Mutable/Source/MutableTools/Private/MuT/ASTOpImageRasterMesh.cpp
+++ b/Engine/Plugins/Mutable/Source/MutableTools/Private/MuT/ASTOpImageRasterMesh.cpp
@@ -302,23 +302,23 @@ namespace mu
 		{
 			// Make an image for every path
 			auto nop = mu::Clone<ASTOpSwitch>(sourceAt.get());
-			nop->type = EOpType::IM_SWITCH;
+			nop->Type = EOpType::IM_SWITCH;
 
-			if (nop->def)
+			if (nop->Default)
 			{
 				Ptr<ASTOpImageRasterMesh> defOp = mu::Clone<ASTOpImageRasterMesh>(this);
-				defOp->mesh = nop->def.child();
-				nop->def = defOp;
+				defOp->mesh = nop->Default.child();
+				nop->Default = defOp;
 			}
 
 			// We need to copy the options because we change them
-			for (size_t o = 0; o < nop->cases.Num(); ++o)
+			for (size_t o = 0; o < nop->Cases.Num(); ++o)
 			{
-				if (nop->cases[o].branch)
+				if (nop->Cases[o].Branch)
 				{
 					Ptr<ASTOpImageRasterMesh> bOp = mu::Clone<ASTOpImageRasterMesh>(this);
-					bOp->mesh = nop->cases[o].branch.child();
-					nop->cases[o].branch = bOp;
+					bOp->mesh = nop->Cases[o].Branch.child();
+					nop->Cases[o].Branch = bOp;
 				}
 			}
 
@@ -362,21 +362,21 @@ namespace mu
 				// Make a project for every path
 				auto nop = mu::Clone<ASTOpSwitch>(imageAt.get());
 
-				if (nop->def)
+				if (nop->Default)
 				{
 					Ptr<ASTOpImageRasterMesh> defOp = mu::Clone<ASTOpImageRasterMesh>(this);
-					defOp->image = nop->def.child();
-					nop->def = defOp;
+					defOp->image = nop->Default.child();
+					nop->Default = defOp;
 				}
 
 				// We need to copy the options because we change them
-				for (size_t o = 0; o < nop->cases.Num(); ++o)
+				for (size_t o = 0; o < nop->Cases.Num(); ++o)
 				{
-					if (nop->cases[o].branch)
+					if (nop->Cases[o].Branch)
 					{
 						Ptr<ASTOpImageRasterMesh> bOp = mu::Clone<ASTOpImageRasterMesh>(this);
-						bOp->image = nop->cases[o].branch.child();
-						nop->cases[o].branch = bOp;
+						bOp->image = nop->Cases[o].Branch.child();
+						nop->Cases[o].Branch = bOp;
 					}
 				}
 
diff --git a/Engine/Plugins/Mutable/Source/MutableTools/Private/MuT/ASTOpImageResize.cpp b/Engine/Plugins/Mutable/Source/MutableTools/Private/MuT/ASTOpImageResize.cpp
--- a/Engine/Plugins/Mutable/Source/MutableTools/Private/MuT/ASTOpImageResize.cpp
+++ b/Engine/Plugins/Mutable/Source/MutableTools/Private/MuT/ASTOpImageResize.cpp
@@ -298,20 +298,20 @@ namespace mu
 			// Move the resize down all the paths
 			Ptr<ASTOpSwitch> NewOp = mu::Clone<ASTOpSwitch>(SourceAt);
 
-			if (NewOp->def)
+			if (NewOp->Default)
 			{
 				Ptr<ASTOpImageResize> DefOp = mu::Clone<ASTOpImageResize>(this);
-				DefOp->Source = NewOp->def.child();
-				NewOp->def = DefOp;
+				DefOp->Source = NewOp->Default.child();
+				NewOp->Default = DefOp;
 			}
 
-			for (ASTOpSwitch::FCase& Case : NewOp->cases)
+			for (ASTOpSwitch::FCase& Case : NewOp->Cases)
 			{
-				if (Case.branch)
+				if (Case.Branch)
 				{
 					Ptr<ASTOpImageResize> CaseOp = mu::Clone<ASTOpImageResize>(this);
-					CaseOp->Source = Case.branch.child();
-					Case.branch = CaseOp;
+					CaseOp->Source = Case.Branch.child();
+					Case.Branch = CaseOp;
 				}
 			}
 
diff --git a/Engine/Plugins/Mutable/Source/MutableTools/Private/MuT/ASTOpImageResizeRel.cpp b/Engine/Plugins/Mutable/Source/MutableTools/Private/MuT/ASTOpImageResizeRel.cpp
--- a/Engine/Plugins/Mutable/Source/MutableTools/Private/MuT/ASTOpImageResizeRel.cpp
+++ b/Engine/Plugins/Mutable/Source/MutableTools/Private/MuT/ASTOpImageResizeRel.cpp
@@ -324,10 +324,10 @@ namespace mu
 			{
 				// We move the mask creation down all the paths
 				Ptr<ASTOpSwitch> NewOp = mu::Clone<ASTOpSwitch>(at);
-				NewOp->def = Visit(NewOp->def.child(), CurrentSinkingOp);
-				for (ASTOpSwitch::FCase& c : NewOp->cases)
+				NewOp->Default = Visit(NewOp->Default.child(), CurrentSinkingOp);
+				for (ASTOpSwitch::FCase& c : NewOp->Cases)
 				{
-					c.branch = Visit(c.branch.child(), CurrentSinkingOp);
+					c.Branch = Visit(c.Branch.child(), CurrentSinkingOp);
 				}
 				Result = NewOp;
 				break;
diff --git a/Engine/Plugins/Mutable/Source/MutableTools/Private/MuT/ASTOpImageSwizzle.cpp b/Engine/Plugins/Mutable/Source/MutableTools/Private/MuT/ASTOpImageSwizzle.cpp
--- a/Engine/Plugins/Mutable/Source/MutableTools/Private/MuT/ASTOpImageSwizzle.cpp
+++ b/Engine/Plugins/Mutable/Source/MutableTools/Private/MuT/ASTOpImageSwizzle.cpp
@@ -417,7 +417,7 @@ namespace mu
 					// Move the swizzle down all the paths
 					Ptr<ASTOpSwitch> nop = mu::Clone<ASTOpSwitch>(channelSourceAt);
 
-					if (nop->def)
+					if (nop->Default)
 					{
 						Ptr<ASTOpImageSwizzle> defOp = mu::Clone<ASTOpImageSwizzle>(this);
 						for (int c = 0; c < MUTABLE_OP_MAX_SWIZZLE_CHANNELS; ++c)
@@ -425,15 +425,15 @@ namespace mu
 							const ASTOpSwitch* ChannelSwitch = static_cast<const ASTOpSwitch*>(Sources[c].child().get());
 							if (ChannelSwitch)
 							{
-								defOp->Sources[c] = ChannelSwitch->def.child();
+								defOp->Sources[c] = ChannelSwitch->Default.child();
 							}
 						}
-						nop->def = defOp;
+						nop->Default = defOp;
 					}
 
-					for (int32 v = 0; v < nop->cases.Num(); ++v)
+					for (int32 v = 0; v < nop->Cases.Num(); ++v)
 					{
-						if (nop->cases[v].branch)
+						if (nop->Cases[v].Branch)
 						{
 							Ptr<ASTOpImageSwizzle> branchOp = mu::Clone<ASTOpImageSwizzle>(this);
 							for (int c = 0; c < MUTABLE_OP_MAX_SWIZZLE_CHANNELS; ++c)
@@ -441,10 +441,10 @@ namespace mu
 								const ASTOpSwitch* ChannelSwitch = static_cast<const ASTOpSwitch*>(Sources[c].child().get());
 								if (ChannelSwitch)
 								{
-									branchOp->Sources[c] = ChannelSwitch->cases[v].branch.child();
+									branchOp->Sources[c] = ChannelSwitch->Cases[v].Branch.child();
 								}
 							}
-							nop->cases[v].branch = branchOp;
+							nop->Cases[v].Branch = branchOp;
 						}
 					}
 
@@ -1207,10 +1207,10 @@ namespace mu
 		{
 			// We move the op down all the paths
 			Ptr<ASTOpSwitch> newOp = mu::Clone<ASTOpSwitch>(at);
-			newOp->def = Visit(newOp->def.child(), CurrentSwizzleOp);
-			for (auto& c : newOp->cases)
+			newOp->Default = Visit(newOp->Default.child(), CurrentSwizzleOp);
+			for (ASTOpSwitch::FCase& c : newOp->Cases)
 			{
-				c.branch = Visit(c.branch.child(), CurrentSwizzleOp);
+				c.Branch = Visit(c.Branch.child(), CurrentSwizzleOp);
 			}
 			newAt = newOp;
 			break;
diff --git a/Engine/Plugins/Mutable/Source/MutableTools/Private/MuT/ASTOpLayoutFromMesh.cpp b/Engine/Plugins/Mutable/Source/MutableTools/Private/MuT/ASTOpLayoutFromMesh.cpp
--- a/Engine/Plugins/Mutable/Source/MutableTools/Private/MuT/ASTOpLayoutFromMesh.cpp
+++ b/Engine/Plugins/Mutable/Source/MutableTools/Private/MuT/ASTOpLayoutFromMesh.cpp
@@ -208,11 +208,11 @@ namespace mu
 				case EOpType::ME_SWITCH:
 				{
 					Ptr<ASTOpSwitch> NewOp = mu::Clone<ASTOpSwitch>(at);
-					NewOp->type = EOpType::LA_SWITCH;
-					NewOp->def = Visit(NewOp->def.child(), CurrentSinkingOp);
-					for (ASTOpSwitch::FCase& c : NewOp->cases)
+					NewOp->Type = EOpType::LA_SWITCH;
+					NewOp->Default = Visit(NewOp->Default.child(), CurrentSinkingOp);
+					for (ASTOpSwitch::FCase& c : NewOp->Cases)
 					{
-						c.branch = Visit(c.branch.child(), CurrentSinkingOp);
+						c.Branch = Visit(c.Branch.child(), CurrentSinkingOp);
 					}
 					newAt = NewOp;
 					break;
diff --git a/Engine/Plugins/Mutable/Source/MutableTools/Private/MuT/ASTOpMeshBindShape.cpp b/Engine/Plugins/Mutable/Source/MutableTools/Private/MuT/ASTOpMeshBindShape.cpp
--- a/Engine/Plugins/Mutable/Source/MutableTools/Private/MuT/ASTOpMeshBindShape.cpp
+++ b/Engine/Plugins/Mutable/Source/MutableTools/Private/MuT/ASTOpMeshBindShape.cpp
@@ -11,6 +11,7 @@
 #include "MuT/ASTOpSwitch.h"
 #include "MuT/ASTOpMeshRemoveMask.h"
 #include "MuT/ASTOpMeshAddTags.h"
+#include "MuT/ASTOpMeshPrepareLayout.h"
 
 
 namespace mu
@@ -237,22 +238,22 @@ namespace mu
 				// Move the operation down all the paths
 				Ptr<ASTOpSwitch> NewSwitch = mu::Clone<ASTOpSwitch>(MeshAt);
 
-				if (NewSwitch->def)
+				if (NewSwitch->Default)
 				{
 					Ptr<ASTOpMeshBindShape> NewBind = mu::Clone<ASTOpMeshBindShape>(this);
-					NewBind->Mesh = MeshSwitch->def.child();
-					NewBind->Shape = ShapeSwitch->def.child();
-					NewSwitch->def = NewBind;
+					NewBind->Mesh = MeshSwitch->Default.child();
+					NewBind->Shape = ShapeSwitch->Default.child();
+					NewSwitch->Default = NewBind;
 				}
 
-				for (int32 v = 0; v < NewSwitch->cases.Num(); ++v)
+				for (int32 v = 0; v < NewSwitch->Cases.Num(); ++v)
 				{
-					if (NewSwitch->cases[v].branch)
+					if (NewSwitch->Cases[v].Branch)
 					{
 						Ptr<ASTOpMeshBindShape> NewBind = mu::Clone<ASTOpMeshBindShape>(this);
-						NewBind->Mesh = MeshSwitch->cases[v].branch.child();
-						NewBind->Shape = ShapeSwitch->FindBranch(MeshSwitch->cases[v].condition);
-						NewSwitch->cases[v].branch = NewBind;
+						NewBind->Mesh = MeshSwitch->Cases[v].Branch.child();
+						NewBind->Shape = ShapeSwitch->FindBranch(MeshSwitch->Cases[v].Condition);
+						NewSwitch->Cases[v].Branch = NewBind;
 					}
 				}
 
@@ -313,20 +314,20 @@ namespace mu
 				// Move the operation down all the paths
 				Ptr<ASTOpSwitch> NewSwitch = mu::Clone<ASTOpSwitch>(MeshAt);
 
-				if (NewSwitch->def)
+				if (NewSwitch->Default)
 				{
 					Ptr<ASTOpMeshBindShape> NewBind = mu::Clone<ASTOpMeshBindShape>(this);
-					NewBind->Mesh = NewSwitch->def.child();
-					NewSwitch->def = NewBind;
+					NewBind->Mesh = NewSwitch->Default.child();
+					NewSwitch->Default = NewBind;
 				}
 
-				for (int32 v = 0; v < NewSwitch->cases.Num(); ++v)
+				for (int32 v = 0; v < NewSwitch->Cases.Num(); ++v)
 				{
-					if (NewSwitch->cases[v].branch)
+					if (NewSwitch->Cases[v].Branch)
 					{
 						Ptr<ASTOpMeshBindShape> NewBind = mu::Clone<ASTOpMeshBindShape>(this);
-						NewBind->Mesh = NewSwitch->cases[v].branch.child();
-						NewSwitch->cases[v].branch = NewBind;
+						NewBind->Mesh = NewSwitch->Cases[v].Branch.child();
+						NewSwitch->Cases[v].Branch = NewBind;
 					}
 				}
 
@@ -388,6 +389,20 @@ namespace mu
 				break;
 			}
 
+			case EOpType::ME_PREPARELAYOUT:
+			{
+				Ptr<ASTOpMeshPrepareLayout> New = mu::Clone<ASTOpMeshPrepareLayout>(MeshAt);
+				if (New->Mesh)
+				{
+					Ptr<ASTOpMeshBindShape> NewBind = mu::Clone<ASTOpMeshBindShape>(this);
+					NewBind->Mesh = New->Mesh.child();
+					New->Mesh = NewBind;
+				}
+
+				NewOp = New;
+				break;
+			}
+
 			default:
 				break;
 
@@ -406,20 +421,20 @@ namespace mu
 				// Move the operation down all the paths
 				Ptr<ASTOpSwitch> NewSwitch = mu::Clone<ASTOpSwitch>(ShapeAt);
 
-				if (NewSwitch->def)
+				if (NewSwitch->Default)
 				{
 					Ptr<ASTOpMeshBindShape> NewBind = mu::Clone<ASTOpMeshBindShape>(this);
-					NewBind->Shape = NewSwitch->def.child();
-					NewSwitch->def = NewBind;
+					NewBind->Shape = NewSwitch->Default.child();
+					NewSwitch->Default = NewBind;
 				}
 
-				for (int32 v = 0; v < NewSwitch->cases.Num(); ++v)
+				for (int32 v = 0; v < NewSwitch->Cases.Num(); ++v)
 				{
-					if (NewSwitch->cases[v].branch)
+					if (NewSwitch->Cases[v].Branch)
 					{
 						Ptr<ASTOpMeshBindShape> NewBind = mu::Clone<ASTOpMeshBindShape>(this);
-						NewBind->Shape = NewSwitch->cases[v].branch.child();
-						NewSwitch->cases[v].branch = NewBind;
+						NewBind->Shape = NewSwitch->Cases[v].Branch.child();
+						NewSwitch->Cases[v].Branch = NewBind;
 					}
 				}
 
@@ -471,7 +486,6 @@ namespace mu
 	}
 
 
-	//-------------------------------------------------------------------------------------------------
 	FSourceDataDescriptor ASTOpMeshBindShape::GetSourceDataDescriptor(FGetSourceDataDescriptorContext* Context) const
 	{
 		if (Mesh)
diff --git a/Engine/Plugins/Mutable/Source/MutableTools/Private/MuT/ASTOpMeshDifference.cpp b/Engine/Plugins/Mutable/Source/MutableTools/Private/MuT/ASTOpMeshDifference.cpp
--- a/Engine/Plugins/Mutable/Source/MutableTools/Private/MuT/ASTOpMeshDifference.cpp
+++ b/Engine/Plugins/Mutable/Source/MutableTools/Private/MuT/ASTOpMeshDifference.cpp
@@ -143,22 +143,22 @@ namespace mu
 				// Move the operation down all the paths
 				Ptr<ASTOpSwitch> NewSwitch = mu::Clone<ASTOpSwitch>(BaseAt);
 
-				if (NewSwitch->def)
+				if (NewSwitch->Default)
 				{
 					Ptr<ASTOpMeshDifference> NewDiff = mu::Clone<ASTOpMeshDifference>(this);
-					NewDiff->Base = BaseSwitch->def.child();
-					NewDiff->Target = TargetSwitch->def.child();
-					NewSwitch->def = NewDiff;
+					NewDiff->Base = BaseSwitch->Default.child();
+					NewDiff->Target = TargetSwitch->Default.child();
+					NewSwitch->Default = NewDiff;
 				}
 
-				for (int32 v = 0; v < NewSwitch->cases.Num(); ++v)
+				for (int32 v = 0; v < NewSwitch->Cases.Num(); ++v)
 				{
-					if (NewSwitch->cases[v].branch)
+					if (NewSwitch->Cases[v].Branch)
 					{
 						Ptr<ASTOpMeshDifference> NewDiff = mu::Clone<ASTOpMeshDifference>(this);
-						NewDiff->Base = BaseSwitch->cases[v].branch.child();
-						NewDiff->Target = TargetSwitch->FindBranch(BaseSwitch->cases[v].condition);
-						NewSwitch->cases[v].branch = NewDiff;
+						NewDiff->Base = BaseSwitch->Cases[v].Branch.child();
+						NewDiff->Target = TargetSwitch->FindBranch(BaseSwitch->Cases[v].Condition);
+						NewSwitch->Cases[v].Branch = NewDiff;
 					}
 				}
 
@@ -219,20 +219,20 @@ namespace mu
 				// Move the operation down all the paths
 				Ptr<ASTOpSwitch> NewSwitch = mu::Clone<ASTOpSwitch>(BaseAt);
 
-				if (NewSwitch->def)
+				if (NewSwitch->Default)
 				{
 					Ptr<ASTOpMeshDifference> NewDiff = mu::Clone<ASTOpMeshDifference>(this);
-					NewDiff->Base = NewSwitch->def.child();
-					NewSwitch->def = NewDiff;
+					NewDiff->Base = NewSwitch->Default.child();
+					NewSwitch->Default = NewDiff;
 				}
 
-				for (int32 v = 0; v < NewSwitch->cases.Num(); ++v)
+				for (int32 v = 0; v < NewSwitch->Cases.Num(); ++v)
 				{
-					if (NewSwitch->cases[v].branch)
+					if (NewSwitch->Cases[v].Branch)
 					{
 						Ptr<ASTOpMeshDifference> NewDiff = mu::Clone<ASTOpMeshDifference>(this);
-						NewDiff->Base = NewSwitch->cases[v].branch.child();
-						NewSwitch->cases[v].branch = NewDiff;
+						NewDiff->Base = NewSwitch->Cases[v].Branch.child();
+						NewSwitch->Cases[v].Branch = NewDiff;
 					}
 				}
 
@@ -296,20 +296,20 @@ namespace mu
 				// Move the operation down all the paths
 				Ptr<ASTOpSwitch> NewSwitch = mu::Clone<ASTOpSwitch>(TargetAt);
 
-				if (NewSwitch->def)
+				if (NewSwitch->Default)
 				{
 					Ptr<ASTOpMeshDifference> NewDiff = mu::Clone<ASTOpMeshDifference>(this);
-					NewDiff->Target = NewSwitch->def.child();
-					NewSwitch->def = NewDiff;
+					NewDiff->Target = NewSwitch->Default.child();
+					NewSwitch->Default = NewDiff;
 				}
 
-				for (int32 v = 0; v < NewSwitch->cases.Num(); ++v)
+				for (int32 v = 0; v < NewSwitch->Cases.Num(); ++v)
 				{
-					if (NewSwitch->cases[v].branch)
+					if (NewSwitch->Cases[v].Branch)
 					{
 						Ptr<ASTOpMeshDifference> NewDiff = mu::Clone<ASTOpMeshDifference>(this);
-						NewDiff->Target = NewSwitch->cases[v].branch.child();
-						NewSwitch->cases[v].branch = NewDiff;
+						NewDiff->Target = NewSwitch->Cases[v].Branch.child();
+						NewSwitch->Cases[v].Branch = NewDiff;
 					}
 				}
 
diff --git a/Engine/Plugins/Mutable/Source/MutableTools/Private/MuT/ASTOpMeshExtractLayoutBlocks.cpp b/Engine/Plugins/Mutable/Source/MutableTools/Private/MuT/ASTOpMeshExtractLayoutBlocks.cpp
--- a/Engine/Plugins/Mutable/Source/MutableTools/Private/MuT/ASTOpMeshExtractLayoutBlocks.cpp
+++ b/Engine/Plugins/Mutable/Source/MutableTools/Private/MuT/ASTOpMeshExtractLayoutBlocks.cpp
@@ -276,10 +276,10 @@ namespace mu
 		case EOpType::ME_SWITCH:
 		{
 			Ptr<ASTOpSwitch> newOp = mu::Clone<ASTOpSwitch>(at);
-			newOp->def = Visit(newOp->def.child(), currentSinkOp);
-			for (ASTOpSwitch::FCase& c : newOp->cases)
+			newOp->Default = Visit(newOp->Default.child(), currentSinkOp);
+			for (ASTOpSwitch::FCase& c : newOp->Cases)
 			{
-				c.branch = Visit(c.branch.child(), currentSinkOp);
+				c.Branch = Visit(c.Branch.child(), currentSinkOp);
 			}
 			newAt = newOp;
 			break;
diff --git a/Engine/Plugins/Mutable/Source/MutableTools/Private/MuT/ASTOpMeshFormat.cpp b/Engine/Plugins/Mutable/Source/MutableTools/Private/MuT/ASTOpMeshFormat.cpp
--- a/Engine/Plugins/Mutable/Source/MutableTools/Private/MuT/ASTOpMeshFormat.cpp
+++ b/Engine/Plugins/Mutable/Source/MutableTools/Private/MuT/ASTOpMeshFormat.cpp
@@ -402,7 +402,7 @@ mu::Ptr<ASTOp> Sink_MeshFormatAST::Visit(const mu::Ptr<ASTOp>& at, const ASTOpMe
 
 	case EOpType::ME_REMOVEMASK:
 	{
-		auto newOp = mu::Clone<ASTOpMeshRemoveMask>(at);
+		Ptr<ASTOpMeshRemoveMask> newOp = mu::Clone<ASTOpMeshRemoveMask>(at);
 		newOp->source = Visit(newOp->source.child(), currentFormatOp);
 		newAt = newOp;
 		break;
@@ -410,7 +410,7 @@ mu::Ptr<ASTOp> Sink_MeshFormatAST::Visit(const mu::Ptr<ASTOp>& at, const ASTOpMe
 
 	case EOpType::ME_CONDITIONAL:
 	{
-		auto newOp = mu::Clone<ASTOpConditional>(at);
+		Ptr<ASTOpConditional> newOp = mu::Clone<ASTOpConditional>(at);
 		newOp->yes = Visit(newOp->yes.child(), currentFormatOp);
 		newOp->no = Visit(newOp->no.child(), currentFormatOp);
 		newAt = newOp;
@@ -419,11 +419,11 @@ mu::Ptr<ASTOp> Sink_MeshFormatAST::Visit(const mu::Ptr<ASTOp>& at, const ASTOpMe
 
 	case EOpType::ME_SWITCH:
 	{
-		auto newOp = mu::Clone<ASTOpSwitch>(at);
-		newOp->def = Visit(newOp->def.child(), currentFormatOp);
-		for (auto& c : newOp->cases)
+		Ptr<ASTOpSwitch> newOp = mu::Clone<ASTOpSwitch>(at);
+		newOp->Default = Visit(newOp->Default.child(), currentFormatOp);
+		for (ASTOpSwitch::FCase& c : newOp->Cases)
 		{
-			c.branch = Visit(c.branch.child(), currentFormatOp);
+			c.Branch = Visit(c.Branch.child(), currentFormatOp);
 		}
 		newAt = newOp;
 		break;
diff --git a/Engine/Plugins/Mutable/Source/MutableTools/Private/MuT/ASTOpMeshMaskClipMesh.cpp b/Engine/Plugins/Mutable/Source/MutableTools/Private/MuT/ASTOpMeshMaskClipMesh.cpp
--- a/Engine/Plugins/Mutable/Source/MutableTools/Private/MuT/ASTOpMeshMaskClipMesh.cpp
+++ b/Engine/Plugins/Mutable/Source/MutableTools/Private/MuT/ASTOpMeshMaskClipMesh.cpp
@@ -9,6 +9,7 @@
 #include "MuT/ASTOpConditional.h"
 #include "MuT/ASTOpMeshRemoveMask.h"
 #include "MuT/ASTOpMeshAddTags.h"
+#include "MuT/ASTOpMeshPrepareLayout.h"
 #include "MuT/ASTOpSwitch.h"
 
 
@@ -170,6 +171,14 @@ namespace mu
 					break;
 				}
 
+				case EOpType::ME_PREPARELAYOUT:
+				{
+					// Ignore the prepare op: it doesn't contribute to the mask generation.
+					const ASTOpMeshPrepareLayout* typedAt = static_cast<const ASTOpMeshPrepareLayout*>(at.get());
+					newAt = Visit(typedAt->Mesh.child());
+					break;
+				}
+
 				case EOpType::ME_ADDTAGS:
 				{
 					Ptr<ASTOpMeshAddTags> newOp = mu::Clone<ASTOpMeshAddTags>(at);
@@ -195,10 +204,10 @@ namespace mu
 				{
 					// We move the mask creation down all the paths
 					Ptr<ASTOpSwitch> newOp = mu::Clone<ASTOpSwitch>(at);
-					newOp->def = Visit(newOp->def.child());
-					for (ASTOpSwitch::FCase& c : newOp->cases)
+					newOp->Default = Visit(newOp->Default.child());
+					for (ASTOpSwitch::FCase& c : newOp->Cases)
 					{
-						c.branch = Visit(c.branch.child());
+						c.Branch = Visit(c.Branch.child());
 					}
 					newAt = newOp;
 					break;
@@ -304,10 +313,10 @@ namespace mu
 				{
 					// We move the mask creation down all the paths
 					Ptr<ASTOpSwitch> newOp = mu::Clone<ASTOpSwitch>(at);
-					newOp->def = Visit(newOp->def.child());
-					for (ASTOpSwitch::FCase& c : newOp->cases)
+					newOp->Default = Visit(newOp->Default.child());
+					for (ASTOpSwitch::FCase& c : newOp->Cases)
 					{
-						c.branch = Visit(c.branch.child());
+						c.Branch = Visit(c.Branch.child());
 					}
 					newAt = newOp;
 					break;
diff --git a/Engine/Plugins/Mutable/Source/MutableTools/Private/MuT/ASTOpMeshMaskClipUVMask.cpp b/Engine/Plugins/Mutable/Source/MutableTools/Private/MuT/ASTOpMeshMaskClipUVMask.cpp
--- a/Engine/Plugins/Mutable/Source/MutableTools/Private/MuT/ASTOpMeshMaskClipUVMask.cpp
+++ b/Engine/Plugins/Mutable/Source/MutableTools/Private/MuT/ASTOpMeshMaskClipUVMask.cpp
@@ -9,6 +9,7 @@
 #include "MuT/ASTOpConditional.h"
 #include "MuT/ASTOpMeshRemoveMask.h"
 #include "MuT/ASTOpMeshAddTags.h"
+#include "MuT/ASTOpMeshPrepareLayout.h"
 #include "MuT/ASTOpSwitch.h"
 
 
@@ -193,6 +194,14 @@ namespace mu
 					break;
 				}
 
+				case EOpType::ME_PREPARELAYOUT:
+				{
+					// Ignore the prepare op in the source: it doesn't contribute to the mask generation (the one in the SourceUV is used for UVs).
+					const ASTOpMeshPrepareLayout* typedAt = static_cast<const ASTOpMeshPrepareLayout*>(at.get());
+					newAt = Visit(typedAt->Mesh.child());
+					break;
+				}
+
 				case EOpType::ME_CONDITIONAL:
 				{
 					// We move the mask creation down the two paths
@@ -210,10 +219,10 @@ namespace mu
 				{
 					// We move the mask creation down all the paths
 					Ptr<ASTOpSwitch> newOp = mu::Clone<ASTOpSwitch>(at);
-					newOp->def = Visit(newOp->def.child());
-					for (ASTOpSwitch::FCase& c : newOp->cases)
+					newOp->Default = Visit(newOp->Default.child());
+					for (ASTOpSwitch::FCase& c : newOp->Cases)
 					{
-						c.branch = Visit(c.branch.child());
+						c.Branch = Visit(c.Branch.child());
 					}
 					newAt = newOp;
 					break;
@@ -319,10 +328,10 @@ namespace mu
 				{
 					// We move the mask creation down all the paths
 					Ptr<ASTOpSwitch> newOp = mu::Clone<ASTOpSwitch>(at);
-					newOp->def = Visit(newOp->def.child());
-					for (ASTOpSwitch::FCase& c : newOp->cases)
+					newOp->Default = Visit(newOp->Default.child());
+					for (ASTOpSwitch::FCase& c : newOp->Cases)
 					{
-						c.branch = Visit(c.branch.child());
+						c.Branch = Visit(c.Branch.child());
 					}
 					newAt = newOp;
 					break;
diff --git a/Engine/Plugins/Mutable/Source/MutableTools/Private/MuT/ASTOpMeshMaskDiff.cpp b/Engine/Plugins/Mutable/Source/MutableTools/Private/MuT/ASTOpMeshMaskDiff.cpp
--- a/Engine/Plugins/Mutable/Source/MutableTools/Private/MuT/ASTOpMeshMaskDiff.cpp
+++ b/Engine/Plugins/Mutable/Source/MutableTools/Private/MuT/ASTOpMeshMaskDiff.cpp
@@ -3,11 +3,11 @@
 #include "MuT/ASTOpMeshMaskDiff.h"
 #include "MuT/ASTOpMeshAddTags.h"
 #include "MuT/ASTOpSwitch.h"
-
-#include "HAL/PlatformMath.h"
+#include "MuT/ASTOpMeshPrepareLayout.h"
 #include "MuR/ModelPrivate.h"
 #include "MuR/RefCounted.h"
 #include "MuR/Types.h"
+#include "HAL/PlatformMath.h"
 
 namespace mu
 {
@@ -104,36 +104,83 @@ namespace mu
 			ResultOp = NewAt;
 			break;
 		}
+
+		case EOpType::ME_PREPARELAYOUT:
+		{
+			// Ignore the prepare op in the source: it doesn't contribute to the diff generation.
+			const ASTOpMeshPrepareLayout* Prepare = static_cast<const ASTOpMeshPrepareLayout*>(Source.child().get());
+
+			Ptr<ASTOpMeshMaskDiff> NewAt = mu::Clone<ASTOpMeshMaskDiff>(this);
+			NewAt->Source = Prepare->Mesh.child();
+			ResultOp = NewAt;
+			break;
+		}
+
 		case EOpType::ME_SWITCH:
 		{
 			// Move the mask diff down all the paths
 			Ptr<ASTOpSwitch> NewOp = mu::Clone<ASTOpSwitch>(Fragment.child().get());
 
-			if (NewOp->def)
+			if (NewOp->Default)
 			{
 				Ptr<ASTOpMeshMaskDiff> defOp = mu::Clone<ASTOpMeshMaskDiff>(this);
-				defOp->Fragment = NewOp->def.child();
-				NewOp->def = defOp;
+				defOp->Fragment = NewOp->Default.child();
+				NewOp->Default = defOp;
 			}
 
 			// We need to copy the options because we change them
-			for (int32 v = 0; v < NewOp->cases.Num(); ++v)
+			for (int32 v = 0; v < NewOp->Cases.Num(); ++v)
 			{
-				if (NewOp->cases[v].branch)
+				if (NewOp->Cases[v].Branch)
 				{
 					Ptr<ASTOpMeshMaskDiff> bOp = mu::Clone<ASTOpMeshMaskDiff>(this);
-					bOp->Fragment = NewOp->cases[v].branch.child();
-					NewOp->cases[v].branch = bOp;
+					bOp->Fragment = NewOp->Cases[v].Branch.child();
+					NewOp->Cases[v].Branch = bOp;
 				}
 			}
 
 			ResultOp = NewOp;
 			break;
 		}
+
 		default:
 			break;
 		}
 
+		// If we didn't optimize the fragment, see ifg we can optimize the source.
+		if (!ResultOp)
+		{
+			EOpType SourceType = Source->GetOpType();
+			switch (SourceType)
+			{
+
+			case EOpType::ME_ADDTAGS:
+			{
+				// Tags in the source can be ignored for mask generation.
+				const ASTOpMeshAddTags* Add = static_cast<const ASTOpMeshAddTags*>(Source.child().get());
+
+				Ptr<ASTOpMeshMaskDiff> NewAt = mu::Clone<ASTOpMeshMaskDiff>(this);
+				NewAt->Source = Add->Source.child();
+				ResultOp = NewAt;
+				break;
+			}
+
+			case EOpType::ME_PREPARELAYOUT:
+			{
+				// Ignore the prepare op in the source: it doesn't contribute to the diff generation.
+				const ASTOpMeshPrepareLayout* Prepare = static_cast<const ASTOpMeshPrepareLayout*>(Source.child().get());
+
+				Ptr<ASTOpMeshMaskDiff> NewAt = mu::Clone<ASTOpMeshMaskDiff>(this);
+				NewAt->Source = Prepare->Mesh.child();
+				ResultOp = NewAt;
+				break;
+			}
+
+			default:
+				break;
+			}
+		}
+
 		return ResultOp;
 	}
 
diff --git a/Engine/Plugins/Mutable/Source/MutableTools/Private/MuT/ASTOpMeshOptimizeSkinning.cpp b/Engine/Plugins/Mutable/Source/MutableTools/Private/MuT/ASTOpMeshOptimizeSkinning.cpp
--- a/Engine/Plugins/Mutable/Source/MutableTools/Private/MuT/ASTOpMeshOptimizeSkinning.cpp
+++ b/Engine/Plugins/Mutable/Source/MutableTools/Private/MuT/ASTOpMeshOptimizeSkinning.cpp
@@ -172,10 +172,10 @@ namespace mu
 		case EOpType::ME_SWITCH:
 		{
 			Ptr<ASTOpSwitch> NewOp = mu::Clone<ASTOpSwitch>(at);
-			NewOp->def = Visit(NewOp->def.child(), CurrentSinkOp);
-			for (auto& c : NewOp->cases)
+			NewOp->Default = Visit(NewOp->Default.child(), CurrentSinkOp);
+			for (ASTOpSwitch::FCase& c : NewOp->Cases)
 			{
-				c.branch = Visit(c.branch.child(), CurrentSinkOp);
+				c.Branch = Visit(c.Branch.child(), CurrentSinkOp);
 			}
 			NewAt = NewOp;
 			break;
diff --git a/Engine/Plugins/Mutable/Source/MutableTools/Private/MuT/ASTOpMeshProject.cpp b/Engine/Plugins/Mutable/Source/MutableTools/Private/MuT/ASTOpMeshProject.cpp
--- a/Engine/Plugins/Mutable/Source/MutableTools/Private/MuT/ASTOpMeshProject.cpp
+++ b/Engine/Plugins/Mutable/Source/MutableTools/Private/MuT/ASTOpMeshProject.cpp
@@ -126,21 +126,21 @@ namespace mu
 				// Move the format down all the paths
 				Ptr<ASTOpSwitch> nop = mu::Clone<ASTOpSwitch>(SourceAt);
 
-				if (nop->def)
+				if (nop->Default)
 				{
 					Ptr<ASTOpMeshProject> defOp = mu::Clone<ASTOpMeshProject>(this);
-					defOp->Mesh = nop->def.child();
-					nop->def = defOp;
+					defOp->Mesh = nop->Default.child();
+					nop->Default = defOp;
 				}
 
 				// We need to copy the options because we change them
-				for (int32 CaseIndex = 0; CaseIndex < nop->cases.Num(); ++CaseIndex)
+				for (int32 CaseIndex = 0; CaseIndex < nop->Cases.Num(); ++CaseIndex)
 				{
-					if (nop->cases[CaseIndex].branch)
+					if (nop->Cases[CaseIndex].Branch)
 					{
 						Ptr<ASTOpMeshProject> bOp = mu::Clone<ASTOpMeshProject>(this);
-						bOp->Mesh = nop->cases[CaseIndex].branch.child();
-						nop->cases[CaseIndex].branch = bOp;
+						bOp->Mesh = nop->Cases[CaseIndex].Branch.child();
+						nop->Cases[CaseIndex].Branch = bOp;
 					}
 				}
 
diff --git a/Engine/Plugins/Mutable/Source/MutableTools/Private/MuT/ASTOpSwitch.cpp b/Engine/Plugins/Mutable/Source/MutableTools/Private/MuT/ASTOpSwitch.cpp
--- a/Engine/Plugins/Mutable/Source/MutableTools/Private/MuT/ASTOpSwitch.cpp
+++ b/Engine/Plugins/Mutable/Source/MutableTools/Private/MuT/ASTOpSwitch.cpp
@@ -4,6 +4,7 @@
 
 #include "MuT/ASTOpParameter.h"
 #include "MuT/ASTOpConstantInt.h"
+#include "MuT/ASTOpMeshAddTags.h"
 #include "Containers/Map.h"
 #include "Misc/AssertionMacros.h"
 #include "MuR/ImagePrivate.h"
@@ -16,16 +17,13 @@
 namespace mu
 {
 
-
-	//---------------------------------------------------------------------------------------------
 	ASTOpSwitch::ASTOpSwitch()
-		: variable(this)
-		, def(this)
+		: Variable(this)
+		, Default(this)
 	{
 	}
 
 
-	//-------------------------------------------------------------------------------------------------
 	ASTOpSwitch::~ASTOpSwitch()
 	{
 		// Explicit call needed to avoid recursive destruction
@@ -33,38 +31,35 @@ namespace mu
 	}
 
 
-	//-------------------------------------------------------------------------------------------------
 	bool ASTOpSwitch::IsEqual(const ASTOp& otherUntyped) const
 	{
 		if (GetOpType()==otherUntyped.GetOpType())
 		{
 			const ASTOpSwitch* other = static_cast<const ASTOpSwitch*>(&otherUntyped);
-			return type == other->type && variable == other->variable &&
-				cases == other->cases && def == other->def;
+			return Type == other->Type && Variable == other->Variable &&
+				Cases == other->Cases && Default == other->Default;
 		}
 		return false;
 	}
 
 
-	//-------------------------------------------------------------------------------------------------
 	mu::Ptr<ASTOp> ASTOpSwitch::Clone(MapChildFuncRef mapChild) const
 	{
 		Ptr<ASTOpSwitch> n = new ASTOpSwitch();
-		n->type = type;
-		n->variable = mapChild(variable.child());
-		n->def = mapChild(def.child());
-		for (const auto& c : cases)
+		n->Type = Type;
+		n->Variable = mapChild(Variable.child());
+		n->Default = mapChild(Default.child());
+		for (const FCase& c : Cases)
 		{
-			n->cases.Emplace(c.condition, n, mapChild(c.branch.child()));
+			n->Cases.Emplace(c.Condition, n, mapChild(c.Branch.child()));
 		}
 		return n;
 	}
 
 
-	//-------------------------------------------------------------------------------------------------
 	void ASTOpSwitch::Assert()
 	{
-		switch (type)
+		switch (Type)
 		{
 		case EOpType::NU_SWITCH:
 		case EOpType::SC_SWITCH:
@@ -76,7 +71,7 @@ namespace mu
 		case EOpType::ED_SWITCH:
 			break;
 		default:
-			// Unexpected type
+			// Unexpected Type
 			check(false);
 			break;
 		}
@@ -85,45 +80,42 @@ namespace mu
 	}
 
 
-	//-------------------------------------------------------------------------------------------------
 	uint64 ASTOpSwitch::Hash() const
 	{
-		uint64 res = std::hash<uint64>()(uint64(type));
-		for (const auto& c : cases)
+		uint64 res = std::hash<uint64>()(uint64(Type));
+		for (const FCase& c : Cases)
 		{
-			hash_combine(res, c.condition);
-			hash_combine(res, c.branch.child().get());
+			hash_combine(res, c.Condition);
+			hash_combine(res, c.Branch.child().get());
 		}
 		return res;
 	}
 
 
-	//-------------------------------------------------------------------------------------------------
 	mu::Ptr<ASTOp> ASTOpSwitch::GetFirstValidValue()
 	{
-		for (int32 i = 0; i < cases.Num(); ++i)
+		for (int32 i = 0; i < Cases.Num(); ++i)
 		{
-			if (cases[i].branch)
+			if (Cases[i].Branch)
 			{
-				return cases[i].branch.child();
+				return Cases[i].Branch.child();
 			}
 		}
 		return nullptr;
 	}
 
 
-	//-------------------------------------------------------------------------------------------------
 	bool ASTOpSwitch::IsCompatibleWith(const ASTOpSwitch* other) const
 	{
 		if (!other) return false;
-		if (variable.child() != other->variable.child()) return false;
-		if (cases.Num() != other->cases.Num()) return false;
-		for (const auto& c : cases)
+		if (Variable.child() != other->Variable.child()) return false;
+		if (Cases.Num() != other->Cases.Num()) return false;
+		for (const FCase& c : Cases)
 		{
 			bool found = false;
-			for (const auto& o : other->cases)
+			for (const FCase& o : other->Cases)
 			{
-				if (c.condition == o.condition)
+				if (c.Condition == o.Condition)
 				{
 					found = true;
 					break;
@@ -139,34 +131,31 @@ namespace mu
 	}
 
 
-	//-------------------------------------------------------------------------------------------------
-	mu::Ptr<ASTOp> ASTOpSwitch::FindBranch(int32 condition) const
+	mu::Ptr<ASTOp> ASTOpSwitch::FindBranch(int32 Condition) const
 	{
-		for (const auto& c : cases)
+		for (const FCase& c : Cases)
 		{
-			if (c.condition == condition)
+			if (c.Condition == Condition)
 			{
-				return c.branch.child();
+				return c.Branch.child();
 			}
 		}
 
-		return def.child();
+		return Default.child();
 	}
 
 
-	//-------------------------------------------------------------------------------------------------
 	void ASTOpSwitch::ForEachChild(const TFunctionRef<void(ASTChild&)> f)
 	{
-		f(variable);
-		f(def);
-		for (auto& cas : cases)
+		f(Variable);
+		f(Default);
+		for (FCase& cas : Cases)
 		{
-			f(cas.branch);
+			f(cas.Branch);
 		}
 	}
 
 
-	//-------------------------------------------------------------------------------------------------
 	void ASTOpSwitch::Link(FProgram& program, FLinkerOptions*)
 	{
 		// Already linked?
@@ -175,24 +164,24 @@ namespace mu
 			linkedAddress = (OP::ADDRESS)program.OpAddress.Num();
 			program.OpAddress.Add((uint32_t)program.ByteCode.Num());
 
-			OP::ADDRESS VarAddress = variable ? variable->linkedAddress : 0;
-			OP::ADDRESS DefAddress = def ? def->linkedAddress : 0;
+			OP::ADDRESS VarAddress = Variable ? Variable->linkedAddress : 0;
+			OP::ADDRESS DefAddress = Default ? Default->linkedAddress : 0;
 
-			AppendCode(program.ByteCode, type);
+			AppendCode(program.ByteCode, Type);
 			AppendCode(program.ByteCode, VarAddress);
 			AppendCode(program.ByteCode, DefAddress);
-			AppendCode(program.ByteCode, (uint32_t)cases.Num());
+			AppendCode(program.ByteCode, (uint32_t)Cases.Num());
 
-			for (const FCase& Case : cases)
+			for (const FCase& Case : Cases)
 			{
-				OP::ADDRESS CaseBranchAddress = Case.branch ? Case.branch->linkedAddress : 0;
-				AppendCode(program.ByteCode, Case.condition);
+				OP::ADDRESS CaseBranchAddress = Case.Branch ? Case.Branch->linkedAddress : 0;
+				AppendCode(program.ByteCode, Case.Condition);
 				AppendCode(program.ByteCode, CaseBranchAddress);
 			}
 		}
 	}
 
-	//-------------------------------------------------------------------------------------------------
+
 	FImageDesc ASTOpSwitch::GetImageDesc(bool bReturnBestOption, class FGetImageDescContext* Context) const
 	{
 		FImageDesc Result;
@@ -218,7 +207,7 @@ namespace mu
 		// Otherwise, we return a descriptor with empty fields in the conflicting ones, size or format.
 		// In some places this will force re-formatting of the image.
 		// The code optimiser will take care then of moving the format operations down to each
-		// branch and remove the unnecessary ones.
+		// Branch and remove the unnecessary ones.
 		FImageDesc Candidate;
 
 		bool bSameSize = true;
@@ -226,18 +215,18 @@ namespace mu
 		bool bSameLods = true;
 		bool bFirst = true;
 
-		if (def)
+		if (Default)
 		{
-			FImageDesc ChildDesc = def->GetImageDesc(bReturnBestOption, Context);
+			FImageDesc ChildDesc = Default->GetImageDesc(bReturnBestOption, Context);
 			Candidate = ChildDesc;
 			bFirst = false;
 		}
 
-		for (int32 CaseIndex = 0; CaseIndex < cases.Num(); ++CaseIndex)
+		for (int32 CaseIndex = 0; CaseIndex < Cases.Num(); ++CaseIndex)
 		{
-			if (cases[CaseIndex].branch)
+			if (Cases[CaseIndex].Branch)
 			{
-				FImageDesc ChildDesc = cases[CaseIndex].branch->GetImageDesc(bReturnBestOption, Context);
+				FImageDesc ChildDesc = Cases[CaseIndex].Branch->GetImageDesc(bReturnBestOption, Context);
 				if (bFirst)
 				{
 					Candidate = ChildDesc;
@@ -264,7 +253,7 @@ namespace mu
 		Result = Candidate;
 
 		// In case of ReturnBestOption the first valid case will be used to determine size and lods.
-		// Format will be the most generic from all cases.
+		// Format will be the most generic from all Cases.
 		if (!bSameFormat && !bReturnBestOption)
 		{
 			Result.m_format = EImageFormat::None;
@@ -290,17 +279,16 @@ namespace mu
 	}
 
 
-	//-------------------------------------------------------------------------------------------------
 	void ASTOpSwitch::GetBlockLayoutSize(uint64 BlockId, int32* pBlockX, int32* pBlockY, FBlockLayoutSizeCache* cache)
 	{
-		switch (type)
+		switch (Type)
 		{
 		case EOpType::LA_SWITCH:
 		{
 			Ptr<ASTOp> child = GetFirstValidValue();
 			if (!child)
 			{
-				child = def.child();
+				child = Default.child();
 			}
 
 			if (child)
@@ -322,17 +310,16 @@ namespace mu
 	}
 
 
-	//-------------------------------------------------------------------------------------------------
 	void ASTOpSwitch::GetLayoutBlockSize(int32* pBlockX, int32* pBlockY)
 	{
-		switch (type)
+		switch (Type)
 		{
 		case EOpType::IM_SWITCH:
 		{
 			Ptr<ASTOp> child = GetFirstValidValue();
 			if (!child)
 			{
-				child = def.child();
+				child = Default.child();
 			}
 
 			if (child)
@@ -352,28 +339,27 @@ namespace mu
 	}
 
 
-	//-------------------------------------------------------------------------------------------------
 	bool ASTOpSwitch::GetNonBlackRect(FImageRect& maskUsage) const
 	{
-		if (type == EOpType::IM_SWITCH)
+		if (Type == EOpType::IM_SWITCH)
 		{
 			FImageRect local;
 			bool localValid = false;
-			if (def)
+			if (Default)
 			{
-				localValid = def->GetNonBlackRect(local);
+				localValid = Default->GetNonBlackRect(local);
 				if (!localValid)
 				{
 					return false;
 				}
 			}
 
-			for (const auto& c : cases)
+			for (const FCase& c : Cases)
 			{
-				if (c.branch)
+				if (c.Branch)
 				{
 					FImageRect branchRect;
-					bool validBranch = c.branch->GetNonBlackRect(branchRect);
+					bool validBranch = c.Branch->GetNonBlackRect(branchRect);
 					if (validBranch)
 					{
 						if (localValid)
@@ -404,7 +390,6 @@ namespace mu
 	}
 
 
-	//-------------------------------------------------------------------------------------------------
 	bool ASTOpSwitch::IsImagePlainConstant(FVector4f&) const
 	{
 		// We could check if every option is plain and exactly the same colour, but probably it is
@@ -413,21 +398,20 @@ namespace mu
 	}
 
 
-	//-------------------------------------------------------------------------------------------------
 	mu::Ptr<ASTOp> ASTOpSwitch::OptimiseSemantic(const FModelOptimizationOptions&, int32 Pass) const
 	{
-		// Constant condition?
-		if (variable->GetOpType() == EOpType::NU_CONSTANT)
+		// Constant Condition?
+		if (Variable->GetOpType() == EOpType::NU_CONSTANT)
 		{
-			Ptr<ASTOp> Branch = def.child();
+			Ptr<ASTOp> Branch = Default.child();
 
-			const ASTOpConstantInt* typedCondition = static_cast<const ASTOpConstantInt*>(variable.child().get());
-			for (int32 o = 0; o < cases.Num(); ++o)
+			const ASTOpConstantInt* typedCondition = static_cast<const ASTOpConstantInt*>(Variable.child().get());
+			for (int32 o = 0; o < Cases.Num(); ++o)
 			{
-				if (cases[o].branch &&
-					typedCondition->Value == (int32)cases[o].condition)
+				if (Cases[o].Branch &&
+					typedCondition->Value == (int32)Cases[o].Condition)
 				{
-					Branch = cases[o].branch.child();
+					Branch = Cases[o].Branch.child();
 					break;
 				}
 			}
@@ -435,10 +419,10 @@ namespace mu
 			return Branch;
 		}
 
-		else if (variable->GetOpType() == EOpType::NU_PARAMETER)
+		else if (Variable->GetOpType() == EOpType::NU_PARAMETER)
 		{
 			// If all the branches for the possible values are the same op remove the instruction
-			const ASTOpParameter* ParamOp = static_cast<const ASTOpParameter*>(variable.child().get());
+			const ASTOpParameter* ParamOp = static_cast<const ASTOpParameter*>(Variable.child().get());
 			check(ParamOp);
 			if(ParamOp->Parameter.PossibleValues.IsEmpty())
 			{
@@ -450,13 +434,13 @@ namespace mu
 			Ptr<ASTOp> SameBranch = nullptr;
 			for (const FParameterDesc::FIntValueDesc& Value : ParamOp->Parameter.PossibleValues)
 			{
-				// Look for the switch branch it would take
-				Ptr<ASTOp> Branch = def.child();
-				for (const FCase& Case : cases)
+				// Look for the switch Branch it would take
+				Ptr<ASTOp> Branch = Default.child();
+				for (const FCase& Case : Cases)
 				{
-					if (Case.condition == Value.Value)
+					if (Case.Condition == Value.Value)
 					{
-						Branch = Case.branch.child();
+						Branch = Case.Branch.child();
 						break;
 					}
 				}
@@ -483,7 +467,7 @@ namespace mu
 	        }
 		}
 
-		// Ad-hoc logic optimization: check if all code paths leading to this operation have a switch with the same variable
+		// Ad-hoc logic optimization: check if all code paths leading to this operation have a switch with the same Variable
 		// and the option on those switches for the path that connects to this one is always the same. In that case, we can 
 		// remove this switch and replace it by the value it has for that option. 
 		// This is something the generic logic optimizer should do whan re-enabled.
@@ -493,7 +477,7 @@ namespace mu
 			Visited.Reserve(64);
 
 			// First is parent, second is what child we are reaching the parent from. This is necessary to find out what 
-			// switch branch we reach the parent from, if it is a switch.
+			// switch Branch we reach the parent from, if it is a switch.
 			TArray< TTuple<const ASTOp*, const ASTOp*>, TInlineAllocator<16>> Pending;
 			ForEachParent([this,&Pending](ASTOp* Parent)
 				{
@@ -518,23 +502,23 @@ namespace mu
 
 					bool bIsMatchingSwitch = false;
 
-					// TODO: Probably it could be a any switch, it doesn't need to be of the same type.
+					// TODO: Probably it could be a any switch, it doesn't need to be of the same Type.
 					if (Parent->GetOpType() == GetOpType())
 					{
 						const ASTOpSwitch* ParentSwitch = static_cast<const ASTOpSwitch*>(Parent);
 						check(ParentSwitch);
 
-						// To be compatible the switch must be on the same variable
-						if (ParentSwitch->variable==variable)
+						// To be compatible the switch must be on the same Variable
+						if (ParentSwitch->Variable==Variable)
 						{
 							bIsMatchingSwitch = true;
 							
 							// Find what switch option we are reaching it from
 							bool bIsSingleOption = true;
 							int OptionIndex = -1;
-							for (int32 CaseIndex = 0; CaseIndex < ParentSwitch->cases.Num(); ++CaseIndex)
+							for (int32 CaseIndex = 0; CaseIndex < ParentSwitch->Cases.Num(); ++CaseIndex)
 							{
-								if (ParentSwitch->cases[CaseIndex].branch.child().get() == ParentChild)
+								if (ParentSwitch->Cases[CaseIndex].Branch.child().get() == ParentChild)
 								{
 									if (OptionIndex != -1)
 									{
@@ -556,9 +540,9 @@ namespace mu
 							{
 								if (MatchingSwitchOption<0)
 								{
-									MatchingSwitchOption = ParentSwitch->cases[OptionIndex].condition;
+									MatchingSwitchOption = ParentSwitch->Cases[OptionIndex].Condition;
 								}
-								else if (MatchingSwitchOption!= ParentSwitch->cases[OptionIndex].condition)
+								else if (MatchingSwitchOption!= ParentSwitch->Cases[OptionIndex].Condition)
 								{
 									bAllPathsHaveMatchingSwitch = false;
 								}
@@ -587,7 +571,7 @@ namespace mu
 
 			if (bAllPathsHaveMatchingSwitch && MatchingSwitchOption>=0)
 			{
-				// We can remove this switch, all paths leading to it have the same condition for this switches variable.
+				// We can remove this switch, all paths leading to it have the same Condition for this switches Variable.
 				return FindBranch(MatchingSwitchOption);
 			}
 
@@ -597,23 +581,139 @@ namespace mu
 	}
 
 
-	//-------------------------------------------------------------------------------------------------
+	Ptr<ASTOp> ASTOpSwitch::OptimiseSink(const FModelOptimizationOptions&, FOptimizeSinkContext&) const
+	{
+		Ptr<ASTOp> NewOp;
+
+		// Detect if all Cases are the same op Type or they are null (same op with some branches being null).
+		EOpType BranchOpType = EOpType::NONE;
+		bool bSameOpTypeOrNull = true;
+
+		if (Default)
+		{
+			BranchOpType = Default->GetOpType();
+		}
+
+		for (const FCase& Case : Cases)
+		{
+			if (!Case.Branch)
+			{
+				continue;
+			}
+
+			if (BranchOpType==EOpType::NONE)
+			{
+				BranchOpType = Case.Branch->GetOpType();
+				break;
+			}
+			else if (Case.Branch->GetOpType() != BranchOpType)
+			{
+				bSameOpTypeOrNull = false;
+				break;
+			}
+		}
+
+		if (bSameOpTypeOrNull)
+		{
+			switch (BranchOpType)
+			{
+			case EOpType::ME_ADDTAGS:
+			{
+				// Move the add tags out of the switch if all tags are the same
+				bool bAllTagsAreTheSame = true;
+				TArray<FString> Tags;
+
+				if (Default)
+				{
+					check(Default->GetOpType()==EOpType::ME_ADDTAGS);
+					const ASTOpMeshAddTags* Typed = static_cast<const ASTOpMeshAddTags*>(Default.child().get());
+					Tags = Typed->Tags;
+				}
+
+				for (const FCase& Case : Cases)
+				{
+					if (!Case.Branch)
+					{
+						continue;
+					}
+
+					check(Case.Branch->GetOpType() == EOpType::ME_ADDTAGS);
+					const ASTOpMeshAddTags* Typed = static_cast<const ASTOpMeshAddTags*>(Case.Branch.child().get());
+					if (Tags.IsEmpty())
+					{
+						Tags = Typed->Tags;
+					}
+					else if (Typed->Tags!=Tags)
+					{
+						bAllTagsAreTheSame = false;
+						break;
+					}
+				}
+
+				if (bAllTagsAreTheSame)
+				{
+					Ptr<ASTOpMeshAddTags> New = new ASTOpMeshAddTags();
+					New->Tags = Tags;
+					
+					{
+						Ptr<ASTOpSwitch> NewSwitch = mu::Clone<ASTOpSwitch>(this);
+
+						// Replace all branches removing the "add tags" operation.
+						if (Default)
+						{
+							check(Default->GetOpType() == EOpType::ME_ADDTAGS);
+							const ASTOpMeshAddTags* Typed = static_cast<const ASTOpMeshAddTags*>(Default.child().get());
+							NewSwitch->Default = Typed->Source.child();
+						}
+
+						for (int32 CaseIndex=0; CaseIndex<Cases.Num();++CaseIndex)
+						{
+							const FCase& SourceCase = Cases[CaseIndex];
+							if (!SourceCase.Branch)
+							{
+								continue;
+							}
+
+							FCase& NewCase = NewSwitch->Cases[CaseIndex];
+
+							check(SourceCase.Branch->GetOpType() == EOpType::ME_ADDTAGS);
+							const ASTOpMeshAddTags* Typed = static_cast<const ASTOpMeshAddTags*>(SourceCase.Branch.child().get());
+							NewCase.Branch = Typed->Source.child();
+						}
+
+						New->Source = NewSwitch;
+					}
+
+					NewOp = New;
+				}
+				break;
+			}
+
+			default:
+				break;
+			}
+		}
+
+		return NewOp;
+	}
+
+
 	mu::Ptr<ImageSizeExpression> ASTOpSwitch::GetImageSizeExpression() const
 	{
 		Ptr<ImageSizeExpression> pRes = new ImageSizeExpression;
 
 		bool first = true;
-		for (const FCase& c : cases)
+		for (const FCase& c : Cases)
 		{
-			if (c.branch)
+			if (c.Branch)
 			{
 				if (first)
 				{
-					pRes = c.branch->GetImageSizeExpression();
+					pRes = c.Branch->GetImageSizeExpression();
 				}
 				else
 				{
-					Ptr<ImageSizeExpression> pOther = c.branch->GetImageSizeExpression();
+					Ptr<ImageSizeExpression> pOther = c.Branch->GetImageSizeExpression();
 					if (!(*pOther == *pRes))
 					{
 						pRes->type = ImageSizeExpression::ISET_UNKNOWN;
@@ -646,11 +746,11 @@ namespace mu
 		// Not cached: calculate
 		FSourceDataDescriptor Result;
 
-		for (const FCase& Case : cases)
+		for (const FCase& Case : Cases)
 		{
-			if (Case.branch)
+			if (Case.Branch)
 			{
-				FSourceDataDescriptor SourceDesc = Case.branch->GetSourceDataDescriptor(Context);
+				FSourceDataDescriptor SourceDesc = Case.Branch->GetSourceDataDescriptor(Context);
 				Result.CombineWith(SourceDesc);
 			}
 		}
diff --git a/Engine/Plugins/Mutable/Source/MutableTools/Private/MuT/ASTOpSwitch.h b/Engine/Plugins/Mutable/Source/MutableTools/Private/MuT/ASTOpSwitch.h
--- a/Engine/Plugins/Mutable/Source/MutableTools/Private/MuT/ASTOpSwitch.h
+++ b/Engine/Plugins/Mutable/Source/MutableTools/Private/MuT/ASTOpSwitch.h
@@ -23,34 +23,29 @@ template <class SCALAR> class vec4;
 	public:
 
 		//! Type of switch
-		EOpType type;
+		EOpType Type;
 
 		//! Variable whose value will be used to choose the switch branch.
-		ASTChild variable;
+		ASTChild Variable;
 
 		//! Default branch in case none matches the value in the variable.
-		ASTChild def;
+		ASTChild Default;
 
 		struct FCase
 		{
 			FCase(int32 cond, Ptr<ASTOp> parent, Ptr<ASTOp> b)
-				: condition(cond)
-				, branch(parent.get(), b)
+				: Condition(cond)
+				, Branch(parent.get(), b)
 			{
 			}
 
-			int32 condition;
-			ASTChild branch;
+			int32 Condition;
+			ASTChild Branch;
 
-			//!
-			bool operator==(const FCase& other) const
-			{
-				return condition == other.condition &&
-					branch == other.branch;
-			}
+			bool operator==(const FCase& other) const = default;
 		};
 
-		TArray<FCase> cases;
+		TArray<FCase> Cases;
 
 	public:
 
@@ -59,7 +54,7 @@ template <class SCALAR> class vec4;
 		~ASTOpSwitch() override;
 
 		// ASTOp interface
-		virtual EOpType GetOpType() const override { return type; }
+		virtual EOpType GetOpType() const override { return Type; }
 		virtual uint64 Hash() const override;
 		virtual void ForEachChild(const TFunctionRef<void(ASTChild&)> f) override;
 		virtual bool IsEqual(const ASTOp& otherUntyped) const override;
@@ -70,6 +65,7 @@ template <class SCALAR> class vec4;
 		virtual void GetBlockLayoutSize(uint64 BlockId, int32* pBlockX, int32* pBlockY, FBlockLayoutSizeCache* cache) override;
 		virtual void GetLayoutBlockSize(int32* pBlockX, int32* pBlockY) override;
 		virtual Ptr<ASTOp> OptimiseSemantic(const FModelOptimizationOptions&, int32 Pass) const override;
+		Ptr<ASTOp> OptimiseSink(const FModelOptimizationOptions&, FOptimizeSinkContext&) const override;
 		virtual bool GetNonBlackRect(FImageRect& maskUsage) const override;
 		virtual bool IsImagePlainConstant(FVector4f& colour) const override;
 		virtual Ptr<ImageSizeExpression> GetImageSizeExpression() const override;
diff --git a/Engine/Plugins/Mutable/Source/MutableTools/Private/MuT/CodeGenerator.cpp b/Engine/Plugins/Mutable/Source/MutableTools/Private/MuT/CodeGenerator.cpp
--- a/Engine/Plugins/Mutable/Source/MutableTools/Private/MuT/CodeGenerator.cpp
+++ b/Engine/Plugins/Mutable/Source/MutableTools/Private/MuT/CodeGenerator.cpp
@@ -706,19 +706,19 @@ namespace mu
 		}
 
 		Ptr<ASTOpSwitch> Op = new ASTOpSwitch();
-		Op->type = EOpType::IN_SWITCH;
+		Op->Type = EOpType::IN_SWITCH;
 
 		// Variable value
 		if (Node->Parameter)
 		{
 			FScalarGenerationResult ParamResult;
 			GenerateScalar(ParamResult, Options, Node->Parameter.get());
-			Op->variable = ParamResult.op;
+			Op->Variable = ParamResult.op;
 		}
 		else
 		{
 			// This argument is required
-			Op->variable = GenerateMissingScalarCode(TEXT("Switch variable"), 0.0f, Node->GetMessageContext());
+			Op->Variable = GenerateMissingScalarCode(TEXT("Switch variable"), 0.0f, Node->GetMessageContext());
 		}
 
 		// Options
@@ -738,7 +738,7 @@ namespace mu
 				Branch = Options.BaseInstance;
 			}
 
-			Op->cases.Emplace(OptionIndex, Op, Branch);
+			Op->Cases.Emplace(OptionIndex, Op, Branch);
 		}
 
 		Result.Op = Op;
diff --git a/Engine/Plugins/Mutable/Source/MutableTools/Private/MuT/CodeGenerator.h b/Engine/Plugins/Mutable/Source/MutableTools/Private/MuT/CodeGenerator.h
--- a/Engine/Plugins/Mutable/Source/MutableTools/Private/MuT/CodeGenerator.h
+++ b/Engine/Plugins/Mutable/Source/MutableTools/Private/MuT/CodeGenerator.h
@@ -955,9 +955,9 @@ namespace mu
         // Create the switch to cover all the options
         Ptr<ASTOp> lastSwitch;
         Ptr<ASTOpSwitch> SwitchOp = new ASTOpSwitch();
-		SwitchOp->type = OPTYPE;
-		SwitchOp->variable = Variable;
-		SwitchOp->def = nullptr;
+		SwitchOp->Type = OPTYPE;
+		SwitchOp->Variable = Variable;
+		SwitchOp->Default = nullptr;
 
 		for (int32 RowIndex = 0; RowIndex < NumRows; ++RowIndex)
         {
@@ -968,7 +968,7 @@ namespace mu
 
 			if (Branch || TYPE != ETableColumnType::Mesh)
 			{
-				SwitchOp->cases.Add(ASTOpSwitch::FCase(Condition, SwitchOp, Branch));
+				SwitchOp->Cases.Add(ASTOpSwitch::FCase(Condition, SwitchOp, Branch));
 			}
         }
 
diff --git a/Engine/Plugins/Mutable/Source/MutableTools/Private/MuT/CodeGenerator_Color.cpp b/Engine/Plugins/Mutable/Source/MutableTools/Private/MuT/CodeGenerator_Color.cpp
--- a/Engine/Plugins/Mutable/Source/MutableTools/Private/MuT/CodeGenerator_Color.cpp
+++ b/Engine/Plugins/Mutable/Source/MutableTools/Private/MuT/CodeGenerator_Color.cpp
@@ -207,19 +207,19 @@ namespace mu
 		}
 
 		Ptr<ASTOpSwitch> op = new ASTOpSwitch();
-		op->type = EOpType::CO_SWITCH;
+		op->Type = EOpType::CO_SWITCH;
 
 		// Variable value
 		if (node.Parameter)
 		{
 			FScalarGenerationResult ParamResult;
 			GenerateScalar(ParamResult, Options, node.Parameter.get());
-			op->variable = ParamResult.op;
+			op->Variable = ParamResult.op;
 		}
 		else
 		{
 			// This argument is required
-			op->variable = GenerateMissingScalarCode(TEXT("Switch variable"), 0.0f, Typed->GetMessageContext());
+			op->Variable = GenerateMissingScalarCode(TEXT("Switch variable"), 0.0f, Typed->GetMessageContext());
 		}
 
 		// Options
@@ -237,7 +237,7 @@ namespace mu
 				// This argument is required
 				branch = GenerateMissingColourCode(TEXT("Switch option"), Typed->GetMessageContext());
 			}
-			op->cases.Emplace((int16)OptionIndex, op, branch);
+			op->Cases.Emplace((int16)OptionIndex, op, branch);
 		}
 
 		result.op = op;
diff --git a/Engine/Plugins/Mutable/Source/MutableTools/Private/MuT/CodeGenerator_ExtensionData.cpp b/Engine/Plugins/Mutable/Source/MutableTools/Private/MuT/CodeGenerator_ExtensionData.cpp
--- a/Engine/Plugins/Mutable/Source/MutableTools/Private/MuT/CodeGenerator_ExtensionData.cpp
+++ b/Engine/Plugins/Mutable/Source/MutableTools/Private/MuT/CodeGenerator_ExtensionData.cpp
@@ -93,19 +93,19 @@ class Node;
 		}
 
 		Ptr<ASTOpSwitch> Op = new ASTOpSwitch;
-		Op->type = EOpType::ED_SWITCH;
+		Op->Type = EOpType::ED_SWITCH;
 
 		// Variable
 		if (Switch->Parameter)
 		{
 			FScalarGenerationResult ParamResult; 
 			GenerateScalar(ParamResult, Options, Switch->Parameter.get());
-			Op->variable = ParamResult.op;
+			Op->Variable = ParamResult.op;
 		}
 		else
 		{
 			// This argument is required
-			Op->variable = GenerateMissingScalarCode(TEXT("Switch variable"), 0.0f, Switch->GetMessageContext());
+			Op->Variable = GenerateMissingScalarCode(TEXT("Switch variable"), 0.0f, Switch->GetMessageContext());
 		}
 
 		// Options
@@ -126,7 +126,7 @@ class Node;
 				// This argument is required
 				Branch = GenerateMissingExtensionDataCode(TEXT("Switch option"), Switch->GetMessageContext());
 			}
-			Op->cases.Emplace(static_cast<int16_t>(OptionIndex), Op, Branch);
+			Op->Cases.Emplace(static_cast<int16_t>(OptionIndex), Op, Branch);
 		}
 
 		OutResult.Op = Op;
diff --git a/Engine/Plugins/Mutable/Source/MutableTools/Private/MuT/CodeGenerator_Image.cpp b/Engine/Plugins/Mutable/Source/MutableTools/Private/MuT/CodeGenerator_Image.cpp
--- a/Engine/Plugins/Mutable/Source/MutableTools/Private/MuT/CodeGenerator_Image.cpp
+++ b/Engine/Plugins/Mutable/Source/MutableTools/Private/MuT/CodeGenerator_Image.cpp
@@ -1177,19 +1177,19 @@ namespace mu
 		}
 
         Ptr<ASTOpSwitch> op = new ASTOpSwitch();
-        op->type = EOpType::IM_SWITCH;
+        op->Type = EOpType::IM_SWITCH;
 
 		// Variable value
 		if (Node->Parameter )
 		{
 			FScalarGenerationResult ParamResult;
 			GenerateScalar(ParamResult, Options, Node->Parameter.get());
-			op->variable = ParamResult.op;
+			op->Variable = ParamResult.op;
 		}
 		else
 		{
 			// This argument is required
-            op->variable = GenerateMissingScalarCode(TEXT("Switch variable"), 0.0f, Node->GetMessageContext() );
+            op->Variable = GenerateMissingScalarCode(TEXT("Switch variable"), 0.0f, Node->GetMessageContext() );
 		}
 
 		// Options
@@ -1209,7 +1209,7 @@ namespace mu
                 branch = GenerateMissingImageCode(TEXT("Switch option"), EImageFormat::RGBA_UByte, Node->GetMessageContext(), Options);
             }
 
-            op->cases.Emplace(int16(OptionIndex),op,branch);
+            op->Cases.Emplace(int16(OptionIndex),op,branch);
         }
 
         Ptr<ASTOp> switchAt = op;
diff --git a/Engine/Plugins/Mutable/Source/MutableTools/Private/MuT/CodeGenerator_Mesh.cpp b/Engine/Plugins/Mutable/Source/MutableTools/Private/MuT/CodeGenerator_Mesh.cpp
--- a/Engine/Plugins/Mutable/Source/MutableTools/Private/MuT/CodeGenerator_Mesh.cpp
+++ b/Engine/Plugins/Mutable/Source/MutableTools/Private/MuT/CodeGenerator_Mesh.cpp
@@ -393,8 +393,8 @@ namespace mu
         }
 
 		Ptr<ASTOpSwitch> SwitchOp = new ASTOpSwitch();
-		SwitchOp->type = EOpType::ME_SWITCH;
-		SwitchOp->variable = Variable;
+		SwitchOp->Type = EOpType::ME_SWITCH;
+		SwitchOp->Variable = Variable;
 
         // Options
 		bool bFirstValidConnectionFound = false;
@@ -449,7 +449,7 @@ namespace mu
 					}
 
 					Ptr<ASTOp> Branch = BranchResult.MeshOp;
-					SwitchOp->cases.Emplace(OptionIndices[ValidOptionIndex], SwitchOp, Branch);
+					SwitchOp->Cases.Emplace(OptionIndices[ValidOptionIndex], SwitchOp, Branch);
 				}
 
 				Result.MeshOp = SwitchOp;
diff --git a/Engine/Plugins/Mutable/Source/MutableTools/Private/MuT/CodeGenerator_Scalar.cpp b/Engine/Plugins/Mutable/Source/MutableTools/Private/MuT/CodeGenerator_Scalar.cpp
--- a/Engine/Plugins/Mutable/Source/MutableTools/Private/MuT/CodeGenerator_Scalar.cpp
+++ b/Engine/Plugins/Mutable/Source/MutableTools/Private/MuT/CodeGenerator_Scalar.cpp
@@ -238,19 +238,19 @@ namespace mu
 		}
 
 		Ptr<ASTOpSwitch> op = new ASTOpSwitch();
-		op->type = EOpType::SC_SWITCH;
+		op->Type = EOpType::SC_SWITCH;
 
 		// Variable value
 		if (node.Parameter)
 		{
 			FScalarGenerationResult ChildResult;
 			GenerateScalar(ChildResult, Options, node.Parameter.get());
-			op->variable = ChildResult.op;
+			op->Variable = ChildResult.op;
 		}
 		else
 		{
 			// This argument is required
-			op->variable = GenerateMissingScalarCode(TEXT("Switch variable"), 0.0f, Typed->GetMessageContext());
+			op->Variable = GenerateMissingScalarCode(TEXT("Switch variable"), 0.0f, Typed->GetMessageContext());
 		}
 
 		// Options
@@ -268,7 +268,7 @@ namespace mu
 				// This argument is required
 				branch = GenerateMissingScalarCode(TEXT("Switch option"), 1.0f, Typed->GetMessageContext());
 			}
-			op->cases.Emplace((int16_t)t, op, branch);
+			op->Cases.Emplace((int16_t)t, op, branch);
 		}
 
 		result.op = op;
diff --git a/Engine/Plugins/Mutable/Source/MutableTools/Private/MuT/CodeOptimiser.cpp b/Engine/Plugins/Mutable/Source/MutableTools/Private/MuT/CodeOptimiser.cpp
--- a/Engine/Plugins/Mutable/Source/MutableTools/Private/MuT/CodeOptimiser.cpp
+++ b/Engine/Plugins/Mutable/Source/MutableTools/Private/MuT/CodeOptimiser.cpp
@@ -683,7 +683,7 @@ namespace mu
 					case EOpType::ME_ADDTAGS:
 					{
 						const ASTOpMeshAddTags* AddTags = static_cast<const ASTOpMeshAddTags*>(SubgraphRoot.get());
-						if (AddTags->Source->GetParentCount()==1)
+						if (AddTags->Source && AddTags->Source->GetParentCount()==1)
 						{
 							// We could resolve the constant if there are no other constraints since this op's mesh op is only referenced once.
 							bGetFromChildren = true;
@@ -707,12 +707,7 @@ namespace mu
 					}
 
 					case EOpType::LA_FROMMESH:
-					{
-						// We want to reduce this type of operation regardless of it having special ops below.
-						bHasSpecialOpInSubgraph = false;
-						break;
-					}
-
+					case EOpType::ME_EXTRACTLAYOUTBLOCK:
 					case EOpType::ME_APPLYLAYOUT:
 					{
 						// We want to reduce this type of operation regardless of it having special ops below.
diff --git a/Engine/Plugins/Mutable/Source/MutableTools/Private/MuT/CodeOptimiser_Logic.cpp b/Engine/Plugins/Mutable/Source/MutableTools/Private/MuT/CodeOptimiser_Logic.cpp
--- a/Engine/Plugins/Mutable/Source/MutableTools/Private/MuT/CodeOptimiser_Logic.cpp
+++ b/Engine/Plugins/Mutable/Source/MutableTools/Private/MuT/CodeOptimiser_Logic.cpp
@@ -179,12 +179,12 @@ namespace mu
                          )
                     {
                         Ptr<ASTOpSwitch> switchOp = new ASTOpSwitch();
-                        switchOp->type = GetSwitchForType(GetOpDataType(topConditional->GetOpType()));
+                        switchOp->Type = GetSwitchForType(GetOpDataType(topConditional->GetOpType()));
 
 						ASTOpBoolEqualIntConst* firstCompare = static_cast<ASTOpBoolEqualIntConst*>(topConditional->condition.child().get());
-                        switchOp->variable = firstCompare->Value.child();
+                        switchOp->Variable = firstCompare->Value.child();
 
-                        auto current = n;
+						Ptr<ASTOp> current = n;
                         while(current)
                         {
                             bool valid = false;
@@ -209,9 +209,9 @@ namespace mu
                                 {
                                     auto compareValue = compare->Value.child();
                                     check(compare);
-                                    if ( compareValue == switchOp->variable.child() )
+                                    if ( compareValue == switchOp->Variable.child() )
                                     {
-                                        switchOp->cases.Emplace( compare->Constant, switchOp, conditional->yes.child() );
+                                        switchOp->Cases.Emplace( compare->Constant, switchOp, conditional->yes.child() );
 
                                         current = conditional->no.child();
                                         valid = true;
@@ -221,13 +221,13 @@ namespace mu
 
                             if (!valid)
                             {
-                                switchOp->def = current;
+                                switchOp->Default = current;
                                 current = nullptr;
                             }
                         }
 
                         const int MIN_CONDITIONS_TO_CREATE_SWITCH = 3;
-                        if (switchOp->cases.Num()>=MIN_CONDITIONS_TO_CREATE_SWITCH)
+                        if (switchOp->Cases.Num()>=MIN_CONDITIONS_TO_CREATE_SWITCH)
                         {
                             ASTOp::Replace(n,switchOp);
                             n = switchOp;
@@ -261,18 +261,18 @@ namespace mu
                 {
                     bool first = true;
 					EOpType caseType = EOpType::NONE;
-                    for (const ASTOpSwitch::FCase& c:topSwitch->cases)
+                    for (const ASTOpSwitch::FCase& c:topSwitch->Cases)
                     {
-                        if ( c.branch )
+                        if ( c.Branch )
                         {
                             if (first)
                             {
                                 first = false;
-                                caseType = c.branch->GetOpType();
+                                caseType = c.Branch->GetOpType();
                             }
                             else
                             {
-                                if (caseType!=c.branch->GetOpType())
+                                if (caseType!=c.Branch->GetOpType())
                                 {
                                     caseType=EOpType::NONE;
                                     break;
