[Mutable] Cloth Streaming

Add support for cloth when streaming Mutable Skeletal Meshes. Currently, when streamingi is enabled, physic mesh editions are not supported.

#rnx

diff --git a/Engine/Plugins/Mutable/Source/CustomizableObject/Internal/MuCO/CustomizableObjectInstancePrivate.h b/Engine/Plugins/Mutable/Source/CustomizableObject/Internal/MuCO/CustomizableObjectInstancePrivate.h
--- a/Engine/Plugins/Mutable/Source/CustomizableObject/Internal/MuCO/CustomizableObjectInstancePrivate.h
+++ b/Engine/Plugins/Mutable/Source/CustomizableObject/Internal/MuCO/CustomizableObjectInstancePrivate.h
@@ -302,7 +302,8 @@ private:
 	void BuildOrCopyElementData(const TSharedRef<FUpdateContextPrivate>& OperationData, USkeletalMesh* SkeletalMesh, UCustomizableObject& CustomizableObjectInstance, FCustomizableObjectInstanceComponentIndex InstanceComponentIndex);
 	void BuildOrCopyMorphTargetsData(const TSharedRef<FUpdateContextPrivate>& OperationData, USkeletalMesh* SkeletalMesh, const USkeletalMesh* SrcSkeletalMesh, UCustomizableObject& CustomizableObject, FCustomizableObjectInstanceComponentIndex InstanceComponentIndex);
 	bool BuildOrCopyRenderData(const TSharedRef<FUpdateContextPrivate>& OperationData, USkeletalMesh* SkeletalMesh, const USkeletalMesh* SrcSkeletalMesh, UCustomizableObjectInstance* CustomizableObjectInstance, FCustomizableObjectInstanceComponentIndex InstanceComponentIndex);
-	void BuildOrCopyClothingData(const TSharedRef<FUpdateContextPrivate>& OperationData, USkeletalMesh* SkeletalMesh, const USkeletalMesh* SrcSkeletalMesh, UCustomizableObjectInstance* CustomizableObjectInstance, FCustomizableObjectInstanceComponentIndex InstanceComponentIndex);
+
+	static void BuildOrCopyClothingData(const TSharedRef<FUpdateContextPrivate>& OperationData, USkeletalMesh* SkeletalMesh, const UModelResources& CustomizableObjectInstance, FCustomizableObjectInstanceComponentIndex InstanceComponentIndex, const TArray<TObjectPtr<UPhysicsAsset>>& ClothingPhysicsAssets);
 
 	// 
 	FCustomizableInstanceComponentData* GetComponentData(FCustomizableObjectComponentIndex ObjectComponentIndex);
diff --git a/Engine/Plugins/Mutable/Source/CustomizableObject/Internal/MuCO/CustomizableObjectSystemPrivate.h b/Engine/Plugins/Mutable/Source/CustomizableObject/Internal/MuCO/CustomizableObjectSystemPrivate.h
--- a/Engine/Plugins/Mutable/Source/CustomizableObject/Internal/MuCO/CustomizableObjectSystemPrivate.h
+++ b/Engine/Plugins/Mutable/Source/CustomizableObject/Internal/MuCO/CustomizableObjectSystemPrivate.h
@@ -46,7 +46,7 @@ struct FTexturePlatformData;
 extern TAutoConsoleVariable<bool> CVarMutableHighPriorityLoading;
 
 
-/** Final data to create a UMorphTarget. */
+/** Serialized Morph data. */
 struct FMorphTargetMeshData
 {
 	/** Names of the Morph Targets. */
@@ -57,6 +57,17 @@ struct FMorphTargetMeshData
 };
 
 
+/** Serialized Cloth data. */
+struct FClothingMeshData
+{
+	int32 ClothingAssetIndex = INDEX_NONE;
+	int32 ClothingAssetLOD = INDEX_NONE;
+
+	/** Per vertex data. */
+	TArray<FCustomizableObjectMeshToMeshVertData> Data;
+};
+
+
 /** Mapping of FMorphTargetVertexData local names to global names. */
 struct FMappedMorphTargetMeshData
 {
@@ -82,6 +93,13 @@ void ReconstructMorphTargets(const mu::FMesh& Mesh, const TArray<FName>& GlobalN
 void LoadMorphTargets(FMutableStreamRequest& MutableDataStreamer, const TSharedRef<const mu::FMesh>& Mesh, TMap<uint32, FMorphTargetMeshData>& StreamingResult);
 
 
+/** Request the Mutable Data Streamer to load Cloth blocks. Does not stream them.
+ *  @param MutableDataStreamer Streamer to request the blocks.
+ *  @param Mesh Mutable mesh to obtain which blocks to load.
+ *  @param StreamingResult Key is the block. Value is the container where the data will be deserialized once streamed. */
+void LoadClothing(FMutableStreamRequest& MutableDataStreamer, const TSharedRef<const mu::FMesh>& Mesh, TMap<uint32, FClothingMeshData>& StreamingResult);
+
+
 /** End a Customizable Object Instance Update. All code paths of an update have to end here. */
 void FinishUpdateGlobal(const TSharedRef<FUpdateContextPrivate>& Context);
 
@@ -511,14 +529,7 @@ struct FInstanceUpdateData
 	/** Key is the Component Name. Value is the final Morph Target data to copy into the Skeletal Mesh. */
 	TMap<FName, FSkeletalMeshMorphTargets> RealTimeMorphTargets;
 
-	struct FClothingMeshData
-	{
-		int32 ClothingAssetIndex = INDEX_NONE;
-		int32 ClothingAssetLOD = INDEX_NONE;
-		int32 PhysicsAssetIndex = INDEX_NONE;
-		TArray<FCustomizableObjectMeshToMeshVertData> Data;
-	};
-
+	/** Key is the Block Id. Value is the LoadAdditionalAssetsAndData read destination. */
 	TMap<uint32, FClothingMeshData> ClothingMeshData;
 
 	struct FSkeletonData
diff --git a/Engine/Plugins/Mutable/Source/CustomizableObject/Internal/MuCO/UnrealConversionUtils.h b/Engine/Plugins/Mutable/Source/CustomizableObject/Internal/MuCO/UnrealConversionUtils.h
--- a/Engine/Plugins/Mutable/Source/CustomizableObject/Internal/MuCO/UnrealConversionUtils.h
+++ b/Engine/Plugins/Mutable/Source/CustomizableObject/Internal/MuCO/UnrealConversionUtils.h
@@ -9,9 +9,14 @@
 
 #include "MuR/Ptr.h"
 #include "Containers/Array.h"
+#include "Containers/ArrayView.h"
 #include "Containers/Map.h"
 #include "Tasks/Task.h"
 
+struct FClothBufferIndexMapping;
+struct FMeshToMeshVertData;
+struct FCustomizableObjectMeshToMeshVertData;
+struct FClothingMeshData;
 struct FMorphTargetMeshData;
 
 namespace mu
@@ -29,6 +34,18 @@ class USkeleton;
 
 namespace UnrealConversionUtils
 {
+	struct FSectionClothData
+	{
+		int32 SectionIndex;
+		int32 LODIndex;
+		int32 BaseVertex;
+		TArrayView<const uint16> SectionIndex16View;
+		TArrayView<const uint32> SectionIndex32View;
+		TArrayView<const int32> ClothingDataIndicesView;
+		TArrayView<const FCustomizableObjectMeshToMeshVertData> ClothingDataView;
+		TArray<FMeshToMeshVertData> MappingData;
+	};
+	
 	/*
 	 * Shared methods between mutable instance update and viewport mesh generation
 	 * These are the result of stripping up parts of the reference code to be able to share it on the two
@@ -138,4 +155,12 @@ namespace UnrealConversionUtils
 
 	CUSTOMIZABLEOBJECT_API UE::Tasks::FTask ConvertSkeletalMeshFromRuntimeData(USkeletalMesh* SkeletalMesh, int32 LODIndex, int32 SectionIndex, UModelResources* ModelResources, mu::FMesh* Result);
 
+	void GetSectionClothData(const mu::FMesh& MutableMesh, int32 LODIndex, const TMap<uint32, FClothingMeshData>& ClothingMeshData, TArray<FSectionClothData>& SectionsClothData, int32& NumClothingDataNotFound);
+	
+	void CopyMeshToMeshClothData(TArray<FSectionClothData>& SectionsClothData);
+
+	/** Based on FSkeletalMeshLODModel::GetClothMappingData() */
+	void CreateClothMapping(const FSectionClothData& SectionClothData, TArray<FMeshToMeshVertData>& MappingData, TArray<FClothBufferIndexMapping>& ClothIndexMapping);
+	
+	void ClothVertexBuffers(FSkeletalMeshLODRenderData& LODResource, const mu::FMesh& MutableMesh, const TMap<uint32, FClothingMeshData>& ClothingMeshData, int32 LODIndex);
 }
diff --git a/Engine/Plugins/Mutable/Source/CustomizableObject/Private/MuCO/CustomizableObjectInstance.cpp b/Engine/Plugins/Mutable/Source/CustomizableObject/Private/MuCO/CustomizableObjectInstance.cpp
--- a/Engine/Plugins/Mutable/Source/CustomizableObject/Private/MuCO/CustomizableObjectInstance.cpp
+++ b/Engine/Plugins/Mutable/Source/CustomizableObject/Private/MuCO/CustomizableObjectInstance.cpp
@@ -8,7 +8,6 @@
 #include "Animation/AnimBlueprintGeneratedClass.h"
 #include "Animation/AnimInstance.h"
 #include "Animation/Skeleton.h"
-#include "AutoRTFM/AutoRTFM.h"
 #include "BoneControllers/AnimNode_RigidBody.h"
 #include "ClothConfig.h"
 #include "ClothingAsset.h"
@@ -2508,7 +2507,7 @@ bool UCustomizableInstancePrivate::UpdateSkeletalMesh_PostBeginUpdate0(UCustomiz
 		}
 
 		BuildOrCopyMorphTargetsData(OperationData, SkeletalMesh, OldSkeletalMesh, *CustomizableObject, InstanceComponentIndex);
-		BuildOrCopyClothingData(OperationData, SkeletalMesh, OldSkeletalMesh, Public, InstanceComponentIndex);
+		BuildOrCopyClothingData(OperationData, SkeletalMesh, ModelResources, InstanceComponentIndex, ClothingPhysicsAssets);
 
 		FSkeletalMeshRenderData* RenderData = SkeletalMesh->GetResourceForRendering();
 		ensure(RenderData && RenderData->LODRenderData.Num() > 0);
@@ -4272,17 +4271,30 @@ namespace
 void UCustomizableInstancePrivate::BuildOrCopyClothingData(
 		const TSharedRef<FUpdateContextPrivate>& OperationData, 
 		USkeletalMesh* SkeletalMesh, 
-		const USkeletalMesh* LastUpdateSkeletalMesh, 
-		UCustomizableObjectInstance* CustomizableObjectInstance, 
-		FCustomizableObjectInstanceComponentIndex InstanceComponentIndex)
+		const UModelResources& ModelResources,
+		FCustomizableObjectInstanceComponentIndex InstanceComponentIndex,
+		const TArray<TObjectPtr<UPhysicsAsset>>& ClothingPhysicsAssets)
 {
 	MUTABLE_CPUPROFILER_SCOPE(UCustomizableInstancePrivate::BuildOrCopyClothingData);
 
-	UCustomizableObject* CustomizableObject = CustomizableObjectInstance->GetCustomizableObject();
-	// It must be not null as it's checked in the calling function
-	check(CustomizableObject);
+	struct FSectionClothMetadata
+	{
+		int32 SectionIndex;
+		int32 LODIndex;
+		int32 ClothAssetIndex;
+		int32 ClothAssetLodIndex;
+		uint32 NumVertices; // Upper bound
+	};
 
-	const UModelResources& ModelResources = CustomizableObject->GetPrivate()->GetModelResourcesChecked();
+	struct FPerClothAssetData
+	{
+		int32 MinLOD = 0;
+		TArray<TArray<int32, TInlineAllocator<8>>, TInlineAllocator<8>> AttachedSections; // Indices in SectionsWithCloth for render sections attached to this ClothAsset.
+		FName Name;
+		UPhysicsAsset* PhysicsAsset = nullptr;
+		UClothingAssetCommon* ClothingAsset = nullptr;
+	};
+	
 	const TArray<FCustomizableObjectClothingAssetData>& ClothingAssetsData = ModelResources.ClothingAssetsData;
 	const TArray<FCustomizableObjectClothConfigData>& ClothSharedConfigsData = ModelResources.ClothSharedConfigsData;
 
@@ -4291,124 +4303,62 @@ void UCustomizableInstancePrivate::BuildOrCopyClothingData(
 		return;
 	}
 
-	const bool bAllowClothingPhysicsEdits = !bDisableClothingPhysicsEditsPropagation && ModelResources.bAllowClothingPhysicsEditsPropagation;
-
-	// First we need to discover if any clothing asset is used for the instance. 
-	struct FSectionWithClothData
-	{
-		int32 ClothAssetIndex;
-		int32 ClothAssetLodIndex;
-		int32 Section;
-		int32 Lod;
-		int32 BaseVertex;
-		TArrayView<const uint16> SectionIndex16View;
-		TArrayView<const uint32> SectionIndex32View;
-		TArrayView<const int32> ClothingDataIndicesView;
-		TArrayView<FCustomizableObjectMeshToMeshVertData> ClothingDataView;
-		TArray<FMeshToMeshVertData> MappingData;
-	};
-
-	TArray<FSectionWithClothData> SectionsWithCloth;
-	SectionsWithCloth.Reserve(32);
-
-	int32 NumClothingDataNotFound = 0;
-
-	const int32 LODCount = OperationData->InstanceUpdateData.LODs.Num();
-
 	const FInstanceUpdateData::FComponent* Component = OperationData->GetComponentUpdateData(InstanceComponentIndex);
 	if (!Component)
 	{
 		return;
 	}
 
-	const FName ComponentName = OperationData->ComponentNames[Component->Id.GetValue()];
+	const FName& ComponentName = OperationData->ComponentNames[Component->Id.GetValue()];
 
-	{
-		MUTABLE_CPUPROFILER_SCOPE(DiscoverSectionsWithCloth);
+	TArray<FSectionClothMetadata> SectionClothMetadata; // Sections must be sorted ascending
+	SectionClothMetadata.Reserve(32);
+	
+	TBitArray<> LODsWithClothing;
+	LODsWithClothing.Init(false, Component->LODCount);
 
-		// This should be FirstRequestedLOD but we because we do not share cloth data between LODs, we must also generate the hacky streaming ones.
-		const int32 FirstLOD = OperationData->FirstLODAvailable[ComponentName];
-		
-		for (int32 LODIndex = FirstLOD; LODIndex < Component->LODCount; ++LODIndex)
+	// Keep in mind that clothing does not do the Hacky Mutable Streaming Copy. This is because LOD data can not be shared between LODs.
+	// This means that LOD loops are a bit different form usual. With the hacky Mutable streaming, we must generate the requested and the hacky copied ones.
+	
+	// Metadata
+	{
+		int32 NumClothingDataNotFound = 0;
+			
+		for (int32 LODIndex = OperationData->FirstLODAvailable[ComponentName]; LODIndex < Component->LODCount; ++LODIndex)
 		{
 			const FInstanceUpdateData::FLOD& LOD = OperationData->InstanceUpdateData.LODs[Component->FirstLOD+LODIndex];
 
 			if (TSharedPtr<const mu::FMesh> MutableMesh = LOD.Mesh)
 			{
-				const mu::FMeshBufferSet& MeshSet = MutableMesh->GetVertexBuffers();
-				const mu::FMeshBufferSet& IndicesSet = MutableMesh->GetIndexBuffers();
-
-				int32 ClothingIndexBuffer, ClothingIndexChannel;
-				MeshSet.FindChannel(mu::EMeshBufferSemantic::Other, 2, &ClothingIndexBuffer, &ClothingIndexChannel);
-
-				int32 ClothingResourceBuffer, ClothingResourceChannel;
-				MeshSet.FindChannel(mu::EMeshBufferSemantic::Other, 3, &ClothingResourceBuffer, &ClothingResourceChannel);
-
-				if (ClothingIndexBuffer < 0 || ClothingResourceBuffer < 0)
+				for (int32 SectionIndex = 0; SectionIndex < MutableMesh->GetSurfaceCount(); ++SectionIndex)
 				{
-					continue;
-				}
-
-				const int32* const ClothingDataBuffer = reinterpret_cast<const int32*>(MeshSet.GetBufferData(ClothingIndexBuffer));
-				const uint32* const ClothingDataResource = reinterpret_cast<const uint32*>(MeshSet.GetBufferData(ClothingResourceBuffer));
-
-				const int32 SurfaceCount = MutableMesh->GetSurfaceCount();
-				for (int32 Section = 0; Section < SurfaceCount; ++Section)
-				{
-					int32 FirstVertex, VerticesCount, FirstIndex, IndicesCount, UnusedBoneIndex, UnusedBoneCount;
-					MutableMesh->GetSurface(Section, FirstVertex, VerticesCount, FirstIndex, IndicesCount, UnusedBoneIndex, UnusedBoneCount);
-
-					if (VerticesCount == 0 || IndicesCount == 0)
+					if (!MutableMesh->Surfaces[SectionIndex].SubMeshes.IsValidIndex(0))
 					{
 						continue;
 					}
 					
-					// A Section has cloth data on all its vertices or it does not have it in any.
-					// It can be determined if this section has clothing data just looking at the 
-					// first vertex of the section.
-					TArrayView<const int32> ClothingDataView(ClothingDataBuffer + FirstVertex, VerticesCount);
-					TArrayView<const uint32> ClothingResourceView(ClothingDataResource + FirstVertex, VerticesCount);
-
-					const int32 IndexCount = MutableMesh->GetIndexBuffers().GetElementCount();
+					const mu::FSurfaceSubMesh& SubMesh = MutableMesh->Surfaces[SectionIndex].SubMeshes[0];
+					const uint32 ClothResourceId = ModelResources.MeshMetadata[SubMesh.ExternalId].ClothingMetadataId;
 
-					TArrayView<const uint16> IndicesView16Bits;
-					TArrayView<const uint32> IndicesView32Bits;
-					
-					if (IndexCount)
-					{
-						 if (MutableMesh->GetIndexBuffers().GetElementSize(0) == 2)
-						 {
-						 	const uint16* IndexPtr = (const uint16*)MutableMesh->GetIndexBuffers().GetBufferData(0);
-							IndicesView16Bits = TArrayView<const uint16>(IndexPtr + FirstIndex, IndicesCount);
-						 }
-						 else
-						 {
-						 	const uint32* IndexPtr = (const uint32*)MutableMesh->GetIndexBuffers().GetBufferData(0);
-							IndicesView32Bits = TArrayView<const uint32>(IndexPtr + FirstIndex, IndicesCount);
-						 }
-					}
-				
-					if (!ClothingDataView.Num())
+					if (ClothResourceId == 0)
 					{
 						continue;
 					}
 
-					const uint32 ClothResourceId = ClothingResourceView[0];
-					if (ClothResourceId == 0)
+					if (MutableMesh->Surfaces[SectionIndex].SubMeshes.Num() > 1)
 					{
+						UE_LOG(LogMutable, Error, TEXT("Section %d has more than one submesh! Skipping section."), SectionIndex);
 						continue;
 					}
+					
+					FClothingMeshData* SectionClothingData = OperationData->InstanceUpdateData.ClothingMeshData.Find(ClothResourceId);
 
-					FInstanceUpdateData::FClothingMeshData* SectionClothingData = 
-							OperationData->InstanceUpdateData.ClothingMeshData.Find(ClothResourceId);
-	
 					if (!SectionClothingData)
 					{
 						++NumClothingDataNotFound;
 						continue;
 					}
 
-					check(SectionClothingData->Data.Num());
 					check(SectionClothingData->ClothingAssetIndex != INDEX_NONE);
 					check(SectionClothingData->ClothingAssetLOD != INDEX_NONE);
 
@@ -4425,101 +4375,83 @@ void UCustomizableInstancePrivate::BuildOrCopyClothingData(
 						continue;
 					}
 
-					SectionsWithCloth.Add(FSectionWithClothData {
-						ClothAssetIndex, ClothAssetLodIndex, 
-						Section, LODIndex, FirstVertex, 
-						IndicesView16Bits, IndicesView32Bits, 
-						ClothingDataView, 
-						MakeArrayView(SectionClothingData->Data.GetData(), SectionClothingData->Data.Num()),
-						TArray<FMeshToMeshVertData>() });
+					const uint32 NumVertices = SubMesh.VertexEnd - SubMesh.VertexBegin;
+					SectionClothMetadata.Add(FSectionClothMetadata { SectionIndex, LODIndex, ClothAssetIndex, ClothAssetLodIndex, NumVertices });
+
+					LODsWithClothing[LODIndex] = true;
 				}
 			}
 		}
+		
+		if (NumClothingDataNotFound > 0)
+		{
+			UE_LOG(LogMutable, Error, TEXT("Some clothing data could not be loaded properly, clothing assets may not behave as expected."));
+		}
 	}
 
-	if (NumClothingDataNotFound > 0)
+	// No clothing, early out.
+	if (!SectionClothMetadata.Num())
 	{
-		UE_LOG(LogMutable, Error, TEXT("Some clothing data could not be loaded properly, clothing assets may not behave as expected."));
+		return;
 	}
+	
+	TMap<int32, FPerClothAssetData> PerClothAssetData;
+	PerClothAssetData.Reserve(32);
 
-	if (!SectionsWithCloth.Num())
-	{
-		return; // Nothing to do.
-	}
-		
+	// Per Cloth Asset data
 	{
-		MUTABLE_CPUPROFILER_SCOPE(CopyMeshToMeshData)
-		
-		// Copy MeshToMeshData.
-		for (FSectionWithClothData& SectionWithCloth : SectionsWithCloth)
+		// Gather attached sections clothing asset LOD.
+		for (int32 MetadataIndex = 0; MetadataIndex < SectionClothMetadata.Num(); ++MetadataIndex)
 		{
-			const int32 NumVertices = SectionWithCloth.ClothingDataIndicesView.Num();
-			
-			TArray<FMeshToMeshVertData>& ClothMappingData = SectionWithCloth.MappingData;
-			ClothMappingData.SetNum(NumVertices);
+			const FSectionClothMetadata& SectionClothing = SectionClothMetadata[MetadataIndex];
+			FPerClothAssetData& AssetData = PerClothAssetData.FindOrAdd(SectionClothing.ClothAssetIndex); 
 
-			// Copy mesh to mesh data indexed by the index stored per vertex at compile time. 
-			for (int32 VertexIdx = 0; VertexIdx < NumVertices; ++VertexIdx)
-			{
-				// Possible Optimization: Gather consecutive indices in ClothingDataView and Memcpy the whole range.
-				// FMeshToMeshVertData and FCustomizableObjectMeshToMeshVertData have the same memory footprint and
-				// bytes in a FCustomizableObjectMeshToMeshVertData form a valid FMeshToMeshVertData (not the other way around).
+			AssetData.MinLOD = FMath::Min(AssetData.MinLOD, SectionClothing.ClothAssetLodIndex);
 
-				static_assert(sizeof(FCustomizableObjectMeshToMeshVertData) == sizeof(FMeshToMeshVertData));
-				static_assert(TIsTrivial<FCustomizableObjectMeshToMeshVertData>::Value);
-				static_assert(TIsTrivial<FMeshToMeshVertData>::Value);
+			int32 MaxLOD = FMath::Max<int32>(AssetData.AttachedSections.Num() - 1, SectionClothing.ClothAssetLodIndex);
 
-				const int32 VertexDataIndex = SectionWithCloth.ClothingDataIndicesView[VertexIdx];
-				check(VertexDataIndex >= 0);
-
-				const FCustomizableObjectMeshToMeshVertData& SrcData = SectionWithCloth.ClothingDataView[VertexDataIndex];
+			AssetData.AttachedSections.SetNum(MaxLOD + 1);
+			AssetData.AttachedSections[SectionClothing.ClothAssetLodIndex].Add(MetadataIndex);
+		}
 
-				FMeshToMeshVertData& DstData = ClothMappingData[VertexIdx];
-				FMemory::Memcpy(&DstData, &SrcData, sizeof(FMeshToMeshVertData));
-			}
+		for (TPair<int32, FPerClothAssetData>& Data : PerClothAssetData)
+		{
+			int32 ClothAssetIndex = Data.Key;
+			FPerClothAssetData& ClothAssetData = Data.Value; 
+	
+			ClothAssetData.Name = ClothingAssetsData[ClothAssetIndex].Name;
+			ClothAssetData.PhysicsAsset = ClothingPhysicsAssets[ClothAssetIndex];
 		}
 	}
-
-	struct FPerClothAssetData
-	{
-		int32 MinLOD = 0;
-
-		// Indices in SectionsWithCloth for render sections attached to this ClothAsset.
-		TArray<TArray<int32, TInlineAllocator<8>>, TInlineAllocator<8>> AttachedSections; 
-
-		FName Name;
-		
-		UPhysicsAsset* PhysicsAsset = nullptr;
-		UClothingAssetCommon* ClothingAsset = nullptr;
-	};
-
-	TMap<int32, FPerClothAssetData> PerClothAssetData;
-	PerClothAssetData.Reserve(32);
-
-	// Gather attached sections clothing asset LOD.
-	const int32 NumSectionsWithCloth = SectionsWithCloth.Num();
-	for (int32 SectionWithClothIndex = 0; SectionWithClothIndex < NumSectionsWithCloth; ++SectionWithClothIndex)
+	
+	TArray<UnrealConversionUtils::FSectionClothData> SectionsClothData; // Sorted by LOD, Section
+	SectionsClothData.Reserve(32);
+	
+	// Data
 	{
-		const FSectionWithClothData& SectionWithCloth = SectionsWithCloth[SectionWithClothIndex];
-		FPerClothAssetData& AssetData = PerClothAssetData.FindOrAdd(SectionWithCloth.ClothAssetIndex); 
+		MUTABLE_CPUPROFILER_SCOPE(DiscoverSectionsWithCloth);
 
-		AssetData.MinLOD = FMath::Min(AssetData.MinLOD, SectionWithCloth.ClothAssetLodIndex);
+		int32 NumClothingDataNotFound = 0;
 
-		int32 MaxLOD = FMath::Max<int32>(AssetData.AttachedSections.Num() - 1, SectionWithCloth.ClothAssetLodIndex);
+		for (int32 LODIndex = OperationData->FirstResidentLOD[ComponentName]; LODIndex < Component->LODCount; ++LODIndex)
+		{
+			const FInstanceUpdateData::FLOD& LOD = OperationData->InstanceUpdateData.LODs[Component->FirstLOD+LODIndex];
 
-		AssetData.AttachedSections.SetNum(MaxLOD + 1);
-		AssetData.AttachedSections[SectionWithCloth.ClothAssetLodIndex].Add(SectionWithClothIndex);
-	}
+			if (TSharedPtr<const mu::FMesh> MutableMesh = LOD.Mesh)
+			{
+				GetSectionClothData(*MutableMesh, LODIndex, OperationData->InstanceUpdateData.ClothingMeshData, SectionsClothData, NumClothingDataNotFound);
+			}
+		}
 
-	for (TPair<int32, FPerClothAssetData>& Data : PerClothAssetData)
-	{
-		int32 ClothAssetIndex = Data.Key;
-		FPerClothAssetData& ClothAssetData = Data.Value; 
-	
-		ClothAssetData.Name = ClothingAssetsData[ClothAssetIndex].Name;
-		ClothAssetData.PhysicsAsset = ClothingPhysicsAssets[ClothAssetIndex];
+		if (NumClothingDataNotFound > 0)
+		{
+			UE_LOG(LogMutable, Error, TEXT("Some clothing data could not be loaded properly, clothing assets may not behave as expected."));
+		}
+		
+		CopyMeshToMeshClothData(SectionsClothData);
 	}
 
+	// Create Clothing Assets
 	{ 
 		MUTABLE_CPUPROFILER_SCOPE(CreateClothingAssets)
 
@@ -4589,13 +4521,12 @@ void UCustomizableInstancePrivate::BuildOrCopyClothingData(
 			}
 
 			// Reconstruct clothing asset lod map.
-			NewClothingAsset->LodMap.Init(INDEX_NONE, LODCount);
+			NewClothingAsset->LodMap.Init(INDEX_NONE, Component->LODCount);
 			for (int32 LODIndex = 0; LODIndex < NumClothLODs; ++LODIndex)
 			{
-				const int32 NumClothSections = ClothAssetData.AttachedSections.Num(); 
 				for (int32 SectionWithClothIndex : ClothAssetData.AttachedSections[LODIndex])
 				{
-					NewClothingAsset->LodMap[SectionsWithCloth[SectionWithClothIndex].Lod] = LODIndex; 
+					NewClothingAsset->LodMap[SectionClothMetadata[SectionWithClothIndex].LODIndex] = LODIndex; 
 				}
 			}
 
@@ -4625,158 +4556,153 @@ void UCustomizableInstancePrivate::BuildOrCopyClothingData(
 		}
 	}
 
-	IMutableClothingModule* MutableClothingModule = nullptr;
+	const bool bAllowClothingPhysicsEdits = !bDisableClothingPhysicsEditsPropagation &&
+		ModelResources.bAllowClothingPhysicsEditsPropagation &&
+		!OperationData->bStreamMeshLODs;
+	
 	if (bAllowClothingPhysicsEdits)
 	{
-		MutableClothingModule = FModuleManager::GetModulePtr<IMutableClothingModule>(MUTABLE_CLOTHING_MODULE_NAME);	
-
-		if (!MutableClothingModule)
-		{
-			UE_LOG(LogMutable, Warning, TEXT("MutableClothing plugin could not be found. Make sure the plugin is enabled if you want to use advanced clothing features."));
-		}
-	}
-
-	if (MutableClothingModule)
-	{
-		for (TPair<int32, FPerClothAssetData>& Data : PerClothAssetData)
+		if (IMutableClothingModule* MutableClothingModule = FModuleManager::GetModulePtr<IMutableClothingModule>(MUTABLE_CLOTHING_MODULE_NAME))
 		{
-			int32 ClothAssetIndex = Data.Key;
-			FPerClothAssetData& ClothAssetData = Data.Value;
-
-			const int32 NumLODs = ClothAssetData.AttachedSections.Num() - ClothAssetData.MinLOD;
-
-			for (int32 LODIndex = ClothAssetData.MinLOD; LODIndex < NumLODs; ++LODIndex)
+			for (TPair<int32, FPerClothAssetData>& Data : PerClothAssetData)
 			{
-				for (int32 AttachedSectionIndex : ClothAssetData.AttachedSections[LODIndex])
+				FPerClothAssetData& ClothAssetData = Data.Value; 
+				UClothingAssetCommon* ClothingAsset = ClothAssetData.ClothingAsset;
+			
+				if (!ClothingAsset)
 				{
-					const FSectionWithClothData& ClothLODRenderSection = SectionsWithCloth[AttachedSectionIndex];
-					const int32 SectionVerticesNum = ClothLODRenderSection.ClothingDataIndicesView.Num();
+					continue;
 				}
-			}
-		}
-
-		for (TPair<int32, FPerClothAssetData>& Data : PerClothAssetData)
-		{
-			FPerClothAssetData& ClothAssetData = Data.Value; 
-			UClothingAssetCommon* ClothingAsset = ClothAssetData.ClothingAsset;
-			
-			if (!ClothingAsset)
-			{
-				continue;
-			}
-
-			bool bNeedsLodTransitionUpdate = false;
-			for (int32 LODIndex = 0; LODIndex < ClothingAsset->LodData.Num(); ++LODIndex)
-			{
-				TArray<TArrayView<FMeshToMeshVertData>, TInlineAllocator<8>> MeshToMeshDataViews;
 
-				for (int32 AttachedSectionIndex : ClothAssetData.AttachedSections[LODIndex])
+				bool bNeedsLodTransitionUpdate = false;
+				for (int32 LODIndex = 0; LODIndex < ClothingAsset->LodData.Num(); ++LODIndex)
 				{
-					MeshToMeshDataViews.Add(MakeArrayView(SectionsWithCloth[AttachedSectionIndex].MappingData));
-				}
+					TArray<TArrayView<FMeshToMeshVertData>, TInlineAllocator<8>> MeshToMeshDataViews;
 
-				bool bModified = MutableClothingModule->UpdateClothSimulationLOD(
-						LODIndex, *ClothingAsset, MakeConstArrayView(MeshToMeshDataViews));
+					for (int32 AttachedSectionIndex : ClothAssetData.AttachedSections[LODIndex])
+					{
+						MeshToMeshDataViews.Add(MakeArrayView(SectionsClothData[AttachedSectionIndex].MappingData));
+					}
+
+					bool bModified = MutableClothingModule->UpdateClothSimulationLOD(
+							LODIndex, *ClothingAsset, MakeConstArrayView(MeshToMeshDataViews));
 				
-				bNeedsLodTransitionUpdate = bNeedsLodTransitionUpdate || bModified;
-			}
+					bNeedsLodTransitionUpdate = bNeedsLodTransitionUpdate || bModified;
+				}
 
-			if (bNeedsLodTransitionUpdate)
-			{
-				// This needs to happen after all LODs have been processed.
-				for (int32 LODIndex = 0; LODIndex < ClothingAsset->LodData.Num(); ++LODIndex)
+				if (bNeedsLodTransitionUpdate)
 				{
-					MutableClothingModule->FixLODTransitionMappings(LODIndex, *ClothingAsset);
+					// This needs to happen after all LODs have been processed.
+					for (int32 LODIndex = 0; LODIndex < ClothingAsset->LodData.Num(); ++LODIndex)
+					{
+						MutableClothingModule->FixLODTransitionMappings(LODIndex, *ClothingAsset);
+					}
 				}
-			}
-		}	
+			}	
+		}
+		else
+		{
+			UE_LOG(LogMutable, Warning, TEXT("MutableClothing plugin could not be found. Make sure the plugin is enabled if you want to use advanced clothing features."));
+		}
+	}
+
+	TArray<TArray<FMeshToMeshVertData>> ResidentLODMappingData;
+	ResidentLODMappingData.SetNum(Component->LODCount);
+
+	TArray<TArray<FClothBufferIndexMapping>> ResidentLODClothIndexMapping;
+	ResidentLODClothIndexMapping.SetNum(Component->LODCount);
+
+	// Zero all LODs (even those which do not use cloth).
+	for (int32 LODIndex = OperationData->FirstLODAvailable[ComponentName]; LODIndex < Component->LODCount; ++LODIndex)
+	{
+		const FInstanceUpdateData::FLOD& LOD = OperationData->InstanceUpdateData.LODs[Component->FirstLOD + LODIndex];
+		
+		TSharedPtr<const mu::FMesh> MutableMesh = LOD.Mesh;
+		if (!MutableMesh)
+		{
+			continue;
+		}
+		
+		ResidentLODClothIndexMapping[LODIndex].SetNumZeroed(MutableMesh->GetSurfaceCount());
 	}
 
-	// Based on FSkeletalMeshLODModel::GetClothMappingData().
-	TArray<TArray<FMeshToMeshVertData>> LodMappingData;
-	LodMappingData.SetNum(LODCount);
+	// Create the mapping of cloth LODs.
+	for (const UnrealConversionUtils::FSectionClothData& Data : SectionsClothData)
+	{
+		CreateClothMapping(Data, ResidentLODMappingData[Data.LODIndex], ResidentLODClothIndexMapping[Data.LODIndex]);
+	}
 
-	TArray<TArray<FClothBufferIndexMapping>> LodClothingIndexMapping;
-	LodClothingIndexMapping.SetNum(LODCount);
+	FSkeletalMeshRenderData* RenderResource = SkeletalMesh->GetResourceForRendering();
 	{
-		int32 NumSectionsWithClothProcessed = 0;
+		MUTABLE_CPUPROFILER_SCOPE(InitClothRenderData)
 
-		for (int32 LODIndex = OperationData->GetFirstRequestedLOD()[ComponentName]; LODIndex < Component->LODCount; ++LODIndex)
+		// Streamed
+		for (int32 LODIndex = OperationData->FirstLODAvailable[ComponentName]; LODIndex < OperationData->FirstResidentLOD[ComponentName]; ++LODIndex)
 		{
-			const FInstanceUpdateData::FLOD& LOD = OperationData->InstanceUpdateData.LODs[Component->FirstLOD+LODIndex];
+			FSkeletalMeshLODRenderData& LODModel = RenderResource->LODRenderData[LODIndex];
 
-			if (TSharedPtr<const mu::FMesh> MutableMesh = LOD.Mesh)
+			if (LODsWithClothing[LODIndex])
 			{
-				TArray<FMeshToMeshVertData>& MappingData = LodMappingData[LODIndex];
-				TArray<FClothBufferIndexMapping>& ClothingIndexMapping = LodClothingIndexMapping[LODIndex];
-				ClothingIndexMapping.Reserve(32);
+				TArray<FClothBufferIndexMapping> ClothIndexMapping;
+				ClothIndexMapping.SetNumZeroed(LODModel.RenderSections.Num());
 
-				const int32 SurfaceCount = MutableMesh->GetSurfaceCount();
-				for (int32 Section = 0; Section < SurfaceCount; ++Section)
+				int32 Stride = sizeof(FMeshToMeshVertData);
+			
+				int32 NumVertices = 0; // Upper bound
+				for (const FSectionClothMetadata& Metadata : SectionClothMetadata)
 				{
-					// Check that is a valid surface.
-					int32 FirstVertex, VerticesCount, FirstIndex, IndicesCount, UnusedBoneIndex, UnusedBoneCount;
-					MutableMesh->GetSurface(Section, FirstVertex, VerticesCount, FirstIndex, IndicesCount, UnusedBoneIndex, UnusedBoneCount);
-
-					if (VerticesCount == 0 || IndicesCount == 0)
+					if (Metadata.LODIndex == LODIndex)
 					{
-						continue;
-					}					
-			
-					// An entry is added for all sections.	
-					FClothBufferIndexMapping& ClothBufferIndexMapping = ClothingIndexMapping.AddZeroed_GetRef();
+						// Based on FSkeletalMeshLODModel::GetClothMappingData().
 
-					if (NumSectionsWithClothProcessed < SectionsWithCloth.Num())
-					{
-						const FSectionWithClothData& SectionWithCloth = SectionsWithCloth[NumSectionsWithClothProcessed];
-						// Section with cloth are sorted by {LOD, Section}
-						if (SectionWithCloth.Lod == LODIndex && SectionWithCloth.Section == Section)
-						{	
-							ClothBufferIndexMapping.BaseVertexIndex = SectionWithCloth.BaseVertex;
-							ClothBufferIndexMapping.MappingOffset = (uint32)MappingData.Num();
-							ClothBufferIndexMapping.LODBiasStride = (uint32)SectionWithCloth.MappingData.Num();
+						FSkelMeshRenderSection& RenderSection = LODModel.RenderSections[Metadata.SectionIndex];
 						
-							MappingData += SectionWithCloth.MappingData; 
-
-							++NumSectionsWithClothProcessed;
-						}
+						check(Metadata.NumVertices == RenderSection.NumVertices); // Both values are upper bounds since we can not know the exact number of vertices without executing the code. 
+						
+						ClothIndexMapping[Metadata.SectionIndex].BaseVertexIndex = RenderSection.BaseVertexIndex;
+						ClothIndexMapping[Metadata.SectionIndex].MappingOffset = NumVertices;
+						ClothIndexMapping[Metadata.SectionIndex].LODBiasStride = Metadata.NumVertices;
+						
+						NumVertices += Metadata.NumVertices;
 					}
 				}
+				
+				LODModel.ClothVertexBuffer.SetMetadata(ClothIndexMapping, Stride, NumVertices);
 			}
 		}
-	}
 
-	FSkeletalMeshRenderData* RenderResource = SkeletalMesh->GetResourceForRendering();
-	{
-		MUTABLE_CPUPROFILER_SCOPE(InitClothRenderData)
-		// Based on FSkeletalMeshLODModel::GetClothMappingData().
-		for (int32 LODIndex = OperationData->GetFirstRequestedLOD()[ComponentName]; LODIndex < OperationData->NumLODsAvailable[ComponentName]; ++LODIndex)
+
+		// Resident
+		for (int32 LODIndex = OperationData->FirstResidentLOD[ComponentName]; LODIndex < OperationData->NumLODsAvailable[ComponentName]; ++LODIndex)
 		{
 			FSkeletalMeshLODRenderData& LODModel = RenderResource->LODRenderData[LODIndex];
 	
-			if (LodMappingData[LODIndex].Num() > 0)
+			if (LODsWithClothing[LODIndex])
 			{
-				LODModel.ClothVertexBuffer.Init(LodMappingData[LODIndex], LodClothingIndexMapping[LODIndex]);
+				LODModel.ClothVertexBuffer.Init(ResidentLODMappingData[LODIndex], ResidentLODClothIndexMapping[LODIndex]);
 			}
 		}
 	}
-
-	for (FSectionWithClothData& SectionWithCloth : SectionsWithCloth)
+	
+	for (const FSectionClothMetadata& Metadata : SectionClothMetadata)
 	{
-		FSkeletalMeshLODRenderData& LODModel = RenderResource->LODRenderData[SectionWithCloth.Lod];
-		FSkelMeshRenderSection& SectionData = LODModel.RenderSections[SectionWithCloth.Section];
-        
-		FPerClothAssetData& AssetData = PerClothAssetData.FindChecked(SectionWithCloth.ClothAssetIndex);
- 
+		FSkeletalMeshLODRenderData& LODModel = RenderResource->LODRenderData[Metadata.LODIndex];
+		FSkelMeshRenderSection& SectionData = LODModel.RenderSections[Metadata.SectionIndex];
+
+		// Ideally we would copy the data of all LODs, but we do not have this information in the initial generation. In any case,
+		// ClothMappingDataLODs is only used for CPU Skinning, and some engine checks (they only check the array size).
+		// The size must be a multiple of SectionData.NumVertices. Currently Mutable only supports one influence per vertex (NumVertices * 1).
 		SectionData.ClothMappingDataLODs.AddDefaulted(1);
-		SectionData.ClothMappingDataLODs[0] = MoveTemp(SectionWithCloth.MappingData);
-		
+		SectionData.ClothMappingDataLODs[0].SetNum(SectionData.NumVertices); // = MoveTemp(SectionWithCloth.MappingData);
+
+		FPerClothAssetData& AssetData = PerClothAssetData.FindChecked(Metadata.ClothAssetIndex);
+
 		SectionData.CorrespondClothAssetIndex = SkeletalMesh->GetClothingAssetIndex(AssetData.ClothingAsset);
 		SectionData.ClothingData.AssetGuid = AssetData.ClothingAsset->GetAssetGuid();
-		SectionData.ClothingData.AssetLodIndex = AssetData.ClothingAsset->LodMap[SectionWithCloth.Lod];
+		SectionData.ClothingData.AssetLodIndex = AssetData.ClothingAsset->LodMap[Metadata.LODIndex];
 	}
-
-	SkeletalMesh->SetHasActiveClothingAssets(static_cast<bool>( SectionsWithCloth.Num()));
+	
+	SkeletalMesh->SetHasActiveClothingAssets(!SectionClothMetadata.IsEmpty());
 }
 
 
@@ -5046,8 +4972,6 @@ UE::Tasks::FTask UCustomizableInstancePrivate::LoadAdditionalAssetsAndData(const
 	FMutableStreamRequest StreamRequest(ModelStreamableBulkData);
 
 	TArray<FSoftObjectPath> AssetsToStream;
-	TArray<uint32> RealTimeMorphStreamableBlocksToStream;
-	TArray<uint32> ClothingStreamableBlocksToStream;
 
 	TArray<FInstanceUpdateData::FComponent>& Components = OperationData->InstanceUpdateData.Components;
 
@@ -5184,7 +5108,8 @@ UE::Tasks::FTask UCustomizableInstancePrivate::LoadAdditionalAssetsAndData(const
 			}
 
 			const FName ComponentName = OperationData->ComponentNames[Component.Id.GetValue()];
-			for (int32 LODIndex = OperationData->FirstLODAvailable[ComponentName]; LODIndex < Component.LODCount; ++LODIndex)
+			
+			for (int32 LODIndex = OperationData->FirstResidentLOD[ComponentName]; LODIndex < Component.LODCount; ++LODIndex)
 			{
 				const FInstanceUpdateData::FLOD& LOD = OperationData->InstanceUpdateData.LODs[Component.FirstLOD + LODIndex];
 
@@ -5197,6 +5122,18 @@ UE::Tasks::FTask UCustomizableInstancePrivate::LoadAdditionalAssetsAndData(const
 				if (bUpdateMeshes)
 				{
 					LoadMorphTargets(StreamRequest, MutableMesh.ToSharedRef(), OperationData->InstanceUpdateData.RealTimeMorphTargetMeshData);
+					LoadClothing(StreamRequest, MutableMesh.ToSharedRef(), OperationData->InstanceUpdateData.ClothingMeshData);
+				}
+			}
+			
+			for (int32 LODIndex = OperationData->FirstLODAvailable[ComponentName]; LODIndex < Component.LODCount; ++LODIndex)
+			{
+				const FInstanceUpdateData::FLOD& LOD = OperationData->InstanceUpdateData.LODs[Component.FirstLOD + LODIndex];
+
+				TSharedPtr<const mu::FMesh> MutableMesh = LOD.Mesh;
+				if (!MutableMesh)
+				{
+					continue;
 				}
 				
 				const TArray<uint64>& StreamedResources = MutableMesh->GetStreamedResources();
@@ -5213,13 +5150,27 @@ UE::Tasks::FTask UCustomizableInstancePrivate::LoadAdditionalAssetsAndData(const
 #endif
 						ComponentData->StreamedResourceIndex.AddUnique(ResourceIndex);
 					}
+					else if (TypedResourceId.Type == static_cast<uint8>(FCustomizableObjectStreameableResourceId::EType::RealTimeMorphTarget))
+					{
+						check(TypedResourceId.Id != 0 && TypedResourceId.Id <= TNumericLimits<uint32>::Max());
+
+						if (ModelStreamableBulkData->RealTimeMorphStreamables.Contains(static_cast<uint32>(TypedResourceId.Id)))
+						{
+							FMorphTargetMeshData& ReadDestData = OperationData->InstanceUpdateData.RealTimeMorphTargetMeshData.FindOrAdd(TypedResourceId.Id);
+							ReadDestData.NameResolutionMap = ModelStreamableBulkData->RealTimeMorphStreamables[TypedResourceId.Id].NameResolutionMap;
+						}
+					}
 					else if (TypedResourceId.Type == (uint8)FCustomizableObjectStreameableResourceId::EType::Clothing)
 					{
 						check(TypedResourceId.Id != 0 && TypedResourceId.Id <= TNumericLimits<uint32>::Max());
-						
+
 						const TMap<uint32, FClothingStreamable>& ClothingStreamables = ModelStreamableBulkData->ClothingStreamables;
 						if (const FClothingStreamable* ClothingStreamable = ClothingStreamables.Find(TypedResourceId.Id))
 						{
+							FClothingMeshData& ReadDestData = OperationData->InstanceUpdateData.ClothingMeshData.FindOrAdd(TypedResourceId.Id);
+							ReadDestData.ClothingAssetIndex = ClothingStreamable->ClothingAssetIndex;
+							ReadDestData.ClothingAssetLOD = ClothingStreamable->ClothingAssetLOD;
+ 							
 							// TODO: Add async loading of ClothingAsset Data. This could be loaded as an streamead resource similar to and the asset user data.
 							int32 ClothingAssetIndex = ClothingStreamable->ClothingAssetIndex; 
 							int32 PhysicsAssetIndex = ClothingStreamable->PhysicsAssetIndex;					
@@ -5243,8 +5194,6 @@ UE::Tasks::FTask UCustomizableInstancePrivate::LoadAdditionalAssetsAndData(const
 									AssetsToStream.AddUnique(PhysicsAsset.ToSoftObjectPath());
 								}
 							}
-
-							ClothingStreamableBlocksToStream.AddUnique(TypedResourceId.Id);
 						}
 						else
 						{
@@ -5460,41 +5409,6 @@ UE::Tasks::FTask UCustomizableInstancePrivate::LoadAdditionalAssetsAndData(const
 			CVarMutableHighPriorityLoading.GetValueOnAnyThread() ? FStreamableManager::AsyncLoadHighPriority : FStreamableManager::DefaultAsyncLoadPriority);
 	}
 	
-	if (bUpdateMeshes)
-	{
-		// Clothing blocks to stream 
-		for (const int32 BlockId : ClothingStreamableBlocksToStream)
-		{
-			MUTABLE_CPUPROFILER_SCOPE(ClothingStreamingRequest_Alloc);
-
-			const FClothingStreamable& ClothingStreamable = ModelStreamableBulkData->ClothingStreamables[BlockId];
-			const FMutableStreamableBlock& Block = ClothingStreamable.Block;
-
-			FInstanceUpdateData::FClothingMeshData& ReadDestData = 
-				OperationData->InstanceUpdateData.ClothingMeshData.FindOrAdd(BlockId);
-
-			// Only request blocks once.
-			if (ReadDestData.Data.Num())
-			{
-				continue;
-			}
-
-			ReadDestData.ClothingAssetIndex = ClothingStreamable.ClothingAssetIndex;
-			ReadDestData.ClothingAssetLOD = ClothingStreamable.ClothingAssetLOD;
-			ReadDestData.PhysicsAssetIndex = ClothingStreamable.PhysicsAssetIndex;
-
-			check(ClothingStreamable.Size % sizeof(FCustomizableObjectMeshToMeshVertData) == 0);
-			const uint32 NumElems = ClothingStreamable.Size / sizeof(FCustomizableObjectMeshToMeshVertData);
-
-			ReadDestData.Data.SetNumUninitialized(NumElems);
-
-			StreamRequest.AddBlock(
-				Block,
-				MutablePrivate::EStreamableDataType::Clothing,
-				MakeArrayView(reinterpret_cast<uint8*>(ReadDestData.Data.GetData()), ReadDestData.Data.Num() * sizeof(FCustomizableObjectMeshToMeshVertData)));
-		}
-	}
-	
 	// Stream files
 	UE::Tasks::FTask StreamingTask = StreamRequest.Stream();
 	Prerequisites.Add(StreamingTask);
diff --git a/Engine/Plugins/Mutable/Source/CustomizableObject/Private/MuCO/CustomizableObjectMeshUpdate.cpp b/Engine/Plugins/Mutable/Source/CustomizableObject/Private/MuCO/CustomizableObjectMeshUpdate.cpp
--- a/Engine/Plugins/Mutable/Source/CustomizableObject/Private/MuCO/CustomizableObjectMeshUpdate.cpp
+++ b/Engine/Plugins/Mutable/Source/CustomizableObject/Private/MuCO/CustomizableObjectMeshUpdate.cpp
@@ -307,6 +307,8 @@ void FCustomizableObjectMeshStreamIn::RequestMeshUpdate(const FContext& Context)
 				{
 					LoadMorphTargets(StreamRequest, Mesh.ToSharedRef(), RefThis->OperationData->MorphTargetMeshData);
 				}
+				
+				LoadClothing(StreamRequest, Mesh.ToSharedRef(), RefThis->OperationData->ClothingMeshData);
 			}
 
 			// Stream data.
@@ -400,12 +402,8 @@ void FCustomizableObjectMeshStreamIn::ConvertMesh(const FContext& Context)
 		UnrealConversionUtils::CopyMutableIndexBuffers(LODResource, MutableMesh.Get());
 		UnrealConversionUtils::CopyMutableSkinWeightProfilesBuffers(LODResource, *const_cast<UCustomizableObjectSkeletalMesh*>(Mesh), LODIndex, MutableMesh.Get(), Mesh->SkinWeightProfileIDs);
 		UnrealConversionUtils::MorphTargetVertexInfoBuffers(LODResource, *Mesh, *MutableMesh.Get(), OperationData->MorphTargetMeshData, LODIndex);
-
-		// Clothing (ClothVertexBuffer)
-		{
-			// TODO PERE
-		}
-
+		UnrealConversionUtils::ClothVertexBuffers(LODResource, *MutableMesh.Get(), OperationData->ClothingMeshData, LODIndex);
+		
 		UnrealConversionUtils::UpdateSkeletalMeshLODRenderDataBuffersSize(LODResource);
 	}
 
diff --git a/Engine/Plugins/Mutable/Source/CustomizableObject/Private/MuCO/CustomizableObjectMeshUpdate.h b/Engine/Plugins/Mutable/Source/CustomizableObject/Private/MuCO/CustomizableObjectMeshUpdate.h
--- a/Engine/Plugins/Mutable/Source/CustomizableObject/Private/MuCO/CustomizableObjectMeshUpdate.h
+++ b/Engine/Plugins/Mutable/Source/CustomizableObject/Private/MuCO/CustomizableObjectMeshUpdate.h
@@ -41,6 +41,9 @@ struct FMutableMeshOperationData
 	/** Serialized data per block. */
 	TMap<uint32, FMorphTargetMeshData> MorphTargetMeshData;
 
+	/** Serialized data per block. */
+	TMap<uint32, FClothingMeshData> ClothingMeshData;
+
 	UE::Tasks::FTaskEvent Event = UE::Tasks::FTaskEvent(TEXT("Task_Mutable_UpdateMesh_End"));
 
 	/** SkeletalMesh */
diff --git a/Engine/Plugins/Mutable/Source/CustomizableObject/Private/MuCO/CustomizableObjectSystem.cpp b/Engine/Plugins/Mutable/Source/CustomizableObject/Private/MuCO/CustomizableObjectSystem.cpp
--- a/Engine/Plugins/Mutable/Source/CustomizableObject/Private/MuCO/CustomizableObjectSystem.cpp
+++ b/Engine/Plugins/Mutable/Source/CustomizableObject/Private/MuCO/CustomizableObjectSystem.cpp
@@ -4861,7 +4861,7 @@ int32 UCustomizableObjectSystemPrivate::BlockTillAllRequestsFinished(float TimeL
 
 void LoadMorphTargets(FMutableStreamRequest& MutableDataStreamer, const TSharedRef<const mu::FMesh>& Mesh, TMap<uint32, FMorphTargetMeshData>& StreamingResult)
 {
-	MUTABLE_CPUPROFILER_SCOPE(RealTimeMorphStreaming);
+	MUTABLE_CPUPROFILER_SCOPE(LoadMorphTargets);
 
 	if (!CVarEnableRealTimeMorphTargets.GetValueOnAnyThread())
 	{
@@ -4921,6 +4921,65 @@ void LoadMorphTargets(FMutableStreamRequest& MutableDataStreamer, const TSharedR
 }
 
 
+void LoadClothing(FMutableStreamRequest& MutableDataStreamer, const TSharedRef<const mu::FMesh>& Mesh, TMap<uint32, FClothingMeshData>& StreamingResult)
+{
+	MUTABLE_CPUPROFILER_SCOPE(LoadClothing);
+
+	TArray<uint32> ClothingStreamableBlocksToStream;
+
+	const TSharedPtr<FModelStreamableBulkData>& ModelStreamableBulkData = MutableDataStreamer.GetModelStreamableBulkData();
+
+	for (uint64 ResourceId : Mesh->GetStreamedResources())
+	{
+		FCustomizableObjectStreameableResourceId TypedResourceId = BitCast<FCustomizableObjectStreameableResourceId>(ResourceId);
+
+		if (TypedResourceId.Type == static_cast<uint8>(FCustomizableObjectStreameableResourceId::EType::Clothing))
+		{
+			check(TypedResourceId.Id != 0 && TypedResourceId.Id <= TNumericLimits<uint32>::Max());
+			
+			if (ModelStreamableBulkData->ClothingStreamables.Contains(TypedResourceId.Id))
+			{
+				ClothingStreamableBlocksToStream.AddUnique(TypedResourceId.Id);
+			}
+			else
+			{
+				UE_LOG(LogMutable, Error, TEXT("Invalid streamed clothing data block [%d] found."), TypedResourceId.Id);
+			}
+		}
+	}
+	
+	// Clothing blocks to stream 
+	for (const int32 BlockId : ClothingStreamableBlocksToStream)
+	{
+		MUTABLE_CPUPROFILER_SCOPE(ClothingStreamingRequest_Alloc);
+
+		const FClothingStreamable& ClothingStreamable = ModelStreamableBulkData->ClothingStreamables[BlockId];
+		const FMutableStreamableBlock& Block = ClothingStreamable.Block;
+
+		FClothingMeshData& ReadDestData = StreamingResult.FindOrAdd(BlockId);
+
+		// Only request blocks once.
+		if (ReadDestData.Data.Num())
+		{
+			continue;
+		}
+
+		ReadDestData.ClothingAssetIndex = ClothingStreamable.ClothingAssetIndex;
+		ReadDestData.ClothingAssetLOD = ClothingStreamable.ClothingAssetLOD;
+
+		check(ClothingStreamable.Size % sizeof(FCustomizableObjectMeshToMeshVertData) == 0);
+		const uint32 NumElems = ClothingStreamable.Size / sizeof(FCustomizableObjectMeshToMeshVertData);
+
+		ReadDestData.Data.SetNumUninitialized(NumElems);
+
+		MutableDataStreamer.AddBlock(
+			Block,
+			MutablePrivate::EStreamableDataType::Clothing,
+			MakeArrayView(reinterpret_cast<uint8*>(ReadDestData.Data.GetData()), ReadDestData.Data.Num() * sizeof(FCustomizableObjectMeshToMeshVertData)));
+	}
+}
+
+
 void ReconstructMorphTargets(const mu::FMesh& Mesh, const TArray<FName>& GlobalNames, const TMap<uint32, FMappedMorphTargetMeshData>& MappedMorphTargets, TArray<FMorphTargetLODModel>& OutMorphTargets)
 {
 	const mu::FMeshBufferSet& MeshSet = Mesh.GetVertexBuffers();
diff --git a/Engine/Plugins/Mutable/Source/CustomizableObject/Private/MuCO/UnrealConversionUtils.cpp b/Engine/Plugins/Mutable/Source/CustomizableObject/Private/MuCO/UnrealConversionUtils.cpp
--- a/Engine/Plugins/Mutable/Source/CustomizableObject/Private/MuCO/UnrealConversionUtils.cpp
+++ b/Engine/Plugins/Mutable/Source/CustomizableObject/Private/MuCO/UnrealConversionUtils.cpp
@@ -1278,4 +1278,173 @@ namespace UnrealConversionUtils
 		return FinalConversionTask;
 	}
 
+
+	void GetSectionClothData(const mu::FMesh& MutableMesh, int32 LODIndex, const TMap<uint32, FClothingMeshData>& ClothingMeshData, TArray<FSectionClothData>& SectionsClothData, int32& NumClothingDataNotFound)
+	{
+		MUTABLE_CPUPROFILER_SCOPE(GetSectionClothData);
+		
+		const mu::FMeshBufferSet& MeshSet = MutableMesh.GetVertexBuffers();
+
+		int32 ClothingIndexBuffer, ClothingIndexChannel;
+		MeshSet.FindChannel(mu::EMeshBufferSemantic::Other, 2, &ClothingIndexBuffer, &ClothingIndexChannel);
+
+		int32 ClothingResourceBuffer, ClothingResourceChannel;
+		MeshSet.FindChannel(mu::EMeshBufferSemantic::Other, 3, &ClothingResourceBuffer, &ClothingResourceChannel);
+
+		if (ClothingIndexBuffer >= 0 && ClothingResourceBuffer >= 0)
+		{
+			const int32* const ClothingDataBuffer = reinterpret_cast<const int32*>(MeshSet.GetBufferData(ClothingIndexBuffer));
+			const uint32* const ClothingDataResource = reinterpret_cast<const uint32*>(MeshSet.GetBufferData(ClothingResourceBuffer));
+
+			const int32 SurfaceCount = MutableMesh.GetSurfaceCount();
+			for (int32 SectionIndex = 0; SectionIndex < SurfaceCount; ++SectionIndex)
+			{
+				int32 FirstVertex, VerticesCount, FirstIndex, IndicesCount, UnusedBoneIndex, UnusedBoneCount;
+				MutableMesh.GetSurface(SectionIndex, FirstVertex, VerticesCount, FirstIndex, IndicesCount, UnusedBoneIndex, UnusedBoneCount);
+
+				if (VerticesCount == 0 || IndicesCount == 0)
+				{
+					continue;
+				}
+				
+				// A Section has cloth data on all its vertices or it does not have it in any.
+				// It can be determined if this section has clothing data just looking at the 
+				// first vertex of the section.
+				TArrayView<const int32> ClothingDataView(ClothingDataBuffer + FirstVertex, VerticesCount);
+				TArrayView<const uint32> ClothingResourceView(ClothingDataResource + FirstVertex, VerticesCount);
+
+				const int32 IndexCount = MutableMesh.GetIndexBuffers().GetElementCount();
+
+				TArrayView<const uint16> IndicesView16Bits;
+				TArrayView<const uint32> IndicesView32Bits;
+				
+				if (IndexCount)
+				{
+					if (MutableMesh.GetIndexBuffers().GetElementSize(0) == 2)
+					{
+						const uint16* IndexPtr = (const uint16*)MutableMesh.GetIndexBuffers().GetBufferData(0);
+						IndicesView16Bits = TArrayView<const uint16>(IndexPtr + FirstIndex, IndicesCount);
+					}
+					else
+					{
+						const uint32* IndexPtr = (const uint32*)MutableMesh.GetIndexBuffers().GetBufferData(0);
+						IndicesView32Bits = TArrayView<const uint32>(IndexPtr + FirstIndex, IndicesCount);
+					}
+				}
+			
+				if (!ClothingDataView.Num())
+				{
+					continue;
+				}
+
+				const uint32 ClothResourceId = ClothingResourceView[0];
+				if (ClothResourceId == 0)
+				{
+					continue;
+				}
+
+				const FClothingMeshData* SectionClothingData = ClothingMeshData.Find(ClothResourceId);
+
+				if (!SectionClothingData)
+				{
+					++NumClothingDataNotFound;
+					continue;
+				}
+
+				check(SectionClothingData->Data.Num());
+
+				SectionsClothData.Add(FSectionClothData {
+					SectionIndex,
+					LODIndex,
+					FirstVertex,
+					IndicesView16Bits,
+					IndicesView32Bits, 
+					ClothingDataView, 
+					MakeArrayView(SectionClothingData->Data.GetData(), SectionClothingData->Data.Num()),
+					TArray<FMeshToMeshVertData>() });
+			}
+		}
+	}
+
+
+	void CopyMeshToMeshClothData(TArray<FSectionClothData>& SectionsClothData)
+	{
+		MUTABLE_CPUPROFILER_SCOPE(ClothCopyMeshToMeshData)
+		
+		// Copy MeshToMeshData.
+		for (FSectionClothData& SectionWithCloth : SectionsClothData)
+		{
+			const int32 NumVertices = SectionWithCloth.ClothingDataIndicesView.Num();
+
+			TArray<FMeshToMeshVertData>& ClothMappingData = SectionWithCloth.MappingData;
+			ClothMappingData.SetNum(NumVertices);
+
+			// Copy mesh to mesh data indexed by the index stored per vertex at compile time. 
+			for (int32 VertexIdx = 0; VertexIdx < NumVertices; ++VertexIdx)
+			{
+				// Possible Optimization: Gather consecutive indices in ClothingDataView and Memcpy the whole range.
+				// FMeshToMeshVertData and FCustomizableObjectMeshToMeshVertData have the same memory footprint and
+				// bytes in a FCustomizableObjectMeshToMeshVertData form a valid FMeshToMeshVertData (not the other way around).
+
+				static_assert(sizeof(FCustomizableObjectMeshToMeshVertData) == sizeof(FMeshToMeshVertData));
+				static_assert(TIsTrivial<FCustomizableObjectMeshToMeshVertData>::Value);
+				static_assert(TIsTrivial<FMeshToMeshVertData>::Value);
+
+				const int32 VertexDataIndex = SectionWithCloth.ClothingDataIndicesView[VertexIdx];
+				check(VertexDataIndex >= 0);
+
+				const FCustomizableObjectMeshToMeshVertData& SrcData = SectionWithCloth.ClothingDataView[VertexDataIndex];
+
+				FMeshToMeshVertData& DstData = ClothMappingData[VertexIdx];
+				FMemory::Memcpy(&DstData, &SrcData, sizeof(FMeshToMeshVertData));
+			}
+		}
+	}
+
+
+	void CreateClothMapping(const FSectionClothData& SectionClothData, TArray<FMeshToMeshVertData>& MappingData, TArray<FClothBufferIndexMapping>& ClothIndexMapping)
+	{
+		ClothIndexMapping[SectionClothData.SectionIndex] = FClothBufferIndexMapping {
+			static_cast<uint32>(SectionClothData.BaseVertex),
+			static_cast<uint32>(MappingData.Num()),
+			static_cast<uint32>(SectionClothData.MappingData.Num()) };
+		
+		MappingData.Append(SectionClothData.MappingData);
+	}
+
+
+	void ClothVertexBuffers(FSkeletalMeshLODRenderData& LODResource, const mu::FMesh& MutableMesh, const TMap<uint32, FClothingMeshData>& ClothingMeshData, int32 LODIndex)
+	{
+		TArray<FSectionClothData> SectionsClothData;
+		SectionsClothData.Reserve(32);
+
+		int32 NumClothingDataNotFound = 0;
+
+		GetSectionClothData(MutableMesh, LODIndex, ClothingMeshData, SectionsClothData, NumClothingDataNotFound);
+	
+		if (NumClothingDataNotFound > 0)
+		{
+			UE_LOG(LogMutable, Error, TEXT("Some clothing data could not be loaded properly, clothing assets may not behave as expected."));
+		}
+
+		if (!SectionsClothData.Num())
+		{
+			return; // Nothing to do.
+		}
+
+		CopyMeshToMeshClothData(SectionsClothData);
+		
+		TArray<FMeshToMeshVertData> MappingData;
+		MappingData.Reserve(Algo::Accumulate(SectionsClothData, 0, [](int32 Sum, const FSectionClothData& Element){	return Sum + Element.MappingData.Num(); })); // Optimization.
+		
+		TArray<FClothBufferIndexMapping> ClothIndexMapping;
+		ClothIndexMapping.SetNumZeroed(LODResource.RenderSections.Num());
+
+		for (const FSectionClothData& Data : SectionsClothData)
+		{
+			CreateClothMapping(Data, MappingData, ClothIndexMapping);
+		}
+		
+		LODResource.ClothVertexBuffer.Init(MappingData, ClothIndexMapping);
+	}
 }
diff --git a/Engine/Plugins/Mutable/Source/CustomizableObjectEditor/Private/MuCOE/CustomizableObjectCompiler.cpp b/Engine/Plugins/Mutable/Source/CustomizableObjectEditor/Private/MuCOE/CustomizableObjectCompiler.cpp
--- a/Engine/Plugins/Mutable/Source/CustomizableObjectEditor/Private/MuCOE/CustomizableObjectCompiler.cpp
+++ b/Engine/Plugins/Mutable/Source/CustomizableObjectEditor/Private/MuCOE/CustomizableObjectCompiler.cpp
@@ -1396,7 +1396,7 @@ void FCustomizableObjectCompiler::FinishCompilationTask()
 		ModelStreamableBulkData->ClothingStreamables.Empty(32);
 
 		uint64 ClothingDataOffsetInBytes = 0;
-		for (TPair<uint32, FClothingMeshData>& MeshData : CompilationContext->ClothingPerMeshData)
+		for (TPair<uint32, FClothingMeshDataSource>& MeshData : CompilationContext->ClothingPerMeshData)
 		{
 			const uint32 DataSizeInBytes = (uint32)MeshData.Value.Data.Num() * sizeof(FCustomizableObjectMeshToMeshVertData);
 			FClothingStreamable& ResourceMeshData = ModelStreamableBulkData->ClothingStreamables.FindOrAdd(MeshData.Key);
diff --git a/Engine/Plugins/Mutable/Source/CustomizableObjectEditor/Private/MuCOE/GenerateMutableSource/GenerateMutableSource.h b/Engine/Plugins/Mutable/Source/CustomizableObjectEditor/Private/MuCOE/GenerateMutableSource/GenerateMutableSource.h
--- a/Engine/Plugins/Mutable/Source/CustomizableObjectEditor/Private/MuCOE/GenerateMutableSource/GenerateMutableSource.h
+++ b/Engine/Plugins/Mutable/Source/CustomizableObjectEditor/Private/MuCOE/GenerateMutableSource/GenerateMutableSource.h
@@ -271,7 +271,9 @@ struct FRealTimeMorphMeshData
 	uint32 SourceId = 0;
 };
 
-struct FClothingMeshData
+
+/** See FClothingMeshData. */
+struct FClothingMeshDataSource
 {
 	int32 ClothingAssetIndex = INDEX_NONE;
 	int32 ClothingAssetLOD = INDEX_NONE;
@@ -466,7 +468,7 @@ struct FMutableCompilationContext
 
 	/** Data used for Clothing reconstruction. */
 	TArray<FCustomizableObjectClothingAssetData> ClothingAssetsData;
-	TMap<uint32, FClothingMeshData> ClothingPerMeshData;
+	TMap<uint32, FClothingMeshDataSource> ClothingPerMeshData;
 
 	/** Stores the physics assets gathered from the SkeletalMesh nodes during compilation, to be used in mesh generation in-game */
 	TArray<TSoftObjectPtr<UPhysicsAsset>> PhysicsAssets;
diff --git a/Engine/Plugins/Mutable/Source/CustomizableObjectEditor/Private/MuCOE/GenerateMutableSource/GenerateMutableSourceMesh.cpp b/Engine/Plugins/Mutable/Source/CustomizableObjectEditor/Private/MuCOE/GenerateMutableSource/GenerateMutableSourceMesh.cpp
--- a/Engine/Plugins/Mutable/Source/CustomizableObjectEditor/Private/MuCOE/GenerateMutableSource/GenerateMutableSourceMesh.cpp
+++ b/Engine/Plugins/Mutable/Source/CustomizableObjectEditor/Private/MuCOE/GenerateMutableSource/GenerateMutableSourceMesh.cpp
@@ -1941,7 +1941,7 @@ TSharedPtr<mu::FMesh> ConvertSkeletalMeshToMutable(
 			const int32 ClothDataStride = bUseMutlipleInfluences ? NumInfluencesPerVertex : 1;
 			const int32 NumClothMappingDataVerts = ClothMappingData.Num() / ClothDataStride;
 
-			FClothingMeshData ClothingMeshData;
+			FClothingMeshDataSource ClothingMeshData;
 
 			const FCustomizableObjectClothingAssetData& ClothingAssetData = Context.ClothingAssetsData[ClothingAssetIndex];
 
@@ -1974,11 +1974,11 @@ TSharedPtr<mu::FMesh> ConvertSkeletalMeshToMutable(
 
 			ClothingUniqueHash = DataHash == InvalidResourceId ? DataHash + 1 : DataHash;
 
-			const FClothingMeshData* FoundHash = Context.ClothingPerMeshData.Find(ClothingUniqueHash);
+			const FClothingMeshDataSource* FoundHash = Context.ClothingPerMeshData.Find(ClothingUniqueHash);
 
 			bool bIsDataAlreadyCollected = false;
 
-			auto CompareClothingData = [](const FClothingMeshData& A, const FClothingMeshData& B) -> bool
+			auto CompareClothingData = [](const FClothingMeshDataSource& A, const FClothingMeshDataSource& B) -> bool
 				{
 					return A.ClothingAssetIndex == B.ClothingAssetIndex &&
 						A.ClothingAssetLOD == B.ClothingAssetLOD &&
@@ -2038,7 +2038,7 @@ TSharedPtr<mu::FMesh> ConvertSkeletalMeshToMutable(
 
 			if (!bIsDataAlreadyCollected)
 			{
-				FClothingMeshData& NewClothingMeshData = Context.ClothingPerMeshData.FindOrAdd(ClothingUniqueHash);
+				FClothingMeshDataSource& NewClothingMeshData = Context.ClothingPerMeshData.FindOrAdd(ClothingUniqueHash);
 
 				check(NewClothingMeshData.ClothingAssetIndex == INDEX_NONE);
 				check(NewClothingMeshData.Data.IsEmpty());
diff --git a/Engine/Source/Runtime/Engine/Private/Rendering/SkeletalMeshVertexClothBuffer.cpp b/Engine/Source/Runtime/Engine/Private/Rendering/SkeletalMeshVertexClothBuffer.cpp
--- a/Engine/Source/Runtime/Engine/Private/Rendering/SkeletalMeshVertexClothBuffer.cpp
+++ b/Engine/Source/Runtime/Engine/Private/Rendering/SkeletalMeshVertexClothBuffer.cpp
@@ -96,6 +96,14 @@ void FSkeletalMeshVertexClothBuffer::ReleaseRHIForStreaming(FRHIResourceReplaceB
 	}
 }
 
+void FSkeletalMeshVertexClothBuffer::SetMetadata(const TArray<FClothBufferIndexMapping>& InClothIndexMapping, uint32 InStride, uint32 InNumVertices)
+{
+	ClothIndexMapping = InClothIndexMapping;
+	Stride = InStride;
+	NumVertices = InNumVertices;
+}
+
+
 /**
 * Initialize the RHI resource for this vertex buffer
 */
diff --git a/Engine/Source/Runtime/Engine/Public/Rendering/SkeletalMeshVertexClothBuffer.h b/Engine/Source/Runtime/Engine/Public/Rendering/SkeletalMeshVertexClothBuffer.h
--- a/Engine/Source/Runtime/Engine/Public/Rendering/SkeletalMeshVertexClothBuffer.h
+++ b/Engine/Source/Runtime/Engine/Public/Rendering/SkeletalMeshVertexClothBuffer.h
@@ -130,6 +130,8 @@ public:
 	void InitRHIForStreaming(FRHIBuffer* IntermediateBuffer, FRHIResourceReplaceBatcher& Batcher);
 	void ReleaseRHIForStreaming(FRHIResourceReplaceBatcher& Batcher);
 
+	ENGINE_API void SetMetadata(const TArray<FClothBufferIndexMapping>& ClothIndexMapping, uint32 Stride, uint32 NumVertices);
+	
 private:
 	FShaderResourceViewRHIRef VertexBufferSRV;
 
