[Mutable] Fix automatic LOD updates not working when streaming enabled.

#rnx

diff --git a/Engine/Plugins/Mutable/Source/CustomizableObject/Private/MuCO/CustomizableObjectSystem.cpp b/Engine/Plugins/Mutable/Source/CustomizableObject/Private/MuCO/CustomizableObjectSystem.cpp
--- a/Engine/Plugins/Mutable/Source/CustomizableObject/Private/MuCO/CustomizableObjectSystem.cpp
+++ b/Engine/Plugins/Mutable/Source/CustomizableObject/Private/MuCO/CustomizableObjectSystem.cpp
@@ -4121,27 +4121,31 @@ int32 UCustomizableObjectSystem::TickInternal(const bool bBlocking)
 			{
 				const UCustomizableObjectInstance* Instance = LODUpdateTuple.Key;
 
-				if (Instance && Instance->GetCustomizableObject())
+				if (Instance)
 				{
-					FMutableUpdateCandidate& LODUpdateCandidate = LODUpdateTuple.Value;
-					ensure(LODUpdateCandidate.HasBeenIssued());
-
-					if (LODUpdateCandidate.Priority <= MaxPriorityFound)
+					if (UCustomizableObject* Object = Instance->GetCustomizableObject())
 					{
-						UCustomizableInstancePrivate* CustomizableInstancePrivate = LODUpdateCandidate.CustomizableObjectInstance->GetPrivate();
 
-						FDescriptorHash LODUpdateDescriptorHash = CustomizableInstancePrivate->CommittedDescriptorHash;
-						LODUpdateDescriptorHash.MinLODs = LODUpdateCandidate.MinLOD;
-						LODUpdateDescriptorHash.QualitySettingMinLODs = LODUpdateCandidate.QualitySettingMinLODs;
-						LODUpdateDescriptorHash.FirstRequestedLOD = LODUpdateCandidate.FirstRequestedLOD;
+						FMutableUpdateCandidate& LODUpdateCandidate = LODUpdateTuple.Value;
+						ensure(LODUpdateCandidate.HasBeenIssued());
 
-						if (CustomizableInstancePrivate->MinSquareDistFromComponentToPlayer < MaxSquareDistanceFound &&
-							!LODUpdateDescriptorHash.IsSubset(CustomizableInstancePrivate->CommittedDescriptorHash))
+						if (LODUpdateCandidate.Priority <= MaxPriorityFound)
 						{
-							MaxPriorityFound = LODUpdateCandidate.Priority;
-							MaxSquareDistanceFound = CustomizableInstancePrivate->MinSquareDistFromComponentToPlayer;
-							PendingInstanceUpdateFound = nullptr;
-							LODUpdateCandidateFound = &LODUpdateCandidate;
+							UCustomizableInstancePrivate* CustomizableInstancePrivate = LODUpdateCandidate.CustomizableObjectInstance->GetPrivate();
+
+							FDescriptorHash LODUpdateDescriptorHash = CustomizableInstancePrivate->CommittedDescriptorHash;
+							LODUpdateDescriptorHash.MinLODs = LODUpdateCandidate.MinLOD;
+							LODUpdateDescriptorHash.QualitySettingMinLODs = LODUpdateCandidate.QualitySettingMinLODs;
+							LODUpdateDescriptorHash.FirstRequestedLOD = LODUpdateCandidate.FirstRequestedLOD;
+
+							if (CustomizableInstancePrivate->MinSquareDistFromComponentToPlayer < MaxSquareDistanceFound &&
+								(IsStreamingEnabled(*Object) || !LODUpdateDescriptorHash.IsSubset(CustomizableInstancePrivate->CommittedDescriptorHash)))
+							{
+								MaxPriorityFound = LODUpdateCandidate.Priority;
+								MaxSquareDistanceFound = CustomizableInstancePrivate->MinSquareDistFromComponentToPlayer;
+								PendingInstanceUpdateFound = nullptr;
+								LODUpdateCandidateFound = &LODUpdateCandidate;
+							}
 						}
 					}
 				}
@@ -5119,4 +5123,4 @@ bool IsStreamingEnabled(const UCustomizableObject& Object)
 	return (Object.bEnableMeshStreaming || bForceStreamMeshLODs) &&
 		bStreamMeshLODs &&
 		IStreamingManager::Get().IsRenderAssetStreamingEnabled(EStreamableRenderAssetType::SkeletalMesh);
-}
\ No newline at end of file
+}
