[Mutable] Fix streaming crash due to an unexpected number of indices and sections.

#rnx

diff --git a/Engine/Plugins/Mutable/Source/CustomizableObject/Internal/MuCO/UnrealConversionUtils.h b/Engine/Plugins/Mutable/Source/CustomizableObject/Internal/MuCO/UnrealConversionUtils.h
--- a/Engine/Plugins/Mutable/Source/CustomizableObject/Internal/MuCO/UnrealConversionUtils.h
+++ b/Engine/Plugins/Mutable/Source/CustomizableObject/Internal/MuCO/UnrealConversionUtils.h
@@ -107,7 +107,8 @@ namespace UnrealConversionUtils
 	 */
 	CUSTOMIZABLEOBJECT_API bool CopyMutableIndexBuffers(
 		FSkeletalMeshLODRenderData& LODResource,
-		const mu::FMesh* InMutableMesh);
+		const mu::FMesh* InMutableMesh,
+		const TArray<uint32>& SurfaceIDs);
 	
 
 	/**
diff --git a/Engine/Plugins/Mutable/Source/CustomizableObject/Private/MuCO/CustomizableObjectInstance.cpp b/Engine/Plugins/Mutable/Source/CustomizableObject/Private/MuCO/CustomizableObjectInstance.cpp
--- a/Engine/Plugins/Mutable/Source/CustomizableObject/Private/MuCO/CustomizableObjectInstance.cpp
+++ b/Engine/Plugins/Mutable/Source/CustomizableObject/Private/MuCO/CustomizableObjectInstance.cpp
@@ -4882,8 +4882,17 @@ bool UCustomizableInstancePrivate::BuildOrCopyRenderData(const TSharedRef<FUpdat
 				LOD.Mesh.Get(),
 				SkeletalMesh->GetLODInfo(LODIndex)->bAllowCPUAccess);
 
+			// SurfaceIDs. Required to copy index buffers with padding
+			TArray<uint32> SurfaceIDs;
+			SurfaceIDs.SetNum(LOD.SurfaceCount);
+
+			for (int32 SurfaceIndex = 0; SurfaceIndex < LOD.SurfaceCount; ++SurfaceIndex)
+			{
+				SurfaceIDs[SurfaceIndex] = LOD.Mesh->GetSurfaceId(SurfaceIndex);
+			}
+
 			// Copy indices.
-			if (!UnrealConversionUtils::CopyMutableIndexBuffers(LODResource, LOD.Mesh.Get()))
+			if (!UnrealConversionUtils::CopyMutableIndexBuffers(LODResource, LOD.Mesh.Get(), SurfaceIDs))
 			{
 				// End with failure
 				return false;
diff --git a/Engine/Plugins/Mutable/Source/CustomizableObject/Private/MuCO/CustomizableObjectMeshUpdate.cpp b/Engine/Plugins/Mutable/Source/CustomizableObject/Private/MuCO/CustomizableObjectMeshUpdate.cpp
--- a/Engine/Plugins/Mutable/Source/CustomizableObject/Private/MuCO/CustomizableObjectMeshUpdate.cpp
+++ b/Engine/Plugins/Mutable/Source/CustomizableObject/Private/MuCO/CustomizableObjectMeshUpdate.cpp
@@ -399,7 +399,7 @@ void FCustomizableObjectMeshStreamIn::ConvertMesh(const FContext& Context)
 
 		FSkeletalMeshLODRenderData& LODResource = *Context.LODResourcesView[LODIndex];
 		UnrealConversionUtils::CopyMutableVertexBuffers(LODResource, MutableMesh.Get(), bNeedsCPUAccess);
-		UnrealConversionUtils::CopyMutableIndexBuffers(LODResource, MutableMesh.Get());
+		UnrealConversionUtils::CopyMutableIndexBuffers(LODResource, MutableMesh.Get(), Mesh->SurfaceIDs[LODIndex]);
 		UnrealConversionUtils::CopyMutableSkinWeightProfilesBuffers(LODResource, *const_cast<UCustomizableObjectSkeletalMesh*>(Mesh), LODIndex, MutableMesh.Get(), Mesh->SkinWeightProfileIDs);
 		UnrealConversionUtils::MorphTargetVertexInfoBuffers(LODResource, *Mesh, *MutableMesh.Get(), OperationData->MorphTargetMeshData, LODIndex);
 		UnrealConversionUtils::ClothVertexBuffers(LODResource, *MutableMesh.Get(), OperationData->ClothingMeshData, LODIndex);
diff --git a/Engine/Plugins/Mutable/Source/CustomizableObject/Private/MuCO/CustomizableObjectSkeletalMesh.cpp b/Engine/Plugins/Mutable/Source/CustomizableObject/Private/MuCO/CustomizableObjectSkeletalMesh.cpp
--- a/Engine/Plugins/Mutable/Source/CustomizableObject/Private/MuCO/CustomizableObjectSkeletalMesh.cpp
+++ b/Engine/Plugins/Mutable/Source/CustomizableObject/Private/MuCO/CustomizableObjectSkeletalMesh.cpp
@@ -30,6 +30,7 @@ void UCustomizableObjectSkeletalMesh::InitMutableStreamingData(const TSharedRef<
 	State = InOperationData->GetCapturedDescriptor().GetState();
 
 	MeshIDs.Init(MAX_uint64, MAX_MESH_LOD_COUNT);
+	SurfaceIDs.SetNum(MAX_MESH_LOD_COUNT);
 
 	const int32 FirstLOD = InOperationData->bStreamMeshLODs ?
 		InOperationData->FirstLODAvailable[ComponentName] :
@@ -39,6 +40,12 @@ void UCustomizableObjectSkeletalMesh::InitMutableStreamingData(const TSharedRef<
 	{
 		const FInstanceUpdateData::FLOD& LOD = InOperationData->InstanceUpdateData.LODs[InstanceUpdateFirstLOD + LODIndex];
 		MeshIDs[LODIndex] = LOD.MeshID;
+
+		SurfaceIDs[LODIndex].SetNum(LOD.SurfaceCount);
+		for (int32 SurfaceIndex = 0; SurfaceIndex < LOD.SurfaceCount; ++SurfaceIndex)
+		{
+			SurfaceIDs[LODIndex][SurfaceIndex] = InOperationData->InstanceUpdateData.Surfaces[LOD.FirstSurface + SurfaceIndex].SurfaceId;
+		}
 	}
 
 	const UModelResources& ModelResources = CustomizableObject.GetPrivate()->GetModelResourcesChecked();
diff --git a/Engine/Plugins/Mutable/Source/CustomizableObject/Private/MuCO/UnrealConversionUtils.cpp b/Engine/Plugins/Mutable/Source/CustomizableObject/Private/MuCO/UnrealConversionUtils.cpp
--- a/Engine/Plugins/Mutable/Source/CustomizableObject/Private/MuCO/UnrealConversionUtils.cpp
+++ b/Engine/Plugins/Mutable/Source/CustomizableObject/Private/MuCO/UnrealConversionUtils.cpp
@@ -463,7 +463,8 @@ namespace UnrealConversionUtils
 
 	bool CopyMutableIndexBuffers(
 		FSkeletalMeshLODRenderData& LODResource,
-		const mu::FMesh* InMutableMesh)
+		const mu::FMesh* InMutableMesh,
+		const TArray<uint32>& SurfaceIDs)
 	{
 		MUTABLE_CPUPROFILER_SCOPE(CopyMutableIndexBuffers);
 
@@ -490,11 +491,29 @@ namespace UnrealConversionUtils
 			LODResource.MultiSizeIndexContainer.CreateIndexBuffer(ElementSize);
 		}
 		
-		check(LODResource.MultiSizeIndexContainer.GetDataTypeSize() == ElementSize)
+		check(LODResource.MultiSizeIndexContainer.GetDataTypeSize() == ElementSize);
 		
-		LODResource.MultiSizeIndexContainer.GetIndexBuffer()->Insert(0, IndexCount);
-		FMemory::Memcpy(LODResource.MultiSizeIndexContainer.GetIndexBuffer()->GetPointerTo(0), DataPtr, IndexCount * ElementSize);
+		const int32 MaxNumIndices = LODResource.RenderSections.Last().BaseIndex + (LODResource.RenderSections.Last().NumTriangles * 3);
+		LODResource.MultiSizeIndexContainer.GetIndexBuffer()->Insert(0, MaxNumIndices);
 
+		if (IndexCount == MaxNumIndices)
+		{
+			FMemory::Memcpy(LODResource.MultiSizeIndexContainer.GetIndexBuffer()->GetPointerTo(0), DataPtr, IndexCount * ElementSize);
+		}
+		else
+		{
+			for (const mu::FMeshSurface& Surface : InMutableMesh->Surfaces)
+			{
+				const int32 SectionIndex = SurfaceIDs.Find(Surface.Id);
+				if (LODResource.RenderSections.IsValidIndex(SectionIndex) && !Surface.SubMeshes.IsEmpty())
+				{
+					void* DestPtr = LODResource.MultiSizeIndexContainer.GetIndexBuffer()->GetPointerTo(LODResource.RenderSections[SectionIndex].BaseIndex);
+					const int32 DataOffset = Surface.SubMeshes[0].IndexBegin * ElementSize;
+					const int32 DataSize = Surface.SubMeshes.Last().IndexEnd * ElementSize - DataOffset;
+					FMemory::Memcpy(DestPtr, DataPtr + DataOffset, DataSize);
+				}
+			}
+		}
 		return true;
 	}
 
diff --git a/Engine/Plugins/Mutable/Source/CustomizableObject/Public/MuCO/CustomizableObjectSkeletalMesh.h b/Engine/Plugins/Mutable/Source/CustomizableObject/Public/MuCO/CustomizableObjectSkeletalMesh.h
--- a/Engine/Plugins/Mutable/Source/CustomizableObject/Public/MuCO/CustomizableObjectSkeletalMesh.h
+++ b/Engine/Plugins/Mutable/Source/CustomizableObject/Public/MuCO/CustomizableObjectSkeletalMesh.h
@@ -53,5 +53,6 @@ class CUSTOMIZABLEOBJECT_API UCustomizableObjectSkeletalMesh : public USkeletalM
 
 	/** Mesh IDs per LOD */
 	TArray<mu::FResourceID> MeshIDs;
+	TArray<TArray<uint32>> SurfaceIDs;
 	TArray<TPair<uint32, FName>> SkinWeightProfileIDs;
 };
diff --git a/Engine/Plugins/Mutable/Source/CustomizableObjectEditor/Private/MuCOE/MutableMeshPreviewUtils.cpp b/Engine/Plugins/Mutable/Source/CustomizableObjectEditor/Private/MuCOE/MutableMeshPreviewUtils.cpp
--- a/Engine/Plugins/Mutable/Source/CustomizableObjectEditor/Private/MuCOE/MutableMeshPreviewUtils.cpp
+++ b/Engine/Plugins/Mutable/Source/CustomizableObjectEditor/Private/MuCOE/MutableMeshPreviewUtils.cpp
@@ -428,8 +428,18 @@ namespace MutableMeshPreviewUtils
 			LODResource.ActiveBoneIndices.Add(0);
 			LODResource.RequiredBones.Add(0);
 			
+			// SurfaceIDs. Required to copy index buffers with padding
+			const int32 NumSurfaces = InMutableMesh->Surfaces.Num();
+			TArray<uint32> SurfaceIDs;
+			SurfaceIDs.SetNum(NumSurfaces);
+
+			for (int32 SurfaceIndex = 0; SurfaceIndex < NumSurfaces; ++SurfaceIndex)
+			{
+				SurfaceIDs[SurfaceIndex] = InMutableMesh->GetSurfaceId(SurfaceIndex);
+			}
+
 			// Copy index buffers or fail to generate the mesh
-			if (!UnrealConversionUtils::CopyMutableIndexBuffers(LODResource, InMutableMesh))
+			if (!UnrealConversionUtils::CopyMutableIndexBuffers(LODResource, InMutableMesh, SurfaceIDs))
 			{
 				return false;
 			}
