[Mutable] Skip MeshPrepareLayout operation on descriptor meshes.
Optimized ME_RemoveMask_2 by 2x factor.

#rnx

diff --git a/Engine/Plugins/Mutable/Source/MutableRuntime/Private/MuR/CodeRunner.cpp b/Engine/Plugins/Mutable/Source/MutableRuntime/Private/MuR/CodeRunner.cpp
--- a/Engine/Plugins/Mutable/Source/MutableRuntime/Private/MuR/CodeRunner.cpp
+++ b/Engine/Plugins/Mutable/Source/MutableRuntime/Private/MuR/CodeRunner.cpp
@@ -1168,31 +1168,46 @@ namespace mu
 			{
 			case 0:
 			{
-				AddOp(FScheduledOp(Item.At, Item, 1),
-					FScheduledOp(Args.Mesh, Item),
-					FScheduledOp(Args.Layout, Item));
+				if (!EnumHasAnyFlags(EMeshContentFlags(Item.ExecutionOptions), EMeshContentFlags::GeometryData))
+				{
+					AddOp(FScheduledOp(Item.At, Item, 1),
+						FScheduledOp(Args.Mesh, Item));
+				}
+				else
+				{
+					AddOp(FScheduledOp(Item.At, Item, 1),
+						FScheduledOp(Args.Mesh, Item),
+						FScheduledOp(Args.Layout, Item));
+				}
 				break;
 			}
 			case 1:
 			{
-				MUTABLE_CPUPROFILER_SCOPE(ME_PREPARELAYOUT)
+				MUTABLE_CPUPROFILER_SCOPE(ME_PREPARELAYOUT);
 
-				TSharedPtr<const FMesh> Base = LoadMesh(FCacheAddress(Args.Mesh, Item));
-				TSharedPtr<const FLayout> Layout = LoadLayout(FCacheAddress(Args.Layout, Item));
-
-				if (Base && Layout)
+				if (!EnumHasAnyFlags(EMeshContentFlags(Item.ExecutionOptions), EMeshContentFlags::GeometryData))
 				{
-					TSharedPtr<FMesh> Result = CloneOrTakeOver(Base);
-
-					MeshPrepareLayout(*Result, *Layout, Args.LayoutChannel, Args.bNormalizeUVs, Args.bClampUVIslands, Args.bEnsureAllVerticesHaveLayoutBlock, Args.bUseAbsoluteBlockIds );
-
-					StoreMesh(Item, Result);
+					TSharedPtr<const FMesh> Base = LoadMesh(FCacheAddress(Args.Mesh, Item));
+					StoreMesh(Item, Base);
 				}
 				else
 				{
-					StoreMesh(Item, Base);
-				}
+					TSharedPtr<const FMesh> Base = LoadMesh(FCacheAddress(Args.Mesh, Item));
+					TSharedPtr<const FLayout> Layout = LoadLayout(FCacheAddress(Args.Layout, Item));
+
+					if (Base && Layout)
+					{
+						TSharedPtr<FMesh> Result = CloneOrTakeOver(Base);
 
+						MeshPrepareLayout(*Result, *Layout, Args.LayoutChannel, Args.bNormalizeUVs, Args.bClampUVIslands, Args.bEnsureAllVerticesHaveLayoutBlock, Args.bUseAbsoluteBlockIds);
+
+						StoreMesh(Item, Result);
+					}
+					else
+					{
+						StoreMesh(Item, Base);
+					}
+				}
 
 				break;
 			}
@@ -2852,10 +2867,7 @@ namespace mu
 
 				if (Source)
 				{
-					TSharedPtr<FMesh> Result = CreateMesh(Source->GetDataSize());
-					Result->CopyFrom(*Source);
-
-					Release(Source);
+					TSharedPtr<FMesh> Result = CloneOrTakeOver(Source);
 
 					for (int32 r = 0; r < conditions.Num(); ++r)
 					{
@@ -2871,24 +2883,10 @@ namespace mu
 							TSharedPtr<const FMesh> Mask = LoadMesh(FCacheAddress(masks[r], Item));
 							if (Mask)
 							{
-								//MeshRemoveMask will make a copy of Result, try to make room for it. 
-								TSharedPtr<FMesh> IterResult = CreateMesh(Result->GetDataSize());
-
-								bool bOutSuccess = false;
 								bool bRemoveIfAllVerticesCulled = FaceCullStrategy == EFaceCullStrategy::AllVerticesCulled;
-								MeshRemoveMask(IterResult.Get(), Result.Get(), Mask.Get(), bRemoveIfAllVerticesCulled, bOutSuccess);
+								MeshRemoveMaskInline(Result.Get(), Mask.Get(), bRemoveIfAllVerticesCulled);
 
 								Release(Mask);
-
-								if (!bOutSuccess)
-								{
-									Release(IterResult);
-								}
-								else
-								{
-									Swap(Result, IterResult);
-									Release(IterResult);
-								}
 							}
 						}
 					}
diff --git a/Engine/Plugins/Mutable/Source/MutableRuntime/Private/MuR/OpMeshRemove.cpp b/Engine/Plugins/Mutable/Source/MutableRuntime/Private/MuR/OpMeshRemove.cpp
--- a/Engine/Plugins/Mutable/Source/MutableRuntime/Private/MuR/OpMeshRemove.cpp
+++ b/Engine/Plugins/Mutable/Source/MutableRuntime/Private/MuR/OpMeshRemove.cpp
@@ -572,28 +572,24 @@ namespace mu
 		MeshRemoveRecreateSurface(Result, UsedVertices, UsedFaces);
     }
 
-	void MeshRemoveMask(FMesh* Result, const FMesh* Source, const FMesh* Mask, bool bRemoveIfAllVerticesCulled, bool& bOutSuccess)
+	void MeshRemoveMaskInline(FMesh* Mesh, const FMesh* Mask, bool bRemoveIfAllVerticesCulled)
     {
         MUTABLE_CPUPROFILER_SCOPE(MeshRemoveMask);
-		bOutSuccess = true;
 
-        if (!Mask->GetVertexCount() || !Source->GetVertexCount() || !Source->GetIndexCount())
+        if (!Mask->GetVertexCount() || !Mesh->GetVertexCount() || !Mesh->GetIndexCount())
         {
-			bOutSuccess = false;
             return;
         }
 
-		Result->CopyFrom(*Source);
-
 		int32 MaskElementCount = Mask->GetVertexBuffers().GetElementCount();
 
         // For each source vertex, true if it is removed.
-        int32 ResultVertexCount = Result->GetVertexCount();
+        int32 MeshVertexCount = Mesh->GetVertexCount();
 		TBitArray<> RemovedVertices;
-		RemovedVertices.SetNum(ResultVertexCount, false);
+		RemovedVertices.SetNum(MeshVertexCount, false);
         {
 			TArray<FIdInterval> Intervals;
-            ExtractVertexIndexIntervals(Intervals, Result);
+            ExtractVertexIndexIntervals(Intervals, Mesh);
 
 			MeshVertexIdIteratorConst itMaskVI(Mask);
 			for ( int32 mv=0; mv<MaskElementCount; ++mv )
@@ -609,6 +605,6 @@ namespace mu
             }
         }
 		
-        MeshRemoveVerticesWithCullSet(Result, RemovedVertices, bRemoveIfAllVerticesCulled);
+        MeshRemoveVerticesWithCullSet(Mesh, RemovedVertices, bRemoveIfAllVerticesCulled);
 	}
 }
diff --git a/Engine/Plugins/Mutable/Source/MutableRuntime/Private/MuR/OpMeshRemove.h b/Engine/Plugins/Mutable/Source/MutableRuntime/Private/MuR/OpMeshRemove.h
--- a/Engine/Plugins/Mutable/Source/MutableRuntime/Private/MuR/OpMeshRemove.h
+++ b/Engine/Plugins/Mutable/Source/MutableRuntime/Private/MuR/OpMeshRemove.h
@@ -15,7 +15,8 @@ namespace mu
 	 * If bRemoveIfAllVerticesCulled is true, a face is remove if all its vertices have the bit set in VerticesToCull.
 	 * If bRemoveIfAllVerticesCulled is false, remove a face if at least one vertex is removed.
 	 */
-	extern void MeshRemoveMask(FMesh* Result, const FMesh* Source, const FMesh* Mask, bool bRemoveIfAllVerticesCulled, bool& bOutSuccess);
+	extern void MeshRemoveMaskInline(FMesh* Mesh, const FMesh* Mask, bool bRemoveIfAllVerticesCulled);
+
 
 	/** 
 	 * Remove a set of vertices and related faces from a mesh in-place. VertexToCull is a bitset where if bit i-th is set, 
