[mutable] Mesh Section SwitchNode pins are named Mesh Section instead of Material. Also fixed for all similar generic nodes with a different pin friendly name.
#tests Editor
#rnx

diff --git a/Engine/Plugins/Mutable/Source/CustomizableObject/Internal/MuCO/CustomizableObjectCustomVersion.h b/Engine/Plugins/Mutable/Source/CustomizableObject/Internal/MuCO/CustomizableObjectCustomVersion.h
--- a/Engine/Plugins/Mutable/Source/CustomizableObject/Internal/MuCO/CustomizableObjectCustomVersion.h
+++ b/Engine/Plugins/Mutable/Source/CustomizableObject/Internal/MuCO/CustomizableObjectCustomVersion.h
@@ -184,6 +184,8 @@ struct CUSTOMIZABLEOBJECT_API FCustomizableObjectCustomVersion
 
 		ChangedComponentsInheritance,
 
+		ChangedSwitchNodesInputPinsFriendlyNames,
+
 		// -----<new versions can be added above this line>-------------------------------------------------
 		VersionPlusOne,
 		LatestVersion = VersionPlusOne - 1
diff --git a/Engine/Plugins/Mutable/Source/CustomizableObjectEditor/Private/MuCOE/Nodes/CustomizableObjectNodeSwitchBase.cpp b/Engine/Plugins/Mutable/Source/CustomizableObjectEditor/Private/MuCOE/Nodes/CustomizableObjectNodeSwitchBase.cpp
--- a/Engine/Plugins/Mutable/Source/CustomizableObjectEditor/Private/MuCOE/Nodes/CustomizableObjectNodeSwitchBase.cpp
+++ b/Engine/Plugins/Mutable/Source/CustomizableObjectEditor/Private/MuCOE/Nodes/CustomizableObjectNodeSwitchBase.cpp
@@ -84,7 +84,8 @@ void UCustomizableObjectNodeSwitchBase::AllocateDefaultPins(UCustomizableObjectN
 		FString PinName = GetPinPrefix(LayerIndex);
 		UEdGraphPin* InputPin = CustomCreatePin(EGPD_Input, GetCategory(), FName(*PinName));
 		FString PinDisplayName = ReloadingElementsNames[LayerIndex].IsEmpty() ? PinName : ReloadingElementsNames[LayerIndex];
-		InputPin->PinFriendlyName = FText::FromString(PinDisplayName);
+		const FString CategoryFriendlyName = UEdGraphSchema_CustomizableObject::GetPinCategoryFriendlyName(GetCategory()).ToString();
+		InputPin->PinFriendlyName = FText::FromString(CategoryFriendlyName + TEXT(" ") + FString::FromInt(LayerIndex));
 		InputPin->bDefaultValueIsIgnored = true;
 		InputPin->SetOwningNode(this);
 	}
@@ -168,6 +169,27 @@ void UCustomizableObjectNodeSwitchBase::PostPasteNode()
 }
 
 
+void UCustomizableObjectNodeSwitchBase::BackwardsCompatibleFixup(int32 CustomizableObjectCustomVersion)
+{
+	Super::BackwardsCompatibleFixup(CustomizableObjectCustomVersion);
+
+	if (CustomizableObjectCustomVersion == FCustomizableObjectCustomVersion::ChangedSwitchNodesInputPinsFriendlyNames)
+	{
+		const FString CategoryFriendlyName = UEdGraphSchema_CustomizableObject::GetPinCategoryFriendlyName(GetCategory()).ToString();
+
+		for (int32 ElementIndex = 0; ElementIndex < GetNumElements(); ++ElementIndex)
+		{
+			UEdGraphPin* ElementPin = GetElementPin(ElementIndex);
+
+			if (ElementPin->PinFriendlyName.ToString().StartsWith(TEXT("Material")))
+			{
+				ElementPin->PinFriendlyName = FText::FromString(CategoryFriendlyName + TEXT(" ") + FString::FromInt(ElementIndex));
+			}
+		}
+	}
+}
+
+
 void UCustomizableObjectNodeSwitchBase::PostBackwardsCompatibleFixup()
 {
 	Super::PostBackwardsCompatibleFixup();
diff --git a/Engine/Plugins/Mutable/Source/CustomizableObjectEditor/Private/MuCOE/Nodes/CustomizableObjectNodeSwitchBase.h b/Engine/Plugins/Mutable/Source/CustomizableObjectEditor/Private/MuCOE/Nodes/CustomizableObjectNodeSwitchBase.h
--- a/Engine/Plugins/Mutable/Source/CustomizableObjectEditor/Private/MuCOE/Nodes/CustomizableObjectNodeSwitchBase.h
+++ b/Engine/Plugins/Mutable/Source/CustomizableObjectEditor/Private/MuCOE/Nodes/CustomizableObjectNodeSwitchBase.h
@@ -31,6 +31,7 @@ public:
 	virtual void PostPasteNode() override;
 
 	// UCustomizableObjectNode interface
+	virtual void BackwardsCompatibleFixup(int32 CustomizableObjectCustomVersion) override;
 	virtual void PostBackwardsCompatibleFixup() override;
 	virtual void AllocateDefaultPins(UCustomizableObjectNodeRemapPins* RemapPins) override;
 	virtual void ReconstructNode(UCustomizableObjectNodeRemapPins* RemapPins) override;
