[mutable] Ensure the bulk of the work of the conflictive triangle fix when preparing layouts is only done when necessary.
#rnx

diff --git a/Engine/Plugins/Mutable/Source/MutableRuntime/Private/MuR/OpMeshPrepareLayout.cpp b/Engine/Plugins/Mutable/Source/MutableRuntime/Private/MuR/OpMeshPrepareLayout.cpp
--- a/Engine/Plugins/Mutable/Source/MutableRuntime/Private/MuR/OpMeshPrepareLayout.cpp
+++ b/Engine/Plugins/Mutable/Source/MutableRuntime/Private/MuR/OpMeshPrepareLayout.cpp
@@ -271,10 +271,40 @@ namespace mu
 			}
 		}
 
+		// Correct UVs and block assignment if necessary
+		const int32 NumTriangles = Mesh.GetIndexCount() / 3;
+		TArray<int32> ConflictiveTriangles;
 
 		if (bClampUVIslands)
 		{
-			const int32 NumTriangles = Mesh.GetIndexCount() / 3;
+			MUTABLE_CPUPROFILER_SCOPE(DetectConflictiveTriangles);
+
+			UntypedMeshBufferIteratorConst ItIndices(Mesh.GetIndexBuffers(), EMeshBufferSemantic::VertexIndex);
+
+			for (int32 TriangleIndex = 0; TriangleIndex < NumTriangles; ++TriangleIndex)
+			{
+				uint32 Index0 = ItIndices.GetAsUINT32();
+				++ItIndices;
+				uint32 Index1 = ItIndices.GetAsUINT32();
+				++ItIndices;
+				uint32 Index2 = ItIndices.GetAsUINT32();
+				++ItIndices;
+
+				const uint16 BlockIndexV0 = LayoutData[Index0];
+				const uint16 BlockIndexV1 = LayoutData[Index1];
+				const uint16 BlockIndexV2 = LayoutData[Index2];
+
+				if (BlockIndexV0 != BlockIndexV1 || BlockIndexV0 != BlockIndexV2)
+				{
+					ConflictiveTriangles.Add(TriangleIndex);
+				}
+			}
+		}
+
+		if (bClampUVIslands && !ConflictiveTriangles.IsEmpty())
+		{
+			MUTABLE_CPUPROFILER_SCOPE(ResolveConflictiveTriangles);
+			
 			TArray<FTriangleInfo> Triangles;
 
 			// Vertices mapped to unique vertex index
@@ -292,8 +322,6 @@ namespace mu
 				MeshCreateCollapsedVertexMap(&Mesh, CollapsedVertices);
 			}
 
-			TArray<int32> ConflictiveTriangles;
-
 			UntypedMeshBufferIteratorConst ItIndices(Mesh.GetIndexBuffers(), EMeshBufferSemantic::VertexIndex);
 			{
 				MUTABLE_CPUPROFILER_SCOPE(CreateTriangles);
@@ -311,11 +339,6 @@ namespace mu
 					const uint16 BlockIndexV1 = LayoutData[Index1];
 					const uint16 BlockIndexV2 = LayoutData[Index2];
 
-					if (BlockIndexV0 != BlockIndexV1 || BlockIndexV0 != BlockIndexV2)
-					{
-						ConflictiveTriangles.Add(TriangleIndex);
-					}
-
 					FTriangleInfo& Triangle = Triangles[TriangleIndex];
 
 					Triangle.Indices[0] = Index0;
@@ -336,7 +359,7 @@ namespace mu
 				}
 			}
 
-			// Clamp UV islands to the predominant block of each island. Will only happen if bClampUVIslands is true.
+			// Clamp UV islands to the predominant block of each island.
 			{
 				MUTABLE_CPUPROFILER_SCOPE(ClampUVs);
 
