[mutable] Additional optimizations that require changes in optimization loop. Ensure all optimizers are run if any changes the result.
#rnx

diff --git a/Engine/Plugins/Mutable/Source/MutableTools/Private/MuT/ASTOpLayoutPack.cpp b/Engine/Plugins/Mutable/Source/MutableTools/Private/MuT/ASTOpLayoutPack.cpp
--- a/Engine/Plugins/Mutable/Source/MutableTools/Private/MuT/ASTOpLayoutPack.cpp
+++ b/Engine/Plugins/Mutable/Source/MutableTools/Private/MuT/ASTOpLayoutPack.cpp
@@ -22,8 +22,6 @@
 namespace mu
 {
 
-
-	//-------------------------------------------------------------------------------------------------
 	ASTOpLayoutPack::ASTOpLayoutPack()
 		: Source(this)
 	{
@@ -96,5 +94,4 @@ namespace mu
 		}
 	}
 
-
 }
diff --git a/Engine/Plugins/Mutable/Source/MutableTools/Private/MuT/ASTOpLayoutRemoveBlocks.cpp b/Engine/Plugins/Mutable/Source/MutableTools/Private/MuT/ASTOpLayoutRemoveBlocks.cpp
--- a/Engine/Plugins/Mutable/Source/MutableTools/Private/MuT/ASTOpLayoutRemoveBlocks.cpp
+++ b/Engine/Plugins/Mutable/Source/MutableTools/Private/MuT/ASTOpLayoutRemoveBlocks.cpp
@@ -22,8 +22,6 @@
 namespace mu
 {
 
-
-	//-------------------------------------------------------------------------------------------------
 	ASTOpLayoutRemoveBlocks::ASTOpLayoutRemoveBlocks()
 		: Source(this)
 		, ReferenceLayout(this)
@@ -101,5 +99,4 @@ namespace mu
 		}
 	}
 
-
 }
diff --git a/Engine/Plugins/Mutable/Source/MutableTools/Private/MuT/ASTOpMeshApplyLayout.cpp b/Engine/Plugins/Mutable/Source/MutableTools/Private/MuT/ASTOpMeshApplyLayout.cpp
--- a/Engine/Plugins/Mutable/Source/MutableTools/Private/MuT/ASTOpMeshApplyLayout.cpp
+++ b/Engine/Plugins/Mutable/Source/MutableTools/Private/MuT/ASTOpMeshApplyLayout.cpp
@@ -2,6 +2,7 @@
 
 #include "MuT/ASTOpMeshApplyLayout.h"
 #include "MuT/ASTOpMeshAddTags.h"
+#include "MuT/ASTOpSwitch.h"
 
 #include "HAL/PlatformMath.h"
 #include "MuR/ModelPrivate.h"
@@ -97,6 +98,8 @@ namespace mu
 		}
 
 		EOpType MeshOpType = MeshAt->GetOpType();
+		
+		bool bLayoutIsConstant = !Layout || Layout->GetOpType() == EOpType::LA_CONSTANT;
 
 		switch (MeshOpType)
 		{
@@ -114,6 +117,37 @@ namespace mu
 				break;
 			}
 
+			case EOpType::ME_SWITCH:
+			{
+				if (bLayoutIsConstant)
+				{
+					// Move the apply down all the paths
+					Ptr<ASTOpSwitch> NewSwitch = mu::Clone<ASTOpSwitch>(MeshAt);
+
+					if (NewSwitch->Default)
+					{
+						Ptr<ASTOpMeshApplyLayout> defOp = mu::Clone<ASTOpMeshApplyLayout>(this);
+						defOp->Mesh = NewSwitch->Default.child();
+						NewSwitch->Default = defOp;
+					}
+
+					// We need to copy the options because we change them
+					for (int32 CaseIndex = 0; CaseIndex < NewSwitch->Cases.Num(); ++CaseIndex)
+					{
+						if (NewSwitch->Cases[CaseIndex].Branch)
+						{
+							Ptr<ASTOpMeshApplyLayout> bOp = mu::Clone<ASTOpMeshApplyLayout>(this);
+							bOp->Mesh = NewSwitch->Cases[CaseIndex].Branch.child();
+							NewSwitch->Cases[CaseIndex].Branch = bOp;
+						}
+					}
+
+					NewOp = NewSwitch;
+				}
+				break;
+			}
+
+
 			default:
 				break;
 		}
diff --git a/Engine/Plugins/Mutable/Source/MutableTools/Private/MuT/ASTOpSwitch.cpp b/Engine/Plugins/Mutable/Source/MutableTools/Private/MuT/ASTOpSwitch.cpp
--- a/Engine/Plugins/Mutable/Source/MutableTools/Private/MuT/ASTOpSwitch.cpp
+++ b/Engine/Plugins/Mutable/Source/MutableTools/Private/MuT/ASTOpSwitch.cpp
@@ -13,7 +13,6 @@
 #include "MuR/RefCounted.h"
 #include "MuR/Types.h"
 
-
 namespace mu
 {
 
@@ -604,7 +603,6 @@ namespace mu
 			if (BranchOpType==EOpType::NONE)
 			{
 				BranchOpType = Case.Branch->GetOpType();
-				break;
 			}
 			else if (Case.Branch->GetOpType() != BranchOpType)
 			{
diff --git a/Engine/Plugins/Mutable/Source/MutableTools/Private/MuT/CodeOptimiser.cpp b/Engine/Plugins/Mutable/Source/MutableTools/Private/MuT/CodeOptimiser.cpp
--- a/Engine/Plugins/Mutable/Source/MutableTools/Private/MuT/CodeOptimiser.cpp
+++ b/Engine/Plugins/Mutable/Source/MutableTools/Private/MuT/CodeOptimiser.cpp
@@ -26,9 +26,7 @@
 #include "MuR/MutableRuntimeModule.h"
 
 #include "Tasks/Task.h"
-
-#include <unordered_set>
-
+#include "Math/NumericLimits.h"
 
 namespace mu
 {
@@ -95,8 +93,7 @@ namespace mu
 	}
 
 
-	//-------------------------------------------------------------------------------------------------
-	void DuplicatedDataRemoverAST( ASTOpList& roots )
+	bool DuplicatedDataRemoverAST( ASTOpList& roots )
 	{
 		MUTABLE_CPUPROFILER_SCOPE(DuplicatedDataRemoverAST);
 
@@ -104,6 +101,8 @@ namespace mu
 		TArray<Ptr<ASTOpConstantResource>> AllImageOps;
 		TArray<Ptr<ASTOpConstantResource>> AllLayoutOps;
 
+		bool bModified = false;
+
 		// Gather constants
 		{
 			MUTABLE_CPUPROFILER_SCOPE(Gather);
@@ -189,6 +188,7 @@ namespace mu
 				if (Found)
 				{
 					ASTOp::Replace(typedNode, Found);
+					bModified = true;
 				}
 				else
 				{
@@ -237,6 +237,7 @@ namespace mu
 				if (Found)
 				{
 					ASTOp::Replace(typedNode, Found);
+					bModified = true;
 				}
 				else
 				{
@@ -285,6 +286,7 @@ namespace mu
 				if (Found)
 				{
 					ASTOp::Replace(typedNode, Found);
+					bModified = true;
 				}
 				else
 				{
@@ -295,56 +297,57 @@ namespace mu
 			}
 		}
 
+		return bModified;
 	}
 
 
-	//-------------------------------------------------------------------------------------------------
-	//-------------------------------------------------------------------------------------------------
-	//-------------------------------------------------------------------------------------------------
-	namespace
+	bool DuplicatedCodeRemoverAST( ASTOpList& roots )
 	{
-		struct op_pointer_hash
+		MUTABLE_CPUPROFILER_SCOPE(DuplicatedCodeRemoverAST);
+
+		bool bModified = false;
+
+		struct FKeyFuncs : BaseKeyFuncs<Ptr<ASTOp>, Ptr<ASTOp>, false>
 		{
-			inline size_t operator() (const mu::Ptr<ASTOp>& n) const
-			{
-				return n->Hash();
-			}
+			static KeyInitType GetSetKey(ElementInitType Element) { return Element; }
+			static bool Matches(const Ptr<ASTOp>& lhs, const Ptr<ASTOp>& rhs) { return lhs == rhs || *lhs == *rhs; }
+			static uint32 GetKeyHash(const Ptr<ASTOp>& Key) { return Key->Hash(); }
 		};
 
-		struct op_pointer_equal
-		{
-			inline bool operator() (const mu::Ptr<ASTOp>& a, const mu::Ptr<ASTOp>& b) const
+		// Visited nodes, per type
+		TSet<Ptr<ASTOp>, FKeyFuncs, TInlineSetAllocator<32>> Visited[int32(EOpType::COUNT)];
+
+		ASTOp::Traverse_BottomUp_Unique_NonReentrant( roots, 
+			[&bModified,&Visited, &roots]
+			(Ptr<ASTOp>& n)
 			{
-				return *a==*b;
-			}
-		};
-	}
+				TSet<Ptr<ASTOp>, FKeyFuncs, TInlineSetAllocator<32>>& Container = Visited[(int32)n->GetOpType()];
 
-	void DuplicatedCodeRemoverAST( ASTOpList& roots )
-	{
-		MUTABLE_CPUPROFILER_SCOPE(DuplicatedCodeRemoverAST);
+				// Insert will tell us if it was already there
+				bool bIsAlreadyInSet = false;
+				Ptr<ASTOp>& Found = Container.FindOrAdd(n, &bIsAlreadyInSet);
+				if( bIsAlreadyInSet)
+				{
+					// It wasn't inserted, so it was already there
+					ASTOp::Replace(n, Found);
 
-		// Visited nodes, per type
-		std::unordered_set<Ptr<ASTOp>,op_pointer_hash,op_pointer_equal> visited[int32(EOpType::COUNT)];
+					// Is it one of the roots? Then we also need to update it
+					for (Ptr<ASTOp>& Root : roots)
+					{
+						if (Root == n)
+						{
+							Root = Found;
+						}
+					}
 
-		ASTOp::Traverse_BottomUp_Unique_NonReentrant( roots, [&](Ptr<ASTOp>& n)
-		{
-			auto& container = visited[(int32)n->GetOpType()];
+					bModified = true;
+				}
+			});
 
-			// Insert will tell us if it was already there
-			auto it = container.insert(n);
-			if( !it.second )
-			{
-				// It wasn't inserted, so it was already there
-				ASTOp::Replace(n,*it.first);
-			}
-		});
+		return bModified;
 	}
 
 
-	//---------------------------------------------------------------------------------------------
-	//---------------------------------------------------------------------------------------------
-	//---------------------------------------------------------------------------------------------
 	FFindOpTypeVisitor::FFindOpTypeVisitor(const TArray<EOpType>& Types)
 		: TypesToFind(Types)
 	{
@@ -418,9 +421,6 @@ namespace mu
 	}
 
 
-	//---------------------------------------------------------------------------------------------
-	//---------------------------------------------------------------------------------------------
-	//---------------------------------------------------------------------------------------------
 	IsConstantVisitor::IsConstantVisitor()
 	{
 		TArray<EOpType> ParameterTypes;
@@ -435,16 +435,12 @@ namespace mu
 	}
 
 
-	//---------------------------------------------------------------------------------------------
 	bool IsConstantVisitor::Apply( FProgram& program, OP::ADDRESS at )
 	{
 		return !FindOpTypeVisitor->Apply(program,at);
 	}
 
 
-	//---------------------------------------------------------------------------------------------
-	//---------------------------------------------------------------------------------------------
-	//---------------------------------------------------------------------------------------------
 	class FConstantTask 
 	{
 	public:
@@ -1004,9 +1000,6 @@ namespace mu
 	}
 
 
-	//-------------------------------------------------------------------------------------------------
-	//-------------------------------------------------------------------------------------------------
-	//-------------------------------------------------------------------------------------------------
 	CodeOptimiser::CodeOptimiser(Ptr<CompilerOptions> InOptions, TArray<FStateCompilationData>& InStates )
 		: States( InStates )
 	{
@@ -1014,108 +1007,84 @@ namespace mu
 	}
 
 
-	//-------------------------------------------------------------------------------------------------
 	void CodeOptimiser::FullOptimiseAST( ASTOpList& roots, int32 Pass )
 	{
 		bool bModified = true;
 		int32 NumIterations = 0;
-		while (bModified && (!OptimizeIterationsMax || (OptimizeIterationsLeft>0) || !NumIterations) )
+		while (bModified && (OptimizeIterationsLeft>0 || !NumIterations))
 		{
-			--OptimizeIterationsLeft;
-			++NumIterations;
-			UE_LOG(LogMutableCore, Verbose, TEXT("Main optimise iteration %d, max %d, left %d"),
-					NumIterations, OptimizeIterationsMax, OptimizeIterationsLeft);
-
-			bModified = false;
-
-			// All kind of optimisations that depend on the meaning of each operation
-			// \TODO: We are doing it for all states.
-			UE_LOG(LogMutableCore, Verbose, TEXT(" - semantic optimiser"));
-			bModified |= SemanticOptimiserAST( roots, Options->GetPrivate()->OptimisationOptions, Pass );
-			//AXE_INT_VALUE("Mutable", Verbose, "ast size", (int64_t)ASTOp::CountNodes(roots));
-			//UE_LOG(LogMutableCore, Verbose, TEXT("(int) %s : %ld"), TEXT("ast size"), int64(ASTOp::CountNodes(roots)));
-			ASTOp::LogHistogram(roots);
-
-			UE_LOG(LogMutableCore, Verbose, TEXT(" - sink optimiser"));
-			bModified |= SinkOptimiserAST( roots, Options->GetPrivate()->OptimisationOptions );
-			//AXE_INT_VALUE("Mutable", Verbose, "ast size", (int64_t)ASTOp::CountNodes(roots));
-			ASTOp::LogHistogram(roots);
+			bool bModifiedInInnerLoop = true;
+			while (bModifiedInInnerLoop && (OptimizeIterationsLeft>0 || !NumIterations))
+			{
+				--OptimizeIterationsLeft;
+				++NumIterations;
+				UE_LOG(LogMutableCore, Verbose, TEXT("Main optimise iteration %d, left %d"), NumIterations, OptimizeIterationsLeft);
 
-			// Image size operations are treated separately
-			UE_LOG(LogMutableCore, Verbose, TEXT(" - size optimiser"));
-			bModified |= SizeOptimiserAST( roots );
-			//AXE_INT_VALUE("Mutable", Verbose, "ast size", (int64_t)ASTOp::CountNodes(roots));
+				bModifiedInInnerLoop = false;
 
-			// We cannot run the logic optimiser here because the constant memory explodes.
-			// TODO: Conservative logic optimiser
-			if (!bModified)
-			{
-	//                AXE_INT_VALUE("Mutable", Verbose, "program size", (int64_t)program.OpAddress.Num());
-	//                UE_LOG(LogMutableCore, Verbose, " - logic optimiser");
-	//                LocalLogicOptimiser log;
-	//                bModified |= log.Apply( program, (int)s );
-			}
-		}
+				// All kind of optimisations that depend on the meaning of each operation
+				// \TODO: We are doing it for all states.
+				UE_LOG(LogMutableCore, Verbose, TEXT(" - semantic optimiser"));
+				bModifiedInInnerLoop |= SemanticOptimiserAST(roots, Options->GetPrivate()->OptimisationOptions, Pass);
+				ASTOp::LogHistogram(roots);
 
-	//    ASTOp::LogHistogram(roots);
+				UE_LOG(LogMutableCore, Verbose, TEXT(" - sink optimiser"));
+				bModifiedInInnerLoop |= SinkOptimiserAST(roots, Options->GetPrivate()->OptimisationOptions);
+				ASTOp::LogHistogram(roots);
 
-		UE_LOG(LogMutableCore, Verbose, TEXT(" - duplicated code remover"));
-		DuplicatedCodeRemoverAST( roots );
-		//UE_LOG(LogMutableCore, Verbose, TEXT("(int) %s : %ld"), TEXT("ast size"), int64(ASTOp::CountNodes(roots)));
+				// Image size operations are treated separately
+				UE_LOG(LogMutableCore, Verbose, TEXT(" - size optimiser"));
+				bModifiedInInnerLoop |= SizeOptimiserAST(roots);
+			}
 
-		ASTOp::LogHistogram(roots);
+			bModified = bModifiedInInnerLoop;
 
-		UE_LOG(LogMutableCore, Verbose, TEXT(" - duplicated data remover"));
-		DuplicatedDataRemoverAST( roots );
-		//UE_LOG(LogMutableCore, Verbose, TEXT("(int) %s : %ld"), TEXT("ast size"), int64(ASTOp::CountNodes(roots)));
+			UE_LOG(LogMutableCore, Verbose, TEXT(" - duplicated code remover"));
+			bModified |= DuplicatedCodeRemoverAST(roots);
+			//UE_LOG(LogMutableCore, Verbose, TEXT("(int) %s : %ld"), TEXT("ast size"), int64(ASTOp::CountNodes(roots)));
 
-		ASTOp::LogHistogram(roots);
+			ASTOp::LogHistogram(roots);
 
-		// Generate constants
-		for ( Ptr<ASTOp>& Root: roots )
-		{
+			UE_LOG(LogMutableCore, Verbose, TEXT(" - duplicated data remover"));
+			bModified |= DuplicatedDataRemoverAST(roots);
 			//UE_LOG(LogMutableCore, Verbose, TEXT("(int) %s : %ld"), TEXT("ast size"), int64(ASTOp::CountNodes(roots)));
-			UE_LOG(LogMutableCore, Verbose, TEXT(" - constant generator"));
 
-			// Constant subtree generation
-			bModified = ConstantGenerator( Options->GetPrivate(), Root, Pass );
-		}
+			ASTOp::LogHistogram(roots);
 
-		//UE_LOG(LogMutableCore, Verbose, TEXT("(int) %s : %ld"), TEXT("ast size"), int64(ASTOp::CountNodes(roots)));
+			// Generate constants
+			bool bModifiedInConstants = false;
+			for (Ptr<ASTOp>& Root : roots)
+			{
+				//UE_LOG(LogMutableCore, Verbose, TEXT("(int) %s : %ld"), TEXT("ast size"), int64(ASTOp::CountNodes(roots)));
+				UE_LOG(LogMutableCore, Verbose, TEXT(" - constant generator"));
 
-	//    ASTOp::LogHistogram(roots);
+				// Constant subtree generation
+				bModifiedInConstants |= ConstantGenerator(Options->GetPrivate(), Root, Pass);
+			}
 
-		// \todo: this code deduplication here, takes very long and doesn't do much... investigate!
-	//    UE_LOG(LogMutableCore, Verbose, " - duplicated code remover");
-	//    DuplicatedCodeRemoverAST( roots );
-	//    AXE_INT_VALUE("Mutable", Verbose, "ast size", (int64_t)ASTOp::CountNodes(roots));
+			ASTOp::LogHistogram(roots);
 
-		ASTOp::LogHistogram(roots);
+			if (bModifiedInConstants)
+			{
+				bModified = true;
 
-		UE_LOG(LogMutableCore, Verbose, TEXT(" - duplicated data remover"));
-		DuplicatedDataRemoverAST( roots );
-		//UE_LOG(LogMutableCore, Verbose, TEXT("(int) %s : %ld"), TEXT("ast size"), int64(ASTOp::CountNodes(roots)));
+				UE_LOG(LogMutableCore, Verbose, TEXT(" - duplicated data remover"));
+				DuplicatedDataRemoverAST(roots);
+			}
 
-		//if (!bModified)
-		{
-			//for ( size_t s=0;  s<program.States.size(); ++s )
+			//if (!bModified)
 			{
-				ASTOp::LogHistogram(roots);
 				UE_LOG(LogMutableCore, Verbose, TEXT(" - logic optimiser"));
-				bModified |= LocalLogicOptimiserAST( roots );
-				ASTOp::LogHistogram(roots);
+				bModified |= LocalLogicOptimiserAST(roots);
 			}
 
-			//AXE_INT_VALUE("Mutable", Verbose, "ast size", (int64_t)ASTOp::CountNodes(roots));
+			ASTOp::LogHistogram(roots);
 		}
-
-		ASTOp::LogHistogram(roots);
 	}
 
 
-	//---------------------------------------------------------------------------------------------
 	// The state represents if there is a parent operation requiring skeleton for current mesh subtree.
-	class CollectAllMeshesForSkeletonVisitorAST : public Visitor_TopDown_Unique_Const<uint8_t>
+	class CollectAllMeshesForSkeletonVisitorAST : public Visitor_TopDown_Unique_Const<uint8>
 	{
 	public:
 
@@ -1208,7 +1177,6 @@ namespace mu
 	};
 
 
-	//---------------------------------------------------------------------------------------------
 	// This stores an ADD_MESH op with the child meshes collected and the final skeleton to use
 	// for this op.
 	struct FAddMeshSkeleton
@@ -1228,7 +1196,6 @@ namespace mu
 	};
 
 
-	//---------------------------------------------------------------------------------------------
 	void SkeletonCleanerAST( TArray<mu::Ptr<ASTOp>>& roots, const FModelOptimizationOptions& options )
 	{
 		// This collects all the meshes that require a skeleton because they are used in operations
@@ -1325,7 +1292,6 @@ namespace mu
 	}
 
 
-	//---------------------------------------------------------------------------------------------
 	void CodeOptimiser::Optimise()
 	{
 		MUTABLE_CPUPROFILER_SCOPE(Optimise);
@@ -1343,8 +1309,8 @@ namespace mu
 		{
 			// We use 4 times the count because at the time we moved to sharing this count it
 			// was being used 4 times, and we want to keep the tests consistent.
-			OptimizeIterationsLeft = Options->GetPrivate()->OptimisationOptions.MaxOptimisationLoopCount * 4;
-			OptimizeIterationsMax = OptimizeIterationsLeft;
+			int32 MaxIterations = Options->GetPrivate()->OptimisationOptions.MaxOptimisationLoopCount;
+			OptimizeIterationsLeft = MaxIterations ? MaxIterations * 4 : TNumericLimits<int32>::Max();
 
 			// The first duplicated data remover has the special mission of removing
 			// duplicated data (meshes) that may have been specified in the source
@@ -1352,13 +1318,11 @@ namespace mu
 			// creation
 			UE_LOG(LogMutableCore, Verbose, TEXT(" - duplicated data remover"));
 			DuplicatedDataRemoverAST( roots );
-			//UE_LOG(LogMutableCore, Verbose, TEXT("(int) %s : %ld"), TEXT("ast size"), int64(ASTOp::CountNodes(roots)));
 
 			ASTOp::LogHistogram(roots);
 
 			UE_LOG(LogMutableCore, Verbose, TEXT(" - duplicated code remover"));
 			DuplicatedCodeRemoverAST( roots );
-			//UE_LOG(LogMutableCore, Verbose, TEXT("(int) %s : %ld"), TEXT("ast size"), int64(ASTOp::CountNodes(roots)));
 
 			// Special optimization stages
 			if ( Options->GetPrivate()->OptimisationOptions.bUniformizeSkeleton )
@@ -1367,7 +1331,6 @@ namespace mu
 				ASTOp::LogHistogram(roots);
 
 				SkeletonCleanerAST( roots, Options->GetPrivate()->OptimisationOptions );
-				//UE_LOG(LogMutableCore, Verbose, TEXT("(int) %s : %ld"), TEXT("ast size"), int64(ASTOp::CountNodes(roots)));
 				ASTOp::LogHistogram(roots);
 			}
 
@@ -1382,8 +1345,7 @@ namespace mu
 
 				--OptimizeIterationsLeft;
 				++NumIterations;
-				UE_LOG(LogMutableCore, Verbose, TEXT("First optimise iteration %d, max %d, left %d"),
-						NumIterations, OptimizeIterationsMax, OptimizeIterationsLeft);
+				UE_LOG(LogMutableCore, Verbose, TEXT("First optimise iteration %d, left %d"), NumIterations, OptimizeIterationsLeft);
 
 				bModified = false;
 
@@ -1391,19 +1353,12 @@ namespace mu
 				bModified |= SizeOptimiserAST( roots );
 			}
 
-			//UE_LOG(LogMutableCore, Verbose, TEXT("(int) %s : %ld"), TEXT("ast size"), int64(ASTOp::CountNodes(roots)));
-			ASTOp::LogHistogram(roots);
-
 			// Main optimisation stage
 			{
 				MUTABLE_CPUPROFILER_SCOPE(MainStage);
 				FullOptimiseAST( roots, 0 );
-				//UE_LOG(LogMutableCore, Verbose, TEXT("(int) %s : %ld"), TEXT("ast size"), int64(ASTOp::CountNodes(roots)));
-				ASTOp::LogHistogram(roots);
 
 				FullOptimiseAST( roots, 1 );
-				//UE_LOG(LogMutableCore, Verbose, TEXT("(int) %s : %ld"), TEXT("ast size"), int64(ASTOp::CountNodes(roots)));
-				ASTOp::LogHistogram(roots);
 			}
 
 			// Constant resolution stage: resolve referenced assets.
@@ -1428,11 +1383,9 @@ namespace mu
 			{
 				MUTABLE_CPUPROFILER_SCOPE(FinalStage);
 				FullOptimiseAST(roots, 0);
-				//UE_LOG(LogMutableCore, Verbose, TEXT("(int) %s : %ld"), TEXT("ast size"), int64(ASTOp::CountNodes(roots)));
 				ASTOp::LogHistogram(roots);
 
 				FullOptimiseAST(roots, 1);
-				//UE_LOG(LogMutableCore, Verbose, TEXT("(int) %s : %ld"), TEXT("ast size"), int64(ASTOp::CountNodes(roots)));
 				ASTOp::LogHistogram(roots);
 			}
 
@@ -1478,11 +1431,9 @@ namespace mu
 			// creation
 			UE_LOG(LogMutableCore, Verbose, TEXT(" - duplicated data remover"));
 			DuplicatedDataRemoverAST( roots );
-			//AXE_INT_VALUE("Mutable", Verbose, "ast size", (int64_t)ASTOp::CountNodes(roots));
 
 			UE_LOG(LogMutableCore, Verbose, TEXT(" - duplicated code remover"));
 			DuplicatedCodeRemoverAST( roots );
-			//UE_LOG(LogMutableCore, Verbose, TEXT("(int) %s : %ld"), TEXT("ast size"), int64(ASTOp::CountNodes(roots)));
 
 			// Constant resolution stage: resolve referenced assets.
 			{
@@ -1496,16 +1447,13 @@ namespace mu
 
 				UE_LOG(LogMutableCore, Verbose, TEXT(" - constant generator"));
 				ConstantGenerator( Options->GetPrivate(), roots[StateIndex], Pass);
-				//AXE_INT_VALUE("Mutable", Verbose, "ast size", (int64_t)ASTOp::CountNodes(roots));
 			}
 
 			UE_LOG(LogMutableCore, Verbose, TEXT(" - duplicated data remover"));
 			DuplicatedDataRemoverAST( roots );
-			//AXE_INT_VALUE("Mutable", Verbose, "ast size", (int64_t)ASTOp::CountNodes(roots));
 
 			UE_LOG(LogMutableCore, Verbose, TEXT(" - duplicated code remover"));
 			DuplicatedCodeRemoverAST( roots );
-			//AXE_INT_VALUE("Mutable", Verbose, "ast size", (int64_t)ASTOp::CountNodes(roots));
 
 			// Reset the state root operations in case they have changed due to optimization
 			for (int32 RootIndex = 0; RootIndex < States.Num(); ++RootIndex)
diff --git a/Engine/Plugins/Mutable/Source/MutableTools/Private/MuT/CodeOptimiser.h b/Engine/Plugins/Mutable/Source/MutableTools/Private/MuT/CodeOptimiser.h
--- a/Engine/Plugins/Mutable/Source/MutableTools/Private/MuT/CodeOptimiser.h
+++ b/Engine/Plugins/Mutable/Source/MutableTools/Private/MuT/CodeOptimiser.h
@@ -24,7 +24,6 @@ namespace mu
 		TArray<FStateCompilationData>& States;
 
         /** The max number of optimize iterations is shared across several stages now. This is how many are left. */
-        int32 OptimizeIterationsMax=0;
         int32 OptimizeIterationsLeft=0;
 
         //! Full optimisation pass
@@ -82,10 +81,10 @@ namespace mu
     extern bool ConstantGenerator( const CompilerOptions::Private*, Ptr<ASTOp>& Root, int32 Pass );
 
     /** \TODO: shapes, projectors, others? but not switches (they must be unique) */
-    extern void DuplicatedDataRemoverAST( ASTOpList& Roots );
+    extern bool DuplicatedDataRemoverAST( ASTOpList& Roots );
 
     /** Mark all the duplicated code instructions to point at the same operation, leaving the copies unreachable. */
-    extern void DuplicatedCodeRemoverAST( ASTOpList& Roots );
+    extern bool DuplicatedCodeRemoverAST( ASTOpList& Roots );
 
     /** All kinds of optimisations that depend on the meaning of each operation. */
     extern bool SemanticOptimiserAST( ASTOpList& Roots, const FModelOptimizationOptions&, int32 Pass);
diff --git a/Engine/Plugins/Mutable/Source/MutableTools/Private/MuT/CodeOptimiser_Logic.cpp b/Engine/Plugins/Mutable/Source/MutableTools/Private/MuT/CodeOptimiser_Logic.cpp
--- a/Engine/Plugins/Mutable/Source/MutableTools/Private/MuT/CodeOptimiser_Logic.cpp
+++ b/Engine/Plugins/Mutable/Source/MutableTools/Private/MuT/CodeOptimiser_Logic.cpp
@@ -21,7 +21,7 @@ namespace mu
     {
         MUTABLE_CPUPROFILER_SCOPE(LocalLogicOptimiserAST);
 
-        bool modified = false;
+        bool bModified = false;
 
         // The separate steps should not be combined into one traversal
 
@@ -31,7 +31,7 @@ namespace mu
             MUTABLE_CPUPROFILER_SCOPE(Unwrap);
             ASTOp::Traverse_TopDown_Unique_Imprecise( roots, [&](Ptr<ASTOp>& n)
             {
-                bool recurse = true;
+                bool bRecurse = true;
 
                 if (n->GetOpType()==EOpType::IN_CONDITIONAL)
                 {
@@ -84,7 +84,7 @@ namespace mu
 										addSurface = newAddSurface.get();
                                     }
 
-                                    modified = true;
+                                    bModified = true;
                                     ended = false;
                                 }
                             }
@@ -140,7 +140,7 @@ namespace mu
                                         Merge = NewMerge;
                                     }
 
-                                    modified = true;
+                                    bModified = true;
                                     ended = false;
                                 }
                             }
@@ -149,7 +149,7 @@ namespace mu
                 }
 
 
-                return recurse;
+                return bRecurse;
             });
         }
 
@@ -161,7 +161,7 @@ namespace mu
             MUTABLE_CPUPROFILER_SCOPE(ConditionalToSwitch);
             ASTOp::Traverse_TopDown_Unique_Imprecise( roots, [&](Ptr<ASTOp>& n)
             {
-                bool recurse = true;
+                bool bRecurse = true;
 
 				ASTOpConditional* topConditional = nullptr;
 				if (n && n->IsConditional())
@@ -187,7 +187,7 @@ namespace mu
 						Ptr<ASTOp> current = n;
                         while(current)
                         {
-                            bool valid = false;
+                            bool bValid = false;
 
 							ASTOpConditional* conditional = nullptr;
 							if (current && current->IsConditional())
@@ -214,12 +214,12 @@ namespace mu
                                         switchOp->Cases.Emplace( compare->Constant, switchOp, conditional->yes.child() );
 
                                         current = conditional->no.child();
-                                        valid = true;
+										bValid = true;
                                     }
                                 }
                             }
 
-                            if (!valid)
+                            if (!bValid)
                             {
                                 switchOp->Default = current;
                                 current = nullptr;
@@ -231,50 +231,51 @@ namespace mu
                         {
                             ASTOp::Replace(n,switchOp);
                             n = switchOp;
-                            modified = true;
+                            bModified = true;
                         }
                     }
 
                 }
 
-                return recurse;
+                return bRecurse;
             });
         }
 
 
         // Float operations up switches, to tidy up the code and reduce its size.
-        //-----------------------------------------------------------------------------------------
+		// TODO?
+		/*
         {
             MUTABLE_CPUPROFILER_SCOPE(FloatSwitches);
 
             ASTOp::Traverse_TopDown_Unique_Imprecise( roots, [&](Ptr<ASTOp>& n)
             {
-                bool recurse = true;
+                bool bRecurse = true;
 
-				ASTOpSwitch* topSwitch = nullptr;
+				ASTOpSwitch* TopSwitch = nullptr;
 				if (n && n->IsSwitch())
 				{
-					topSwitch = static_cast<ASTOpSwitch*>(n.get());
+					TopSwitch = static_cast<ASTOpSwitch*>(n.get());
 				}
 
-                if (topSwitch)
+                if (TopSwitch)
                 {
-                    bool first = true;
-					EOpType caseType = EOpType::NONE;
-                    for (const ASTOpSwitch::FCase& c:topSwitch->Cases)
+                    bool bFirst = true;
+					EOpType CaseType = EOpType::NONE;
+                    for (const ASTOpSwitch::FCase& c: TopSwitch->Cases)
                     {
                         if ( c.Branch )
                         {
-                            if (first)
+                            if (bFirst)
                             {
-                                first = false;
-                                caseType = c.Branch->GetOpType();
+								bFirst = false;
+								CaseType = c.Branch->GetOpType();
                             }
                             else
                             {
-                                if (caseType!=c.Branch->GetOpType())
+                                if (CaseType !=c.Branch->GetOpType())
                                 {
-                                    caseType=EOpType::NONE;
+									CaseType = EOpType::NONE;
                                     break;
                                 }
                             }
@@ -282,11 +283,12 @@ namespace mu
                     }
                 }
 
-                return recurse;
+                return bRecurse;
             });
         }
+		*/
 
-        return modified;
+        return bModified;
     }
 
 }
diff --git a/Engine/Plugins/Mutable/Source/MutableTools/Private/MuT/CodeOptimiser_Parameters.cpp b/Engine/Plugins/Mutable/Source/MutableTools/Private/MuT/CodeOptimiser_Parameters.cpp
--- a/Engine/Plugins/Mutable/Source/MutableTools/Private/MuT/CodeOptimiser_Parameters.cpp
+++ b/Engine/Plugins/Mutable/Source/MutableTools/Private/MuT/CodeOptimiser_Parameters.cpp
@@ -41,9 +41,6 @@
 namespace mu
 {
 
-    //---------------------------------------------------------------------------------------------
-    //---------------------------------------------------------------------------------------------
-    //---------------------------------------------------------------------------------------------
     RuntimeParameterVisitorAST::RuntimeParameterVisitorAST(const FStateCompilationData* InState)
         : State(InState)
     {
@@ -2034,80 +2031,75 @@ namespace mu
             {
                 // Promote the intructions that depend on runtime parameters, and sink new
                 // format instructions.
-                bool modified = true;
-                int32 numIterations = 0;
-                while (modified && (!OptimizeIterationsMax || OptimizeIterationsLeft>0 || !numIterations ))
+                bModified = true;
+                int32 NumIterations = 0;
+                while (bModified && ( OptimizeIterationsLeft>0 || !NumIterations))
                 {
-                    modified = false;
+					bModified = false;
 
-                    ++numIterations;
+                    ++NumIterations;
                     --OptimizeIterationsLeft;
-                    UE_LOG(LogMutableCore, Verbose, TEXT("State optimise iteration %d, max %d, left %d"),
-                            numIterations, OptimizeIterationsMax, OptimizeIterationsLeft);
+                    UE_LOG(LogMutableCore, Verbose, TEXT("State optimise iteration %d, left %d"), NumIterations, OptimizeIterationsLeft);
 
                     UE_LOG(LogMutableCore, Verbose, TEXT(" - before parameter optimiser"));
 
                     ParameterOptimiserAST param( States[s], Options->GetPrivate()->OptimisationOptions );
-                    modified = param.Apply();
+					bModified = param.Apply();
 
-                    TArray<Ptr<ASTOp>> roots;
-                    roots.Add(States[s].root);
+                    TArray<Ptr<ASTOp>> Roots;
+                    Roots.Add(States[s].root);
 
                     UE_LOG(LogMutableCore, Verbose, TEXT(" - after parameter optimiser"));
-					//UE_LOG(LogMutableCore, Verbose, TEXT("(int) %s : %ld"), TEXT("ast size"), int64(ASTOp::CountNodes(roots)));
 
                     // All kind of optimisations that depend on the meaning of each operation
                     UE_LOG(LogMutableCore, Verbose, TEXT(" - semantic optimiser"));
-                    modified |= SemanticOptimiserAST( roots, Options->GetPrivate()->OptimisationOptions, 1 );
-					//UE_LOG(LogMutableCore, Verbose, TEXT("(int) %s : %ld"), TEXT("ast size"), int64(ASTOp::CountNodes(roots)));
-					//ASTOp::LogHistogram(roots);
+					bModified |= SemanticOptimiserAST(Roots, Options->GetPrivate()->OptimisationOptions, 1 );
 
                     UE_LOG(LogMutableCore, Verbose, TEXT(" - sink optimiser"));
-                    modified |= SinkOptimiserAST( roots, Options->GetPrivate()->OptimisationOptions );
-					//UE_LOG(LogMutableCore, Verbose, TEXT("(int) %s : %ld"), TEXT("ast size"), int64(ASTOp::CountNodes(roots)));
-					//ASTOp::LogHistogram(roots);
+					bModified |= SinkOptimiserAST(Roots, Options->GetPrivate()->OptimisationOptions );
 
                     // Image size operations are treated separately
                     UE_LOG(LogMutableCore, Verbose, TEXT(" - size optimiser"));
-                    modified |= SizeOptimiserAST( roots );
-					//UE_LOG(LogMutableCore, Verbose, TEXT("(int) %s : %ld"), TEXT("ast size"), int64(ASTOp::CountNodes(roots)));
+					bModified |= SizeOptimiserAST(Roots);
+
+					// Some sink optimizations can only be applied after some constant reductions
+					for (Ptr<ASTOp>& Root : Roots)
+					{
+						bModified |= ConstantGenerator(Options->GetPrivate(), Root, 1);
+					}
 
-                    //LogicOptimiser log;
-                    //modified |= log.Apply( program, s );
                 }
 
-                TArray<Ptr<ASTOp>> roots;
-                roots.Add(States[s].root);
+                TArray<Ptr<ASTOp>> Roots;
+                Roots.Add(States[s].root);
 
                 UE_LOG(LogMutableCore, Verbose, TEXT(" - duplicated data remover"));
-                DuplicatedDataRemoverAST( roots );
-				//UE_LOG(LogMutableCore, Verbose, TEXT("(int) %s : %ld"), TEXT("ast size"), int64(ASTOp::CountNodes(roots)));
+				bModified |= DuplicatedDataRemoverAST(Roots);
 
                 UE_LOG(LogMutableCore, Verbose, TEXT(" - duplicated code remover"));
-                DuplicatedCodeRemoverAST( roots );
-				//UE_LOG(LogMutableCore, Verbose, TEXT("(int) %s : %ld"), TEXT("ast size"), int64(ASTOp::CountNodes(roots)));
+				bModified |= DuplicatedCodeRemoverAST(Roots);
 
-                States[s].root = roots[0];
+                States[s].root = Roots[0];
             }
         }
 
 
-		TArray<Ptr<ASTOp>> roots;
+		TArray<Ptr<ASTOp>> Roots;
 		for (const FStateCompilationData& s : States)
 		{
-			roots.Add(s.root);
+			Roots.Add(s.root);
 		}
 
         // Mark the instructions that don't depend on runtime parameters to be cached. This is
         // necessary At this stage before GPU optimisation.
         {
             AccumulateAllImageFormatsOpAST opFormats;
-            opFormats.Run(roots);
+            opFormats.Run(Roots);
 
 			// Reset the state root operations in case they have changed due to optimization
 			for (int32 RootIndex = 0; RootIndex < States.Num(); ++RootIndex)
 			{
-				States[RootIndex].root = roots[RootIndex];
+				States[RootIndex].root = Roots[RootIndex];
 			}
 
             for (FStateCompilationData& s: States )
@@ -2126,66 +2118,58 @@ namespace mu
             }
         }
 
-        // Reoptimise because of state cache reformats and gpuization
+        // Reoptimise because of state cache reformats
         {
             MUTABLE_CPUPROFILER_SCOPE(Reoptimise);
-            bool modified = true;
-            int32 numIterations = 0;
+            bool bModified = true;
+            int32 NumIterations = 0;
 			int32 Pass = 1;
-            while (modified && (!OptimizeIterationsMax || OptimizeIterationsLeft>0 || !numIterations ))
+            while (bModified && (OptimizeIterationsLeft>0 || !NumIterations))
             {
-                ++numIterations;
+                ++NumIterations;
                 --OptimizeIterationsLeft;
-                UE_LOG(LogMutableCore, Verbose, TEXT("State reoptimise iteration %d, max %d, left %d"),
-                        numIterations, OptimizeIterationsMax, OptimizeIterationsLeft);
+                UE_LOG(LogMutableCore, Verbose, TEXT("State reoptimise iteration %d, left %d"), NumIterations, OptimizeIterationsLeft);
 
-                modified = false;
+				bModified = false;
 
                 UE_LOG(LogMutableCore, Verbose, TEXT(" - semantic optimiser"));
-                modified |= SemanticOptimiserAST( roots, Options->GetPrivate()->OptimisationOptions, Pass );
-				//UE_LOG(LogMutableCore, Verbose, TEXT("(int) %s : %ld"), TEXT("ast size"), int64(ASTOp::CountNodes(roots)));
+				bModified |= SemanticOptimiserAST( Roots, Options->GetPrivate()->OptimisationOptions, Pass );
 
                 // Image size operations are treated separately
                 UE_LOG(LogMutableCore, Verbose, TEXT(" - size optimiser"));
-                modified |= SizeOptimiserAST( roots );
-				//UE_LOG(LogMutableCore, Verbose, TEXT("(int) %s : %ld"), TEXT("ast size"), int64(ASTOp::CountNodes(roots)));
+				bModified |= SizeOptimiserAST( Roots );
 			}
 
-            for(Ptr<ASTOp>& Root : roots)
+            for(Ptr<ASTOp>& Root : Roots)
             {
                 UE_LOG(LogMutableCore, Verbose, TEXT(" - constant optimiser"));
-				modified = ConstantGenerator( Options->GetPrivate(), Root, Pass );
+				bModified = ConstantGenerator( Options->GetPrivate(), Root, Pass );
 			}
 
-			//UE_LOG(LogMutableCore, Verbose, TEXT("(int) %s : %ld"), TEXT("ast size"), int64(ASTOp::CountNodes(roots)));
-
             UE_LOG(LogMutableCore, Verbose, TEXT(" - duplicated data remover"));
-            DuplicatedDataRemoverAST( roots );
-			//UE_LOG(LogMutableCore, Verbose, TEXT("(int) %s : %ld"), TEXT("ast size"), int64(ASTOp::CountNodes(roots)));
+            DuplicatedDataRemoverAST( Roots );
 
             UE_LOG(LogMutableCore, Verbose, TEXT(" - duplicated code remover"));
-            DuplicatedCodeRemoverAST( roots );
-			//UE_LOG(LogMutableCore, Verbose, TEXT("(int) %s : %ld"), TEXT("ast size"), int64(ASTOp::CountNodes(roots)));
+            DuplicatedCodeRemoverAST( Roots );
 		}
 
 		// Reset the state root operations in case they have changed due to optimization
 		for (int32 RootIndex = 0; RootIndex < States.Num(); ++RootIndex)
 		{
-			States[RootIndex].root = roots[RootIndex];
+			States[RootIndex].root = Roots[RootIndex];
 		}
 
         // Optimise the data formats
         {
             MUTABLE_CPUPROFILER_SCOPE(DataFormats);
 
-            DataOptimise( Options.get(), roots);
+            DataOptimise( Options.get(), Roots);
 
             // After optimising the data formats, we may remove more constants
-            DuplicatedDataRemoverAST( roots );
-            DuplicatedCodeRemoverAST( roots );
+            DuplicatedDataRemoverAST( Roots );
+            DuplicatedCodeRemoverAST( Roots );
 
-            // Update the marks for the instructions that don't depend on runtime parameters to be
-            // cached.
+            // Update the marks for the instructions that don't depend on runtime parameters to be cached.
             for (FStateCompilationData& s:States)
             {
                 StateCacheDetectorAST c( &s );
