[Mutable] Fix morph streaming issues.

- Fixed morph crash without streaming.
- Reenables morphs.
- Fixed DoComponentsNeedUpdate issue when streaming.

diff --git a/Engine/Plugins/Mutable/Source/CustomizableObject/Private/MuCO/CustomizableObjectInstance.cpp b/Engine/Plugins/Mutable/Source/CustomizableObject/Private/MuCO/CustomizableObjectInstance.cpp
--- a/Engine/Plugins/Mutable/Source/CustomizableObject/Private/MuCO/CustomizableObjectInstance.cpp
+++ b/Engine/Plugins/Mutable/Source/CustomizableObject/Private/MuCO/CustomizableObjectInstance.cpp
@@ -2075,7 +2075,9 @@ bool UCustomizableInstancePrivate::DoComponentsNeedUpdate(UCustomizableObjectIns
 		// If the component SkeletalMesh was build with scalability settings different than the current, we need to update.
 		// The mismatch can happen when the scalability setting is changed and an update is requested with a MinLOD lower 
 		// than the new FirstLODAvailable. This is an artifact of our LOD management system setting MinLOD.
-		if (!OperationData->MeshChangedPerInstanceComponent[InstanceComponentIndex] && bHadSkeletalMesh) 
+		if (!OperationData->bStreamMeshLODs &&
+			!OperationData->MeshChangedPerInstanceComponent[InstanceComponentIndex] &&
+			bHadSkeletalMesh) 
 		{
 			const int32 FirstRequestedLOD = OperationData->GetFirstRequestedLOD()[ComponentName];
 			for (int32 LODIndex = OperationData->FirstLODAvailable[ComponentName]; LODIndex < FirstRequestedLOD; ++LODIndex)
@@ -4728,6 +4730,18 @@ bool UCustomizableInstancePrivate::BuildOrCopyRenderData(const TSharedRef<FUpdat
 	const UModelResources& ModelResources = CustomizableObject->GetPrivate()->GetModelResourcesChecked();
 	const FName ComponentName = ModelResources.ComponentNamesPerObjectComponent[Component->Id.GetValue()];
 
+	for (int32 LODIndex = OperationData->GetFirstRequestedLOD()[ComponentName]; LODIndex < Component->LODCount; ++LODIndex)
+	{
+		const FInstanceUpdateData::FLOD& LOD = OperationData->InstanceUpdateData.LODs[Component->FirstLOD+LODIndex];
+
+		if (!LOD.Mesh || LOD.SurfaceCount == 0)
+		{
+			continue;	
+		}
+
+		SetLastMeshId(Component->Id, LODIndex, LOD.MeshID);
+	}
+	
 	const int32 FirstLOD = OperationData->bStreamMeshLODs ?
 		OperationData->FirstLODAvailable[ComponentName] :
 		OperationData->GetFirstRequestedLOD()[ComponentName];
@@ -4751,9 +4765,7 @@ bool UCustomizableInstancePrivate::BuildOrCopyRenderData(const TSharedRef<FUpdat
 		}
 
 		TRACE_CPUPROFILER_EVENT_SCOPE_TEXT(*FString::Printf(TEXT("BuildRenderData: Component index %d, LOD %d"), Component->Id.GetValue(), LODIndex));
-
-		SetLastMeshId(Component->Id, LODIndex, LOD.MeshID);
-
+		
 		FSkeletalMeshLODRenderData& LODResource = RenderData->LODRenderData[LODIndex];
 
 		const TMap<mu::FBoneName, TPair<FName, uint16>>& BoneInfoMap = OperationData->InstanceUpdateData.SkeletonsPerInstanceComponent[InstanceComponentIndex.GetValue()].BoneInfoMap;
@@ -4773,7 +4785,6 @@ bool UCustomizableInstancePrivate::BuildOrCopyRenderData(const TSharedRef<FUpdat
 			LODResource.RequiredBones.Sort();
 		}
 
-
 		// Find referenced surface metadata.
 		const int32 MeshNumSurfaces = LOD.Mesh->Surfaces.Num();
 		TArray<const FMutableSurfaceMetadata*> MeshSurfacesMetadata;
@@ -5117,21 +5128,139 @@ UE::Tasks::FTask UCustomizableInstancePrivate::LoadAdditionalAssetsAndData(const
 			}
 
 			const FName ComponentName = OperationData->ComponentNames[Component.Id.GetValue()];
-			
-			for (int32 LODIndex = OperationData->FirstResidentLOD[ComponentName]; LODIndex < Component.LODCount; ++LODIndex)
-			{
-				const FInstanceUpdateData::FLOD& LOD = OperationData->InstanceUpdateData.LODs[Component.FirstLOD + LODIndex];
 
-				TSharedPtr<const mu::FMesh> MutableMesh = LOD.Mesh;
-				if (!MutableMesh)
+			if (bUpdateMeshes)
+			{
+				// Morphs
 				{
-					continue;
+					// Data
+					for (int32 LODIndex = OperationData->GetFirstRequestedLOD()[ComponentName]; LODIndex < Component.LODCount; ++LODIndex)
+					{
+						const FInstanceUpdateData::FLOD& LOD = OperationData->InstanceUpdateData.LODs[Component.FirstLOD + LODIndex];
+
+						TSharedPtr<const mu::FMesh> MutableMesh = LOD.Mesh;
+						if (!MutableMesh)
+						{
+							continue;
+						}
+
+						LoadMorphTargets(StreamRequest, MutableMesh.ToSharedRef(), OperationData->InstanceUpdateData.RealTimeMorphTargetMeshData);
+					}
+
+					// Metadata
+					const int32 FirstLOD = OperationData->bStreamMeshLODs ?
+						OperationData->FirstResidentLOD[ComponentName] :
+						OperationData->GetFirstRequestedLOD()[ComponentName];
+				
+					for (int32 LODIndex = FirstLOD; LODIndex < Component.LODCount; ++LODIndex)
+					{
+						const FInstanceUpdateData::FLOD& LOD = OperationData->InstanceUpdateData.LODs[Component.FirstLOD + LODIndex];
+
+						TSharedPtr<const mu::FMesh> MutableMesh = LOD.Mesh;
+						if (!MutableMesh)
+						{
+							continue;
+						}
+					
+						const TArray<uint64>& StreamedResources = MutableMesh->GetStreamedResources();
+
+						for (uint64 ResourceId : StreamedResources)
+						{
+							FCustomizableObjectStreameableResourceId TypedResourceId = BitCast<FCustomizableObjectStreameableResourceId>(ResourceId);
+						
+							if (TypedResourceId.Type == static_cast<uint8>(FCustomizableObjectStreameableResourceId::EType::RealTimeMorphTarget))
+							{
+								check(TypedResourceId.Id != 0 && TypedResourceId.Id <= TNumericLimits<uint32>::Max());
+
+								if (ModelStreamableBulkData->RealTimeMorphStreamables.Contains(static_cast<uint32>(TypedResourceId.Id)))
+								{
+									FMorphTargetMeshData& ReadDestData = OperationData->InstanceUpdateData.RealTimeMorphTargetMeshData.FindOrAdd(TypedResourceId.Id);
+									ReadDestData.NameResolutionMap = ModelStreamableBulkData->RealTimeMorphStreamables[TypedResourceId.Id].NameResolutionMap;
+								}
+							}
+						}
+					}
 				}
 
-				if (bUpdateMeshes)
+				// Cloth
 				{
-					LoadMorphTargets(StreamRequest, MutableMesh.ToSharedRef(), OperationData->InstanceUpdateData.RealTimeMorphTargetMeshData);
-					LoadClothing(StreamRequest, MutableMesh.ToSharedRef(), OperationData->InstanceUpdateData.ClothingMeshData);
+					// Data
+					// From FirstResidentLOD instead of FirstRequestedLOD since clothing we generate all LODs, even the hacky streaming copied ones.
+					for (int32 LODIndex = OperationData->FirstResidentLOD[ComponentName]; LODIndex < Component.LODCount; ++LODIndex)
+					{
+						const FInstanceUpdateData::FLOD& LOD = OperationData->InstanceUpdateData.LODs[Component.FirstLOD + LODIndex];
+
+						TSharedPtr<const mu::FMesh> MutableMesh = LOD.Mesh;
+						if (!MutableMesh)
+						{
+							continue;
+						}
+
+						if (bUpdateMeshes)
+						{
+							LoadClothing(StreamRequest, MutableMesh.ToSharedRef(), OperationData->InstanceUpdateData.ClothingMeshData);
+						}
+					}
+
+					// Metadata
+					for (int32 LODIndex = OperationData->FirstLODAvailable[ComponentName]; LODIndex < Component.LODCount; ++LODIndex)
+					{
+						const FInstanceUpdateData::FLOD& LOD = OperationData->InstanceUpdateData.LODs[Component.FirstLOD + LODIndex];
+
+						TSharedPtr<const mu::FMesh> MutableMesh = LOD.Mesh;
+						if (!MutableMesh)
+						{
+							continue;
+						}
+					
+						const TArray<uint64>& StreamedResources = MutableMesh->GetStreamedResources();
+
+						for (uint64 ResourceId : StreamedResources)
+						{
+							FCustomizableObjectStreameableResourceId TypedResourceId = BitCast<FCustomizableObjectStreameableResourceId>(ResourceId);
+
+							if (TypedResourceId.Type == (uint8)FCustomizableObjectStreameableResourceId::EType::Clothing)
+							{
+								check(TypedResourceId.Id != 0 && TypedResourceId.Id <= TNumericLimits<uint32>::Max());
+
+								const TMap<uint32, FClothingStreamable>& ClothingStreamables = ModelStreamableBulkData->ClothingStreamables;
+								if (const FClothingStreamable* ClothingStreamable = ClothingStreamables.Find(TypedResourceId.Id))
+								{
+									FClothingMeshData& ReadDestData = OperationData->InstanceUpdateData.ClothingMeshData.FindOrAdd(TypedResourceId.Id);
+									ReadDestData.ClothingAssetIndex = ClothingStreamable->ClothingAssetIndex;
+									ReadDestData.ClothingAssetLOD = ClothingStreamable->ClothingAssetLOD;
+ 								
+									// TODO: Add async loading of ClothingAsset Data. This could be loaded as an streamead resource similar to and the asset user data.
+									int32 ClothingAssetIndex = ClothingStreamable->ClothingAssetIndex; 
+									int32 PhysicsAssetIndex = ClothingStreamable->PhysicsAssetIndex;					
+									const TSoftObjectPtr<UPhysicsAsset>& PhysicsAsset = ModelResources.PhysicsAssets.IsValidIndex(PhysicsAssetIndex) 
+											? ModelResources.PhysicsAssets[PhysicsAssetIndex] 
+											: nullptr;
+
+									// The entry should always be in the map
+									if (!PhysicsAsset.IsNull())
+									{
+										if (PhysicsAsset.Get())
+										{
+											if (ClothingPhysicsAssets.IsValidIndex(ClothingAssetIndex))
+											{
+												ClothingPhysicsAssets[ClothingAssetIndex] = PhysicsAsset.Get();
+											}
+										}
+										else
+										{
+											ComponentData->ClothingPhysicsAssetsToStream.Emplace(ClothingAssetIndex, PhysicsAssetIndex);
+											AssetsToStream.AddUnique(PhysicsAsset.ToSoftObjectPath());
+										}
+									}
+								}
+								else
+								{
+									UE_LOG(LogMutable, Error, TEXT("Invalid streamed clothing data block [%d] found."), TypedResourceId.Id);
+								}
+							}
+						}
+					}
 				}
 			}
 			
@@ -5159,56 +5288,6 @@ UE::Tasks::FTask UCustomizableInstancePrivate::LoadAdditionalAssetsAndData(const
 #endif
 						ComponentData->StreamedResourceIndex.AddUnique(ResourceIndex);
 					}
-					else if (TypedResourceId.Type == static_cast<uint8>(FCustomizableObjectStreameableResourceId::EType::RealTimeMorphTarget))
-					{
-						check(TypedResourceId.Id != 0 && TypedResourceId.Id <= TNumericLimits<uint32>::Max());
-
-						if (ModelStreamableBulkData->RealTimeMorphStreamables.Contains(static_cast<uint32>(TypedResourceId.Id)))
-						{
-							FMorphTargetMeshData& ReadDestData = OperationData->InstanceUpdateData.RealTimeMorphTargetMeshData.FindOrAdd(TypedResourceId.Id);
-							ReadDestData.NameResolutionMap = ModelStreamableBulkData->RealTimeMorphStreamables[TypedResourceId.Id].NameResolutionMap;
-						}
-					}
-					else if (TypedResourceId.Type == (uint8)FCustomizableObjectStreameableResourceId::EType::Clothing)
-					{
-						check(TypedResourceId.Id != 0 && TypedResourceId.Id <= TNumericLimits<uint32>::Max());
-
-						const TMap<uint32, FClothingStreamable>& ClothingStreamables = ModelStreamableBulkData->ClothingStreamables;
-						if (const FClothingStreamable* ClothingStreamable = ClothingStreamables.Find(TypedResourceId.Id))
-						{
-							FClothingMeshData& ReadDestData = OperationData->InstanceUpdateData.ClothingMeshData.FindOrAdd(TypedResourceId.Id);
-							ReadDestData.ClothingAssetIndex = ClothingStreamable->ClothingAssetIndex;
-							ReadDestData.ClothingAssetLOD = ClothingStreamable->ClothingAssetLOD;
- 							
-							// TODO: Add async loading of ClothingAsset Data. This could be loaded as an streamead resource similar to and the asset user data.
-							int32 ClothingAssetIndex = ClothingStreamable->ClothingAssetIndex; 
-							int32 PhysicsAssetIndex = ClothingStreamable->PhysicsAssetIndex;					
-							const TSoftObjectPtr<UPhysicsAsset>& PhysicsAsset = ModelResources.PhysicsAssets.IsValidIndex(PhysicsAssetIndex) 
-									? ModelResources.PhysicsAssets[PhysicsAssetIndex] 
-								    : nullptr;
-
-							// The entry should always be in the map
-							if (!PhysicsAsset.IsNull())
-							{
-								if (PhysicsAsset.Get())
-								{
-									if (ClothingPhysicsAssets.IsValidIndex(ClothingAssetIndex))
-									{
-										ClothingPhysicsAssets[ClothingAssetIndex] = PhysicsAsset.Get();
-									}
-								}
-								else
-								{
-									ComponentData->ClothingPhysicsAssetsToStream.Emplace(ClothingAssetIndex, PhysicsAssetIndex);
-									AssetsToStream.AddUnique(PhysicsAsset.ToSoftObjectPath());
-								}
-							}
-						}
-						else
-						{
-							UE_LOG(LogMutable, Error, TEXT("Invalid streamed clothing data block [%d] found."), TypedResourceId.Id);
-						}
-					}
 				}
 				
 				for (int32 TagIndex = 0; TagIndex < MutableMesh->GetTagCount(); ++TagIndex)
diff --git a/Engine/Plugins/Mutable/Source/CustomizableObject/Private/MuCO/CustomizableObjectSystem.cpp b/Engine/Plugins/Mutable/Source/CustomizableObject/Private/MuCO/CustomizableObjectSystem.cpp
--- a/Engine/Plugins/Mutable/Source/CustomizableObject/Private/MuCO/CustomizableObjectSystem.cpp
+++ b/Engine/Plugins/Mutable/Source/CustomizableObject/Private/MuCO/CustomizableObjectSystem.cpp
@@ -119,7 +119,7 @@ TAutoConsoleVariable<bool> CVarEnableUpdateOptimization(
 
 TAutoConsoleVariable<bool> CVarEnableRealTimeMorphTargets(
 	TEXT("mutable.EnableRealTimeMorphTargets"),
-	false,
+	true,
 	TEXT("Enable or disable generation of realtime morph targets."));
 
 
@@ -4909,9 +4909,7 @@ void LoadMorphTargets(FMutableStreamRequest& MutableDataStreamer, const TSharedR
 		{
 			continue;
 		}
-
-		ReadDestData.NameResolutionMap = ModelStreamableBulkData->RealTimeMorphStreamables[BlockId].NameResolutionMap;
-
+		
 		check(Streamable.Size % sizeof(FMorphTargetVertexData) == 0);
 		uint32 NumElems = Streamable.Size / sizeof(FMorphTargetVertexData);
 
