[mutable] Moved object section scope now that the MTU can handle it

#rnx

diff --git a/Engine/Plugins/Mutable/Source/MutableValidation/Private/MuV/CustomizableObjectBulkValidationCommandlet.cpp b/Engine/Plugins/Mutable/Source/MutableValidation/Private/MuV/CustomizableObjectBulkValidationCommandlet.cpp
--- a/Engine/Plugins/Mutable/Source/MutableValidation/Private/MuV/CustomizableObjectBulkValidationCommandlet.cpp
+++ b/Engine/Plugins/Mutable/Source/MutableValidation/Private/MuV/CustomizableObjectBulkValidationCommandlet.cpp
@@ -125,18 +125,9 @@ int32 UCustomizableObjectBulkValidationCommandlet::Main(const FString& Params)
 			continue;
 		}
 		
-		// Get the CustomizableObject.
-		{
-			// TODO: MOVE THIS Scope section inside TestCustomizableObject once the parsing system has been updated to handle this new section type
-			// By keeping it here we avoid our current implementation of the single CO Validation from using this new type of not yet supported by the
-			// parser section.
-			const FScopedLogSection ObjectSection (EMutableLogSection::Object, FName(CustomizableObjectName));
-			
-			// Body of the test ----------------------------------------------------------------------------------------------------------------------
-			TestCustomizableObject(*TargetCustomizableObject, *TargetCompilationPlatform, InstancesToGenerate, bUseDiskCompilation);
-			// ---------------------------------------------------------------------------------------------------------------------------------------
-		}
-
+		// Body of the test ----------------------------------------------------------------------------------------------------------------------
+		TestCustomizableObject(*TargetCustomizableObject, *TargetCompilationPlatform, InstancesToGenerate, bUseDiskCompilation);
+		// ---------------------------------------------------------------------------------------------------------------------------------------
 		
 		// Try to collect the garbage
 		{
diff --git a/Engine/Plugins/Mutable/Source/MutableValidation/Private/MuV/ValidationUtils.cpp b/Engine/Plugins/Mutable/Source/MutableValidation/Private/MuV/ValidationUtils.cpp
--- a/Engine/Plugins/Mutable/Source/MutableValidation/Private/MuV/ValidationUtils.cpp
+++ b/Engine/Plugins/Mutable/Source/MutableValidation/Private/MuV/ValidationUtils.cpp
@@ -5,6 +5,7 @@
 #include "CustomizableObjectCompilationUtility.h"
 #include "CustomizableObjectInstanceUpdateUtility.h"
 #include "RHIGlobals.h"
+#include "ScopedLogSection.h"
 #include "AssetRegistry/AssetRegistryModule.h"
 #include "Commandlets/Commandlet.h"
 #include "Interfaces/ITargetPlatform.h"
@@ -169,6 +170,8 @@ ITargetPlatform* GetCompilationPlatform(const FString& Params)
 bool TestCustomizableObject(UCustomizableObject& InTargetCustomizableObject, const ITargetPlatform& TargetCompilationPlatform,
 	const uint32 InstancesToGenerate, const bool bUseDiskCompilation)
 {
+	const FScopedLogSection ObjectSection (EMutableLogSection::Object, FName( InTargetCustomizableObject.GetPathName()));
+	
 	// Keep a strong object pointer pointing at the CO to prevent it from being GCd during the test
 	const TStrongObjectPtr<UCustomizableObject> TargetCO = TStrongObjectPtr{&InTargetCustomizableObject};
 	
