[Mutable] Fix crash with the mesh preview utils when there is a descriptor mesh.

#rnx

diff --git a/Engine/Plugins/Mutable/Source/CustomizableObjectEditor/Private/MuCOE/MutableMeshPreviewUtils.cpp b/Engine/Plugins/Mutable/Source/CustomizableObjectEditor/Private/MuCOE/MutableMeshPreviewUtils.cpp
--- a/Engine/Plugins/Mutable/Source/CustomizableObjectEditor/Private/MuCOE/MutableMeshPreviewUtils.cpp
+++ b/Engine/Plugins/Mutable/Source/CustomizableObjectEditor/Private/MuCOE/MutableMeshPreviewUtils.cpp
@@ -463,6 +463,11 @@ namespace MutableMeshPreviewUtils
 			return nullptr;
 		}
 
+		if (InMutableMesh->GetVertexCount() == 0 || InMutableMesh->GetVertexBuffers().IsDescriptor())
+		{
+			return nullptr;
+		}
+
 		TSharedPtr<const mu::FMesh> MeshCopy;
 
 		// Buffer structures checks and corrective operations
diff --git a/Engine/Plugins/Mutable/Source/CustomizableObjectEditor/Private/MuCOE/SMutableMeshViewer.cpp b/Engine/Plugins/Mutable/Source/CustomizableObjectEditor/Private/MuCOE/SMutableMeshViewer.cpp
--- a/Engine/Plugins/Mutable/Source/CustomizableObjectEditor/Private/MuCOE/SMutableMeshViewer.cpp
+++ b/Engine/Plugins/Mutable/Source/CustomizableObjectEditor/Private/MuCOE/SMutableMeshViewer.cpp
@@ -543,7 +543,7 @@ TArray<FVector2f> SMutableMeshViewer::GetUVs() const
 {
 	TArray<FVector2f> Lines;
 
-	if (MutableMesh)
+	if (MutableMesh && !MutableMesh->VertexBuffers.IsDescriptor())
 	{
 		int32 NumFaces = MutableMesh->GetFaceCount();
 		Lines.Reserve(NumFaces*6);
