[Mutable] Replace CurrentSupportedVersion checks with its hashed version.

#rnx

diff --git a/Engine/Plugins/Mutable/Source/CustomizableObject/Internal/MuCO/CustomizableObjectPrivate.h b/Engine/Plugins/Mutable/Source/CustomizableObject/Internal/MuCO/CustomizableObjectPrivate.h
--- a/Engine/Plugins/Mutable/Source/CustomizableObject/Internal/MuCO/CustomizableObjectPrivate.h
+++ b/Engine/Plugins/Mutable/Source/CustomizableObject/Internal/MuCO/CustomizableObjectPrivate.h
@@ -1690,6 +1690,8 @@ public:
 
 #if WITH_EDITOR
 
+	static uint32 GetECustomizableObjectVersionEnumHash();
+
 	// Compose folder name where the data is stored
 	static FString GetCompiledDataFolderPath();
 	static FString GetDataTypeExtension(MutablePrivate::EStreamableDataType DataType);
@@ -1882,7 +1884,6 @@ public:
 	UPROPERTY()
 	int32 ImageTiling = 0;
 
-	static constexpr int32 CurrentSupportedVersion = static_cast<int32>(ECustomizableObjectVersions::LastCustomizableObjectVersion);
 #endif
 	
 	static constexpr int32 DerivedDataVersion = 0x46799a1c;
diff --git a/Engine/Plugins/Mutable/Source/CustomizableObject/Private/MuCO/CustomizableObject.cpp b/Engine/Plugins/Mutable/Source/CustomizableObject/Private/MuCO/CustomizableObject.cpp
--- a/Engine/Plugins/Mutable/Source/CustomizableObject/Private/MuCO/CustomizableObject.cpp
+++ b/Engine/Plugins/Mutable/Source/CustomizableObject/Private/MuCO/CustomizableObject.cpp
@@ -451,7 +451,7 @@ void UCustomizableObject::Serialize(FArchive& Ar_Asset)
 	else
 	{
 		// Can't remove this or saved customizable objects will fail to load
-		int64 InternalVersion = UCustomizableObjectPrivate::CurrentSupportedVersion;
+		int64 InternalVersion = 0;
 		Ar_Asset << InternalVersion;
 	}
 #else
@@ -537,7 +537,7 @@ TObjectPtr<UModelResources> UCustomizableObjectPrivate::LoadModelResources_Inter
 	FObjectAndNameAsStringProxyArchive ObjectReader(MemoryReader, true);
 	LocalModelResources->Serialize(ObjectReader);
 
-	const bool bLoadedSuccessfully = LocalModelResources->CodeVersion == UCustomizableObjectPrivate::CurrentSupportedVersion;
+	const bool bLoadedSuccessfully = LocalModelResources->CodeVersion == GetECustomizableObjectVersionEnumHash();
 	return bLoadedSuccessfully ? LocalModelResources : nullptr;
 }
 
@@ -672,7 +672,7 @@ void UCustomizableObjectPrivate::LoadCompiledDataFromDisk()
 			MutableCompiledDataStreamHeader DataTypeHeader;
 			AuxMemoryReader << DataTypeHeader;
 
-			if(DataTypeHeader.InternalVersion != CurrentSupportedVersion
+			if(DataTypeHeader.InternalVersion != GetECustomizableObjectVersionEnumHash()
 				|| (!IsRunningGame() && DataTypeHeader.VersionId != GetVersionId()) 
 				|| (CompiledVersionId.IsValid() && CompiledVersionId != DataTypeHeader.VersionId))
 			{
@@ -823,6 +823,31 @@ FReply UCustomizableObjectPrivate::AddNewParameterProfile(FString Name, UCustomi
 }
 
 
+uint32 UCustomizableObjectPrivate::GetECustomizableObjectVersionEnumHash()
+{
+	static const uint32 VersionsHash = []() -> uint32
+		{
+			UEnum* Enum = StaticEnum<ECustomizableObjectVersions>();
+			check(Enum);
+
+			const int32 StartIndex = Enum->GetIndexByValue(static_cast<int64>(ECustomizableObjectVersions::FirstEnumeratedVersion));
+			const int32 EndIndex = Enum->GetIndexByValue(static_cast<int64>(ECustomizableObjectVersions::LastCustomizableObjectVersion));
+			check(StartIndex < EndIndex);
+
+			uint32 CombinedHash = 0;
+			for (int32 Index = StartIndex; Index <= EndIndex; Index++)
+			{
+				const FString VersionString = Enum->GetNameStringByIndex(Index);
+				CombinedHash = HashCombine(GetTypeHash(VersionString), CombinedHash);
+			}
+
+			return CombinedHash;
+		}();
+
+	return VersionsHash;
+}
+
+
 FString UCustomizableObjectPrivate::GetCompiledDataFolderPath()
 {	
 	return FPaths::ConvertRelativePathToFull(FPaths::ProjectSavedDir() + TEXT("MutableStreamedDataEditor/"));
@@ -870,7 +895,7 @@ void UCustomizableObjectPrivate::SaveEmbeddedData(FArchive& Ar)
 
 	TSharedPtr<mu::FModel> Model = GetModel();
 
-	int32 InternalVersion = Model ? CurrentSupportedVersion : INDEX_NONE;
+	int64 InternalVersion = Model ? GetECustomizableObjectVersionEnumHash() : INDEX_NONE;
 	Ar << InternalVersion;
 
 	if (Model)
@@ -893,7 +918,7 @@ void UCustomizableObjectPrivate::LoadEmbeddedData(FArchive& Ar)
 {
 	MUTABLE_CPUPROFILER_SCOPE(UCustomizableObject::LoadEmbeddedData)
 
-	int32 InternalVersion;
+	int64 InternalVersion;
 	Ar << InternalVersion;
 
 	// If this fails, something went wrong with the packaging: we have data that belongs
@@ -1380,31 +1405,6 @@ const FString& UCustomizableObject::GetParameterName(int32 ParamIndex) const
 	return s_EmptyString;
 }
 
-#if WITH_EDITOR
-uint32 GetECustomizableObjectVersionEnumHash()
-{
-	static const uint32 VersionsHash = []() -> uint32
-	{
-		UEnum* Enum = StaticEnum<ECustomizableObjectVersions>();
-		check(Enum);
-		
-		const int32 StartIndex = Enum->GetIndexByValue(static_cast<int64>(ECustomizableObjectVersions::FirstEnumeratedVersion));
-		const int32 EndIndex = Enum->GetIndexByValue(static_cast<int64>(ECustomizableObjectVersions::LastCustomizableObjectVersion));
-		check(StartIndex < EndIndex);
-		
-		uint32 CombinedHash = 0;
-		for (int32 Index = StartIndex; Index <= EndIndex; Index++)
-		{
-			const FString VersionString = Enum->GetNameStringByIndex(Index);
-			CombinedHash = HashCombine( GetTypeHash(VersionString), CombinedHash);
-		}
-		
-		return CombinedHash;
-	}();
-
-	return VersionsHash;
-}
-#endif
 
 void UCustomizableObjectPrivate::UpdateParameterPropertiesFromModel(const TSharedPtr<mu::FModel>& Model)
 {
diff --git a/Engine/Plugins/Mutable/Source/CustomizableObjectEditor/Private/MuCOE/CustomizableObjectCompileRunnable.cpp b/Engine/Plugins/Mutable/Source/CustomizableObjectEditor/Private/MuCOE/CustomizableObjectCompileRunnable.cpp
--- a/Engine/Plugins/Mutable/Source/CustomizableObjectEditor/Private/MuCOE/CustomizableObjectCompileRunnable.cpp
+++ b/Engine/Plugins/Mutable/Source/CustomizableObjectEditor/Private/MuCOE/CustomizableObjectCompileRunnable.cpp
@@ -351,7 +351,7 @@ FCustomizableObjectSaveDDRunnable::FCustomizableObjectSaveDDRunnable(const TShar
 	UCustomizableObject* CustomizableObject = InRequest->GetCustomizableObject();
 	CustomizableObjectName = GetNameSafe(CustomizableObject);
 
-	CustomizableObjectHeader.InternalVersion = CustomizableObject->GetPrivate()->CurrentSupportedVersion;
+	CustomizableObjectHeader.InternalVersion = UCustomizableObjectPrivate::GetECustomizableObjectVersionEnumHash();
 	CustomizableObjectHeader.VersionId = CustomizableObject->GetPrivate()->GetVersionId();
 
 	// Cache ModelResources
diff --git a/Engine/Plugins/Mutable/Source/CustomizableObjectEditor/Private/MuCOE/CustomizableObjectCompiler.cpp b/Engine/Plugins/Mutable/Source/CustomizableObjectEditor/Private/MuCOE/CustomizableObjectCompiler.cpp
--- a/Engine/Plugins/Mutable/Source/CustomizableObjectEditor/Private/MuCOE/CustomizableObjectCompiler.cpp
+++ b/Engine/Plugins/Mutable/Source/CustomizableObjectEditor/Private/MuCOE/CustomizableObjectCompiler.cpp
@@ -858,7 +858,7 @@ void FCustomizableObjectCompiler::CompileInternal(bool bAsync)
 		
 		PlatformData->ModelStreamableBulkData = MakeShared<FModelStreamableBulkData>();
 
-		ModelResources->CodeVersion = UCustomizableObjectPrivate::CurrentSupportedVersion;
+		ModelResources->CodeVersion = UCustomizableObjectPrivate::GetECustomizableObjectVersionEnumHash();
 
 		ModelResources->ReferenceSkeletalMeshesData = MoveTemp(GenerationContext.ReferenceSkeletalMeshesData);
 
