[Mutable] Fix for static analysis warning.

#rnx

diff --git a/Engine/Plugins/Mutable/Source/CustomizableObject/Private/MuCO/CustomizableObject.cpp b/Engine/Plugins/Mutable/Source/CustomizableObject/Private/MuCO/CustomizableObject.cpp
--- a/Engine/Plugins/Mutable/Source/CustomizableObject/Private/MuCO/CustomizableObject.cpp
+++ b/Engine/Plugins/Mutable/Source/CustomizableObject/Private/MuCO/CustomizableObject.cpp
@@ -895,7 +895,7 @@ void UCustomizableObjectPrivate::SaveEmbeddedData(FArchive& Ar)
 
 	TSharedPtr<mu::FModel> Model = GetModel();
 
-	int64 InternalVersion = Model ? GetECustomizableObjectVersionEnumHash() : INDEX_NONE;
+	int64 InternalVersion = Model ? (int64)GetECustomizableObjectVersionEnumHash() : INDEX_NONE;
 	Ar << InternalVersion;
 
 	if (Model)
