[mutable] Additional optimization removing the "prepare layout" operation when layouts end up not being used.
#rnx

diff --git a/Engine/Plugins/Mutable/Source/MutableTools/Private/MuT/DataPacker.cpp b/Engine/Plugins/Mutable/Source/MutableTools/Private/MuT/DataPacker.cpp
--- a/Engine/Plugins/Mutable/Source/MutableTools/Private/MuT/DataPacker.cpp
+++ b/Engine/Plugins/Mutable/Source/MutableTools/Private/MuT/DataPacker.cpp
@@ -30,6 +30,7 @@
 #include "MuT/ASTOpMeshMorph.h"
 #include "MuT/ASTOpMeshProject.h"
 #include "MuT/ASTOpMeshApplyLayout.h"
+#include "MuT/ASTOpMeshPrepareLayout.h"
 #include "MuT/ASTOpLayoutFromMesh.h"
 #include "MuT/ASTOpImageRasterMesh.h"
 #include "MuT/CompilerPrivate.h"
@@ -367,16 +368,16 @@ namespace mu
                 recurse = false;
 				const ASTOpImageCompose* op = static_cast<const ASTOpImageCompose*>(node.get());
 
-				TArray<bool> newState;
-				newState.Init(false, int32(EImageFormat::Count));
-                RecurseWithState( op->Layout.child(), newState );
-                RecurseWithState( op->Base.child(), newState );
-                RecurseWithState( op->BlockImage.child(), newState );
+				TArray<bool> NewState;
+				NewState.Init(false, int32(EImageFormat::Count));
+                RecurseWithState( op->Layout.child(), NewState );
+                RecurseWithState( op->Base.child(), NewState );
+                RecurseWithState( op->BlockImage.child(), NewState );
 
                 if ( op->Mask )
                 {
-                    newState[(int32)EImageFormat::L_UBitRLE] = true;
-                    RecurseWithState( op->Mask.child(), newState );
+                    NewState[(int32)EImageFormat::L_UBitRLE] = true;
+                    RecurseWithState( op->Mask.child(), NewState );
                 }
                 break;
             }
@@ -385,17 +386,17 @@ namespace mu
             {
                 recurse = false;
 				const ASTOpImageLayerColor* op = static_cast<const ASTOpImageLayerColor*>(node.get());
-				TArray<bool> newState;
-				newState.Init(false, int32(EImageFormat::Count));
-				RecurseWithState( op->base.child(), newState );
-                RecurseWithState( op->color.child(), newState );
+				TArray<bool> NewState;
+				NewState.Init(false, int32(EImageFormat::Count));
+				RecurseWithState( op->base.child(), NewState );
+                RecurseWithState( op->color.child(), NewState );
 
                 if ( op->mask )
                 {
-                    newState[(int32)EImageFormat::L_UByte] = true;
-                    newState[(int32)EImageFormat::L_UByteRLE] = true;
+                    NewState[(int32)EImageFormat::L_UByte] = true;
+                    NewState[(int32)EImageFormat::L_UByteRLE] = true;
 
-                    RecurseWithState( op->mask.child(), newState );
+                    RecurseWithState( op->mask.child(), NewState );
                 }
                 break;
             }
@@ -404,17 +405,17 @@ namespace mu
             {
                 recurse = false;
 				const ASTOpImageLayer* op = static_cast<const ASTOpImageLayer*>(node.get());
-				TArray<bool> newState;
-				newState.Init(false, int32(EImageFormat::Count));
-				RecurseWithState( op->base.child(), newState );
-                RecurseWithState( op->blend.child(), newState );
+				TArray<bool> NewState;
+				NewState.Init(false, int32(EImageFormat::Count));
+				RecurseWithState( op->base.child(), NewState );
+                RecurseWithState( op->blend.child(), NewState );
 
                 if (op->mask)
                 {
-                    newState[(int32)EImageFormat::L_UByte] = true;
-                    newState[(int32)EImageFormat::L_UByteRLE] = true;
+                    NewState[(int32)EImageFormat::L_UByte] = true;
+                    NewState[(int32)EImageFormat::L_UByteRLE] = true;
 
-                    RecurseWithState( op->mask.child(), newState );
+                    RecurseWithState( op->mask.child(), NewState );
                 }
                 break;
             }
@@ -423,17 +424,17 @@ namespace mu
             {
                 recurse = false;
 				const ASTOpImageMultiLayer* op = static_cast<const ASTOpImageMultiLayer*>(node.get());
-				TArray<bool> newState;
-				newState.Init(false, int32(EImageFormat::Count));
-				RecurseWithState( op->base.child(), newState );
-                RecurseWithState( op->blend.child(), newState );
+				TArray<bool> NewState;
+				NewState.Init(false, int32(EImageFormat::Count));
+				RecurseWithState( op->base.child(), NewState );
+                RecurseWithState( op->blend.child(), NewState );
 
                 if (op->mask)
                 {
-                    newState[(size_t)EImageFormat::L_UByte] = true;
-                    newState[(size_t)EImageFormat::L_UByteRLE] = true;
+                    NewState[(size_t)EImageFormat::L_UByte] = true;
+                    NewState[(size_t)EImageFormat::L_UByteRLE] = true;
 
-                    RecurseWithState( op->mask.child(), newState );
+                    RecurseWithState( op->mask.child(), NewState );
                 }
                 break;
             }
@@ -442,14 +443,14 @@ namespace mu
             {
                 recurse = false;
 				const ASTOpImageDisplace* op = static_cast<const ASTOpImageDisplace*>(node.get());
-				TArray<bool> newState;
-				newState.Init(false, int32(EImageFormat::Count));
-				RecurseWithState( op->Source.child(), newState );
+				TArray<bool> NewState;
+				NewState.Init(false, int32(EImageFormat::Count));
+				RecurseWithState( op->Source.child(), NewState );
 
-                newState[(int32)EImageFormat::L_UByte ] = true;
-                newState[(int32)EImageFormat::L_UByteRLE ] = true;
+                NewState[(int32)EImageFormat::L_UByte ] = true;
+                NewState[(int32)EImageFormat::L_UByteRLE ] = true;
 
-                RecurseWithState( op->DisplacementMap.child(), newState );
+                RecurseWithState( op->DisplacementMap.child(), NewState );
                 break;
             }
 
@@ -459,11 +460,11 @@ namespace mu
                 //Recurse(at, program);
                 //m_currentFormats.pop_back();
 
-				TArray<bool> newState;
-				newState.Init(false, int32(EImageFormat::Count));
-				if (currentFormats != newState)
+				TArray<bool> NewState;
+				NewState.Init(false, int32(EImageFormat::Count));
+				if (currentFormats != NewState)
                 {
-                    RecurseWithState(node, newState);
+                    RecurseWithState(node, NewState);
                     recurse = false;
                 }
                 else
@@ -519,180 +520,182 @@ namespace mu
 
         bool Visit( const Ptr<ASTOp>& node ) override
         {
-            bool recurse = true;
+            bool bRecurse = true;
 
-            uint64 currentSemantics = GetCurrentState();
+            uint64 CurrentSemantics = GetCurrentState();
 
             switch ( node->GetOpType() )
             {
 
-            case EOpType::ME_CONSTANT:
-            {
-                // Accumulate necessary semantics
-				const ASTOpConstantResource* op = static_cast<const ASTOpConstantResource*>(node.get());
-                uint64 currentFlags = 0;
-                if (m_requiredSemantics.count(op))
-                {
-                    currentFlags = m_requiredSemantics[op];
-                }
-                else
-                {
-                    m_requiredSemantics.insert(std::make_pair(op,currentFlags));
-                }
+			case EOpType::ME_CONSTANT:
+			{
+				// Accumulate necessary semantics
+				ASTOpConstantResource* Op = static_cast<ASTOpConstantResource*>(node.get());
+				uint64 InitialFlags = 0;
+				uint64& CurrentFlags = RequiredSemanticsPerConstant.FindOrAdd(Op, InitialFlags);
+				CurrentFlags |= CurrentSemantics;
+				bRecurse = false;
+				break;
+			}
 
-                currentFlags |= currentSemantics;
-                m_requiredSemantics[op] = currentFlags;
-                recurse = false;
-                break;
-            }
+			case EOpType::ME_PREPARELAYOUT:
+			{
+				// Accumulate necessary semantics
+				ASTOpMeshPrepareLayout* Op = static_cast<ASTOpMeshPrepareLayout*>(node.get());
+				uint64 InitialFlags = 0;
+				uint64& CurrentFlags = RequiredSemanticsPerPrepareLayout.FindOrAdd(Op, InitialFlags);
+				CurrentFlags |= CurrentSemantics;
+				break;
+			}
 
             // TODO: These could probably optimise something
             //case EOpType::IM_RASTERMESH: break;
 
             case EOpType::ME_DIFFERENCE:
             {
-                recurse = false;
+				bRecurse = false;
 
 				const ASTOpMeshDifference* op = static_cast<const ASTOpMeshDifference*>(node.get());
 
-                uint64 newState = currentSemantics;
-                newState |= (UINT64_C(1)<<uint32(EMeshBufferSemantic::VertexIndex));
-                RecurseWithState( op->Base.child(), newState );
+                uint64 NewState = CurrentSemantics;
+                NewState |= (UINT64_C(1)<<uint32(EMeshBufferSemantic::VertexIndex));
+                RecurseWithState( op->Base.child(), NewState );
 
-                RecurseWithState( op->Target.child(), currentSemantics );
+                RecurseWithState( op->Target.child(), CurrentSemantics );
                 break;
              }
 
             case EOpType::ME_REMOVEMASK:
             {
-                recurse = false;
+				bRecurse = false;
 
 				const ASTOpMeshRemoveMask* op = static_cast<const ASTOpMeshRemoveMask*>(node.get());
 
-                uint64 newState = currentSemantics;
-                newState |= (UINT64_C(1)<< uint32(EMeshBufferSemantic::VertexIndex));
-                RecurseWithState( op->source.child(), newState );
+                uint64 NewState = CurrentSemantics;
+                NewState |= (UINT64_C(1)<< uint32(EMeshBufferSemantic::VertexIndex));
+                RecurseWithState( op->source.child(), NewState );
                 for( const TPair<ASTChild, ASTChild>& r: op->removes )
                 {
-                    RecurseWithState( r.Value.child(), newState );
+                    RecurseWithState( r.Value.child(), NewState );
                 }
                 break;
              }
 
             case EOpType::ME_MORPH:
             {
-                recurse = false;
+				bRecurse = false;
 
 				const ASTOpMeshMorph* op = static_cast<const ASTOpMeshMorph*>(node.get());
 
-                uint64 newState = currentSemantics;
-                newState |= (UINT64_C(1)<< uint32(EMeshBufferSemantic::VertexIndex));
-                RecurseWithState( op->Base.child(), newState );
-                RecurseWithState( op->Target.child(), newState );
+                uint64 NewState = CurrentSemantics;
+                NewState |= (UINT64_C(1)<< uint32(EMeshBufferSemantic::VertexIndex));
+                RecurseWithState( op->Base.child(), NewState );
+                RecurseWithState( op->Target.child(), NewState );
                 break;
              }
 
             case EOpType::ME_APPLYLAYOUT:
             {
-                recurse = false;
+				bRecurse = false;
 
 				const ASTOpMeshApplyLayout* Op = static_cast<const ASTOpMeshApplyLayout*>(node.get());
 
-                uint64 NewState = currentSemantics;
+                uint64 NewState = CurrentSemantics;
 				NewState |= (UINT64_C(1)<< uint32(EMeshBufferSemantic::LayoutBlock));
                 RecurseWithState( Op->Mesh.child(), NewState);
 
-                RecurseWithState( Op->Layout.child(), currentSemantics );
+                RecurseWithState( Op->Layout.child(), CurrentSemantics );
                 break;
             }
 
 			case EOpType::ME_PROJECT:
 			{
-				recurse = false;
+				bRecurse = false;
 
 				const ASTOpMeshProject* op = static_cast<const ASTOpMeshProject*>(node.get());
 
-				uint64 newState = currentSemantics;
-				newState |= (UINT64_C(1) << uint32(EMeshBufferSemantic::LayoutBlock));
-				RecurseWithState(op->Mesh.child(), newState);
+				uint64 NewState = CurrentSemantics;
+				NewState |= (UINT64_C(1) << uint32(EMeshBufferSemantic::LayoutBlock));
+				RecurseWithState(op->Mesh.child(), NewState);
 
-				RecurseWithState(op->Projector.child(), currentSemantics);
+				RecurseWithState(op->Projector.child(), CurrentSemantics);
 				break;
 			}
 
 			case EOpType::IM_RASTERMESH:
 			{
-				recurse = false;
+				bRecurse = false;
 
 				const ASTOpImageRasterMesh* op = static_cast<const ASTOpImageRasterMesh*>(node.get());
 
-				uint64 newState = currentSemantics;
-				newState |= (UINT64_C(1) << uint32(EMeshBufferSemantic::LayoutBlock));
-				RecurseWithState(op->mesh.child(), newState);
+				uint64 NewState = CurrentSemantics;
+				NewState |= (UINT64_C(1) << uint32(EMeshBufferSemantic::LayoutBlock));
+				RecurseWithState(op->mesh.child(), NewState);
 
-				RecurseWithState(op->image.child(), currentSemantics);
-				RecurseWithState(op->angleFadeProperties.child(), currentSemantics);
-				RecurseWithState(op->mask.child(), currentSemantics);
-				RecurseWithState(op->projector.child(), currentSemantics);
+				RecurseWithState(op->image.child(), CurrentSemantics);
+				RecurseWithState(op->angleFadeProperties.child(), CurrentSemantics);
+				RecurseWithState(op->mask.child(), CurrentSemantics);
+				RecurseWithState(op->projector.child(), CurrentSemantics);
 				break;
 			}
 
             case EOpType::ME_EXTRACTLAYOUTBLOCK:
             {
-                recurse = false;
+				bRecurse = false;
 
 				const ASTOpMeshExtractLayoutBlocks* op = static_cast<const ASTOpMeshExtractLayoutBlocks*>(node.get());
 
                 // todo: check if we really need all of them
-                uint64 newState = currentSemantics;
-                newState |= (UINT64_C(1)<< uint32(EMeshBufferSemantic::LayoutBlock));
-                newState |= (UINT64_C(1)<< uint32(EMeshBufferSemantic::VertexIndex));
+                uint64 NewState = CurrentSemantics;
+                NewState |= (UINT64_C(1)<< uint32(EMeshBufferSemantic::LayoutBlock));
+                NewState |= (UINT64_C(1)<< uint32(EMeshBufferSemantic::VertexIndex));
 
-                RecurseWithState( op->Source.child(), newState );
+                RecurseWithState( op->Source.child(), NewState );
                 break;
             }
 
 			case EOpType::LA_FROMMESH:
 			{
-				recurse = false;
+				bRecurse = false;
 
 				const ASTOpLayoutFromMesh* op = static_cast<const ASTOpLayoutFromMesh*>(node.get());
 
-				uint64 newState = currentSemantics;
-				newState |= (UINT64_C(1) << uint32(EMeshBufferSemantic::LayoutBlock));
+				uint64 NewState = CurrentSemantics;
+				NewState |= (UINT64_C(1) << uint32(EMeshBufferSemantic::LayoutBlock));
 
-				RecurseWithState(op->Mesh.child(), newState);
+				RecurseWithState(op->Mesh.child(), NewState);
 				break;
 			}
 
             case EOpType::IN_ADDMESH:
             {
-                recurse = false;
+				bRecurse = false;
 
 				const ASTOpInstanceAdd* op = static_cast<const ASTOpInstanceAdd*>(node.get());
 
-                RecurseWithState( op->instance.child(), currentSemantics );
+                RecurseWithState( op->instance.child(), CurrentSemantics );
 
-                uint64 newState = GetDefaultState();
-                RecurseWithState( op->value.child(), newState );
+                uint64 NewState = GetDefaultState();
+                RecurseWithState( op->value.child(), NewState );
                 break;
             }
 
             default:
                 // Unhandled op, we may need everything? Recurse with current state?
-                //uint64 newState = 0xffffffffffffffff;
+                //uint64 NewState = 0xffffffffffffffff;
                 break;
 
             }
 
-            return recurse;
+            return bRecurse;
         }
 
     public:
 
-        //! Result of this visitor:
-        //! Used mesh channel semantics for each constant mesh
-        std::unordered_map< Ptr<const ASTOpConstantResource>, uint64 > m_requiredSemantics;
+        // Result of this visitor:
+        // Used mesh channel semantics for some relevant operation types mesh
+		TMap< Ptr<ASTOpConstantResource>, uint64 > RequiredSemanticsPerConstant;
+		TMap< Ptr<ASTOpMeshPrepareLayout>, uint64 > RequiredSemanticsPerPrepareLayout;
 
     };
 
@@ -702,40 +705,37 @@ namespace mu
     //---------------------------------------------------------------------------------------------
 
     // Todo: move to its own file
-    inline void MeshRemoveUnusedBufferSemantics( FMesh* pMesh, uint64 usedSemantics )
+    inline void MeshRemoveUnusedBufferSemantics( FMesh* Mesh, uint64 UsedSemantics )
     {
         // right now we only remove entire buffers if no channel is used
         // TODO: remove from inside the buffer?
-        for (int v=0; v<pMesh->GetVertexBuffers().GetBufferCount(); )
+        for (int32 BufferIndex=0; BufferIndex < Mesh->GetVertexBuffers().GetBufferCount(); )
         {
-            bool used = false;
-            for (int c=0; !used && c<pMesh->GetVertexBuffers().GetBufferChannelCount(v); ++c)
+            bool bUsed = false;
+            for (int32 ChannelIndex=0; !bUsed && ChannelIndex < Mesh->GetVertexBuffers().GetBufferChannelCount(BufferIndex); ++ChannelIndex)
             {
-                EMeshBufferSemantic semantic;
-                pMesh->GetVertexBuffers().GetChannel(v,c,&semantic,
-                                                      nullptr, nullptr, nullptr, nullptr);
-                used = (( (UINT64_C(1)<< uint32(semantic)) ) & usedSemantics) != 0;
+                EMeshBufferSemantic Semantic;
+				Mesh->GetVertexBuffers().GetChannel(BufferIndex, ChannelIndex, &Semantic, nullptr, nullptr, nullptr, nullptr);
+				bUsed = (( (UINT64_C(1)<< uint32(Semantic)) ) & UsedSemantics) != 0;
             }
 
-            if (!used)
+            if (!bUsed)
             {
-                TArray<FMeshBuffer>& buffers = pMesh->GetVertexBuffers().Buffers;
-                buffers.RemoveAt(v);
+				Mesh->VertexBuffers.Buffers.RemoveAt(BufferIndex);
             }
             else
             {
-                ++v;
+                ++BufferIndex;
             }
         }
 
-        // TODO: hack, if we don't need layouts, remove them.
+        // If we don't need layouts, remove them.
         {
-            uint64 layoutSemantics = 0;
-            layoutSemantics |= (UINT64_C(1)<< uint32(EMeshBufferSemantic::LayoutBlock));
+            constexpr uint64 LayoutSemantics = (UINT64_C(1)<< uint32(EMeshBufferSemantic::LayoutBlock));
 
-            if ( (usedSemantics&layoutSemantics) == 0)
+            if ( (UsedSemantics & LayoutSemantics) == 0)
             {
-                pMesh->Layouts.Empty();
+                Mesh->Layouts.Empty();
             }
         }
     }
@@ -748,8 +748,8 @@ namespace mu
 		const FModelOptimizationOptions& OptimizeOptions = Options->GetPrivate()->OptimisationOptions;
 
         // Images
-        AccumulateImageFormatsAST accFormat;
-        accFormat.Run( roots );
+        AccumulateImageFormatsAST ImageFormatAccumulator;
+		ImageFormatAccumulator.Run( roots );
 
         // See if we can convert some constants to more efficient formats
         ASTOp::Traverse_BottomUp_Unique_NonReentrant( roots, [&](Ptr<ASTOp>& n)
@@ -781,7 +781,7 @@ namespace mu
 					ASTOp::Replace(n, NewPlain);
 
 				}
-				else if ( accFormat.SupportedFormats[Typed][(int32)EImageFormat::L_UBitRLE] )
+				else if (ImageFormatAccumulator.SupportedFormats[Typed][(int32)EImageFormat::L_UBitRLE] )
                 {
                     TSharedPtr<FImage> pNew = ImOp.ImagePixelFormat( ImageCompressionQuality, pOld.Get(), EImageFormat::L_UBitRLE );
 
@@ -793,7 +793,7 @@ namespace mu
 						Typed->SetValue(pNew, OptimizeOptions.DiskCacheContext);
                     }
                 }
-                else if ( accFormat.SupportedFormats[Typed][(int32)EImageFormat::L_UByteRLE] )
+                else if (ImageFormatAccumulator.SupportedFormats[Typed][(int32)EImageFormat::L_UByteRLE] )
                 {
                     TSharedPtr<FImage> pNew = ImOp.ImagePixelFormat( ImageCompressionQuality, pOld.Get(), EImageFormat::L_UByteRLE );
 
@@ -810,22 +810,31 @@ namespace mu
 
 
         // Meshes
-        AccumulateMeshChannelUsageAST meshSemanticsVisitor;
-        meshSemanticsVisitor.Run( roots );
-
-        // See if we can remove some buffers from the constants
-        ASTOp::Traverse_BottomUp_Unique_NonReentrant( roots, [&](Ptr<ASTOp>& n)
-        {
-            if (n->GetOpType()==EOpType::ME_CONSTANT)
-            {
-				ASTOpConstantResource* Typed = static_cast<ASTOpConstantResource*>(n.get());
-                TSharedPtr<FMesh> Mesh = static_cast<const FMesh*>(Typed->GetValue().Get())->Clone();
-                MeshRemoveUnusedBufferSemantics(Mesh.Get(), meshSemanticsVisitor.m_requiredSemantics[Typed]);
-				Typed->SetValue(Mesh, OptimizeOptions.DiskCacheContext);
-            }
-        });
-
-    }
+        AccumulateMeshChannelUsageAST MeshSemanticsAccumulator;
+		MeshSemanticsAccumulator.Run( roots );
+
+		// See if we can remove some buffers from the constants
+		for (const TPair<Ptr<ASTOpConstantResource>, uint64>& Entry : MeshSemanticsAccumulator.RequiredSemanticsPerConstant)
+		{
+			ASTOpConstantResource* Op = Entry.Key.get();
+			TSharedPtr<FMesh> Mesh = static_cast<const FMesh*>(Op->GetValue().Get())->Clone();
+			MeshRemoveUnusedBufferSemantics(Mesh.Get(), Entry.Value);
+			Op->SetValue(Mesh, OptimizeOptions.DiskCacheContext);
+		}
+
+		// See if we can remove entire "prepare layout" operations
+		for (const TPair<Ptr<ASTOpMeshPrepareLayout>, uint64>& Entry : MeshSemanticsAccumulator.RequiredSemanticsPerPrepareLayout)
+		{
+			constexpr uint64 LayoutBlockMask = UINT64_C(1) << uint64(EMeshBufferSemantic::LayoutBlock);
+			bool bRequiresLayouts = Entry.Value & LayoutBlockMask;
+			if (!bRequiresLayouts)
+			{
+				// Directly remove the operation: the layout is never applied
+				Ptr<ASTOpMeshPrepareLayout> Op = Entry.Key;
+				ASTOp::Replace( Op, Op->Mesh.child() );
+			}
+		}
+	}
 
 
 }
