[Mutable] Fix non-deterministic SurfaceMetadata ID. The data used to generate the hash could be invalid if the owner mesh was not loaded beforehand.

#rnx

diff --git a/Engine/Plugins/Mutable/Source/CustomizableObject/Internal/MuCO/UnrealToMutableTextureConversionUtils.h b/Engine/Plugins/Mutable/Source/CustomizableObject/Internal/MuCO/UnrealToMutableTextureConversionUtils.h
--- a/Engine/Plugins/Mutable/Source/CustomizableObject/Internal/MuCO/UnrealToMutableTextureConversionUtils.h
+++ b/Engine/Plugins/Mutable/Source/CustomizableObject/Internal/MuCO/UnrealToMutableTextureConversionUtils.h
@@ -62,6 +62,16 @@ enum class EMutableMeshConversionFlags : uint32
 
 ENUM_CLASS_FLAGS(EMutableMeshConversionFlags)
 
+struct CUSTOMIZABLEOBJECT_API FMutableSourceSurfaceMetadata
+{
+	TSoftObjectPtr<UStreamableRenderAsset> Mesh;
+	uint8 LODIndex = 0;
+	uint8 SectionIndex = 0;
+
+	bool operator==(const FMutableSourceSurfaceMetadata&) const = default;
+};
+
+uint32 CUSTOMIZABLEOBJECT_API GetTypeHash(const FMutableSourceSurfaceMetadata& Key);
 
 struct CUSTOMIZABLEOBJECT_API FMutableSourceMeshData
 {
@@ -70,7 +80,13 @@ struct CUSTOMIZABLEOBJECT_API FMutableSourceMeshData
 	TSoftClassPtr<UAnimInstance> AnimInstance;
 	TSoftObjectPtr<const USkeletalMesh> TableReferenceSkeletalMesh;
 	FName Component = NAME_None;
-	uint32 SurfaceMetadataId = 0;
+
+	/** */
+	bool bIsPassthrough = false;
+
+	/** Required for SurfaceMetadataID */
+
+	FMutableSourceSurfaceMetadata Metadata;
 
 	/** Selection of the mesh section*/
 	uint8 LODIndex=0;
diff --git a/Engine/Plugins/Mutable/Source/CustomizableObject/Private/MuCO/UnrealToMutableTextureConversionUtils.cpp b/Engine/Plugins/Mutable/Source/CustomizableObject/Private/MuCO/UnrealToMutableTextureConversionUtils.cpp
--- a/Engine/Plugins/Mutable/Source/CustomizableObject/Private/MuCO/UnrealToMutableTextureConversionUtils.cpp
+++ b/Engine/Plugins/Mutable/Source/CustomizableObject/Private/MuCO/UnrealToMutableTextureConversionUtils.cpp
@@ -362,4 +362,11 @@ EUnrealToMutableConversionError ConvertTextureUnrealSourceToMutable(mu::FImage*
     return EUnrealToMutableConversionError::Success;
 }
 
+uint32 GetTypeHash(const FMutableSourceSurfaceMetadata& Key)
+{
+	uint32 MetadataKey = HashCombine(GetTypeHash(Key.Mesh.ToString()), (uint32)Key.LODIndex);
+	MetadataKey = HashCombine(MetadataKey, (uint32)Key.SectionIndex);
+	return MetadataKey;
+}
+
 #endif // WITH_EDITOR
diff --git a/Engine/Plugins/Mutable/Source/CustomizableObjectEditor/Private/MuCOE/CustomizableObjectCompiler.cpp b/Engine/Plugins/Mutable/Source/CustomizableObjectEditor/Private/MuCOE/CustomizableObjectCompiler.cpp
--- a/Engine/Plugins/Mutable/Source/CustomizableObjectEditor/Private/MuCOE/CustomizableObjectCompiler.cpp
+++ b/Engine/Plugins/Mutable/Source/CustomizableObjectEditor/Private/MuCOE/CustomizableObjectCompiler.cpp
@@ -939,7 +939,6 @@ void FCustomizableObjectCompiler::CompileInternal(bool bAsync)
 		ModelResources->StateUIDataMap = MoveTemp(GenerationContext.StateUIDataMap);
 		ModelResources->IntParameterOptionDataTable = MoveTemp(GenerationContext.IntParameterOptionDataTable);
 
-		ModelResources->SurfaceMetadata = MoveTemp(GenerationContext.SurfaceMetadata);
 
 		ModelResources->GroupNodeMap = GenerationContext.GroupNodeMap;
 
@@ -1463,6 +1462,7 @@ void FCustomizableObjectCompiler::FinishCompilationTask()
 			ClothingAssetData.ConfigsData.RemoveAllSwap(IsSharedConfigData);
 		}
 
+		ModelResources->SurfaceMetadata = MoveTemp(CompilationContext->SurfaceMetadata);
 		ModelResources->MeshMetadata = MoveTemp(CompilationContext->MeshMetadata);
 	}
 
diff --git a/Engine/Plugins/Mutable/Source/CustomizableObjectEditor/Private/MuCOE/GenerateMutableSource/GenerateMutableSource.cpp b/Engine/Plugins/Mutable/Source/CustomizableObjectEditor/Private/MuCOE/GenerateMutableSource/GenerateMutableSource.cpp
--- a/Engine/Plugins/Mutable/Source/CustomizableObjectEditor/Private/MuCOE/GenerateMutableSource/GenerateMutableSource.cpp
+++ b/Engine/Plugins/Mutable/Source/CustomizableObjectEditor/Private/MuCOE/GenerateMutableSource/GenerateMutableSource.cpp
@@ -706,8 +706,61 @@ mu::Ptr<mu::NodeMeshApplyPose> CreateNodeMeshApplyPose(FMutableGraphGenerationCo
 	return NodeMeshApplyPose;
 }
 
-uint32 AddUniqueSurfaceMetadata(const FSkeletalMaterial* Material, const FSkelMeshSection* MeshSection, TMap<uint32, FMutableSurfaceMetadata>& InOutHashSurfaceMetadataSet)
+
+uint32 AddUniqueSurfaceMetadata(const FMutableSourceSurfaceMetadata& SourceMetadata, FMutableCompilationContext& CompilationContext)
 {
+	auto GetSkeletalMaterialFunc = [](const USkeletalMesh* SkeletalMesh, uint8 LODIndex, uint8 SectionIndex) -> const FSkeletalMaterial*
+		{
+			if (!SkeletalMesh)
+			{
+				return nullptr;
+			}
+
+			// We assume that LODIndex and MaterialIndex are valid for the imported model
+			int32 MaterialIndex = INDEX_NONE;
+
+			// Check if we have lod info map to get the correct material index
+			if (const FSkeletalMeshLODInfo* LodInfo = SkeletalMesh->GetLODInfo(LODIndex))
+			{
+				if (LodInfo->LODMaterialMap.IsValidIndex(SectionIndex))
+				{
+					MaterialIndex = LodInfo->LODMaterialMap[SectionIndex];
+				}
+			}
+
+			if (MaterialIndex == INDEX_NONE)
+			{
+				const FSkeletalMeshModel* ImportedModel = SkeletalMesh->GetImportedModel();
+				if (ImportedModel && ImportedModel->LODModels.IsValidIndex(LODIndex) && ImportedModel->LODModels[LODIndex].Sections.IsValidIndex(SectionIndex))
+				{
+					MaterialIndex = ImportedModel->LODModels[LODIndex].Sections[SectionIndex].MaterialIndex;
+				}
+			}
+
+			if (SkeletalMesh->GetMaterials().IsValidIndex(MaterialIndex))
+			{
+				return &SkeletalMesh->GetMaterials()[MaterialIndex];
+			}
+			
+			return nullptr;
+		};
+
+	auto GetMeshSectionFunc = [](const USkeletalMesh* SkeletalMesh, uint8 LODIndex, uint8 SectionIndex) -> const FSkelMeshSection*
+	{
+		if (!SkeletalMesh)
+		{
+			return nullptr;
+		}
+
+		const FSkeletalMeshModel* ImportedModel = SkeletalMesh->GetImportedModel();
+		if (ImportedModel && ImportedModel->LODModels.IsValidIndex(LODIndex) && ImportedModel->LODModels[LODIndex].Sections.IsValidIndex(SectionIndex))
+		{
+			return &ImportedModel->LODModels[LODIndex].Sections[SectionIndex];
+		}
+
+		return nullptr;
+	};
+
 	auto HashSurfaceMetadataFunc = [](const FMutableSurfaceMetadata& Data) -> uint32
 	{
 		// Create a string representation of the metadata.
@@ -730,24 +783,32 @@ uint32 AddUniqueSurfaceMetadata(const FSkeletalMaterial* Material, const FSkelMe
 		return A.bCastShadow == B.bCastShadow && A.MaterialSlotName == B.MaterialSlotName;
 	};
 
+	if (uint32* Found = CompilationContext.CachedSurfaceMetadataIds.Find(SourceMetadata))
+	{
+		return *Found;
+	}
+	
+	USkeletalMesh* MetadataBaseMesh = Cast<USkeletalMesh>(MutablePrivate::LoadObject(SourceMetadata.Mesh));
+	
 	FMutableSurfaceMetadata SurfaceMetadata;
 
-	if (Material)
+	if (const FSkeletalMaterial* Material = GetSkeletalMaterialFunc(MetadataBaseMesh, SourceMetadata.LODIndex, SourceMetadata.SectionIndex))
 	{
 		SurfaceMetadata.MaterialSlotName = Material->MaterialSlotName;
 	}
 
-	if (MeshSection)
+	if (const FSkelMeshSection* MeshSection = GetMeshSectionFunc(MetadataBaseMesh, SourceMetadata.LODIndex, SourceMetadata.SectionIndex))
 	{
 		SurfaceMetadata.bCastShadow = MeshSection->bCastShadow;
 	}
 
 	uint32 SurfaceMetadataUniqueHash = Private::GenerateUniquePersistentHash(
-			SurfaceMetadata, InOutHashSurfaceMetadataSet, HashSurfaceMetadataFunc, CompareSurfaceMetadataFunc);
+			SurfaceMetadata, CompilationContext.SurfaceMetadata, HashSurfaceMetadataFunc, CompareSurfaceMetadataFunc);
 
 	if (SurfaceMetadataUniqueHash != 0)
 	{
-		InOutHashSurfaceMetadataSet.FindOrAdd(SurfaceMetadataUniqueHash, SurfaceMetadata);
+		CompilationContext.SurfaceMetadata.FindOrAdd(SurfaceMetadataUniqueHash, SurfaceMetadata);
+		CompilationContext.CachedSurfaceMetadataIds.FindOrAdd(SourceMetadata, SurfaceMetadataUniqueHash);
 	}
 	else
 	{
@@ -1643,19 +1704,18 @@ TSharedPtr<mu::FImage> GenerateImageConstant(UTexture* Texture, FMutableGraphGen
 
 TSharedPtr<mu::FMesh> GenerateMeshConstant(
 	const FMutableSourceMeshData& Source,
-	FMutableGraphGenerationContext& GenerationContext, 
-	bool bIsPassthrough)
+	FMutableGraphGenerationContext& GenerationContext)
 {
 	MUTABLE_CPUPROFILER_SCOPE(GenerateMeshConstant)
 
-	if (!Source.Mesh)
+	if (Source.Mesh.IsNull())
 	{
 		return nullptr;
 	}
 
 	bool bForceLoad = false;
 	bool bIsCompileTime = false;
-	if (!bIsPassthrough)
+	if (!Source.bIsPassthrough)
 	{
 		bForceLoad = true;
 		// TODO: For now, even in none, we turn them into compile-time references, since runtime meshes are not
@@ -1677,7 +1737,7 @@ TSharedPtr<mu::FMesh> GenerateMeshConstant(
 	FMutableGraphGenerationContext::FGeneratedReferencedMesh* Entry = nullptr;
 	int32 Num = 0;
 
-	if (bIsPassthrough)
+	if (Source.bIsPassthrough)
 	{
 		Entry = &GenerationContext.PassthroughMeshMap.FindOrAdd(Source.Mesh, InvalidEntry);
 		Num = GenerationContext.PassthroughMeshMap.Num();
diff --git a/Engine/Plugins/Mutable/Source/CustomizableObjectEditor/Private/MuCOE/GenerateMutableSource/GenerateMutableSource.h b/Engine/Plugins/Mutable/Source/CustomizableObjectEditor/Private/MuCOE/GenerateMutableSource/GenerateMutableSource.h
--- a/Engine/Plugins/Mutable/Source/CustomizableObjectEditor/Private/MuCOE/GenerateMutableSource/GenerateMutableSource.h
+++ b/Engine/Plugins/Mutable/Source/CustomizableObjectEditor/Private/MuCOE/GenerateMutableSource/GenerateMutableSource.h
@@ -58,6 +58,7 @@ struct FMutableRefSkeletalMeshData;
 struct FMutableRefSocket;
 struct FMutableSkinWeightProfileInfo;
 struct FMorphTargetVertexData;
+struct FMutableSourceSurfaceMetadata;
 enum class EPinMode;
 enum class ETableCompilationFilterOperationType : uint8;
 
@@ -478,6 +479,10 @@ struct FMutableCompilationContext
 	/** Data used for SkinWeightProfiles reconstruction. */
 	TArray<FMutableSkinWeightProfileInfo> SkinWeightProfilesInfo;
 
+	/** */
+	TMap<FMutableSourceSurfaceMetadata, uint32> CachedSurfaceMetadataIds;
+	TMap<uint32, FMutableSurfaceMetadata> SurfaceMetadata;
+
 	/** */
 	TMap<uint32, FMutableMeshMetadata> MeshMetadata;
 
@@ -776,8 +781,6 @@ struct FMutableGraphGenerationContext
 	// For a given material node (the key is node package path + node uid + image index in node) stores images generated for the same node at a higher quality LOD to reuse that image node
 	TMap<FString, FGroupProjectorImageInfo> GroupProjectorLODCache;
 	
-	TMap<uint32, FMutableSurfaceMetadata> SurfaceMetadata;
-
 	uint8 FromLOD = 0; // LOD to append to the CurrentLOD when using AutomaticLODs. 
 	uint8 CurrentLOD = 0;
 	FName CurrentMeshComponent;
@@ -972,7 +975,7 @@ int32 GetMaxTextureSize(const UTexture2D& ReferenceTexture, const UTextureLODSet
 int32 GetTextureSizeInGame(const UTexture2D& Texture, const UTextureLODSettings& LODSettings);
 
 TSharedPtr<mu::FImage> GenerateImageConstant(UTexture*, FMutableGraphGenerationContext&, bool bIsReference);
-TSharedPtr<mu::FMesh> GenerateMeshConstant(const FMutableSourceMeshData& Source, FMutableGraphGenerationContext&, bool bIsReference);
+TSharedPtr<mu::FMesh> GenerateMeshConstant(const FMutableSourceMeshData& Source, FMutableGraphGenerationContext&);
 
 /** Generates a mutable image descriptor from an unreal engine texture */
 mu::FImageDesc GenerateImageDescriptor(UTexture* Texture);
@@ -982,7 +985,7 @@ mu::FImageDesc GenerateImageDescriptor(UTexture* Texture);
  * 
  * return the unique id for the SurfaceMetadata in HashSurfaceMetadataSet.
  **/
-uint32 AddUniqueSurfaceMetadata(const FSkeletalMaterial* Material, const FSkelMeshSection* MeshSection, TMap<uint32, FMutableSurfaceMetadata>& InOutHashSurfaceMetadataSet);
+uint32 AddUniqueSurfaceMetadata(const FMutableSourceSurfaceMetadata& MeshData, FMutableCompilationContext& CompilationContext);
 
 /** Convert a mesh to Mutable format.
 */
diff --git a/Engine/Plugins/Mutable/Source/CustomizableObjectEditor/Private/MuCOE/GenerateMutableSource/GenerateMutableSourceColor.cpp b/Engine/Plugins/Mutable/Source/CustomizableObjectEditor/Private/MuCOE/GenerateMutableSource/GenerateMutableSourceColor.cpp
--- a/Engine/Plugins/Mutable/Source/CustomizableObjectEditor/Private/MuCOE/GenerateMutableSource/GenerateMutableSourceColor.cpp
+++ b/Engine/Plugins/Mutable/Source/CustomizableObjectEditor/Private/MuCOE/GenerateMutableSource/GenerateMutableSourceColor.cpp
@@ -313,7 +313,7 @@ mu::Ptr<mu::NodeColour> GenerateMutableSourceColor(const UEdGraphPin* Pin, FMuta
 					{
 						int32 Dummy = -1; // TODO MTBL-1512
 						bool Dummy2 = false;
-						bSuccess = GenerateTableColumn(TypedNodeTable, Pin, Table, ColumnName, Property, Dummy, Dummy, GenerationContext.CurrentLOD, Dummy, 0, Dummy2, GenerationContext);
+						bSuccess = GenerateTableColumn(TypedNodeTable, Pin, Table, ColumnName, Property, FMutableSourceMeshData(), Dummy, Dummy, GenerationContext.CurrentLOD, Dummy, Dummy2, GenerationContext);
 
 						if (!bSuccess)
 						{
diff --git a/Engine/Plugins/Mutable/Source/CustomizableObjectEditor/Private/MuCOE/GenerateMutableSource/GenerateMutableSourceComponent.cpp b/Engine/Plugins/Mutable/Source/CustomizableObjectEditor/Private/MuCOE/GenerateMutableSource/GenerateMutableSourceComponent.cpp
--- a/Engine/Plugins/Mutable/Source/CustomizableObjectEditor/Private/MuCOE/GenerateMutableSource/GenerateMutableSourceComponent.cpp
+++ b/Engine/Plugins/Mutable/Source/CustomizableObjectEditor/Private/MuCOE/GenerateMutableSource/GenerateMutableSourceComponent.cpp
@@ -285,11 +285,10 @@ mu::Ptr<mu::NodeComponent> GenerateMutableSourceComponent(const UEdGraphPin* Pin
 		{
 			MeshNode = new mu::NodeMeshConstant();
 
-			FString MeshUniqueTags;
-			constexpr bool bIsPassthrough = true;
-			TSoftClassPtr<UAnimInstance> AnimInstance;
-			TSharedPtr<mu::FMesh> MutableMesh = GenerateMutableSkeletalMesh(SkeletalMesh, AnimInstance, 0, 0, 0, 0, MeshUniqueTags, 0, 
-																GenerationContext, TypedComponentPassthroughMesh, nullptr, bIsPassthrough);
+			FMutableSourceMeshData Source;
+			Source.Mesh = SkeletalMesh;
+			Source.bIsPassthrough = true;
+			TSharedPtr<mu::FMesh> MutableMesh = GenerateMutableSkeletalMesh(Source, 0, 0, 0, 0, {}, GenerationContext, Node);
 
 			MeshNode->Value = MutableMesh;
 		}
diff --git a/Engine/Plugins/Mutable/Source/CustomizableObjectEditor/Private/MuCOE/GenerateMutableSource/GenerateMutableSourceFloat.cpp b/Engine/Plugins/Mutable/Source/CustomizableObjectEditor/Private/MuCOE/GenerateMutableSource/GenerateMutableSourceFloat.cpp
--- a/Engine/Plugins/Mutable/Source/CustomizableObjectEditor/Private/MuCOE/GenerateMutableSource/GenerateMutableSourceFloat.cpp
+++ b/Engine/Plugins/Mutable/Source/CustomizableObjectEditor/Private/MuCOE/GenerateMutableSource/GenerateMutableSourceFloat.cpp
@@ -357,7 +357,7 @@ mu::Ptr<mu::NodeScalar> GenerateMutableSourceFloat(const UEdGraphPin* Pin, FMuta
 					{
 						int32 Dummy = -1; // TODO MTBL-1512
 						bool Dummy2 = false;
-						bSuccess = GenerateTableColumn(TypedNodeTable, Pin, Table, ColumnName, Property, Dummy, Dummy, GenerationContext.CurrentLOD, Dummy, 0, Dummy2, GenerationContext);
+						bSuccess = GenerateTableColumn(TypedNodeTable, Pin, Table, ColumnName, Property, FMutableSourceMeshData(), Dummy, Dummy, GenerationContext.CurrentLOD, Dummy, Dummy2, GenerationContext);
 
 						if (!bSuccess)
 						{
diff --git a/Engine/Plugins/Mutable/Source/CustomizableObjectEditor/Private/MuCOE/GenerateMutableSource/GenerateMutableSourceImage.cpp b/Engine/Plugins/Mutable/Source/CustomizableObjectEditor/Private/MuCOE/GenerateMutableSource/GenerateMutableSourceImage.cpp
--- a/Engine/Plugins/Mutable/Source/CustomizableObjectEditor/Private/MuCOE/GenerateMutableSource/GenerateMutableSourceImage.cpp
+++ b/Engine/Plugins/Mutable/Source/CustomizableObjectEditor/Private/MuCOE/GenerateMutableSource/GenerateMutableSourceImage.cpp
@@ -711,7 +711,7 @@ mu::Ptr<mu::NodeImage> GenerateMutableSourceImage(const UEdGraphPin* Pin, FMutab
 			FLayoutGenerationFlags LayoutGenerationFlags;
 			LayoutGenerationFlags.TexturePinModes.Init(EPinMode::Mutable, TEXSTREAM_MAX_NUM_UVCHANNELS);
 			GenerationContext.LayoutGenerationFlags.Push(LayoutGenerationFlags);
-			mu::Ptr<mu::NodeMesh> MeshNode = GenerateMutableSourceMesh(ConnectedPin, GenerationContext, 0, false, false);
+			mu::Ptr<mu::NodeMesh> MeshNode = GenerateMutableSourceMesh(ConnectedPin, GenerationContext, FMutableSourceMeshData(), false, false);
 			GenerationContext.LayoutGenerationFlags.Pop();
 			ImageNode->Mesh = MeshNode;
 		}
@@ -1099,7 +1099,8 @@ mu::Ptr<mu::NodeImage> GenerateMutableSourceImage(const UEdGraphPin* Pin, FMutab
 						{
 							int32 Dummy = -1; // TODO MTBL-1512
 							bool Dummy2 = false;
-							bSuccess = GenerateTableColumn(TypedNodeTable, Pin, Table, ColumnName, Property, Dummy, Dummy, GenerationContext.CurrentLOD, Dummy, 0, Dummy2, GenerationContext);
+							bSuccess = GenerateTableColumn(TypedNodeTable, Pin, Table, ColumnName, Property, 
+								FMutableSourceMeshData(), Dummy, Dummy, GenerationContext.CurrentLOD, Dummy, Dummy2, GenerationContext);
 
 							if (!bSuccess)
 							{
diff --git a/Engine/Plugins/Mutable/Source/CustomizableObjectEditor/Private/MuCOE/GenerateMutableSource/GenerateMutableSourceLayout.cpp b/Engine/Plugins/Mutable/Source/CustomizableObjectEditor/Private/MuCOE/GenerateMutableSource/GenerateMutableSourceLayout.cpp
--- a/Engine/Plugins/Mutable/Source/CustomizableObjectEditor/Private/MuCOE/GenerateMutableSource/GenerateMutableSourceLayout.cpp
+++ b/Engine/Plugins/Mutable/Source/CustomizableObjectEditor/Private/MuCOE/GenerateMutableSource/GenerateMutableSourceLayout.cpp
@@ -92,13 +92,16 @@ mu::Ptr<mu::NodeLayout> CreateMutableLayoutNode(FMutableGraphGenerationContext&
 			const EMutableMeshConversionFlags ShapeFlags =
 				EMutableMeshConversionFlags::IgnoreSkinning |
 				EMutableMeshConversionFlags::IgnorePhysics |
-				EMutableMeshConversionFlags::IgnoreMorphs;
+				EMutableMeshConversionFlags::IgnoreMorphs |
+				EMutableMeshConversionFlags::DoNotCreateMeshMetadata;
 
 			GenerationContext.MeshGenerationFlags.Push(ShapeFlags);
 
 			bool bMeshMustExist = true;
 			bool bForceImmediateConversion = true;
-			MutableMesh = ConvertSkeletalMeshToMutable(SkeletalMesh, TSoftClassPtr<UAnimInstance>(), bMeshMustExist, UnrealLayout->GetLOD(), UnrealLayout->GetMaterial(), 0, GenerationContext, nullptr, nullptr, bForceImmediateConversion);
+			FMutableSourceMeshData Source;
+			Source.Mesh = SkeletalMesh;
+			MutableMesh = ConvertSkeletalMeshToMutable(Source, bMeshMustExist, UnrealLayout->GetLOD(), UnrealLayout->GetMaterial(), GenerationContext, nullptr, bForceImmediateConversion);
 
 			GenerationContext.MeshGenerationFlags.Pop();
 		}
diff --git a/Engine/Plugins/Mutable/Source/CustomizableObjectEditor/Private/MuCOE/GenerateMutableSource/GenerateMutableSourceMesh.cpp b/Engine/Plugins/Mutable/Source/CustomizableObjectEditor/Private/MuCOE/GenerateMutableSource/GenerateMutableSourceMesh.cpp
--- a/Engine/Plugins/Mutable/Source/CustomizableObjectEditor/Private/MuCOE/GenerateMutableSource/GenerateMutableSourceMesh.cpp
+++ b/Engine/Plugins/Mutable/Source/CustomizableObjectEditor/Private/MuCOE/GenerateMutableSource/GenerateMutableSourceMesh.cpp
@@ -2422,7 +2422,7 @@ TSharedPtr<mu::FMesh> ConvertSkeletalMeshToMutable(
 
 		MeshMetadata.MorphMetadataId = MorphUniqueHash;
 		MeshMetadata.ClothingMetadataId = ClothingUniqueHash;
-		MeshMetadata.SurfaceMetadataId = Source.SurfaceMetadataId;
+		MeshMetadata.SurfaceMetadataId = AddUniqueSurfaceMetadata(Source.Metadata, Context);
 
 		auto HashMeshMetadataFunc = [](const FMutableMeshMetadata& Data) -> uint32
 			{
@@ -2517,33 +2517,26 @@ TSharedPtr<mu::FMesh> ConvertSkeletalMeshToMutable(
 
 
 TSharedPtr<mu::FMesh> ConvertSkeletalMeshToMutable(
-	TSoftObjectPtr<UStreamableRenderAsset> InSkeletalMesh, const TSoftClassPtr<UAnimInstance>& AnimBp,
-	bool bMeshMustExist, int32 LODIndex, int32 SectionIndex, 
-	uint32 SurfaceMetadataId,
+	FMutableSourceMeshData& Source,
+	bool bMeshMustExist, int32 LODIndex, int32 SectionIndex,
 	FMutableGraphGenerationContext& GenerationContext, 
 	const UCustomizableObjectNode* CurrentNode, 
-	USkeletalMesh* TableReferenceSkeletalMesh,
 	bool bForceImmediateConversion)
 {
 	MUTABLE_CPUPROFILER_SCOPE(ConvertSkeletalMeshToMutable);
 
-	if(!InSkeletalMesh)
+	if(Source.Mesh.IsNull())
 	{
 		return nullptr;
 	}
 
 	// Prepare the data that is needed for the core mesh conversion.
-	FMutableSourceMeshData Source;
-	Source.Mesh = InSkeletalMesh;
-	Source.AnimInstance = AnimBp;
 	Source.bMeshMustExist = bMeshMustExist;
 	Source.LODIndex = LODIndex;
 	Source.SectionIndex = SectionIndex;
-	Source.SurfaceMetadataId = SurfaceMetadataId;
 	Source.Flags = GenerationContext.MeshGenerationFlags.Last();
 	Source.MessageContext = CurrentNode;
 	Source.Component = GenerationContext.CurrentMeshComponent;
-	Source.TableReferenceSkeletalMesh = TableReferenceSkeletalMesh;
 
 	// Prepare source data for realtime morph targets
 	const UCustomizableObjectNodeSkeletalMesh* NodeTypedSkMesh = Cast<UCustomizableObjectNodeSkeletalMesh>(CurrentNode);
@@ -2558,11 +2551,13 @@ TSharedPtr<mu::FMesh> ConvertSkeletalMeshToMutable(
 
 	// Add additional participating objects
 	{
+		const FName LongPackageFName = Source.Mesh.GetLongPackageFName();
+
 		FAssetRegistryModule& AssetRegistryModule = FModuleManager::LoadModuleChecked<FAssetRegistryModule>("AssetRegistry");
 		TArray<FAssetData> Assets;
 		{
 			TArray<FName> PackageDependencies;
-			AssetRegistryModule.Get().GetDependencies(InSkeletalMesh.GetLongPackageFName(), PackageDependencies);
+			AssetRegistryModule.Get().GetDependencies(LongPackageFName, PackageDependencies);
 
 			FARFilter Filter;
 			for (const FName& PackageName : PackageDependencies)
@@ -2591,7 +2586,7 @@ TSharedPtr<mu::FMesh> ConvertSkeletalMeshToMutable(
 
 			if (!bSkeletonFound)
 			{
-				FString Msg = FString::Printf(TEXT("No skeleton provided when converting SkeletalMesh [%s]."), *InSkeletalMesh.GetLongPackageName());
+				FString Msg = FString::Printf(TEXT("No skeleton provided when converting SkeletalMesh [%s]."), *LongPackageFName.ToString());
 				GenerationContext.Log(FText::FromString(Msg), Source.MessageContext);
 				return nullptr;
 			}
@@ -2608,7 +2603,7 @@ TSharedPtr<mu::FMesh> ConvertSkeletalMeshToMutable(
 					TArray<FAssetData> ClothAssets;
 					{
 						TArray<FName> ClothPackageDependencies;
-						AssetRegistryModule.Get().GetDependencies(InSkeletalMesh.GetLongPackageFName(), ClothPackageDependencies);
+						AssetRegistryModule.Get().GetDependencies(LongPackageFName, ClothPackageDependencies);
 
 						FARFilter ClothFilter;
 						for (const FName& PackageName : ClothPackageDependencies)
@@ -2625,11 +2620,10 @@ TSharedPtr<mu::FMesh> ConvertSkeletalMeshToMutable(
 		}
 	}
 
-	bool bIsPassthrough = false;
 	TSharedPtr<mu::FMesh> MutableMesh;
 	if (!bForceImmediateConversion)
 	{
-		MutableMesh = GenerateMeshConstant(Source, GenerationContext, bIsPassthrough);
+		MutableMesh = GenerateMeshConstant(Source, GenerationContext);
 	}
 	else
 	{
@@ -2829,14 +2823,14 @@ TSharedPtr<mu::FMesh> ConvertStaticMeshToMutable(const UStaticMesh* StaticMesh,
 
 // Convert a Mesh constant to a mutable format. UniqueTags are the tags that make this Mesh unique that cannot be merged in the cache 
 // with the exact same Mesh with other tags
-TSharedPtr<mu::FMesh> GenerateMutableSkeletalMesh( TSoftObjectPtr<UStreamableRenderAsset> Mesh, const TSoftClassPtr<UAnimInstance>& AnimInstance, int32 LODIndexConnected, int32 SectionIndexConnected,
-									  int32 LODIndex, int32 SectionIndex, const FString& UniqueTags, uint32 SurfaceMetadataId, FMutableGraphGenerationContext & GenerationContext, 
-									  const UCustomizableObjectNode* CurrentNode, USkeletalMesh* TableReferenceSkeletalMesh, bool bIsPassthrough)
+TSharedPtr<mu::FMesh> GenerateMutableSkeletalMesh(FMutableSourceMeshData& Source, int32 LODIndexConnected, int32 SectionIndexConnected,
+									  int32 LODIndex, int32 SectionIndex, const FString& UniqueTags, FMutableGraphGenerationContext & GenerationContext, 
+									  const UCustomizableObjectNode* CurrentNode)
 {
 	// Get the mesh generation flags to use
 	EMutableMeshConversionFlags CurrentFlags = GenerationContext.MeshGenerationFlags.Last();
 
-	FMutableGraphGenerationContext::FGeneratedMeshData::FKey Key = { Mesh, LODIndex, GenerationContext.CurrentLOD, SectionIndex, CurrentFlags, UniqueTags, CurrentNode };
+	FMutableGraphGenerationContext::FGeneratedMeshData::FKey Key = { Source.Mesh, LODIndex, GenerationContext.CurrentLOD, SectionIndex, CurrentFlags, UniqueTags, CurrentNode };
 	TSharedPtr<mu::FMesh> MutableMesh = GenerationContext.FindGeneratedMesh(Key);
 	if (MutableMesh)
 	{
@@ -2852,15 +2846,13 @@ TSharedPtr<mu::FMesh> GenerateMutableSkeletalMesh( TSoftObjectPtr<UStreamableRen
 		((LODIndex == LODIndexConnected) && (SectionIndex == SectionIndexConnected));
 
 	// At some point we will want all meshes to be references at compile-time. For now, just create the actual pass-through meshes.
-	if (bIsPassthrough)
+	if (Source.bIsPassthrough)
 	{
-		FMutableSourceMeshData Source;
-		Source.Mesh = Mesh;
-		MutableMesh = GenerateMeshConstant(Source, GenerationContext, bIsPassthrough);
+		MutableMesh = GenerateMeshConstant(Source, GenerationContext);
 	}
 	else
 	{
-		MutableMesh = ConvertSkeletalMeshToMutable(Mesh, AnimInstance, bMeshMustExist, LODIndex, SectionIndex, SurfaceMetadataId, GenerationContext, CurrentNode, TableReferenceSkeletalMesh);
+		MutableMesh = ConvertSkeletalMeshToMutable(Source, bMeshMustExist, LODIndex, SectionIndex, GenerationContext, CurrentNode);
 	}
 
 	if (MutableMesh)
@@ -2875,7 +2867,7 @@ TSharedPtr<mu::FMesh> GenerateMutableSkeletalMesh( TSoftObjectPtr<UStreamableRen
 // Convert a Mesh constant to a mutable format. UniqueTags are the tags that make this Mesh unique that cannot be merged in the cache 
 // with the exact same Mesh with other tags
 TSharedPtr<mu::FMesh> GenerateMutableStaticMesh(TSoftObjectPtr<UStreamableRenderAsset> Mesh, const TSoftClassPtr<UAnimInstance>& AnimInstance,
-	int32 LODIndex, int32 SectionIndex, const FString& UniqueTags, uint32 SurfaceMetadataId, FMutableGraphGenerationContext& GenerationContext,
+	int32 LODIndex, int32 SectionIndex, const FString& UniqueTags, FMutableGraphGenerationContext& GenerationContext,
 	const UCustomizableObjectNode* CurrentNode, USkeletalMesh* TableReferenceSkeletalMesh, bool bIsPassthrough)
 {
 	// Get the mesh generation flags to use
@@ -2952,9 +2944,10 @@ TSharedPtr<mu::FMesh> BuildMorphedMutableMesh(const UEdGraphPin* BaseSourcePin,
 	{
 		GetLODAndSectionForAutomaticLODs(GenerationContext, *Node, *SkeletalMesh, LODIndexConnected, SectionIndexConnected, LODIndex, SectionIndex, bOnlyConnectedLOD);
 		// Get the base mesh
-		constexpr bool bIsPassthrough = false;
-		TSharedPtr<mu::FMesh> BaseSourceMesh = GenerateMutableSkeletalMesh(SkeletalMesh, TSoftClassPtr<UAnimInstance>(), LODIndexConnected, SectionIndexConnected,
-			LODIndex, SectionIndex, FString(), 0, GenerationContext, Node, nullptr, bIsPassthrough);
+		FMutableSourceMeshData Source;
+		Source.Mesh = SkeletalMesh;
+		TSharedPtr<mu::FMesh> BaseSourceMesh = GenerateMutableSkeletalMesh(Source, LODIndexConnected, SectionIndexConnected,
+			LODIndex, SectionIndex, FString(), GenerationContext, Node);
 		check(BaseSourceMesh->IsReference());
 		// The mesh will be modified and it may come from a cache, so we need to clone it.
 		MorphedSourceMesh = BaseSourceMesh->Clone(); 
@@ -3636,7 +3629,7 @@ mu::NodeMeshPtr GenerateMorphMesh(const UEdGraphPin* Pin,
 
 mu::Ptr<mu::NodeMesh> GenerateMutableSourceMesh(const UEdGraphPin* Pin,
 	FMutableGraphGenerationContext& GenerationContext,
-	uint32 SurfaceMetadataId,
+	const FMutableSourceMeshData& BaseMeshData,
 	const bool bLinkedToExtendMaterial,
 	const bool bOnlyConnectedLOD)
 {
@@ -3751,9 +3744,12 @@ mu::Ptr<mu::NodeMesh> GenerateMutableSourceMesh(const UEdGraphPin* Pin,
 
 			FSkeletalMeshModel* ImportedModel = SkeletalMesh->GetImportedModel();
 			
-			constexpr bool bIsPassthrough = false;
-			TSharedPtr<mu::FMesh> MutableMesh = GenerateMutableSkeletalMesh(SkeletalMesh, TypedNodeSkel->AnimInstance, LODIndexConnected, SectionIndexConnected,
-																LODIndex, SectionIndex, MeshUniqueTags, SurfaceMetadataId, GenerationContext, TypedNodeSkel, nullptr, bIsPassthrough);
+			FMutableSourceMeshData Source = BaseMeshData;
+			Source.Mesh = SkeletalMesh;
+			Source.AnimInstance = TypedNodeSkel->AnimInstance;
+
+			TSharedPtr<mu::FMesh> MutableMesh = GenerateMutableSkeletalMesh(Source, LODIndexConnected, SectionIndexConnected,
+																LODIndex, SectionIndex, MeshUniqueTags, GenerationContext, TypedNodeSkel);
 			if (MutableMesh)
 			{
 				MeshNode->Value = MutableMesh;
@@ -3893,7 +3889,7 @@ mu::Ptr<mu::NodeMesh> GenerateMutableSourceMesh(const UEdGraphPin* Pin,
 
 			constexpr bool bIsPassthrough = false;
 			TSharedPtr<mu::FMesh> MutableMesh = GenerateMutableStaticMesh(TypedNodeStatic->StaticMesh, TSoftClassPtr<UAnimInstance>(), LODIndex, SectionIndex,
-														  FString(), 0, GenerationContext, TypedNodeStatic, nullptr, bIsPassthrough);
+														  FString(), GenerationContext, TypedNodeStatic, nullptr, bIsPassthrough);
 			if (MutableMesh)
 			{
 				MeshNode->Value = MutableMesh;
@@ -3967,7 +3963,7 @@ mu::Ptr<mu::NodeMesh> GenerateMutableSourceMesh(const UEdGraphPin* Pin,
 			// Mesh Morph Stack Management
 			FMorphNodeData NewMorphData = { TypedNodeMorph, MorphTargetName, FactorNode, TypedNodeMorph->MeshPin() };
 			GenerationContext.MeshMorphStack.Push(NewMorphData);
-			Result = GenerateMutableSourceMesh(ConnectedPin, GenerationContext, SurfaceMetadataId, false, bOnlyConnectedLOD);
+			Result = GenerateMutableSourceMesh(ConnectedPin, GenerationContext, BaseMeshData, false, bOnlyConnectedLOD);
 			GenerationContext.MeshMorphStack.Pop(EAllowShrinking::Yes);
 		}
 		else
@@ -4032,7 +4028,7 @@ mu::Ptr<mu::NodeMesh> GenerateMutableSourceMesh(const UEdGraphPin* Pin,
 
 				if (const UEdGraphPin* MeshConnectedPin = FollowInputPin(*TypedNodeMeshMorphStackApp->GetMeshPin()))
 				{
-					Result = GenerateMutableSourceMesh(MeshConnectedPin, GenerationContext, SurfaceMetadataId, false, bOnlyConnectedLOD);
+					Result = GenerateMutableSourceMesh(MeshConnectedPin, GenerationContext, BaseMeshData, false, bOnlyConnectedLOD);
 				}
 
 				for (int32 MorphIndex = 0; MorphIndex < AddedMorphs; ++MorphIndex)
@@ -4103,7 +4099,7 @@ mu::Ptr<mu::NodeMesh> GenerateMutableSourceMesh(const UEdGraphPin* Pin,
 				{
 					if (const UEdGraphPin* ConnectedPin = FollowInputPin(*TypedNodeMeshSwitch->GetElementPin(SelectorIndex)))
 					{
-						Result = GenerateMutableSourceMesh(ConnectedPin, GenerationContext, SurfaceMetadataId, false, bOnlyConnectedLOD);
+						Result = GenerateMutableSourceMesh(ConnectedPin, GenerationContext, BaseMeshData, false, bOnlyConnectedLOD);
 						SwitchNode->Options[SelectorIndex] = Result;
 					}
 				}
@@ -4126,7 +4122,7 @@ mu::Ptr<mu::NodeMesh> GenerateMutableSourceMesh(const UEdGraphPin* Pin,
 
 		if (const UEdGraphPin* ConnectedPin = FollowInputPin(*TypedNodeMeshVar->DefaultPin()))
 		{
-			mu::NodeMeshPtr ChildNode = GenerateMutableSourceMesh(ConnectedPin, GenerationContext, SurfaceMetadataId, false, bOnlyConnectedLOD);
+			mu::NodeMeshPtr ChildNode = GenerateMutableSourceMesh(ConnectedPin, GenerationContext, BaseMeshData, false, bOnlyConnectedLOD);
 			if (ChildNode)
 			{
 				MeshNode->DefaultMesh = ChildNode;
@@ -4151,7 +4147,7 @@ mu::Ptr<mu::NodeMesh> GenerateMutableSourceMesh(const UEdGraphPin* Pin,
 
 			if (const UEdGraphPin* ConnectedPin = FollowInputPin(*VariationPin))
 			{
-				mu::NodeMeshPtr ChildNode = GenerateMutableSourceMesh(ConnectedPin, GenerationContext, SurfaceMetadataId, false, bOnlyConnectedLOD);
+				mu::NodeMeshPtr ChildNode = GenerateMutableSourceMesh(ConnectedPin, GenerationContext, BaseMeshData, false, bOnlyConnectedLOD);
 				MeshNode->Variations[VariationIndex].Mesh = ChildNode;
 			}
 		}
@@ -4164,7 +4160,7 @@ mu::Ptr<mu::NodeMesh> GenerateMutableSourceMesh(const UEdGraphPin* Pin,
 
 		if (const UEdGraphPin* ConnectedPin = FollowInputPin(*TypedNodeReshape->BaseMeshPin()))
 		{
-			mu::Ptr<mu::NodeMesh> ChildNode = GenerateMutableSourceMesh(ConnectedPin, GenerationContext, SurfaceMetadataId, false, bOnlyConnectedLOD);
+			mu::Ptr<mu::NodeMesh> ChildNode = GenerateMutableSourceMesh(ConnectedPin, GenerationContext, BaseMeshData, false, bOnlyConnectedLOD);
 			if (ChildNode)
 			{
 				MeshNode->BaseMesh = ChildNode;
@@ -4281,13 +4277,14 @@ mu::Ptr<mu::NodeMesh> GenerateMutableSourceMesh(const UEdGraphPin* Pin,
 		// We don't need all the data for the shape meshes
 		const EMutableMeshConversionFlags ShapeFlags = 
 				EMutableMeshConversionFlags::IgnoreSkinning |
-				EMutableMeshConversionFlags::IgnorePhysics;
+				EMutableMeshConversionFlags::IgnorePhysics |
+			EMutableMeshConversionFlags::DoNotCreateMeshMetadata;
 
 		GenerationContext.MeshGenerationFlags.Push( ShapeFlags );
 			
 		if (const UEdGraphPin* ConnectedPin = FollowInputPin(*TypedNodeReshape->BaseShapePin()))
 		{
-			mu::Ptr<mu::NodeMesh> ChildNode = GenerateMutableSourceMesh(ConnectedPin, GenerationContext, 0, false, true);
+			mu::Ptr<mu::NodeMesh> ChildNode = GenerateMutableSourceMesh(ConnectedPin, GenerationContext, BaseMeshData, false, true);
 	
 			if (ChildNode)
 			{
@@ -4301,7 +4298,7 @@ mu::Ptr<mu::NodeMesh> GenerateMutableSourceMesh(const UEdGraphPin* Pin,
 
 		if (const UEdGraphPin* ConnectedPin = FollowInputPin(*TypedNodeReshape->TargetShapePin()))
 		{
-			mu::Ptr<mu::NodeMesh> ChildNode = GenerateMutableSourceMesh(ConnectedPin, GenerationContext, 0, false, true);
+			mu::Ptr<mu::NodeMesh> ChildNode = GenerateMutableSourceMesh(ConnectedPin, GenerationContext, BaseMeshData, false, true);
 			
 			if (ChildNode)
 			{
@@ -4320,7 +4317,7 @@ mu::Ptr<mu::NodeMesh> GenerateMutableSourceMesh(const UEdGraphPin* Pin,
 	{
 		if (const UEdGraphPin* InputMeshPin = FollowInputPin(*TypedNode->GetInputMeshPin()))
 		{
-			mu::Ptr<mu::NodeMesh> InputMeshNode = GenerateMutableSourceMesh(InputMeshPin, GenerationContext, SurfaceMetadataId, false, bOnlyConnectedLOD);
+			mu::Ptr<mu::NodeMesh> InputMeshNode = GenerateMutableSourceMesh(InputMeshPin, GenerationContext, BaseMeshData, false, bOnlyConnectedLOD);
 
 			if (GenerationContext.GetCurrentComponentInfo()->RefSkeletalMesh)
 			{
@@ -4348,7 +4345,7 @@ mu::Ptr<mu::NodeMesh> GenerateMutableSourceMesh(const UEdGraphPin* Pin,
 					if (const UCustomizableObjectNodeTable* TypedNodeTable = Cast<UCustomizableObjectNodeTable>(TablePosePin->GetOwningNode()))
 					{
 						mu::Ptr<mu::NodeMeshApplyPose> NodeMeshApplyPose = new mu::NodeMeshApplyPose();
-						mu::Ptr<mu::NodeMesh> MeshTableNode = GenerateMutableSourceMesh(TablePosePin, GenerationContext, SurfaceMetadataId, false, bOnlyConnectedLOD);
+						mu::Ptr<mu::NodeMesh> MeshTableNode = GenerateMutableSourceMesh(TablePosePin, GenerationContext, BaseMeshData, false, bOnlyConnectedLOD);
 
 						NodeMeshApplyPose->Base = InputMeshNode;
 						NodeMeshApplyPose->Pose = MeshTableNode;
@@ -4447,7 +4444,7 @@ mu::Ptr<mu::NodeMesh> GenerateMutableSourceMesh(const UEdGraphPin* Pin,
 					// Generating a new FMesh column if not exists
 					if (Table->FindColumn(MutableColumnName) == INDEX_NONE)
 					{
-						bSuccess = GenerateTableColumn(TypedNodeTable, Pin, Table, ColumnName, Property, LODIndexConnected, SectionIndexConnected, LODIndex, SectionIndex, SurfaceMetadataId, bOnlyConnectedLOD, GenerationContext);
+						bSuccess = GenerateTableColumn(TypedNodeTable, Pin, Table, ColumnName, Property, BaseMeshData, LODIndexConnected, SectionIndexConnected, LODIndex, SectionIndex, bOnlyConnectedLOD, GenerationContext);
 
 						if (!bSuccess)
 						{
@@ -4610,7 +4607,7 @@ mu::Ptr<mu::NodeMesh> GenerateMutableSourceMesh(const UEdGraphPin* Pin,
 			if (const UEdGraphPin* FollowPin = FollowInputPin(*OutputPin))
 			{
 				GenerationContext.MacroNodesStack.Push(TypedNodeMacro);
-				Result = GenerateMutableSourceMesh(FollowPin, GenerationContext, SurfaceMetadataId, bLinkedToExtendMaterial, bOnlyConnectedLOD);
+				Result = GenerateMutableSourceMesh(FollowPin, GenerationContext, BaseMeshData, bLinkedToExtendMaterial, bOnlyConnectedLOD);
 				GenerationContext.MacroNodesStack.Pop();
 			}
 			else
@@ -4641,7 +4638,7 @@ mu::Ptr<mu::NodeMesh> GenerateMutableSourceMesh(const UEdGraphPin* Pin,
 		{
 			if (const UEdGraphPin* FollowPin = FollowInputPin(*InputPin))
 			{
-				Result = GenerateMutableSourceMesh(FollowPin, GenerationContext, SurfaceMetadataId, bLinkedToExtendMaterial, bOnlyConnectedLOD);
+				Result = GenerateMutableSourceMesh(FollowPin, GenerationContext, BaseMeshData, bLinkedToExtendMaterial, bOnlyConnectedLOD);
 			}
 			else
 			{
diff --git a/Engine/Plugins/Mutable/Source/CustomizableObjectEditor/Private/MuCOE/GenerateMutableSource/GenerateMutableSourceMesh.h b/Engine/Plugins/Mutable/Source/CustomizableObjectEditor/Private/MuCOE/GenerateMutableSource/GenerateMutableSourceMesh.h
--- a/Engine/Plugins/Mutable/Source/CustomizableObjectEditor/Private/MuCOE/GenerateMutableSource/GenerateMutableSourceMesh.h
+++ b/Engine/Plugins/Mutable/Source/CustomizableObjectEditor/Private/MuCOE/GenerateMutableSource/GenerateMutableSourceMesh.h
@@ -18,6 +18,7 @@ class UStreamableRenderAsset;
 struct FMutableGraphGenerationContext;
 struct FMutableGraphMeshGenerationData;
 struct FMutableCompilationContext;
+struct FMutableSourceMeshData;
 
 namespace mu
 {
@@ -46,22 +47,19 @@ void GetLODAndSectionForAutomaticLODs(const FMutableGraphGenerationContext& Cont
 
 /** Converts an Unreal Skeletal Mesh to Mutable Mesh.
  *
- * @param InSkeletalMesh Unreal Skeletal Mesh.
- * @param AnimBp 
+ * @param Source. Mesh data required to generate the mutable mesh
  * @param LODIndexConnected LOD which the pin is connected to.
  * @param SectionIndexConnected Section which the pin is connected to.
  * @param LODIndex LOD we are generating. Will be different from LODIndexConnected only when using Automatic LOD From Mesh. 
  * @param SectionIndex Section we are generating. Will be different from SectionIndexConnected only when using Automatic LOD From Mesh.
  * @param GenerationContext 
  * @param CurrentNode The node being compiled, can be a skeletal mesh or a table node.
- * @param TableReferenceSkeletalMesh Only set if compiling a table node. It refers to the reference mesh of the skeletal mesh column being compiled.
  * @return Mutable Mesh. Nullptr if there has been an error. Empty mesh if the Skeletal Mesh does not contain the requested LOD + Section. */
-TSharedPtr<mu::FMesh> ConvertSkeletalMeshToMutable(TSoftObjectPtr<UStreamableRenderAsset> InSkeletalMesh, const TSoftClassPtr<UAnimInstance>& AnimBp,
+TSharedPtr<mu::FMesh> ConvertSkeletalMeshToMutable(FMutableSourceMeshData& Source,
                                          bool bMeshMustExist,
-                                         int32 LODIndex, int32 SectionIndex, uint32 SurfaceMetadataId,
+                                         int32 LODIndex, int32 SectionIndex,
                                          FMutableGraphGenerationContext& GenerationContext,
                                          const UCustomizableObjectNode* CurrentNode, 
-										 USkeletalMesh* TableReferenceSkeletalMesh,
 										 bool bForceImmediateConversion = false );
 
 
@@ -70,27 +68,24 @@ TSharedPtr<mu::FMesh> ConvertStaticMeshToMutable(const UStaticMesh* StaticMesh,
 
 
 /** 
- * @param Mesh Unreal Mesh.
- * @param AnimBp 
+ * @param Source. Mesh data required to generate the mutable mesh
  * @param LODIndexConnected LOD which the pin is connected to.
  * @param SectionIndexConnected Section which the pin is connected to.
  * @param LODIndex LOD we are generating. Will be different from LODIndexConnected only when using Automatic LOD From Mesh. 
  * @param SectionIndex Section we are generating. Will be different from SectionIndexConnected only when using Automatic LOD From Mesh.
- * @param MeshUniqueTags 
- * @param GenerationContext 
+ * @param MeshUniqueTags
+ * @param GenerationContext
  * @param CurrentNode The node being compiled, can be a skeletal mesh or a table node.
- * @param TableReferenceSkeletalMesh Only set if compiling a table node. It refers to the reference mesh of the skeletal mesh column being compiled.
  * @return Mutable Mesh. Nullptr if there has been an error. Empty mesh if the Skeletal Mesh does not contain the requested LOD + Section. */
-TSharedPtr<mu::FMesh> GenerateMutableSkeletalMesh(TSoftObjectPtr<UStreamableRenderAsset> Mesh, const TSoftClassPtr<UAnimInstance>& AnimBp,
+TSharedPtr<mu::FMesh> GenerateMutableSkeletalMesh(FMutableSourceMeshData& Source,
                                 int32 LODIndexConnected, int32 SectionIndexConnected,
-                                int32 LODIndex, int32 SectionIndex, const FString& MeshUniqueTags, uint32 SurfaceMetadataId, 
+                                int32 LODIndex, int32 SectionIndex, const FString& MeshUniqueTags, 
                                 FMutableGraphGenerationContext& GenerationContext, 
-								const UCustomizableObjectNode* CurrentNode, USkeletalMesh* TableReferenceSkeletalMesh,
-								bool bIsPassthrough);
+								const UCustomizableObjectNode* CurrentNode);
 
 /** */
 TSharedPtr<mu::FMesh> GenerateMutableStaticMesh(TSoftObjectPtr<UStreamableRenderAsset> Mesh, const TSoftClassPtr<UAnimInstance>& AnimBp,
-	int32 LODIndex, int32 SectionIndex, const FString& MeshUniqueTags, uint32 SurfaceMetadataId,
+	int32 LODIndex, int32 SectionIndex, const FString& MeshUniqueTags,
 	FMutableGraphGenerationContext& GenerationContext,
 	const UCustomizableObjectNode* CurrentNode, USkeletalMesh* TableReferenceSkeletalMesh,
 	bool bIsPassthrough);
@@ -104,12 +99,11 @@ TSharedPtr<mu::FMesh> BuildMorphedMutableMesh(const UEdGraphPin* BaseSourcePin,
  *
  * @param Pin 
  * @param GenerationContext 
- * @param MeshData 
- * @param SurfaceMetadataId
+ * @param BaseMeshData 
  * @param bLinkedToExtendMaterial 
  * @param bOnlyConnectedLOD Corrected LOD and Section will unconditionally always be the connected ones.
  * @return  Mutable Mesh Node. */
-mu::Ptr<mu::NodeMesh> GenerateMutableSourceMesh(const UEdGraphPin* Pin, FMutableGraphGenerationContext& GenerationContext, uint32 SurfaceMetadataId, bool bLinkedToExtendMaterial, bool bOnlyConnectedLOD);
+mu::Ptr<mu::NodeMesh> GenerateMutableSourceMesh(const UEdGraphPin* Pin, FMutableGraphGenerationContext& GenerationContext, const FMutableSourceMeshData& BaseMeshData, bool bLinkedToExtendMaterial, bool bOnlyConnectedLOD);
 
 
 /** */
diff --git a/Engine/Plugins/Mutable/Source/CustomizableObjectEditor/Private/MuCOE/GenerateMutableSource/GenerateMutableSourceModifier.cpp b/Engine/Plugins/Mutable/Source/CustomizableObjectEditor/Private/MuCOE/GenerateMutableSource/GenerateMutableSourceModifier.cpp
--- a/Engine/Plugins/Mutable/Source/CustomizableObjectEditor/Private/MuCOE/GenerateMutableSource/GenerateMutableSourceModifier.cpp
+++ b/Engine/Plugins/Mutable/Source/CustomizableObjectEditor/Private/MuCOE/GenerateMutableSource/GenerateMutableSourceModifier.cpp
@@ -106,7 +106,8 @@ mu::Ptr<mu::NodeModifier> GenerateMutableSourceModifier(const UEdGraphPin * Pin,
 		const EMutableMeshConversionFlags ModifiersMeshFlags =
 			EMutableMeshConversionFlags::IgnoreSkinning |
 			EMutableMeshConversionFlags::IgnorePhysics |
-			EMutableMeshConversionFlags::IgnoreMorphs;
+			EMutableMeshConversionFlags::IgnoreMorphs |
+			EMutableMeshConversionFlags::DoNotCreateMeshMetadata;
 		GenerationContext.MeshGenerationFlags.Push(ModifiersMeshFlags);
 
 		// This modifier can be connected to multiple nodes at the same time and, when that happens and if the cache is being used, only the first node to be processed does work. 
@@ -138,7 +139,8 @@ mu::Ptr<mu::NodeModifier> GenerateMutableSourceModifier(const UEdGraphPin * Pin,
 		const EMutableMeshConversionFlags ModifiersMeshFlags =
 			EMutableMeshConversionFlags::IgnoreSkinning |
 			EMutableMeshConversionFlags::IgnorePhysics |
-			EMutableMeshConversionFlags::IgnoreMorphs;
+			EMutableMeshConversionFlags::IgnoreMorphs |
+			EMutableMeshConversionFlags::DoNotCreateMeshMetadata;
 		GenerationContext.MeshGenerationFlags.Push(ModifiersMeshFlags);
 
 		mu::Ptr<mu::NodeModifierMeshClipDeform> ClipNode = new mu::NodeModifierMeshClipDeform();
@@ -148,7 +150,7 @@ mu::Ptr<mu::NodeModifier> GenerateMutableSourceModifier(const UEdGraphPin * Pin,
 
 		if (const UEdGraphPin* ConnectedPin = FollowInputPin(*TypedNodeClipDeform->ClipShapePin()))
 		{
-			mu::Ptr<mu::NodeMesh> ClipMesh = GenerateMutableSourceMesh(ConnectedPin, GenerationContext, 0, false, true);
+			mu::Ptr<mu::NodeMesh> ClipMesh = GenerateMutableSourceMesh(ConnectedPin, GenerationContext, FMutableSourceMeshData(), false, true);
 
 			ClipNode->ClipMesh = ClipMesh;
 
@@ -189,7 +191,8 @@ mu::Ptr<mu::NodeModifier> GenerateMutableSourceModifier(const UEdGraphPin * Pin,
 		const EMutableMeshConversionFlags ModifiersMeshFlags =
 			EMutableMeshConversionFlags::IgnoreSkinning |
 			EMutableMeshConversionFlags::IgnorePhysics |
-			EMutableMeshConversionFlags::IgnoreMorphs;
+			EMutableMeshConversionFlags::IgnoreMorphs |
+			EMutableMeshConversionFlags::DoNotCreateMeshMetadata;
 		GenerationContext.MeshGenerationFlags.Push(ModifiersMeshFlags);
 
 		// MeshClipWithMesh can be connected to multiple objects, so the compiled NodeModifierMeshClipWithMesh
@@ -203,7 +206,7 @@ mu::Ptr<mu::NodeModifier> GenerateMutableSourceModifier(const UEdGraphPin * Pin,
 
 		if (const UEdGraphPin* ConnectedPin = FollowInputPin(*TypedNodeClipMesh->ClipMeshPin()))
 		{
-			mu::NodeMeshPtr ClipMesh = GenerateMutableSourceMesh(ConnectedPin, GenerationContext, 0, false, true);
+			mu::NodeMeshPtr ClipMesh = GenerateMutableSourceMesh(ConnectedPin, GenerationContext, FMutableSourceMeshData(), false, true);
 
 			if (FMatrix Matrix = TypedNodeClipMesh->Transform.ToMatrixWithScale(); Matrix != FMatrix::Identity)
 			{
@@ -234,7 +237,8 @@ mu::Ptr<mu::NodeModifier> GenerateMutableSourceModifier(const UEdGraphPin * Pin,
 		const EMutableMeshConversionFlags ModifiersMeshFlags =
 			EMutableMeshConversionFlags::IgnoreSkinning |
 			EMutableMeshConversionFlags::IgnorePhysics |
-			EMutableMeshConversionFlags::IgnoreMorphs;
+			EMutableMeshConversionFlags::IgnoreMorphs |
+			EMutableMeshConversionFlags::DoNotCreateMeshMetadata;
 		GenerationContext.MeshGenerationFlags.Push(ModifiersMeshFlags);
 
 		// This modifier can be connected to multiple objects, so the compiled node
@@ -299,33 +303,34 @@ mu::Ptr<mu::NodeModifier> GenerateMutableSourceModifier(const UEdGraphPin * Pin,
 
 				GenerationContext.LayoutGenerationFlags.Push(LayoutGenerationFlags);
 
-				// Generate surface metadata for this fragment.
-				uint32 SurfaceMetadataUniqueHash = 0;	
-				if (ConnectedPin)
+				FMutableSourceMeshData MeshData;
+
+				// Find reference mesh used to generate the surface metadata for this fragment.
+				if(ConnectedPin)
 				{
 					//NOTE: This is the same is done in GenerateMutableSourceSurface. 
 					if (const UEdGraphPin* SkeletalMeshPin = FindMeshBaseSource(*ConnectedPin, false, &GenerationContext.MacroNodesStack))
 					{
-						FSkeletalMaterial* SkeletalMaterial = nullptr;
-						const FSkelMeshSection* ReferenceSkelMeshSection = nullptr;
-						
+						int32 MetadataLODIndex, MetadataSectionIndex, MetadataLayoutIndex;
+						MetadataLODIndex = MetadataSectionIndex = MetadataLayoutIndex = INDEX_NONE;
+
 						if (const UCustomizableObjectNodeSkeletalMesh* SkeletalMeshNode = Cast<UCustomizableObjectNodeSkeletalMesh>(SkeletalMeshPin->GetOwningNode()))
 						{
-							SkeletalMaterial = SkeletalMeshNode->GetSkeletalMaterialFor(*SkeletalMeshPin);
-							ReferenceSkelMeshSection = SkeletalMeshNode->GetSkeletalMeshSectionFor(*SkeletalMeshPin);
+							MeshData.Metadata.Mesh = SkeletalMeshNode->GetMesh().ToSoftObjectPath();
+							SkeletalMeshNode->GetPinSection(*SkeletalMeshPin, MetadataLODIndex, MetadataSectionIndex, MetadataLayoutIndex);
 						}
-
 						else if (const UCustomizableObjectNodeTable* TableNode = Cast<UCustomizableObjectNodeTable>(SkeletalMeshPin->GetOwningNode()))
 						{
-							SkeletalMaterial = TableNode->GetDefaultSkeletalMaterialFor(*SkeletalMeshPin);
-							ReferenceSkelMeshSection = TableNode->GetDefaultSkeletalMeshSectionFor(*SkeletalMeshPin);
+							MeshData.Metadata.Mesh = TableNode->GetColumnDefaultAssetByType<USkeletalMesh>(SkeletalMeshPin);
+							TableNode->GetPinLODAndSection(SkeletalMeshPin, MetadataLODIndex, MetadataSectionIndex);
 						}
 
-						SurfaceMetadataUniqueHash = AddUniqueSurfaceMetadata(SkeletalMaterial, ReferenceSkelMeshSection, GenerationContext.SurfaceMetadata);
+						MeshData.Metadata.LODIndex = MetadataLODIndex;
+						MeshData.Metadata.SectionIndex = MetadataSectionIndex;
 					}
-				}	
-
-				AddMeshNode = GenerateMutableSourceMesh(ConnectedPin, GenerationContext, SurfaceMetadataUniqueHash, true, false);
+				}
+				
+				AddMeshNode = GenerateMutableSourceMesh(ConnectedPin, GenerationContext, MeshData, true, false);
 
 				GenerationContext.LayoutGenerationFlags.Pop();
 			}
@@ -429,7 +434,8 @@ mu::Ptr<mu::NodeModifier> GenerateMutableSourceModifier(const UEdGraphPin * Pin,
 		const EMutableMeshConversionFlags ModifiersMeshFlags =
 			EMutableMeshConversionFlags::IgnoreSkinning |
 			EMutableMeshConversionFlags::IgnorePhysics |
-			EMutableMeshConversionFlags::IgnoreMorphs;
+			EMutableMeshConversionFlags::IgnoreMorphs |
+			EMutableMeshConversionFlags::DoNotCreateMeshMetadata;
 		GenerationContext.MeshGenerationFlags.Push(ModifiersMeshFlags);
 
 		mu::Ptr<mu::NodeModifierSurfaceEdit> SurfNode = new mu::NodeModifierSurfaceEdit();
@@ -451,7 +457,7 @@ mu::Ptr<mu::NodeModifier> GenerateMutableSourceModifier(const UEdGraphPin * Pin,
 				GenerationContext.FromLOD = 0;
 				GenerationContext.CurrentLOD = LODIndex;
 
-				mu::Ptr<mu::NodeMesh> RemoveMeshNode = GenerateMutableSourceMesh(ConnectedPin, GenerationContext, 0, false, true);
+				mu::Ptr<mu::NodeMesh> RemoveMeshNode = GenerateMutableSourceMesh(ConnectedPin, GenerationContext, FMutableSourceMeshData(), false, true);
 				SurfNode->LODs[LODIndex].MeshRemove = RemoveMeshNode;
 			}
 		}
@@ -466,7 +472,8 @@ mu::Ptr<mu::NodeModifier> GenerateMutableSourceModifier(const UEdGraphPin * Pin,
 		const EMutableMeshConversionFlags ModifiersMeshFlags =
 			EMutableMeshConversionFlags::IgnoreSkinning |
 			EMutableMeshConversionFlags::IgnorePhysics |
-			EMutableMeshConversionFlags::IgnoreMorphs;
+			EMutableMeshConversionFlags::IgnoreMorphs |
+			EMutableMeshConversionFlags::DoNotCreateMeshMetadata;
 		GenerationContext.MeshGenerationFlags.Push(ModifiersMeshFlags);
 
 		mu::Ptr<mu::NodeModifierMeshClipWithUVMask> ClipNode = new mu::NodeModifierMeshClipWithUVMask();
@@ -490,7 +497,8 @@ mu::Ptr<mu::NodeModifier> GenerateMutableSourceModifier(const UEdGraphPin * Pin,
 		const EMutableMeshConversionFlags ModifiersMeshFlags =
 			EMutableMeshConversionFlags::IgnoreSkinning |
 			EMutableMeshConversionFlags::IgnorePhysics |
-			EMutableMeshConversionFlags::IgnoreMorphs;
+			EMutableMeshConversionFlags::IgnoreMorphs |
+			EMutableMeshConversionFlags::DoNotCreateMeshMetadata;
 		GenerationContext.MeshGenerationFlags.Push(ModifiersMeshFlags);
 
 		mu::Ptr<mu::NodeModifierSurfaceEdit> SurfNode = new mu::NodeModifierSurfaceEdit();
@@ -568,7 +576,8 @@ mu::Ptr<mu::NodeModifier> GenerateMutableSourceModifier(const UEdGraphPin * Pin,
 		const EMutableMeshConversionFlags ModifiersMeshFlags =
 			EMutableMeshConversionFlags::IgnoreSkinning |
 			EMutableMeshConversionFlags::IgnorePhysics |
-			EMutableMeshConversionFlags::IgnoreMorphs;
+			EMutableMeshConversionFlags::IgnoreMorphs |
+			EMutableMeshConversionFlags::DoNotCreateMeshMetadata;
 		GenerationContext.MeshGenerationFlags.Push(ModifiersMeshFlags);
 
 		mu::Ptr<mu::NodeModifierSurfaceEdit> SurfNode = new mu::NodeModifierSurfaceEdit();
@@ -652,7 +661,8 @@ mu::Ptr<mu::NodeModifier> GenerateMutableSourceModifier(const UEdGraphPin * Pin,
 		const EMutableMeshConversionFlags ModifiersMeshFlags =
 			EMutableMeshConversionFlags::IgnoreSkinning |
 			EMutableMeshConversionFlags::IgnorePhysics |
-			EMutableMeshConversionFlags::IgnoreMorphs;
+			EMutableMeshConversionFlags::IgnoreMorphs |
+			EMutableMeshConversionFlags::DoNotCreateMeshMetadata;
 		GenerationContext.MeshGenerationFlags.Push(ModifiersMeshFlags);
 
 		// MeshTransformInMesh can be connected to multiple objects, so the compiled NodeModifierMeshTransformInMesh
@@ -670,7 +680,7 @@ mu::Ptr<mu::NodeModifier> GenerateMutableSourceModifier(const UEdGraphPin * Pin,
 		// If no bounding mesh is provided, we transform the entire mesh.
 		if (const UEdGraphPin* ConnectedPin = FollowInputPin(*TypedNodeTransformMesh->BoundingMeshPin()))
 		{
-			mu::Ptr<mu::NodeMesh> BoundingMesh = GenerateMutableSourceMesh(ConnectedPin, GenerationContext, 0, false, true);
+			mu::Ptr<mu::NodeMesh> BoundingMesh = GenerateMutableSourceMesh(ConnectedPin, GenerationContext, FMutableSourceMeshData(), false, true);
 
 			if (FMatrix Matrix = TypedNodeTransformMesh->BoundingMeshTransform.ToMatrixWithScale(); Matrix != FMatrix::Identity)
 			{
diff --git a/Engine/Plugins/Mutable/Source/CustomizableObjectEditor/Private/MuCOE/GenerateMutableSource/GenerateMutableSourceSurface.cpp b/Engine/Plugins/Mutable/Source/CustomizableObjectEditor/Private/MuCOE/GenerateMutableSource/GenerateMutableSourceSurface.cpp
--- a/Engine/Plugins/Mutable/Source/CustomizableObjectEditor/Private/MuCOE/GenerateMutableSource/GenerateMutableSourceSurface.cpp
+++ b/Engine/Plugins/Mutable/Source/CustomizableObjectEditor/Private/MuCOE/GenerateMutableSource/GenerateMutableSourceSurface.cpp
@@ -111,7 +111,8 @@ mu::Ptr<mu::NodeSurface> GenerateMutableSourceSurface(const UEdGraphPin * Pin, F
 			}
 		}
 
-		if (!TypedNodeMat->GetMaterial())
+		UMaterialInterface* Material = TypedNodeMat->GetMaterial();
+		if (!Material)
 		{
 			const FText Message = LOCTEXT("FailedToGenerateMeshSection", "Could not generate a mesh section because it didn't have a material selected. Please assign one and recompile.");
 			GenerationContext.Log(Message, Node);
@@ -128,32 +129,31 @@ mu::Ptr<mu::NodeSurface> GenerateMutableSourceSurface(const UEdGraphPin * Pin, F
 		FMutableGraphGenerationContext::FSharedSurface& SharedSurface = SharedSurfaces.Add_GetRef(FMutableGraphGenerationContext::FSharedSurface(GenerationContext.CurrentLOD, SurfNode));
 		SharedSurface.bMakeUnique = !TypedNodeMat->IsReuseMaterialBetweenLODs();
 
-		uint32 SurfaceMetadataUniqueHash = 0; // Value 0 is used as invalid hash.
-		if (TypedNodeMat->GetMaterial())
-		{
-			GenerationContext.CurrentReferencedMaterialIndex = GenerationContext.ReferencedMaterials.AddUnique(TypedNodeMat->GetMaterial());
+		GenerationContext.CurrentReferencedMaterialIndex = GenerationContext.ReferencedMaterials.AddUnique(Material);
 
-			if (ConnectedMaterialPin)
+		// Find reference mesh used to generate the surface metadata for this fragment.
+		FMutableSourceMeshData MeshData;
+		if (ConnectedMaterialPin)
+		{
+			//NOTE: This is the same is done in GenerateMutableSourceSurface. 
+			if (const UEdGraphPin* SkeletalMeshPin = FindMeshBaseSource(*ConnectedMaterialPin, false, &GenerationContext.MacroNodesStack))
 			{
-				if (const UEdGraphPin* SkeletalMeshPin = FindMeshBaseSource(*ConnectedMaterialPin, false, &GenerationContext.MacroNodesStack))
-				{
-					FSkeletalMaterial* SkeletalMaterial = nullptr;
-					const FSkelMeshSection* ReferenceSkelMeshSection = nullptr;
-					
-					if (const UCustomizableObjectNodeSkeletalMesh* SkeletalMeshNode = Cast<UCustomizableObjectNodeSkeletalMesh>(SkeletalMeshPin->GetOwningNode()))
-					{
-						SkeletalMaterial = SkeletalMeshNode->GetSkeletalMaterialFor(*SkeletalMeshPin);
-						ReferenceSkelMeshSection = SkeletalMeshNode->GetSkeletalMeshSectionFor(*SkeletalMeshPin);
-					}
+				int32 MetadataLODIndex, MetadataSectionIndex, MetadataLayoutIndex;
+				MetadataLODIndex = MetadataSectionIndex = MetadataLayoutIndex = INDEX_NONE;
 
-					else if (const UCustomizableObjectNodeTable* TableNode = Cast<UCustomizableObjectNodeTable>(SkeletalMeshPin->GetOwningNode()))
-					{
-						SkeletalMaterial = TableNode->GetDefaultSkeletalMaterialFor(*SkeletalMeshPin);
-						ReferenceSkelMeshSection = TableNode->GetDefaultSkeletalMeshSectionFor(*SkeletalMeshPin);
-					}
-
-					SurfaceMetadataUniqueHash = AddUniqueSurfaceMetadata(SkeletalMaterial, ReferenceSkelMeshSection, GenerationContext.SurfaceMetadata);
+				if (const UCustomizableObjectNodeSkeletalMesh* SkeletalMeshNode = Cast<UCustomizableObjectNodeSkeletalMesh>(SkeletalMeshPin->GetOwningNode()))
+				{
+					MeshData.Metadata.Mesh = SkeletalMeshNode->GetMesh().ToSoftObjectPath();
+					SkeletalMeshNode->GetPinSection(*SkeletalMeshPin, MetadataLODIndex, MetadataSectionIndex, MetadataLayoutIndex);
 				}
+				else if (const UCustomizableObjectNodeTable* TableNode = Cast<UCustomizableObjectNodeTable>(SkeletalMeshPin->GetOwningNode()))
+				{
+					MeshData.Metadata.Mesh = TableNode->GetColumnDefaultAssetByType<USkeletalMesh>(SkeletalMeshPin);
+					TableNode->GetPinLODAndSection(SkeletalMeshPin, MetadataLODIndex, MetadataSectionIndex);
+				}
+
+				MeshData.Metadata.LODIndex = MetadataLODIndex;
+				MeshData.Metadata.SectionIndex = MetadataSectionIndex;
 			}
 		}
 
@@ -194,7 +194,7 @@ mu::Ptr<mu::NodeSurface> GenerateMutableSourceSurface(const UEdGraphPin * Pin, F
 
 				GenerationContext.LayoutGenerationFlags.Push(LayoutGenerationFlags);
 
-				MeshNode = GenerateMutableSourceMesh(ConnectedPin, GenerationContext, SurfaceMetadataUniqueHash, false, false);
+				MeshNode = GenerateMutableSourceMesh(ConnectedPin, GenerationContext, MeshData, false, false);
 
 				GenerationContext.LayoutGenerationFlags.Pop();
 
@@ -221,7 +221,7 @@ mu::Ptr<mu::NodeSurface> GenerateMutableSourceSurface(const UEdGraphPin * Pin, F
 					if (UMaterialInstance* TableMaterial = TypedNodeTable->GetColumnDefaultAssetByType<UMaterialInstance>(MaterialAssetConnectedPin))
 					{
 						// Checking if the reference material of the Table Node has the same parent as the material of the Material Node 
-						if (!TypedNodeMat->GetMaterial() || TableMaterial->GetMaterial() != TypedNodeMat->GetMaterial()->GetMaterial())
+						if (!Material || TableMaterial->GetMaterial() != Material->GetMaterial())
 						{
 							bTableMaterialPinLinked = false;
 
@@ -275,7 +275,7 @@ mu::Ptr<mu::NodeSurface> GenerateMutableSourceSurface(const UEdGraphPin * Pin, F
 					const FString SurfNodeImageName = FString::Printf(TEXT("%d"), Props.ImagePropertiesIndex);
 					SurfNode->Images[ImageIndex].Name = SurfNodeImageName;
 					SurfNode->Images[ImageIndex].LayoutIndex = -1;
-					SurfNode->Images[ImageIndex].MaterialName = TypedNodeMat->GetMaterial()->GetName();
+					SurfNode->Images[ImageIndex].MaterialName = Material->GetName();
 					SurfNode->Images[ImageIndex].MaterialParameterName = Props.TextureParameterName;
 
 				}
@@ -681,7 +681,7 @@ mu::Ptr<mu::NodeSurface> GenerateMutableSourceSurface(const UEdGraphPin * Pin, F
 							UVLayout = TypedNodeMat->GetImageUVLayout(ImageIndex);;
 						}
 						SurfNode->Images[ImageIndex].LayoutIndex = UVLayout;
-						SurfNode->Images[ImageIndex].MaterialName = TypedNodeMat->GetMaterial()->GetName();
+						SurfNode->Images[ImageIndex].MaterialName = Material->GetName();
 						SurfNode->Images[ImageIndex].MaterialParameterName = ImageName;
 
 						if (bShareProjectionTexturesBetweenLODs && bIsGroupProjectorImage)
diff --git a/Engine/Plugins/Mutable/Source/CustomizableObjectEditor/Private/MuCOE/GenerateMutableSource/GenerateMutableSourceTable.cpp b/Engine/Plugins/Mutable/Source/CustomizableObjectEditor/Private/MuCOE/GenerateMutableSource/GenerateMutableSourceTable.cpp
--- a/Engine/Plugins/Mutable/Source/CustomizableObjectEditor/Private/MuCOE/GenerateMutableSource/GenerateMutableSourceTable.cpp
+++ b/Engine/Plugins/Mutable/Source/CustomizableObjectEditor/Private/MuCOE/GenerateMutableSource/GenerateMutableSourceTable.cpp
@@ -34,7 +34,7 @@
 
 
 bool FillTableColumn(const UCustomizableObjectNodeTable* TableNode, mu::Ptr<mu::FTable> MutableTable, const FString& ColumnName, const FString& RowName, const uint32 RowId, uint8* CellData, const FProperty* ColumnProperty,
-	const int LODIndexConnected, const int32 SectionIndexConnected, int32 LODIndex, int32 SectionIndex, uint32 SectionMetadataId, const bool bOnlyConnectedLOD, FMutableGraphGenerationContext& GenerationContext)
+	const FMutableSourceMeshData& BaseMeshData, const int LODIndexConnected, const int32 SectionIndexConnected, int32 LODIndex, int32 SectionIndex, const bool bOnlyConnectedLOD, FMutableGraphGenerationContext& GenerationContext)
 {
 	int32 CurrentColumn;
 	UDataTable* DataTablePtr = GetDataTable(TableNode, GenerationContext);
@@ -276,10 +276,11 @@ bool FillTableColumn(const UCustomizableObjectNodeTable* TableNode, mu::Ptr<mu::
 			}
 
 			//TODO: Add AnimBp physics to Tables.
-			constexpr bool bIsPassthrough = false;
-			TSoftObjectPtr<UStreamableRenderAsset> SkeletalMeshSoftObjectPtr(SoftObject.ToSoftObjectPath());
-			TSharedPtr<mu::FMesh> MutableMesh = GenerateMutableSkeletalMesh(SkeletalMeshSoftObjectPtr, TSoftClassPtr<UAnimInstance>(), LODIndexConnected, SectionIndexConnected,
-																LODIndex, SectionIndex, MeshUniqueTags, SectionMetadataId, GenerationContext, TableNode, ReferenceSkeletalMesh, bIsPassthrough);
+			FMutableSourceMeshData Source = BaseMeshData;
+			Source.Mesh = SoftObject.ToSoftObjectPath();
+			Source.TableReferenceSkeletalMesh = ReferenceSkeletalMesh;
+			TSharedPtr<mu::FMesh> MutableMesh = GenerateMutableSkeletalMesh(Source, LODIndexConnected, SectionIndexConnected,
+																LODIndex, SectionIndex, MeshUniqueTags, GenerationContext, TableNode);
 
 			if (MutableMesh)
 			{
@@ -376,7 +377,7 @@ bool FillTableColumn(const UCustomizableObjectNodeTable* TableNode, mu::Ptr<mu::
 
 			constexpr bool bIsPassthrough = false;
 			TSharedPtr<mu::FMesh> MutableMesh = GenerateMutableStaticMesh(StaticMesh, TSoftClassPtr<UAnimInstance>(), CurrentLOD, SectionIndex,
-														  FString(), 0, GenerationContext, TableNode, nullptr, bIsPassthrough);
+														  FString(), GenerationContext, TableNode, nullptr, bIsPassthrough);
 
 			if (MutableMesh)
 			{
@@ -827,8 +828,8 @@ TArray<FName> GetRowsToCompile(const UDataTable& DataTable, const UCustomizableO
 }
 
 
-bool GenerateTableColumn(const UCustomizableObjectNodeTable* TableNode, const UEdGraphPin* Pin, mu::Ptr<mu::FTable> MutableTable, const FString& DataTableColumnName, const FProperty* ColumnProperty,
-	const int32 LODIndexConnected, const int32 SectionIndexConnected, const int32 LODIndex, const int32 SectionIndex, uint32 SectionMetadataId, const bool bOnlyConnectedLOD, FMutableGraphGenerationContext& GenerationContext)
+bool GenerateTableColumn(const UCustomizableObjectNodeTable* TableNode, const UEdGraphPin* Pin, mu::Ptr<mu::FTable> MutableTable, const FString& DataTableColumnName, const FProperty* ColumnProperty, 
+	const FMutableSourceMeshData& BaseMeshData,	const int32 LODIndexConnected, const int32 SectionIndexConnected, const int32 LODIndex, const int32 SectionIndex, const bool bOnlyConnectedLOD, FMutableGraphGenerationContext& GenerationContext)
 {
 	MUTABLE_CPUPROFILER_SCOPE(GenerateTableColumn);
 
@@ -914,7 +915,7 @@ bool GenerateTableColumn(const UCustomizableObjectNodeTable* TableNode, const UE
 		if (uint8* CellData = UCustomizableObjectNodeTable::GetCellData(RowNames[RowIndex], *DataTable, *ColumnProperty))
 		{
 			bCellGenerated = FillTableColumn(TableNode, MutableTable, DataTableColumnName, RowNames[RowIndex].ToString(), RowIds[RowIndex], CellData, ColumnProperty,
-				LODIndexConnected, SectionIndexConnected, LODIndex, SectionIndex, SectionMetadataId, bOnlyConnectedLOD, GenerationContext);
+				BaseMeshData, LODIndexConnected, SectionIndexConnected, LODIndex, SectionIndex, bOnlyConnectedLOD, GenerationContext);
 
 			// Stop the compilation if something fails
 			if (!bCellGenerated)
diff --git a/Engine/Plugins/Mutable/Source/CustomizableObjectEditor/Private/MuCOE/GenerateMutableSource/GenerateMutableSourceTable.h b/Engine/Plugins/Mutable/Source/CustomizableObjectEditor/Private/MuCOE/GenerateMutableSource/GenerateMutableSourceTable.h
--- a/Engine/Plugins/Mutable/Source/CustomizableObjectEditor/Private/MuCOE/GenerateMutableSource/GenerateMutableSourceTable.h
+++ b/Engine/Plugins/Mutable/Source/CustomizableObjectEditor/Private/MuCOE/GenerateMutableSource/GenerateMutableSourceTable.h
@@ -10,6 +10,7 @@ class UCustomizableObjectNodeTable;
 class UDataTable;
 class UEdGraphPin;
 struct FMutableGraphGenerationContext;
+struct FMutableSourceMeshData;
 
 
 /** 
@@ -27,8 +28,9 @@ struct FMutableGraphGenerationContext;
  * @param GenerationContext 
  * @return  */
 bool GenerateTableColumn(const UCustomizableObjectNodeTable* TableNode, const UEdGraphPin* Pin, mu::Ptr<mu::FTable> MutableTable, const FString& DataTableColumnName, const FProperty* ColumnProperty,
+	const FMutableSourceMeshData& BaseMeshData,
 	int32 LODIndexConnected, int32 SectionIndexConnected,
-	int32 LODIndex, int32 SectionIndex, uint32 SurfaceMetadataId,
+	int32 LODIndex, int32 SectionIndex,
 	bool bOnlyConnectedLOD,
 	FMutableGraphGenerationContext& GenerationContext);
 
@@ -52,7 +54,8 @@ bool GenerateTableColumn(const UCustomizableObjectNodeTable* TableNode, const UE
  * @return  */
 bool FillTableColumn(const UCustomizableObjectNodeTable* TableNode, mu::Ptr<mu::FTable> MutableTable, const FString& ColumnName,
 	const FString& RowName, uint32 RowId, uint8* CellData, const FProperty* ColumnProperty,
-	int32 LODIndexConnected, int32 SectionIndexConnected, int32 LODIndex, int32 SectionIndex, uint32 SectionMetadataId,
+	const FMutableSourceMeshData& BaseMeshData,
+	int32 LODIndexConnected, int32 SectionIndexConnected, int32 LODIndex, int32 SectionIndex,
 	bool bOnlyConnectedLOD,
 	FMutableGraphGenerationContext& GenerationContext);
 
