[Mutable] Add debug callstack tracking for coderunner ops.

- Disabled by default, enabled setting MUTABLE_DEBUG_CODERUNNER_TASK_SCHEDULE_CALLSATACK to 1.

#rnx

diff --git a/Engine/Plugins/Mutable/Source/MutableRuntime/Internal/MuR/SystemPrivate.h b/Engine/Plugins/Mutable/Source/MutableRuntime/Internal/MuR/SystemPrivate.h
--- a/Engine/Plugins/Mutable/Source/MutableRuntime/Internal/MuR/SystemPrivate.h
+++ b/Engine/Plugins/Mutable/Source/MutableRuntime/Internal/MuR/SystemPrivate.h
@@ -21,6 +21,10 @@
 #include "HAL/PlatformTLS.h"
 #endif
 
+#ifndef MUTABLE_DEBUG_CODERUNNER_TASK_SCHEDULE_CALLSTACK
+	#define MUTABLE_DEBUG_CODERUNNER_TASK_SCHEDULE_CALLSTACK 0	
+#endif
+
 namespace mu::MemoryCounters
 {
 	struct MUTABLERUNTIME_API FInternalMemoryCounter
@@ -157,6 +161,12 @@ namespace mu
 			ImageDesc
 		};
 		EType Type : 1;
+
+#if MUTABLE_DEBUG_CODERUNNER_TASK_SCHEDULE_CALLSTACK
+		static constexpr uint64 CallstackMaxDepth = 16;
+		uint64 StackDepth = 0; 
+		uint64 ScheduleCallstack[CallstackMaxDepth];
+#endif
 	};
 
 	inline uint32 GetTypeHash(const FScheduledOp& Op)
diff --git a/Engine/Plugins/Mutable/Source/MutableRuntime/Private/MuR/CodeRunner.cpp b/Engine/Plugins/Mutable/Source/MutableRuntime/Private/MuR/CodeRunner.cpp
--- a/Engine/Plugins/Mutable/Source/MutableRuntime/Private/MuR/CodeRunner.cpp
+++ b/Engine/Plugins/Mutable/Source/MutableRuntime/Private/MuR/CodeRunner.cpp
@@ -59,6 +59,10 @@
 #include "MuR/SystemPrivate.h"
 #include "Templates/Tuple.h"
 
+#if MUTABLE_DEBUG_CODERUNNER_TASK_SCHEDULE_CALLSTACK
+#include "GenericPlatform/GenericPlatformStackWalk.h"
+#endif
+
 namespace
 {
 
@@ -98,6 +102,33 @@ static FAutoConsoleVariableRef CVarUseImageTransformVectorImpl (
 	ECVF_Default);
 }
 
+#if MUTABLE_DEBUG_CODERUNNER_TASK_SCHEDULE_CALLSTACK
+
+#define AddOp(...) Invoke([&](){ AddOp(__VA_ARGS__); });
+
+namespace mu::Private
+{
+	FString DumpItemScheduledCallstack(const FScheduledOp& Item)
+	{
+		constexpr SIZE_T MaxStringSize = 16 * 1024;
+		ANSICHAR StackTrace[MaxStringSize];
+
+		FString OutputString;
+
+		constexpr uint32 EntriesToSkip = 3;
+		for (uint32 Index = EntriesToSkip; Index < Item.StackDepth; ++Index)
+		{
+			StackTrace[0] = 0;
+			FPlatformStackWalk::ProgramCounterToHumanReadableString(Index, Item.ScheduleCallstack[Index], StackTrace, MaxStringSize, nullptr);
+			OutputString += FString::Printf(TEXT("\t\t%d %s\n"), Index - EntriesToSkip, ANSI_TO_TCHAR(StackTrace));
+		}
+
+		return OutputString;
+	}
+}
+
+#endif
+
 namespace mu
 {
 
@@ -7338,5 +7369,10 @@ namespace mu
 			Runner->System->ImagePixelFormatOverride
 		);
 	}
-
 }
+
+#if MUTABLE_DEBUG_CODERUNNER_TASK_SCHEDULE_CALLSTACK
+
+#undef AddOp
+
+#endif
diff --git a/Engine/Plugins/Mutable/Source/MutableRuntime/Private/MuR/CodeRunner.h b/Engine/Plugins/Mutable/Source/MutableRuntime/Private/MuR/CodeRunner.h
--- a/Engine/Plugins/Mutable/Source/MutableRuntime/Private/MuR/CodeRunner.h
+++ b/Engine/Plugins/Mutable/Source/MutableRuntime/Private/MuR/CodeRunner.h
@@ -707,3 +707,11 @@ namespace  mu
 	extern FImageOperator MakeImageOperator(CodeRunner* Runner);
 
 }
+
+#if MUTABLE_DEBUG_CODERUNNER_TASK_SCHEDULE_CALLSTACK
+namespace mu::Private
+{
+	FString DumpItemScheduledCallstack(const FScheduledOp& Item);
+}
+#endif
+
diff --git a/Engine/Plugins/Mutable/Source/MutableRuntime/Private/MuR/CodeRunnerTasked.cpp b/Engine/Plugins/Mutable/Source/MutableRuntime/Private/MuR/CodeRunnerTasked.cpp
--- a/Engine/Plugins/Mutable/Source/MutableRuntime/Private/MuR/CodeRunnerTasked.cpp
+++ b/Engine/Plugins/Mutable/Source/MutableRuntime/Private/MuR/CodeRunnerTasked.cpp
@@ -42,6 +42,9 @@
 #include "ProfilingDebugging/CountersTrace.h"
 #include "Async/Fundamental/Scheduler.h"
 
+#if MUTABLE_DEBUG_CODERUNNER_TASK_SCHEDULE_CALLSTACK
+#include "GenericPlatform/GenericPlatformStackWalk.h"
+#endif
 
 static TAutoConsoleVariable<bool> CVarCodeRunnerForceInline(
 	TEXT("mutable.CodeRunnerForceInline"),
@@ -67,7 +70,12 @@ namespace mu
 		{
 			if (ScheduledStagePerOp.get(at) <= dep.Stage)
 			{
+#if MUTABLE_DEBUG_CODERUNNER_TASK_SCHEDULE_CALLSTACK
+				FScheduledOp& AddedOp = OpenTasks.Add_GetRef(dep);
+				AddedOp.StackDepth = FPlatformStackWalk::CaptureStackBackTrace(&AddedOp.ScheduleCallstack[0], FScheduledOp::CallstackMaxDepth);
+#else
 				OpenTasks.Add(dep);
+#endif
 				ScheduledStagePerOp[at] = dep.Stage + 1;
 			}
 		}
