[Mutable] Updated Mutable Macros.
- Unified CO nodes to have the same behaviour in Macros than in CO graphs.
- Renamed String node to Static String.
- Implemented new pin Static String to all nodes that need a name.
- Unified all common parameter nodes variable into the same parent class UCustomizableObjectNodeParameter.
- Implemented functions to disable some node widgets when its the static string node has something linked.
- Added a CVar to enable and disable Macros.

#rnx

diff --git a/Engine/Plugins/Mutable/Source/CustomizableObjectEditor/Private/MuCOE/CustomizableObjectCompiler.cpp b/Engine/Plugins/Mutable/Source/CustomizableObjectEditor/Private/MuCOE/CustomizableObjectCompiler.cpp
--- a/Engine/Plugins/Mutable/Source/CustomizableObjectEditor/Private/MuCOE/CustomizableObjectCompiler.cpp
+++ b/Engine/Plugins/Mutable/Source/CustomizableObjectEditor/Private/MuCOE/CustomizableObjectCompiler.cpp
@@ -516,7 +516,7 @@ void ProcessChildObjectsRecursively(const UCustomizableObject* ParentObject, FMu
 
 		GenerationContext.GroupIdToExternalNodeMap.Add(Root->ParentObjectGroupId, Root);
 
-		TArray<UCustomizableObjectNodeObjectGroup*> GroupNodes;
+		TArray<UCustomizableObjectNodeObjectGroup*> GroupNodes; 
 		ChildObject->GetPrivate()->GetSource()->GetNodesOfClass<UCustomizableObjectNodeObjectGroup>(GroupNodes);
 
 		if (GroupNodes.Num() > 0) // Only grafs with group nodes should have child grafs
@@ -578,7 +578,7 @@ mu::Ptr<mu::NodeObject> GenerateMutableRoot(
 		return nullptr;
 	}
 	
-	if (LocalRootNodeObject->ObjectName.IsEmpty())
+	if (LocalRootNodeObject->GetObjectName(&GenerationContext.MacroNodesStack).IsEmpty())
 	{
 		GenerationContext.NoNameNodeObjectArray.AddUnique(LocalRootNodeObject);
 	}
diff --git a/Engine/Plugins/Mutable/Source/CustomizableObjectEditor/Private/MuCOE/CustomizableObjectEditor.cpp b/Engine/Plugins/Mutable/Source/CustomizableObjectEditor/Private/MuCOE/CustomizableObjectEditor.cpp
--- a/Engine/Plugins/Mutable/Source/CustomizableObjectEditor/Private/MuCOE/CustomizableObjectEditor.cpp
+++ b/Engine/Plugins/Mutable/Source/CustomizableObjectEditor/Private/MuCOE/CustomizableObjectEditor.cpp
@@ -512,7 +512,7 @@ bool FCustomizableObjectEditor::GroupNodeIsLinkedToParentByName(UCustomizableObj
 
 	for (UCustomizableObjectNodeObjectGroup* GroupNode : GroupNodes)
 	{
-		if ((Node->ParentObjectGroupId == GroupNode->NodeGuid) && (ParentGroupName == GroupNode->GroupName))
+		if ((Node->ParentObjectGroupId == GroupNode->NodeGuid) && (ParentGroupName == GroupNode->GetGroupName()))
 		{
 			return true;
 		}
@@ -562,7 +562,7 @@ void FCustomizableObjectEditor::ReconstructAllChildNodes(UCustomizableObjectNode
 			{
 				if (UCustomizableObjectNode* TypedNode = Cast<UCustomizableObjectNode>(ConnectedPin->GetOwningNode()))
 				{
-					NodesToVisit.Add(TypedNode);						
+					NodesToVisit.Add(TypedNode);
 				}
 			}
 		}
diff --git a/Engine/Plugins/Mutable/Source/CustomizableObjectEditor/Private/MuCOE/CustomizableObjectEditorModule.cpp b/Engine/Plugins/Mutable/Source/CustomizableObjectEditor/Private/MuCOE/CustomizableObjectEditorModule.cpp
--- a/Engine/Plugins/Mutable/Source/CustomizableObjectEditor/Private/MuCOE/CustomizableObjectEditorModule.cpp
+++ b/Engine/Plugins/Mutable/Source/CustomizableObjectEditor/Private/MuCOE/CustomizableObjectEditorModule.cpp
@@ -64,6 +64,7 @@
 #include "MuCOE/Nodes/CustomizableObjectNodeObjectDetails.h"
 #include "MuCOE/Nodes/CustomizableObjectNodeObjectGroup.h"
 #include "MuCOE/Nodes/CustomizableObjectNodeDetails.h"
+#include "MuCOE/Nodes/CustomizableObjectNodeParameterDetails.h"
 #include "MuCOE/Nodes/CustomizableObjectNodeProjectorConstant.h"
 #include "MuCOE/Nodes/CustomizableObjectNodeProjectorParameterDetails.h"
 #include "MuCOE/Nodes/CustomizableObjectNodeSkeletalMesh.h"
@@ -240,6 +241,7 @@ void FCustomizableObjectEditorModule::StartupModule()
 	RegisterCustomDetails(PropertyModule, UCustomizableObjectNodeTable::StaticClass(), FOnGetDetailCustomizationInstance::CreateStatic(&FCustomizableObjectNodeTableDetails::MakeInstance));
 	RegisterCustomDetails(PropertyModule, UCustomizableObjectNodeComponentMesh::StaticClass(), FOnGetDetailCustomizationInstance::CreateStatic(&FCustomizableObjectNodeComponentMeshDetails::MakeInstance));
 	RegisterCustomDetails(PropertyModule, UCustomizableObjectNodeMacroInstance::StaticClass(), FOnGetDetailCustomizationInstance::CreateStatic(&FCustomizableObjectNodeMacroInstanceDetails::MakeInstance));
+	RegisterCustomDetails(PropertyModule, UCustomizableObjectNodeParameter::StaticClass(), FOnGetDetailCustomizationInstance::CreateStatic(&FCustomizableObjectNodeParameterDetails::MakeInstance));
 
 	// Other Objects
 	RegisterCustomDetails(PropertyModule, UCustomizableObject::StaticClass(), FOnGetDetailCustomizationInstance::CreateStatic(&FCustomizableObjectDetails::MakeInstance));
@@ -816,7 +818,7 @@ USkeletalMesh* FCustomizableObjectEditorModule::GetReferenceSkeletalMesh(const U
 		{
 			if (UCustomizableObjectNodeComponentMesh* NodeComponentMesh = Cast<UCustomizableObjectNodeComponentMesh>(Node))
 			{
-				if (NodeComponentMesh->ComponentName == ComponentName)
+				if (NodeComponentMesh->GetComponentName() == ComponentName)
 				{
 					return NodeComponentMesh->ReferenceSkeletalMesh;
 				}
diff --git a/Engine/Plugins/Mutable/Source/CustomizableObjectEditor/Private/MuCOE/CustomizableObjectMacroLibrary/CustomizableObjectMacroDetails.cpp b/Engine/Plugins/Mutable/Source/CustomizableObjectEditor/Private/MuCOE/CustomizableObjectMacroLibrary/CustomizableObjectMacroDetails.cpp
--- a/Engine/Plugins/Mutable/Source/CustomizableObjectEditor/Private/MuCOE/CustomizableObjectMacroLibrary/CustomizableObjectMacroDetails.cpp
+++ b/Engine/Plugins/Mutable/Source/CustomizableObjectEditor/Private/MuCOE/CustomizableObjectMacroLibrary/CustomizableObjectMacroDetails.cpp
@@ -236,7 +236,7 @@ void SCOMacroPinTypeSelector::Construct(const FArguments& InArgs)
 	{
 		TSharedPtr<FPinNameRowData> NewSourceData = MakeShareable(new FPinNameRowData());
 		NewSourceData->PinCategory = PinType;
-		NewSourceData->PinFriendlyName = Schema->GetPinCategoryName(PinType);
+		NewSourceData->PinFriendlyName = Schema->GetPinCategoryFriendlyName(PinType).ToString();
 
 		ComboBoxSource.Add(NewSourceData);
 
@@ -297,7 +297,7 @@ TSharedRef<SWidget> SCOMacroPinTypeSelector::OnGenerateRow(TSharedPtr<FPinNameRo
 	.Padding(7.5f, 0.0f, 0.0f, 0.0f)
 	[
 		SNew(STextBlock)
-		.Text(FText::FromName(Option.IsValid() ? Option->PinFriendlyName : FName("Invalid")))
+		.Text(FText::FromString(Option.IsValid() ? Option->PinFriendlyName : "Invalid"))
 		.Font(IDetailLayoutBuilder::GetDetailFont())
 	];
 
@@ -328,7 +328,7 @@ TSharedRef<SWidget> SCOMacroPinTypeSelector::GenerateSelectedContent() const
 		.Font(IDetailLayoutBuilder::GetDetailFont())
 		.Text_Lambda([this]() 
 		{
-			return FText::FromName(Variable ? UEdGraphSchema_CustomizableObject::GetPinCategoryName(Variable->PinCategoryType) : FName("Invalid"));
+			return Variable ? UEdGraphSchema_CustomizableObject::GetPinCategoryFriendlyName(Variable->PinCategoryType) : FText::FromString("Invalid");
 		})
 	];
 
diff --git a/Engine/Plugins/Mutable/Source/CustomizableObjectEditor/Private/MuCOE/CustomizableObjectMacroLibrary/CustomizableObjectMacroDetails.h b/Engine/Plugins/Mutable/Source/CustomizableObjectEditor/Private/MuCOE/CustomizableObjectMacroLibrary/CustomizableObjectMacroDetails.h
--- a/Engine/Plugins/Mutable/Source/CustomizableObjectEditor/Private/MuCOE/CustomizableObjectMacroLibrary/CustomizableObjectMacroDetails.h
+++ b/Engine/Plugins/Mutable/Source/CustomizableObjectEditor/Private/MuCOE/CustomizableObjectMacroLibrary/CustomizableObjectMacroDetails.h
@@ -19,7 +19,7 @@ enum class ECOMacroIOType : uint8;
 
 struct FPinNameRowData
 {
-	FName PinFriendlyName;
+	FString PinFriendlyName;
 	FName PinCategory;
 };
 
diff --git a/Engine/Plugins/Mutable/Source/CustomizableObjectEditor/Private/MuCOE/CustomizableObjectMacroLibrary/CustomizableObjectMacroLibrary.cpp b/Engine/Plugins/Mutable/Source/CustomizableObjectEditor/Private/MuCOE/CustomizableObjectMacroLibrary/CustomizableObjectMacroLibrary.cpp
--- a/Engine/Plugins/Mutable/Source/CustomizableObjectEditor/Private/MuCOE/CustomizableObjectMacroLibrary/CustomizableObjectMacroLibrary.cpp
+++ b/Engine/Plugins/Mutable/Source/CustomizableObjectEditor/Private/MuCOE/CustomizableObjectMacroLibrary/CustomizableObjectMacroLibrary.cpp
@@ -8,6 +8,13 @@
 
 #define LOCTEXT_NAMESPACE "CustomizableObjectMacroLibrary"
 
+TAutoConsoleVariable<bool> CVarEnableMutableMacros(
+	TEXT("mutable.EnableMutableMacros"),
+	false,
+	TEXT("Enables or disables the new nodes and assets to be able to create and use Mutable Macros."),
+	ECVF_Scalability);
+
+
 UCustomizableObjectMacro* UCustomizableObjectMacroLibrary::AddMacro()
 {
 	UCustomizableObjectMacro* NewMacro = nullptr;
diff --git a/Engine/Plugins/Mutable/Source/CustomizableObjectEditor/Private/MuCOE/CustomizableObjectMacroLibrary/CustomizableObjectMacroLibrary.h b/Engine/Plugins/Mutable/Source/CustomizableObjectEditor/Private/MuCOE/CustomizableObjectMacroLibrary/CustomizableObjectMacroLibrary.h
--- a/Engine/Plugins/Mutable/Source/CustomizableObjectEditor/Private/MuCOE/CustomizableObjectMacroLibrary/CustomizableObjectMacroLibrary.h
+++ b/Engine/Plugins/Mutable/Source/CustomizableObjectEditor/Private/MuCOE/CustomizableObjectMacroLibrary/CustomizableObjectMacroLibrary.h
@@ -2,11 +2,11 @@
 
 #pragma once
 
+#include "HAL/IConsoleManager.h"
+
 #include "CustomizableObjectMacroLibrary.generated.h"
 
-#ifndef ENABLE_MUTABLE_MACROS
-	#define ENABLE_MUTABLE_MACROS 0
-#endif // !ENABLE_MUTABLE_MACROS
+extern TAutoConsoleVariable<bool> CVarEnableMutableMacros;
 
 class UEdGraph;
 class UCustomizableObjectNodeTunnel;
diff --git a/Engine/Plugins/Mutable/Source/CustomizableObjectEditor/Private/MuCOE/CustomizableObjectMacroLibrary/CustomizableObjectMacroLibraryEditor.cpp b/Engine/Plugins/Mutable/Source/CustomizableObjectEditor/Private/MuCOE/CustomizableObjectMacroLibrary/CustomizableObjectMacroLibraryEditor.cpp
--- a/Engine/Plugins/Mutable/Source/CustomizableObjectEditor/Private/MuCOE/CustomizableObjectMacroLibrary/CustomizableObjectMacroLibraryEditor.cpp
+++ b/Engine/Plugins/Mutable/Source/CustomizableObjectEditor/Private/MuCOE/CustomizableObjectMacroLibrary/CustomizableObjectMacroLibraryEditor.cpp
@@ -11,6 +11,7 @@
 #include "MuCOE/Nodes/CustomizableObjectNodeTunnel.h"
 #include "PropertyEditorModule.h"
 #include "Widgets/Docking/SDockTab.h"
+#include "Widgets/Layout/SScrollBox.h"
 #include "Widgets/SNullWidget.h"
 
 
@@ -112,6 +113,7 @@ TSharedRef<SDockTab> FCustomizableObjectMacroLibraryEditor::SpawnTab_Details(con
 		DetailsViewArgs.bAllowSearch = true;
 		DetailsViewArgs.bShowObjectLabel = false;
 		DetailsViewArgs.bShowScrollBar = false;
+		DetailsViewArgs.ExternalScrollbar = SNew(SScrollBar);
 
 		DetailsView = PropPlugin.CreateDetailView(DetailsViewArgs);
 		DetailsView->SetObject(SelectedMacro, true);
@@ -121,7 +123,11 @@ TSharedRef<SDockTab> FCustomizableObjectMacroLibraryEditor::SpawnTab_Details(con
 	.Label(LOCTEXT("DetailsDockTab", "Details"))
 	.TabColorScale(GetTabColorScale())
 	[
-		DetailsView.ToSharedRef()
+		SNew(SScrollBox)
+		+ SScrollBox::Slot()
+		[
+			DetailsView.ToSharedRef()
+		]
 	];
 }
 
@@ -182,7 +188,11 @@ void FCustomizableObjectMacroLibraryEditor::OnSelectedGraphNodesChanged(const FG
 {
 	if (DetailsView.IsValid())
 	{
-		if (NewSelection.Num() == 1)
+		if (NewSelection.Num() == 0)
+		{
+			DetailsView->SetObject(SelectedMacro, true);
+		}
+		else if (NewSelection.Num() == 1)
 		{
 			if (Cast<UCustomizableObjectNodeTunnel>(NewSelection.Array()[0]) && SelectedMacro)
 			{
@@ -193,7 +203,7 @@ void FCustomizableObjectMacroLibraryEditor::OnSelectedGraphNodesChanged(const FG
 				DetailsView->SetObject(NewSelection.Array()[0], true);
 			}
 		}
-		else
+		else if (NewSelection.Num() > 1)
 		{
 			DetailsView->SetObjects(NewSelection.Array(), true);
 		}
diff --git a/Engine/Plugins/Mutable/Source/CustomizableObjectEditor/Private/MuCOE/CustomizableObjectMacroLibrary/CustomizableObjectMacroLibraryFactory.cpp b/Engine/Plugins/Mutable/Source/CustomizableObjectEditor/Private/MuCOE/CustomizableObjectMacroLibrary/CustomizableObjectMacroLibraryFactory.cpp
--- a/Engine/Plugins/Mutable/Source/CustomizableObjectEditor/Private/MuCOE/CustomizableObjectMacroLibrary/CustomizableObjectMacroLibraryFactory.cpp
+++ b/Engine/Plugins/Mutable/Source/CustomizableObjectEditor/Private/MuCOE/CustomizableObjectMacroLibrary/CustomizableObjectMacroLibraryFactory.cpp
@@ -7,12 +7,6 @@
 
 UCustomizableObjectMacroLibraryFactory::UCustomizableObjectMacroLibraryFactory() 
 {
-#if ENABLE_MUTABLE_MACROS
-	bCreateNew = true;
-#else
-	bCreateNew = false;
-#endif
-
 	SupportedClass = UCustomizableObjectMacroLibrary::StaticClass();
 	bEditAfterNew = true;
 }
@@ -28,3 +22,9 @@ UObject* UCustomizableObjectMacroLibraryFactory::FactoryCreateNew(UClass* Class,
 
 	return NewMacroLibrary;
 }
+
+
+bool UCustomizableObjectMacroLibraryFactory::CanCreateNew() const
+{
+	return CVarEnableMutableMacros.GetValueOnAnyThread();
+}
diff --git a/Engine/Plugins/Mutable/Source/CustomizableObjectEditor/Private/MuCOE/CustomizableObjectMacroLibrary/CustomizableObjectMacroLibraryFactory.h b/Engine/Plugins/Mutable/Source/CustomizableObjectEditor/Private/MuCOE/CustomizableObjectMacroLibrary/CustomizableObjectMacroLibraryFactory.h
--- a/Engine/Plugins/Mutable/Source/CustomizableObjectEditor/Private/MuCOE/CustomizableObjectMacroLibrary/CustomizableObjectMacroLibraryFactory.h
+++ b/Engine/Plugins/Mutable/Source/CustomizableObjectEditor/Private/MuCOE/CustomizableObjectMacroLibrary/CustomizableObjectMacroLibraryFactory.h
@@ -16,7 +16,8 @@ public:
 	UCustomizableObjectMacroLibraryFactory();
 
 	// UFactory Interface
-	UObject* FactoryCreateNew(UClass* Class, UObject* InParent, FName Name, EObjectFlags Flags, UObject* Context, FFeedbackContext* Warn);
+	UObject* FactoryCreateNew(UClass* Class, UObject* InParent, FName Name, EObjectFlags Flags, UObject* Context, FFeedbackContext* Warn) override;
+	virtual bool CanCreateNew() const override;
 
 };
 
diff --git a/Engine/Plugins/Mutable/Source/CustomizableObjectEditor/Private/MuCOE/CustomizableObjectNodeObjectGroupDetails.cpp b/Engine/Plugins/Mutable/Source/CustomizableObjectEditor/Private/MuCOE/CustomizableObjectNodeObjectGroupDetails.cpp
--- a/Engine/Plugins/Mutable/Source/CustomizableObjectEditor/Private/MuCOE/CustomizableObjectNodeObjectGroupDetails.cpp
+++ b/Engine/Plugins/Mutable/Source/CustomizableObjectEditor/Private/MuCOE/CustomizableObjectNodeObjectGroupDetails.cpp
@@ -7,6 +7,7 @@
 #include "Misc/Attribute.h"
 #include "MuCOE/CustomizableObjectEditor.h"
 #include "MuCOE/CustomizableObjectGraph.h"
+#include "MuCOE/CustomizableObjectMacroLibrary/CustomizableObjectMacroLibrary.h"
 #include "MuCOE/GraphTraversal.h"
 #include "MuCOE/Nodes/CustomizableObjectNodeObject.h"
 #include "MuCOE/Nodes/CustomizableObjectNodeObjectGroup.h"
@@ -44,8 +45,19 @@ void FCustomizableObjectNodeObjectGroupDetails::CustomizeDetails(IDetailLayoutBu
 		IDetailCategoryBuilder& GroupInfoCategory = DetailBuilder.EditCategory("GroupInfo");
 
 		// Forcing property order
-		GroupInfoCategory.AddProperty("GroupName");
-		GroupInfoCategory.AddProperty("GroupType");
+		{
+			if (CVarEnableMutableMacros.GetValueOnAnyThread())
+			{
+				DetailBuilder.HideProperty("GroupName");
+			}
+			else
+			{
+				GroupInfoCategory.AddProperty("GroupName");
+			}
+
+			GroupInfoCategory.AddProperty("GroupType");
+		}
+
 
 		// Getting group node children names
 		GenerateChildrenObjectNames();
@@ -130,18 +142,18 @@ void FCustomizableObjectNodeObjectGroupDetails::GenerateChildrenObjectNames()
 		{
 			if (const UCustomizableObjectNodeObject* ChildObjectNode = Cast<UCustomizableObjectNodeObject>(ChildPin->GetOwningNode()))
 			{
-				ChildrenNameOptions.Add(MakeShareable(new FString(ChildObjectNode->ObjectName)));
+				ChildrenNameOptions.Add(MakeShareable(new FString(ChildObjectNode->GetObjectName())));
 
-				if (ChildObjectNode->ObjectName.Equals(NodeGroup->DefaultValue))
+				if (ChildObjectNode->GetObjectName().Equals(NodeGroup->DefaultValue))
 				{
 					InitialNameOption = ChildrenNameOptions.Last();
 				}
 			}
 			else if (const UCustomizableObjectNodeObjectGroup* ChildGroupObjectNode = Cast<UCustomizableObjectNodeObjectGroup>(ChildPin->GetOwningNode()))
 			{
-				ChildrenNameOptions.Add(MakeShareable(new FString(ChildGroupObjectNode->GroupName)));
+				ChildrenNameOptions.Add(MakeShareable(new FString(ChildGroupObjectNode->GetGroupName())));
 
-				if (ChildGroupObjectNode->GroupName.Equals(NodeGroup->DefaultValue))
+				if (ChildGroupObjectNode->GetGroupName().Equals(NodeGroup->DefaultValue))
 				{
 					InitialNameOption = ChildrenNameOptions.Last();
 				}
@@ -163,9 +175,9 @@ void FCustomizableObjectNodeObjectGroupDetails::GenerateChildrenObjectNames()
 			{
 				if (ChildObjectNode)
 				{
-					ChildrenNameOptions.Add(MakeShareable(new FString(ChildObjectNode->ObjectName)));
+					ChildrenNameOptions.Add(MakeShareable(new FString(ChildObjectNode->GetObjectName())));
 
-					if (ChildObjectNode->ObjectName.Equals(NodeGroup->DefaultValue))
+					if (ChildObjectNode->GetObjectName().Equals(NodeGroup->DefaultValue))
 					{
 						InitialNameOption = ChildrenNameOptions.Last();
 					}
diff --git a/Engine/Plugins/Mutable/Source/CustomizableObjectEditor/Private/MuCOE/EdGraphSchema_CustomizableObject.cpp b/Engine/Plugins/Mutable/Source/CustomizableObjectEditor/Private/MuCOE/EdGraphSchema_CustomizableObject.cpp
--- a/Engine/Plugins/Mutable/Source/CustomizableObjectEditor/Private/MuCOE/EdGraphSchema_CustomizableObject.cpp
+++ b/Engine/Plugins/Mutable/Source/CustomizableObjectEditor/Private/MuCOE/EdGraphSchema_CustomizableObject.cpp
@@ -65,7 +65,7 @@
 #include "MuCOE/Nodes/CustomizableObjectNodeTable.h"
 #include "MuCOE/Nodes/CustomizableObjectNodeTexture.h"
 #include "MuCOE/Nodes/CustomizableObjectNodePassThroughTexture.h"
-#include "MuCOE/Nodes/CustomizableObjectNodeStringConstant.h"
+#include "MuCOE/Nodes/CustomizableObjectNodeStaticString.h"
 #include "MuCOE/Nodes/CustomizableObjectNodeTextureBinarise.h"
 #include "MuCOE/Nodes/CustomizableObjectNodeTextureColourMap.h"
 #include "MuCOE/Nodes/CustomizableObjectNodeTextureFromChannels.h"
@@ -444,15 +444,21 @@ void UEdGraphSchema_CustomizableObject::GetGraphContextActions(FGraphContextMenu
 			ContextMenuBuilder.CreateTemplateNode<UCustomizableObjectNodeMaterialVariation>(),
 			ContextMenuBuilder.CreateTemplateNode<UCustomizableObjectNodeMaterialSwitch>(),
 			ContextMenuBuilder.CreateTemplateNode<UCustomizableObjectNodeCopyMaterial>(),
-			ContextMenuBuilder.CreateTemplateNode<UCustomizableObjectNodeTable>(),
-
-#if	ENABLE_MUTABLE_MACROS
-			ContextMenuBuilder.CreateTemplateNode<UCustomizableObjectNodeMacroInstance>(),
-			ContextMenuBuilder.CreateTemplateNode<UCustomizableObjectNodeStringConstant>(),
-#endif
+			ContextMenuBuilder.CreateTemplateNode<UCustomizableObjectNodeTable>()
 		};
 
 		AddNewNodeCategoryActionsFiltered(ObjectTemplateNodes, ContextMenuBuilder, TEXT("Object"), GeneralGrouping, Filter);
+
+		if (CVarEnableMutableMacros.GetValueOnAnyThread())
+		{
+			UCustomizableObjectNode* MacroTemplateNodes[]
+			{
+				ContextMenuBuilder.CreateTemplateNode<UCustomizableObjectNodeMacroInstance>(),
+				ContextMenuBuilder.CreateTemplateNode<UCustomizableObjectNodeStaticString>()
+			};
+
+			AddNewNodeCategoryActionsFiltered(MacroTemplateNodes, ContextMenuBuilder, TEXT("Object"), GeneralGrouping, Filter);
+		}
 	}
 
 	{
diff --git a/Engine/Plugins/Mutable/Source/CustomizableObjectEditor/Private/MuCOE/GenerateMutableSource/GenerateMutableSource.cpp b/Engine/Plugins/Mutable/Source/CustomizableObjectEditor/Private/MuCOE/GenerateMutableSource/GenerateMutableSource.cpp
--- a/Engine/Plugins/Mutable/Source/CustomizableObjectEditor/Private/MuCOE/GenerateMutableSource/GenerateMutableSource.cpp
+++ b/Engine/Plugins/Mutable/Source/CustomizableObjectEditor/Private/MuCOE/GenerateMutableSource/GenerateMutableSource.cpp
@@ -23,7 +23,7 @@
 #include "MuCOE/Nodes/CustomizableObjectNodeMesh.h"
 #include "MuCOE/Nodes/CustomizableObjectNodeModifierClipWithMesh.h"
 #include "MuCOE/Nodes/CustomizableObjectNodeObjectGroup.h"
-#include "MuCOE/Nodes/CustomizableObjectNodeStringConstant.h"
+#include "MuCOE/Nodes/CustomizableObjectNodeStaticString.h"
 #include "MuCOE/Nodes/CustomizableObjectNodeTable.h"
 #include "MuCOE/Nodes/CustomizableObjectNodeTexture.h"
 #include "MuCOE/Nodes/CustomizableObjectNodePassThroughTexture.h"
@@ -329,14 +329,26 @@ int32 FMutableGraphGenerationContext::AddStreamedResource(UAssetUserData& AssetU
 
 const FGuid FMutableGraphGenerationContext::GetNodeIdUnique(const UCustomizableObjectNode* Node)
 {
-	TArray<const UObject*>* ArrayResult = NodeIdsMap.Find(Node->NodeGuid);
+	FGuid NodeID = Node->NodeGuid;
+
+	if (Node->IsInMacro())
+	{
+		check(MacroNodesStack.Num()); // To ensure that we only enter here when compiling
+		
+		for (int32 MacroIndex = 0; MacroIndex < MacroNodesStack.Num(); ++MacroIndex)
+		{
+			NodeID = FGuid::Combine(NodeID, MacroNodesStack[MacroIndex]->NodeGuid);
+		}
+	}
+
+	TArray<const UObject*>* ArrayResult = NodeIdsMap.Find(NodeID);
 
 	if (ArrayResult == nullptr)
 	{
 		TArray<const UObject*> ArrayTemp;
 		ArrayTemp.Add(Node);
-		NodeIdsMap.Add(Node->NodeGuid, ArrayTemp);
-		return Node->NodeGuid;
+		NodeIdsMap.Add(NodeID, ArrayTemp);
+		return NodeID;
 	}
 
 	ArrayResult->Add(Node);
@@ -848,9 +860,9 @@ mu::Ptr<mu::NodeObject> GenerateMutableSource(const UEdGraphPin * Pin, FMutableG
 		Result = ObjectNode;
 
 		ObjectNode->SetMessageContext(Node);
-		ObjectNode->SetName(TypedNodeObj->ObjectName);
+		ObjectNode->SetName(TypedNodeObj->GetObjectName(&GenerationContext.MacroNodesStack));
 		FGuid FinalGuid = GenerationContext.GetNodeIdUnique(TypedNodeObj);
-		if (FinalGuid != TypedNodeObj->NodeGuid)
+		if (FinalGuid != TypedNodeObj->NodeGuid && !TypedNodeObj->IsInMacro())
 		{
 			GenerationContext.Log(FText::FromString(TEXT("Warning: Node has a duplicated GUID. A new ID has been generated, but cooked data will not be deterministic.")), Node, EMessageSeverity::Warning);
 		}
@@ -1038,7 +1050,7 @@ mu::Ptr<mu::NodeObject> GenerateMutableSource(const UEdGraphPin * Pin, FMutableG
 		GenerationContext.SocketPriorityStack.Push(TypedNodeGroup->SocketPriority);
 
 		GroupNode->SetMessageContext(Node);
-		GroupNode->SetName(TypedNodeGroup->GroupName);
+		GroupNode->SetName(TypedNodeGroup->GetGroupName(&GenerationContext.MacroNodesStack));
 		GroupNode->SetUid(TypedNodeGroup->NodeGuid.ToString());
 		
 		TArray<UCustomizableObjectNodeGroupProjectorParameter*> GroupProjectors;
@@ -1075,7 +1087,7 @@ mu::Ptr<mu::NodeObject> GenerateMutableSource(const UEdGraphPin * Pin, FMutableG
 		// External children
 		TArray<UCustomizableObjectNodeObject*> ExternalChildNodes;
 		GenerationContext.GroupIdToExternalNodeMap.MultiFind(TypedNodeGroup->NodeGuid, ExternalChildNodes);
-		GenerationContext.GuidToParamNameMap.Add(TypedNodeGroup->NodeGuid, TypedNodeGroup->GroupName);
+		GenerationContext.GuidToParamNameMap.Add(TypedNodeGroup->NodeGuid, TypedNodeGroup->GetGroupName(&GenerationContext.MacroNodesStack));
 
 		// Children
 		const TArray<UEdGraphPin*> ConnectedChildrenPins = FollowInputPinArray(*TypedNodeGroup->ObjectsPin());
@@ -1103,10 +1115,10 @@ mu::Ptr<mu::NodeObject> GenerateMutableSource(const UEdGraphPin * Pin, FMutableG
 
 			UCustomizableObjectNodeObject* CustomizableObjectNodeObject = Cast<UCustomizableObjectNodeObject>(ConnectedChildrenPins[ChildIndex]->GetOwningNode());
 
-			const FString* SelectedOptionName = GenerationContext.CompilationContext.Options.ParamNamesToSelectedOptions.Find(TypedNodeGroup->GroupName); // If the param is in the map restrict to only the selected option
+			const FString* SelectedOptionName = GenerationContext.CompilationContext.Options.ParamNamesToSelectedOptions.Find(TypedNodeGroup->GetGroupName(&GenerationContext.MacroNodesStack)); // If the param is in the map restrict to only the selected option
 			mu::Ptr<mu::NodeObject> ChildNode;
 
-			if (bConnectAtLeastTheLastChild || !SelectedOptionName || (CustomizableObjectNodeObject && *SelectedOptionName == CustomizableObjectNodeObject->ObjectName) )
+			if (bConnectAtLeastTheLastChild || !SelectedOptionName || (CustomizableObjectNodeObject && *SelectedOptionName == CustomizableObjectNodeObject->GetObjectName(&GenerationContext.MacroNodesStack)) )
 			{
 				bAtLeastOneConnected = true;
 
@@ -1118,15 +1130,15 @@ mu::Ptr<mu::NodeObject> GenerateMutableSource(const UEdGraphPin * Pin, FMutableG
 					FString LeftSplit = CustomizableObjectNodeObject->GetPathName();
 					LeftSplit.Split(".", &LeftSplit, nullptr);
 					GenerationContext.CustomizableObjectPathMap.Add(CustomizableObjectNodeObject->Identifier.ToString(), LeftSplit);
-					GenerationContext.GroupNodeMap.Add(CustomizableObjectNodeObject->Identifier.ToString(), FCustomizableObjectIdPair(TypedNodeGroup->GroupName, ChildNode->GetName()));
+					GenerationContext.GroupNodeMap.Add(CustomizableObjectNodeObject->Identifier.ToString(), FCustomizableObjectIdPair(TypedNodeGroup->GetGroupName(&GenerationContext.MacroNodesStack), ChildNode->GetName()));
 					ParameterUIData.ArrayIntegerParameterOption.Add(
-						CustomizableObjectNodeObject->ObjectName,
+						CustomizableObjectNodeObject->GetObjectName(&GenerationContext.MacroNodesStack),
 				FIntegerParameterUIData(CustomizableObjectNodeObject->ParamUIMetadata));
 
 					if (TypedNodeGroup->GroupType == ECustomizableObjectGroupType::COGT_TOGGLE)
 					{
 						// UI Data is only relevant when the group node is set to Toggle
-						GenerationContext.ParameterUIDataMap.Add(CustomizableObjectNodeObject->ObjectName, FMutableParameterData(
+						GenerationContext.ParameterUIDataMap.Add(CustomizableObjectNodeObject->GetObjectName(&GenerationContext.MacroNodesStack), FMutableParameterData(
 							CustomizableObjectNodeObject->ParamUIMetadata,
 							EMutableParameterType::Int));
 					}
@@ -1135,7 +1147,7 @@ mu::Ptr<mu::NodeObject> GenerateMutableSource(const UEdGraphPin * Pin, FMutableG
 			else
 			{
 				ChildNode = new mu::NodeObjectNew;
-				ChildNode->SetName(CustomizableObjectNodeObject->ObjectName);
+				ChildNode->SetName(CustomizableObjectNodeObject->GetObjectName(&GenerationContext.MacroNodesStack));
 				GroupNode->Children[ChildIndex] = ChildNode;
 			}
 
@@ -1160,7 +1172,7 @@ mu::Ptr<mu::NodeObject> GenerateMutableSource(const UEdGraphPin * Pin, FMutableG
 				if (UCustomizableObjectNodeObject* NodeObject = Cast<UCustomizableObjectNodeObject>(ConnectedPin->GetOwningNode()))
 				{
 					ParentId = GenerationContext.GroupIdToExternalNodeMap.FindKey(NodeObject);
-					ValueName = NodeObject->ObjectName;
+					ValueName = NodeObject->GetObjectName(&GenerationContext.MacroNodesStack);
 
 					// Group objects in the same graph aren't in the GroupIdToExternalNodeMap, so follow the pins instead
 					if (!ParentId && NodeObject->OutputPin())
@@ -1177,7 +1189,7 @@ mu::Ptr<mu::NodeObject> GenerateMutableSource(const UEdGraphPin * Pin, FMutableG
 				else if (UCustomizableObjectNodeObjectGroup* NodeObjectGroup = Cast<UCustomizableObjectNodeObjectGroup>(ConnectedPin->GetOwningNode()))
 				{
 					ParentId = &NodeObjectGroup->NodeGuid;
-					ValueName = TypedNodeGroup->GroupName;
+					ValueName = TypedNodeGroup->GetGroupName(&GenerationContext.MacroNodesStack);
 				}
 
 				if (ParentId)
@@ -1216,10 +1228,10 @@ mu::Ptr<mu::NodeObject> GenerateMutableSource(const UEdGraphPin * Pin, FMutableG
 
 			UCustomizableObjectNodeObject* CustomizableObjectNodeObject = Cast<UCustomizableObjectNodeObject>(ExternalChildNode->OutputPin()->GetOwningNode());
 
-			const FString* SelectedOptionName = GenerationContext.CompilationContext.Options.ParamNamesToSelectedOptions.Find(TypedNodeGroup->GroupName); // If the param is in the map restrict to only the selected option
+			const FString* SelectedOptionName = GenerationContext.CompilationContext.Options.ParamNamesToSelectedOptions.Find(TypedNodeGroup->GetGroupName(&GenerationContext.MacroNodesStack)); // If the param is in the map restrict to only the selected option
 			mu::Ptr<mu::NodeObject> ChildNode;
 
-			if (bConnectAtLeastTheLastChild || !SelectedOptionName || (CustomizableObjectNodeObject  && *SelectedOptionName == CustomizableObjectNodeObject->ObjectName) )
+			if (bConnectAtLeastTheLastChild || !SelectedOptionName || (CustomizableObjectNodeObject  && *SelectedOptionName == CustomizableObjectNodeObject->GetObjectName(&GenerationContext.MacroNodesStack)) )
 			{
 				bAtLeastOneConnected = true;
 
@@ -1231,12 +1243,12 @@ mu::Ptr<mu::NodeObject> GenerateMutableSource(const UEdGraphPin * Pin, FMutableG
 					FString LeftSplit = ExternalChildNode->GetPathName();
 					LeftSplit.Split(".", &LeftSplit, nullptr);
 					GenerationContext.CustomizableObjectPathMap.Add(CustomizableObjectNodeObject->Identifier.ToString(), LeftSplit);
-					GenerationContext.GroupNodeMap.Add(CustomizableObjectNodeObject->Identifier.ToString(), FCustomizableObjectIdPair(TypedNodeGroup->GroupName, ChildNode->GetName()));
+					GenerationContext.GroupNodeMap.Add(CustomizableObjectNodeObject->Identifier.ToString(), FCustomizableObjectIdPair(TypedNodeGroup->GetGroupName(&GenerationContext.MacroNodesStack), ChildNode->GetName()));
 					ParameterUIData.ArrayIntegerParameterOption.Add(
-						CustomizableObjectNodeObject->ObjectName,
+						CustomizableObjectNodeObject->GetObjectName(&GenerationContext.MacroNodesStack),
 						FIntegerParameterUIData(CustomizableObjectNodeObject->ParamUIMetadata));
 
-					if (CustomizableObjectNodeObject->ObjectName.IsEmpty())
+					if (CustomizableObjectNodeObject->GetObjectName(&GenerationContext.MacroNodesStack).IsEmpty())
 					{
 						GenerationContext.NoNameNodeObjectArray.AddUnique(CustomizableObjectNodeObject);
 					}
@@ -1244,7 +1256,7 @@ mu::Ptr<mu::NodeObject> GenerateMutableSource(const UEdGraphPin * Pin, FMutableG
 					if (TypedNodeGroup->GroupType == ECustomizableObjectGroupType::COGT_TOGGLE)
 					{
 						// UI Data is only relevant when the group node is set to Toggle
-						GenerationContext.ParameterUIDataMap.Add(CustomizableObjectNodeObject->ObjectName, FMutableParameterData(
+						GenerationContext.ParameterUIDataMap.Add(CustomizableObjectNodeObject->GetObjectName(&GenerationContext.MacroNodesStack), FMutableParameterData(
 							CustomizableObjectNodeObject->ParamUIMetadata,
 							EMutableParameterType::Int));
 					}
@@ -1253,7 +1265,7 @@ mu::Ptr<mu::NodeObject> GenerateMutableSource(const UEdGraphPin * Pin, FMutableG
 			else
 			{
 				ChildNode = new mu::NodeObjectNew;
-				ChildNode->SetName(CustomizableObjectNodeObject->ObjectName);
+				ChildNode->SetName(CustomizableObjectNodeObject->GetObjectName(&GenerationContext.MacroNodesStack));
 				GroupNode->Children[ChildIndex] = ChildNode;
 			}
 
@@ -1267,7 +1279,7 @@ mu::Ptr<mu::NodeObject> GenerateMutableSource(const UEdGraphPin * Pin, FMutableG
 			ChildIndex++;
 		}
 
-		const FMutableParameterData* ChildFilledUIData = GenerationContext.ParameterUIDataMap.Find(TypedNodeGroup->GroupName);
+		const FMutableParameterData* ChildFilledUIData = GenerationContext.ParameterUIDataMap.Find(TypedNodeGroup->GetGroupName(&GenerationContext.MacroNodesStack));
 		if (ChildFilledUIData && ChildFilledUIData->ParamUIMetadata.ExtraInformation.Find(FString("__HasCollapsibleChildren")))
 		{
 			// Some child param filled the HasCollapsibleChildren UI info, refill it so it's not lost
@@ -1289,7 +1301,7 @@ mu::Ptr<mu::NodeObject> GenerateMutableSource(const UEdGraphPin * Pin, FMutableG
 		}
 		else
 		{
-			GenerationContext.ParameterUIDataMap.Add(TypedNodeGroup->GroupName, ParameterUIData);
+			GenerationContext.ParameterUIDataMap.Add(TypedNodeGroup->GetGroupName(&GenerationContext.MacroNodesStack), ParameterUIData);
 		}
 
 		// Go back to the parent group node's socket priority if it exists
@@ -1856,37 +1868,6 @@ void FMutableComponentInfo::AccumulateBonesToRemovePerLOD(const TArray<FLODReduc
 	}
 }
 
-FString GetVariationTag(const UCustomizableObjectNodeVariation& VariationNode, int32 VariationIndex, FMutableGraphGenerationContext& GenerationContext)
-{
-	if (VariationNode.IsInMacro())
-	{
-		if (const UEdGraphPin* VariationPin = VariationNode.VariationPin(VariationIndex))
-		{
-			if (const UEdGraphPin* TagPin = VariationNode.VariationTagPin(VariationIndex))
-			{
-				if (const UEdGraphPin* ConnectedStringPin = FollowInputPin(*TagPin))
-				{
-					const UEdGraphPin* SourceStringPin = ConnectedStringPin ? GraphTraversal::FindStringSource(*ConnectedStringPin, &GenerationContext.MacroNodesStack) : nullptr;
-
-					if (SourceStringPin)
-					{
-						if (const UCustomizableObjectNodeStringConstant* StringNode = Cast<UCustomizableObjectNodeStringConstant>(SourceStringPin->GetOwningNode()))
-						{
-							return StringNode->Value;
-						}
-					}
-					else
-					{
-						GenerationContext.Log(LOCTEXT("VariationNodeTagError", "Could not found a linked String node."), &VariationNode);
-					}
-				}
-			}
-		}
-	}
-
-	return VariationNode.GetVariation(VariationIndex).Tag;
-}
-
 #undef LOCTEXT_NAMESPACE
 
 FMutableGraphGenerationContext::FSharedSurface::FSharedSurface(uint8 InLOD, const mu::Ptr<mu::NodeSurfaceNew>& InNodeSurfaceNew)
diff --git a/Engine/Plugins/Mutable/Source/CustomizableObjectEditor/Private/MuCOE/GenerateMutableSource/GenerateMutableSource.h b/Engine/Plugins/Mutable/Source/CustomizableObjectEditor/Private/MuCOE/GenerateMutableSource/GenerateMutableSource.h
--- a/Engine/Plugins/Mutable/Source/CustomizableObjectEditor/Private/MuCOE/GenerateMutableSource/GenerateMutableSource.h
+++ b/Engine/Plugins/Mutable/Source/CustomizableObjectEditor/Private/MuCOE/GenerateMutableSource/GenerateMutableSource.h
@@ -990,5 +990,3 @@ uint32 AddUniqueSurfaceMetadata(const FMutableSourceSurfaceMetadata& MeshData, F
 /** Convert a mesh to Mutable format.
 */
 TSharedPtr<mu::FMesh> ConvertSkeletalMeshToMutable(const FMutableSourceMeshData& Source, FMutableCompilationContext& Context, const FString& MorphName);
-
-FString GetVariationTag(const UCustomizableObjectNodeVariation& VariationNode, int32 VariationIndex, FMutableGraphGenerationContext& GenerationContext);
diff --git a/Engine/Plugins/Mutable/Source/CustomizableObjectEditor/Private/MuCOE/GenerateMutableSource/GenerateMutableSourceColor.cpp b/Engine/Plugins/Mutable/Source/CustomizableObjectEditor/Private/MuCOE/GenerateMutableSource/GenerateMutableSourceColor.cpp
--- a/Engine/Plugins/Mutable/Source/CustomizableObjectEditor/Private/MuCOE/GenerateMutableSource/GenerateMutableSourceColor.cpp
+++ b/Engine/Plugins/Mutable/Source/CustomizableObjectEditor/Private/MuCOE/GenerateMutableSource/GenerateMutableSourceColor.cpp
@@ -76,11 +76,11 @@ mu::Ptr<mu::NodeColour> GenerateMutableSourceColor(const UEdGraphPin* Pin, FMuta
 		mu::Ptr<mu::NodeColourParameter> ColorNode = new mu::NodeColourParameter();
 		Result = ColorNode;
 
-		ColorNode->Name = TypedNodeColorParam->ParameterName;
+		ColorNode->Name = TypedNodeColorParam->GetParameterName(&GenerationContext.MacroNodesStack);
 		ColorNode->Uid = GenerationContext.GetNodeIdUnique(Node).ToString();
 		ColorNode->DefaultValue = TypedNodeColorParam->DefaultValue;
 
-		GenerationContext.ParameterUIDataMap.Add(TypedNodeColorParam->ParameterName, FMutableParameterData(
+		GenerationContext.ParameterUIDataMap.Add(TypedNodeColorParam->GetParameterName(&GenerationContext.MacroNodesStack), FMutableParameterData(
 			TypedNodeColorParam->ParamUIMetadata,
 			EMutableParameterType::Color));
 	}
@@ -264,7 +264,7 @@ mu::Ptr<mu::NodeColour> GenerateMutableSourceColor(const UEdGraphPin* Pin, FMuta
 			const UEdGraphPin* VariationPin = TypedNodeColorVar->VariationPin(VariationIndex);
 			if (!VariationPin) continue;
 
-			ColorNode->Variations[VariationIndex].Tag = GetVariationTag(*TypedNodeColorVar,VariationIndex, GenerationContext);
+			ColorNode->Variations[VariationIndex].Tag = TypedNodeColorVar->GetVariationTag(VariationIndex, &GenerationContext.MacroNodesStack);
 
 			if (const UEdGraphPin* ConnectedPin = FollowInputPin(*VariationPin))
 			{
diff --git a/Engine/Plugins/Mutable/Source/CustomizableObjectEditor/Private/MuCOE/GenerateMutableSource/GenerateMutableSourceComponent.cpp b/Engine/Plugins/Mutable/Source/CustomizableObjectEditor/Private/MuCOE/GenerateMutableSource/GenerateMutableSourceComponent.cpp
--- a/Engine/Plugins/Mutable/Source/CustomizableObjectEditor/Private/MuCOE/GenerateMutableSource/GenerateMutableSourceComponent.cpp
+++ b/Engine/Plugins/Mutable/Source/CustomizableObjectEditor/Private/MuCOE/GenerateMutableSource/GenerateMutableSourceComponent.cpp
@@ -17,7 +17,7 @@
 #include "MuCOE/Nodes/CustomizableObjectNodeComponentVariation.h"
 #include "MuCOE/Nodes/CustomizableObjectNodeModifierBase.h"
 #include "MuCOE/Nodes/CustomizableObjectNodeMacroInstance.h"
-#include "MuCOE/Nodes/CustomizableObjectNodeStringConstant.h"
+#include "MuCOE/Nodes/CustomizableObjectNodeStaticString.h"
 #include "MuCOE/Nodes/CustomizableObjectNodeTunnel.h"
 #include "MuT/NodeComponent.h"
 #include "MuT/NodeComponentEdit.h"
@@ -37,42 +37,6 @@
 
 #define LOCTEXT_NAMESPACE "CustomizableObjectEditor"
 
-FName GetComponentName(const UCustomizableObjectNode& Node, FMutableGraphGenerationContext& GenerationContext)
-{
-	// find the name of the linked string node
-	if (Node.IsInMacro())
-	{
-		if (const UEdGraphPin* ComponentNamePin = Node.FindPin(TEXT("Component Name"), EEdGraphPinDirection::EGPD_Input))
-		{
-			const UEdGraphPin* StringPin = GraphTraversal::FindStringSource(*ComponentNamePin, &GenerationContext.MacroNodesStack);
-
-			if (const UCustomizableObjectNodeStringConstant* StringNode = StringPin ? Cast<UCustomizableObjectNodeStringConstant>(StringPin->GetOwningNode()) : nullptr)
-			{
-				return FName(StringNode->Value);
-			}
-		}
-	}
-
-	// return the name set in the node
-	if (const UCustomizableObjectNodeComponentMesh* TypedComponentMesh = Cast<UCustomizableObjectNodeComponentMesh>(&Node))
-	{
-		return TypedComponentMesh->ComponentName;
-	}
-
-	else if(const UCustomizableObjectNodeComponentMeshAddTo* TypedComponentMeshExtend = Cast<UCustomizableObjectNodeComponentMeshAddTo>(&Node))
-	{
-		return TypedComponentMeshExtend->ParentComponentName;
-	}
-
-	else if(const UCustomizableObjectNodeComponentPassthroughMesh* TypedComponentPassthroughMesh = Cast<UCustomizableObjectNodeComponentPassthroughMesh>(&Node))
-	{
-		return TypedComponentPassthroughMesh->ComponentName;
-	}
-
-	return FName();
-}
-
-
 /** Generate LOD pins of the given NodeComponentBase (NodeComponent, NodeComponentExtend...).
  * @param TypedComponentMesh Given component node.
  * @param NodeComponent Core node to connect LOD generated pins. */
@@ -179,7 +143,9 @@ mu::Ptr<mu::NodeComponent> GenerateMutableSourceComponent(const UEdGraphPin* Pin
 		}
 		
 		mu::Ptr<mu::NodeComponentNew> NodeComponentNew = new mu::NodeComponentNew();
-		NodeComponentNew->Id = GenerationContext.ComponentNames.Find(TypedComponentMesh->ComponentName);
+		const FName& ComponentName = TypedComponentMesh->GetComponentName(&GenerationContext.MacroNodesStack);
+
+		NodeComponentNew->Id = GenerationContext.ComponentNames.Find(ComponentName);
 		NodeComponentNew->SetMessageContext(Node);
 		
 		UEdGraphPin* MaterialAssetPin = TypedComponentMesh->GetOverlayMaterialAssetPin();
@@ -198,7 +164,7 @@ mu::Ptr<mu::NodeComponent> GenerateMutableSourceComponent(const UEdGraphPin* Pin
 
 		Result = NodeComponentNew;
 
-		GenerationContext.CurrentMeshComponent = TypedComponentMesh->ComponentName;
+		GenerationContext.CurrentMeshComponent = ComponentName;
 		GenerationContext.CurrentAutoLODStrategy = TypedComponentMesh->AutoLODStrategy;
 		
 		GenerateMutableSourceComponentMesh(GenerationContext, *TypedComponentMesh, NodeComponentNew);
@@ -209,16 +175,18 @@ mu::Ptr<mu::NodeComponent> GenerateMutableSourceComponent(const UEdGraphPin* Pin
 
 	else if (const UCustomizableObjectNodeComponentMeshAddTo* TypedComponentMeshExtend = Cast<UCustomizableObjectNodeComponentMeshAddTo>(Node))
 	{
+		const FName& ParentComponentName = TypedComponentMeshExtend->GetParentComponentName(&GenerationContext.MacroNodesStack);
+
 		if (FMutableComponentInfo* FindResult = GenerationContext.CompilationContext.ComponentInfos.FindByPredicate([&](const FMutableComponentInfo& Element)
 		{
-			return Element.ComponentName == TypedComponentMeshExtend->ParentComponentName;
+			return Element.ComponentName == ParentComponentName;
 		}))
 		{
 			UCustomizableObjectNodeComponentMesh* TypedParentComponentMesh = FindResult->Node;
 
 			if (TypedComponentMeshExtend->NumLODs > TypedParentComponentMesh->NumLODs)
 			{
-				FText Msg = FText::Format(LOCTEXT("ExtendMeshComponentLODs", "Add To Mesh Component can not have more LODs than its parent Mesh Component [{0}]."), FText::FromName(TypedComponentMeshExtend->ParentComponentName));
+				FText Msg = FText::Format(LOCTEXT("ExtendMeshComponentLODs", "Add To Mesh Component can not have more LODs than its parent Mesh Component [{0}]."), FText::FromName(ParentComponentName));
 				GenerationContext.Log(Msg, TypedComponentMeshExtend, EMessageSeverity::Warning);
 			}
 
@@ -228,6 +196,9 @@ mu::Ptr<mu::NodeComponent> GenerateMutableSourceComponent(const UEdGraphPin* Pin
 
 			mu::Ptr<mu::NodeComponent> ParentNodeComponent = GenerateMutableSourceComponent(TypedParentComponentMesh->OutputPin.Get(), GenerationContext);
 
+			// Get the parent name using the parent's macro context
+			GenerationContext.CurrentMeshComponent = TypedParentComponentMesh->GetComponentName(&GenerationContext.MacroNodesStack);
+
 			// Restore GenerationContext MacroContex
 			GenerationContext.MacroNodesStack = MacroContextCopy;
 
@@ -235,7 +206,6 @@ mu::Ptr<mu::NodeComponent> GenerateMutableSourceComponent(const UEdGraphPin* Pin
 			NodeComponentEdit->Parent = ParentNodeComponent.get();
 			NodeComponentEdit->SetMessageContext(TypedComponentMeshExtend);
 		
-			GenerationContext.CurrentMeshComponent = TypedParentComponentMesh->ComponentName;
 			GenerationContext.CurrentAutoLODStrategy = TypedComponentMeshExtend->AutoLODStrategy == ECustomizableObjectAutomaticLODStrategy::Inherited ?
 				TypedParentComponentMesh->AutoLODStrategy :
 				TypedComponentMeshExtend->AutoLODStrategy;
@@ -249,16 +219,17 @@ mu::Ptr<mu::NodeComponent> GenerateMutableSourceComponent(const UEdGraphPin* Pin
 		}
 		else
 		{
-			FText Msg = FText::Format(LOCTEXT("ExtendMeshComponent", "Can not find parent Mesh Component {0}."), FText::FromName(TypedComponentMeshExtend->ParentComponentName));
+			FText Msg = FText::Format(LOCTEXT("ExtendMeshComponent", "Can not find parent Mesh Component {0}."), FText::FromName(ParentComponentName));
 			GenerationContext.Log(Msg, TypedComponentMeshExtend, EMessageSeverity::Error);
 		}
 	}
 	
 	else if (const UCustomizableObjectNodeComponentPassthroughMesh* TypedComponentPassthroughMesh = Cast<UCustomizableObjectNodeComponentPassthroughMesh>(Node))
 	{
-		GenerationContext.CurrentMeshComponent = TypedComponentPassthroughMesh->ComponentName;
+		const FName& ComponentName = TypedComponentPassthroughMesh->GetComponentName(&GenerationContext.MacroNodesStack);
+		GenerationContext.CurrentMeshComponent = ComponentName;
 
-		if (TypedComponentPassthroughMesh->ComponentName.IsNone())
+		if (ComponentName.IsNone())
 		{
 			FString Msg = FString::Printf(TEXT("Invalid Component Name."));
 			GenerationContext.Log(FText::FromString(Msg), TypedComponentPassthroughMesh, EMessageSeverity::Warning);
@@ -295,12 +266,12 @@ mu::Ptr<mu::NodeComponent> GenerateMutableSourceComponent(const UEdGraphPin* Pin
 
 		// Create the component node
 		mu::Ptr<mu::NodeComponentNew> ComponentNode = new mu::NodeComponentNew;
-		if (GenerationContext.ComponentNames.Contains(TypedComponentPassthroughMesh->ComponentName))
+		if (GenerationContext.ComponentNames.Contains(ComponentName))
 		{
-			FString Msg = FString::Printf(TEXT("More than one component with the same name [%s] have been found. This is not supported."), *TypedComponentPassthroughMesh->ComponentName.ToString());
+			FString Msg = FString::Printf(TEXT("More than one component with the same name [%s] have been found. This is not supported."), *ComponentName.ToString());
 			GenerationContext.Log(FText::FromString(Msg), TypedComponentPassthroughMesh, EMessageSeverity::Warning);
 		}
-		ComponentNode->Id = GenerationContext.ComponentNames.Add(TypedComponentPassthroughMesh->ComponentName);
+		ComponentNode->Id = GenerationContext.ComponentNames.Add(ComponentName);
 
 		// Create a LOD for each pass-through mesh LOD.
 		const FSkeletalMeshModel* Model = SkeletalMesh->GetImportedModel();
@@ -440,7 +411,7 @@ mu::Ptr<mu::NodeComponent> GenerateMutableSourceComponent(const UEdGraphPin* Pin
 
 			if (UEdGraphPin* VariationPin = TypedNodeVar->VariationPin(VariationIndex))
 			{
-				SurfNode->Variations[VariationIndex].Tag = GetVariationTag(*TypedNodeVar, VariationIndex, GenerationContext);
+				SurfNode->Variations[VariationIndex].Tag = TypedNodeVar->GetVariationTag(VariationIndex, &GenerationContext.MacroNodesStack);
 				for (const UEdGraphPin* ConnectedPin : FollowInputPinArray(*VariationPin))
 				{
 					// Is it a modifier?
@@ -489,7 +460,9 @@ mu::Ptr<mu::NodeComponent> GenerateMutableSourceComponent(const UEdGraphPin* Pin
 
 void FirstPass(UCustomizableObjectNodeComponentMesh& Node, FMutableGraphGenerationContext& GenerationContext)
 {
-	if (Node.ComponentName.IsNone())
+	FName ComponentName = Node.GetComponentName(&GenerationContext.MacroNodesStack);
+
+	if (ComponentName.IsNone())
 	{
 		GenerationContext.Log(LOCTEXT("EmptyComponentNameError", "Error! Missing name in a component of the Customizable Object."), &Node, EMessageSeverity::Error);
 		return;
@@ -497,7 +470,7 @@ void FirstPass(UCustomizableObjectNodeComponentMesh& Node, FMutableGraphGenerati
 	
 	if (FMutableComponentInfo* Result = GenerationContext.CompilationContext.ComponentInfos.FindByPredicate([&](const FMutableComponentInfo& Element)
 	{
-		return Element.ComponentName == Node.ComponentName;
+		return Element.ComponentName == ComponentName;
 	}))
 	{
 		FText Msg = FText::Format(LOCTEXT("ComponentNodeWithSameNameExists", "Error! A Mesh Component node with the same name already exists in the Customizable Object [{0}]"), FText::FromString(GraphTraversal::GetObject(*Result->Node)->GetName()));
@@ -602,12 +575,12 @@ void FirstPass(UCustomizableObjectNodeComponentMesh& Node, FMutableGraphGenerati
 
 	NumMaxLODsToStream = FMath::Clamp(NumMaxLODsToStream, 0, NumLODs - 1);
 
-	GenerationContext.NumLODs.Add(Node.ComponentName, NumLODs);
-	GenerationContext.FirstLODAvailable.Add(Node.ComponentName, FirstLODAvailable);
-	GenerationContext.NumMaxLODsToStream.Add(Node.ComponentName, NumMaxLODsToStream);
+	GenerationContext.NumLODs.Add(ComponentName, NumLODs);
+	GenerationContext.FirstLODAvailable.Add(ComponentName, FirstLODAvailable);
+	GenerationContext.NumMaxLODsToStream.Add(ComponentName, NumMaxLODsToStream);
 	
 	// Add a new entry to the list of Component Infos
-	FMutableComponentInfo ComponentInfo(Node.ComponentName, RefSkeletalMesh);
+	FMutableComponentInfo ComponentInfo(ComponentName, RefSkeletalMesh);
 	ComponentInfo.Node = &Node;
 	ComponentInfo.AccumulateBonesToRemovePerLOD(Node.LODReductionSettings, Node.NumLODs);
 	ComponentInfo.MacroContext = GenerationContext.MacroNodesStack;
@@ -617,7 +590,7 @@ void FirstPass(UCustomizableObjectNodeComponentMesh& Node, FMutableGraphGenerati
 	// Make sure the Skeleton from the reference mesh is added to the list of referenced Skeletons.
 	GenerationContext.CompilationContext.ReferencedSkeletons.Add(RefSkeleton);
 
-	GenerationContext.ComponentNames.Add(Node.ComponentName);
+	GenerationContext.ComponentNames.Add(ComponentName);
 }
 
 
diff --git a/Engine/Plugins/Mutable/Source/CustomizableObjectEditor/Private/MuCOE/GenerateMutableSource/GenerateMutableSourceFloat.cpp b/Engine/Plugins/Mutable/Source/CustomizableObjectEditor/Private/MuCOE/GenerateMutableSource/GenerateMutableSourceFloat.cpp
--- a/Engine/Plugins/Mutable/Source/CustomizableObjectEditor/Private/MuCOE/GenerateMutableSource/GenerateMutableSourceFloat.cpp
+++ b/Engine/Plugins/Mutable/Source/CustomizableObjectEditor/Private/MuCOE/GenerateMutableSource/GenerateMutableSourceFloat.cpp
@@ -73,11 +73,11 @@ mu::Ptr<mu::NodeScalar> GenerateMutableSourceFloat(const UEdGraphPin* Pin, FMuta
 		mu::Ptr<mu::NodeScalarParameter> ScalarNode = new mu::NodeScalarParameter();
 		Result = ScalarNode;
 
-		ScalarNode->Name = FloatParameterNode->ParameterName;
+		ScalarNode->Name = FloatParameterNode->GetParameterName(&GenerationContext.MacroNodesStack);
 		ScalarNode->UID = GenerationContext.GetNodeIdUnique(Node).ToString();
 		ScalarNode->DefaultValue = FloatParameterNode->DefaultValue;
 
-		GenerationContext.ParameterUIDataMap.Add(FloatParameterNode->ParameterName, FMutableParameterData(
+		GenerationContext.ParameterUIDataMap.Add(FloatParameterNode->GetParameterName(&GenerationContext.MacroNodesStack), FMutableParameterData(
 			FloatParameterNode->ParamUIMetadata,
 			EMutableParameterType::Float));
 	}
@@ -90,7 +90,7 @@ mu::Ptr<mu::NodeScalar> GenerateMutableSourceFloat(const UEdGraphPin* Pin, FMuta
 
 		int32 DefaultValue = FMath::Clamp(EnumParamNode->DefaultIndex, 0, NumSelectors - 1);
 
-		EnumParameterNode->Name = EnumParamNode->ParameterName;
+		EnumParameterNode->Name = EnumParamNode->GetParameterName(&GenerationContext.MacroNodesStack);
 		EnumParameterNode->UID = GenerationContext.GetNodeIdUnique(Node).ToString();
 		EnumParameterNode->Options.SetNum(NumSelectors);
 		EnumParameterNode->DefaultValue = DefaultValue;
@@ -110,7 +110,7 @@ mu::Ptr<mu::NodeScalar> GenerateMutableSourceFloat(const UEdGraphPin* Pin, FMuta
 
 		Result = EnumParameterNode;
 
-		GenerationContext.ParameterUIDataMap.Add(EnumParamNode->ParameterName, ParameterUIData);
+		GenerationContext.ParameterUIDataMap.Add(EnumParamNode->GetParameterName(&GenerationContext.MacroNodesStack), ParameterUIData);
 	}
 
 	else if (const UCustomizableObjectNodeFloatSwitch* TypedNodeFloatSwitch = Cast<UCustomizableObjectNodeFloatSwitch>(Node))
@@ -256,7 +256,7 @@ mu::Ptr<mu::NodeScalar> GenerateMutableSourceFloat(const UEdGraphPin* Pin, FMuta
 			UEdGraphPin* VariationPin = TypedNodeFloatVar->VariationPin(VariationIndex);
 			if (!VariationPin) continue;
 
-			FloatNode->Variations[VariationIndex].Tag = GetVariationTag(*TypedNodeFloatVar, VariationIndex, GenerationContext);
+			FloatNode->Variations[VariationIndex].Tag = TypedNodeFloatVar->GetVariationTag(VariationIndex, &GenerationContext.MacroNodesStack);
 
 			if (const UEdGraphPin* ConnectedPin = FollowInputPin(*VariationPin))
 			{
diff --git a/Engine/Plugins/Mutable/Source/CustomizableObjectEditor/Private/MuCOE/GenerateMutableSource/GenerateMutableSourceGroupProjector.cpp b/Engine/Plugins/Mutable/Source/CustomizableObjectEditor/Private/MuCOE/GenerateMutableSource/GenerateMutableSourceGroupProjector.cpp
--- a/Engine/Plugins/Mutable/Source/CustomizableObjectEditor/Private/MuCOE/GenerateMutableSource/GenerateMutableSourceGroupProjector.cpp
+++ b/Engine/Plugins/Mutable/Source/CustomizableObjectEditor/Private/MuCOE/GenerateMutableSource/GenerateMutableSourceGroupProjector.cpp
@@ -348,7 +348,7 @@ TOptional<FGroupProjectorTempData> GenerateMutableGroupProjector(UCustomizableOb
 		else
 		{
 			NodeScalarParam = new mu::NodeScalarParameter;
-			FString NodeScalarParamName = ProjParamNode->ParameterName + NUM_LAYERS_PARAMETER_POSTFIX;
+			FString NodeScalarParamName = ProjParamNode->GetParameterName(&GenerationContext.MacroNodesStack) + NUM_LAYERS_PARAMETER_POSTFIX;
 			NodeScalarParam->Name = NodeScalarParamName;
 			NodeScalarParam->UID = NumLayersParamUid.ToString();
 			NodeScalarParam->DefaultValue =0.f;
@@ -373,7 +373,7 @@ TOptional<FGroupProjectorTempData> GenerateMutableGroupProjector(UCustomizableOb
 		else
 		{
 			PoseEnumParameterNode = new mu::NodeScalarEnumParameter;
-			FString PoseNodeEnumParamName = ProjParamNode->ParameterName + POSE_PARAMETER_POSTFIX;
+			FString PoseNodeEnumParamName = ProjParamNode->GetParameterName(&GenerationContext.MacroNodesStack) + POSE_PARAMETER_POSTFIX;
 			PoseEnumParameterNode->Name = PoseNodeEnumParamName;
 			PoseEnumParameterNode->UID = SelectedPoseParamUid.ToString();
 			PoseEnumParameterNode->Options.SetNum( ProjParamNode->OptionPoses.Num() + 1 );
@@ -394,7 +394,7 @@ TOptional<FGroupProjectorTempData> GenerateMutableGroupProjector(UCustomizableOb
 		else
 		{
 			OpacityParameterNode = new mu::NodeScalarParameter;
-			FString OpacityParameterNodeName = ProjParamNode->ParameterName + OPACITY_PARAMETER_POSTFIX;
+			FString OpacityParameterNodeName = ProjParamNode->GetParameterName(&GenerationContext.MacroNodesStack) + OPACITY_PARAMETER_POSTFIX;
 			OpacityParameterNode->Name = OpacityParameterNodeName;
 			OpacityParameterNode->UID = OpacityParamUid.ToString();
 			OpacityParameterNode->DefaultValue = 0.75f;
@@ -458,7 +458,7 @@ TOptional<FGroupProjectorTempData> GenerateMutableGroupProjector(UCustomizableOb
 		{
 			EnumParameterNode = new mu::NodeScalarEnumParameter;
 
-			FString NodeEnumParamName = ProjParamNode->ParameterName + IMAGE_PARAMETER_POSTFIX;
+			FString NodeEnumParamName = ProjParamNode->GetParameterName(&GenerationContext.MacroNodesStack) + IMAGE_PARAMETER_POSTFIX;
 			EnumParameterNode->Name = NodeEnumParamName;
 			EnumParameterNode->UID = SelectedImageParamUid.ToString();
 			EnumParameterNode->Options.SetNum(ArrayOptionTexture.Num());
diff --git a/Engine/Plugins/Mutable/Source/CustomizableObjectEditor/Private/MuCOE/GenerateMutableSource/GenerateMutableSourceImage.cpp b/Engine/Plugins/Mutable/Source/CustomizableObjectEditor/Private/MuCOE/GenerateMutableSource/GenerateMutableSourceImage.cpp
--- a/Engine/Plugins/Mutable/Source/CustomizableObjectEditor/Private/MuCOE/GenerateMutableSource/GenerateMutableSourceImage.cpp
+++ b/Engine/Plugins/Mutable/Source/CustomizableObjectEditor/Private/MuCOE/GenerateMutableSource/GenerateMutableSourceImage.cpp
@@ -158,7 +158,7 @@ mu::Ptr<mu::NodeImage> GenerateMutableSourceImage(const UEdGraphPin* Pin, FMutab
 	{
 		mu::Ptr<mu::NodeImageParameter> TextureNode = new mu::NodeImageParameter();
 
-		TextureNode->Name = TypedNodeParam->ParameterName;
+		TextureNode->Name = TypedNodeParam->GetParameterName(&GenerationContext.MacroNodesStack);
 		TextureNode->UID = GenerationContext.GetNodeIdUnique(Node).ToString();
 
 		if (TypedNodeParam->DefaultValue)
@@ -166,7 +166,7 @@ mu::Ptr<mu::NodeImage> GenerateMutableSourceImage(const UEdGraphPin* Pin, FMutab
 			TextureNode->DefaultValue = FName(TypedNodeParam->DefaultValue->GetPathName());
 		}
 
-		GenerationContext.ParameterUIDataMap.Add(TypedNodeParam->ParameterName, FMutableParameterData(
+		GenerationContext.ParameterUIDataMap.Add(TypedNodeParam->GetParameterName(&GenerationContext.MacroNodesStack), FMutableParameterData(
 			TypedNodeParam->ParamUIMetadata,
 			EMutableParameterType::Texture));
 
@@ -443,7 +443,7 @@ mu::Ptr<mu::NodeImage> GenerateMutableSourceImage(const UEdGraphPin* Pin, FMutab
 				continue;
 			}
 
-			TextureNode->Variations[VariationIndex].Tag = GetVariationTag(*TypedNodeImageVar, VariationIndex, GenerationContext);
+			TextureNode->Variations[VariationIndex].Tag = TypedNodeImageVar->GetVariationTag(VariationIndex, &GenerationContext.MacroNodesStack);
 
 			if (const UEdGraphPin* ConnectedPin = FollowInputPin(*VariationPin))
 			{
diff --git a/Engine/Plugins/Mutable/Source/CustomizableObjectEditor/Private/MuCOE/GenerateMutableSource/GenerateMutableSourceMesh.cpp b/Engine/Plugins/Mutable/Source/CustomizableObjectEditor/Private/MuCOE/GenerateMutableSource/GenerateMutableSourceMesh.cpp
--- a/Engine/Plugins/Mutable/Source/CustomizableObjectEditor/Private/MuCOE/GenerateMutableSource/GenerateMutableSourceMesh.cpp
+++ b/Engine/Plugins/Mutable/Source/CustomizableObjectEditor/Private/MuCOE/GenerateMutableSource/GenerateMutableSourceMesh.cpp
@@ -38,7 +38,7 @@
 #include "MuCOE/Nodes/CustomizableObjectNodeMeshParameter.h"
 #include "MuCOE/Nodes/CustomizableObjectNodeSkeletalMesh.h"
 #include "MuCOE/Nodes/CustomizableObjectNodeStaticMesh.h"
-#include "MuCOE/Nodes/CustomizableObjectNodeStringConstant.h"
+#include "MuCOE/Nodes/CustomizableObjectNodeStaticString.h"
 #include "MuCOE/Nodes/CustomizableObjectNodeTable.h"
 #include "MuCOE/Nodes/CustomizableObjectNodeTunnel.h"
 #include "MuCOE/UnrealEditorPortabilityHelpers.h"
@@ -3941,19 +3941,19 @@ mu::Ptr<mu::NodeMesh> GenerateMutableSourceMesh(const UEdGraphPin* Pin,
 
 			if (UEdGraphPin* MorphTargetNamePin = TypedNodeMorph->MorphTargetNamePin())
 			{
-				const UEdGraphPin* ConnectedStringPin = FollowInputPin(*MorphTargetNamePin);
-				const UEdGraphPin* SourceStringPin = ConnectedStringPin ? GraphTraversal::FindStringSource(*ConnectedStringPin, &GenerationContext.MacroNodesStack) : nullptr;
-
-				if (SourceStringPin)
+				if (const UEdGraphPin* ConnectedStringPin = FollowInputPin(*MorphTargetNamePin))
 				{
-					if (const UCustomizableObjectNodeStringConstant* StringNode = Cast<UCustomizableObjectNodeStringConstant>(SourceStringPin->GetOwningNode()))
+					if (const UEdGraphPin* SourceStringPin = GraphTraversal::FindStringSource(*ConnectedStringPin, &GenerationContext.MacroNodesStack))
 					{
-						MorphTargetName = StringNode->Value;
+						if (const UCustomizableObjectNodeStaticString* StringNode = Cast<UCustomizableObjectNodeStaticString>(SourceStringPin->GetOwningNode()))
+						{
+							MorphTargetName = StringNode->Value;
+						}
+					}
+					else
+					{
+						GenerationContext.Log(LOCTEXT("MorphStringNodeFailed", "Could not find a linked String node."), Node);
 					}
-				}
-				else
-				{
-					GenerationContext.Log(LOCTEXT("MorphStringNodeFailed", "Could not found a linked String node."), Node);
 				}
 			}
 
@@ -4142,7 +4142,7 @@ mu::Ptr<mu::NodeMesh> GenerateMutableSourceMesh(const UEdGraphPin* Pin,
 			const UEdGraphPin* VariationPin = TypedNodeMeshVar->VariationPin(VariationIndex);
 			if (!VariationPin) continue;
 
-			FString VariationTag = GetVariationTag(*TypedNodeMeshVar, VariationIndex, GenerationContext);
+			FString VariationTag = TypedNodeMeshVar->GetVariationTag(VariationIndex, &GenerationContext.MacroNodesStack);
 			MeshNode->Variations[VariationIndex].Tag = StringCast<ANSICHAR>(*VariationTag).Get();
 
 			if (const UEdGraphPin* ConnectedPin = FollowInputPin(*VariationPin))
diff --git a/Engine/Plugins/Mutable/Source/CustomizableObjectEditor/Private/MuCOE/GenerateMutableSource/GenerateMutableSourceModifier.cpp b/Engine/Plugins/Mutable/Source/CustomizableObjectEditor/Private/MuCOE/GenerateMutableSource/GenerateMutableSourceModifier.cpp
--- a/Engine/Plugins/Mutable/Source/CustomizableObjectEditor/Private/MuCOE/GenerateMutableSource/GenerateMutableSourceModifier.cpp
+++ b/Engine/Plugins/Mutable/Source/CustomizableObjectEditor/Private/MuCOE/GenerateMutableSource/GenerateMutableSourceModifier.cpp
@@ -31,7 +31,7 @@
 #include "MuCOE/Nodes/CustomizableObjectNodeFloatParameter.h"
 #include "MuCOE/Nodes/CustomizableObjectNodeFloatConstant.h"
 #include "MuCOE/Nodes/CustomizableObjectNodeModifierTransformInMesh.h"
-#include "MuCOE/Nodes/CustomizableObjectNodeStringConstant.h"
+#include "MuCOE/Nodes/CustomizableObjectNodeStaticString.h"
 #include "MuCOE/Nodes/CustomizableObjectNodeTunnel.h"
 #include "MuR/Mesh.h"
 #include "MuT/NodeMeshTransform.h"
@@ -47,37 +47,6 @@
 
 #define LOCTEXT_NAMESPACE "CustomizableObjectEditor"
 
-TArray<FString> GetNodeRequiredTags(const UCustomizableObjectNodeModifierBase& Node, FMutableGraphGenerationContext& GenerationContext)
-{
-	TArray<FString> OutTags = Node.RequiredTags;
-
-	if (Node.IsInMacro())
-	{
-		UEdGraphPin* TagsPin = Node.RequiredTagsPin();
-		check(TagsPin);
-
-		for (const UEdGraphPin* StringPin : FollowInputPinArray(*TagsPin))
-		{
-			const UEdGraphPin* SourceStringPin = GraphTraversal::FindStringSource(*StringPin, &GenerationContext.MacroNodesStack);
-
-			if (SourceStringPin)
-			{
-				if (const UCustomizableObjectNodeStringConstant* StringNode = Cast<UCustomizableObjectNodeStringConstant>(SourceStringPin->GetOwningNode()))
-				{
-					OutTags.Add(StringNode->Value);
-				}
-			}
-			else
-			{
-				GenerationContext.Log(LOCTEXT("ModifierNodeTagError", "Could not found a linked String node."), &Node);
-			}
-		}
-	}
-
-	return OutTags;
-}
-
-
 mu::Ptr<mu::NodeModifier> GenerateMutableSourceModifier(const UEdGraphPin * Pin, FMutableGraphGenerationContext & GenerationContext)
 {
 	check(Pin)
@@ -127,7 +96,7 @@ mu::Ptr<mu::NodeModifier> GenerateMutableSourceModifier(const UEdGraphPin * Pin,
 		ClipNode->SetVertexSelectionBone(GenerationContext.CompilationContext.GetBoneUnique(TypedNodeClip->BoneName), TypedNodeClip->MaxEffectRadius);
 
 		ClipNode->MultipleTagsPolicy = TypedNodeClip->MultipleTagPolicy;
-		ClipNode->RequiredTags = GetNodeRequiredTags(*TypedNodeClip, GenerationContext);
+		ClipNode->RequiredTags = TypedNodeClip->GetNodeRequiredTags(&GenerationContext.MacroNodesStack);
 
 		ClipNode->Parameters.FaceCullStrategy = TypedNodeClip->FaceCullStrategy;
 
@@ -181,7 +150,7 @@ mu::Ptr<mu::NodeModifier> GenerateMutableSourceModifier(const UEdGraphPin * Pin,
 		}
 	
 		ClipNode->MultipleTagsPolicy = TypedNodeClipDeform->MultipleTagPolicy;
-		ClipNode->RequiredTags = GetNodeRequiredTags(*TypedNodeClipDeform, GenerationContext);
+		ClipNode->RequiredTags = TypedNodeClipDeform->GetNodeRequiredTags(&GenerationContext.MacroNodesStack);
 
 		GenerationContext.MeshGenerationFlags.Pop();
 	}
@@ -227,7 +196,7 @@ mu::Ptr<mu::NodeModifier> GenerateMutableSourceModifier(const UEdGraphPin * Pin,
 		}
 
 		ClipNode->MultipleTagsPolicy = TypedNodeClipMesh->MultipleTagPolicy;
-		ClipNode->RequiredTags = GetNodeRequiredTags(*TypedNodeClipMesh, GenerationContext);
+		ClipNode->RequiredTags = TypedNodeClipMesh->GetNodeRequiredTags(&GenerationContext.MacroNodesStack);
 
 		GenerationContext.MeshGenerationFlags.Pop();
 	}
@@ -266,7 +235,7 @@ mu::Ptr<mu::NodeModifier> GenerateMutableSourceModifier(const UEdGraphPin * Pin,
 		ClipNode->LayoutIndex = TypedNodeClipUVMask->UVChannelForMask;
 
 		ClipNode->MultipleTagsPolicy = TypedNodeClipUVMask->MultipleTagPolicy;
-		ClipNode->RequiredTags = GetNodeRequiredTags(*TypedNodeClipUVMask, GenerationContext);
+		ClipNode->RequiredTags = TypedNodeClipUVMask->GetNodeRequiredTags(&GenerationContext.MacroNodesStack);
 
 		GenerationContext.MeshGenerationFlags.Pop();
 	}
@@ -283,7 +252,7 @@ mu::Ptr<mu::NodeModifier> GenerateMutableSourceModifier(const UEdGraphPin * Pin,
 		const int32 LOD = Node->IsAffectedByLOD() ? GenerationContext.CurrentLOD : 0;
 
 		SurfNode->MultipleTagsPolicy = TypedNodeExt->MultipleTagPolicy;
-		SurfNode->RequiredTags = GetNodeRequiredTags(*TypedNodeExt, GenerationContext);
+		SurfNode->RequiredTags = TypedNodeExt->GetNodeRequiredTags(&GenerationContext.MacroNodesStack);
 
 		// Is this enough? Should we try to narrow down with potential mesh sections modified by this?
 		int32 LODCount = GenerationContext.NumLODs[GenerationContext.CurrentMeshComponent];
@@ -395,32 +364,7 @@ mu::Ptr<mu::NodeModifier> GenerateMutableSourceModifier(const UEdGraphPin * Pin,
 			}
 		}
 
-		SurfNode->EnableTags = TypedNodeExt->Tags;
-		
-		// Add Macro linked Tags
-		if (TypedNodeExt->IsInMacro())
-		{
-			UEdGraphPin* TagsPin = TypedNodeExt->GetEnableTagsPin();
-			check(TagsPin);
-
-			for (const UEdGraphPin* StringPin : FollowInputPinArray(*TagsPin))
-			{
-				const UEdGraphPin* SourceStringPin = GraphTraversal::FindStringSource(*StringPin, &GenerationContext.MacroNodesStack);
-
-				if (SourceStringPin)
-				{
-					if (const UCustomizableObjectNodeStringConstant* StringNode = Cast<UCustomizableObjectNodeStringConstant>(SourceStringPin->GetOwningNode()))
-					{
-						SurfNode->EnableTags.Add(StringNode->Value);
-					}
-				}
-				else
-				{
-					GenerationContext.Log(LOCTEXT("ModifierNodeTagError", "Could not found a linked String node."), Node);
-				}
-			}
-		}
-
+		SurfNode->EnableTags = TypedNodeExt->GetEnableTags(&GenerationContext.MacroNodesStack);
 		SurfNode->EnableTags.AddUnique(TypedNodeExt->GetInternalTag());
 
 		GenerationContext.MeshGenerationFlags.Pop();
@@ -442,7 +386,7 @@ mu::Ptr<mu::NodeModifier> GenerateMutableSourceModifier(const UEdGraphPin * Pin,
 		Result = SurfNode;
 
 		SurfNode->MultipleTagsPolicy = TypedNodeRem->MultipleTagPolicy;
-		SurfNode->RequiredTags = GetNodeRequiredTags(*TypedNodeRem, GenerationContext);
+		SurfNode->RequiredTags = TypedNodeRem->GetNodeRequiredTags(&GenerationContext.MacroNodesStack);
 
 		if (const UEdGraphPin* ConnectedPin = FollowInputPin(*TypedNodeRem->RemoveMeshPin()))
 		{
@@ -482,7 +426,7 @@ mu::Ptr<mu::NodeModifier> GenerateMutableSourceModifier(const UEdGraphPin * Pin,
 		ClipNode->FaceCullStrategy = TypedNodeRemBlocks->FaceCullStrategy;
 
 		ClipNode->MultipleTagsPolicy = TypedNodeRemBlocks->MultipleTagPolicy;
-		ClipNode->RequiredTags = GetNodeRequiredTags(*TypedNodeRemBlocks, GenerationContext);
+		ClipNode->RequiredTags = TypedNodeRemBlocks->GetNodeRequiredTags(&GenerationContext.MacroNodesStack);
 
 		bool bWasEmpty = false;
 		mu::Ptr<mu::NodeLayout> SourceLayout = CreateMutableLayoutNode(GenerationContext, TypedNodeRemBlocks->Layout, true, bWasEmpty);
@@ -505,7 +449,7 @@ mu::Ptr<mu::NodeModifier> GenerateMutableSourceModifier(const UEdGraphPin * Pin,
 		Result = SurfNode;
 
 		SurfNode->MultipleTagsPolicy = TypedNodeEdit->MultipleTagPolicy;
-		SurfNode->RequiredTags = GetNodeRequiredTags(*TypedNodeEdit, GenerationContext);
+		SurfNode->RequiredTags = TypedNodeEdit->GetNodeRequiredTags(&GenerationContext.MacroNodesStack);
 
 		// Is this enough? Should we try to narrow down with potential mesh sections modified by this?
 		int32 LODCount = GenerationContext.NumLODs[GenerationContext.CurrentMeshComponent];
@@ -587,26 +531,26 @@ mu::Ptr<mu::NodeModifier> GenerateMutableSourceModifier(const UEdGraphPin * Pin,
 		SurfNode->bApplyBeforeNormalOperations = true;
 
 		SurfNode->MultipleTagsPolicy = TypedNodeMorph->MultipleTagPolicy;
-		SurfNode->RequiredTags = GetNodeRequiredTags(*TypedNodeMorph, GenerationContext);
+		SurfNode->RequiredTags = TypedNodeMorph->GetNodeRequiredTags(&GenerationContext.MacroNodesStack);
 
 		FString MorphTargetName = TypedNodeMorph->MorphTargetName;
 
 		if (UEdGraphPin* MorphTargetNamePin = TypedNodeMorph->MorphTargetNamePin())
 		{
-			const UEdGraphPin* ConnectedStringPin = FollowInputPin(*MorphTargetNamePin);
-			const UEdGraphPin* SourceStringPin = ConnectedStringPin ? GraphTraversal::FindStringSource(*ConnectedStringPin, &GenerationContext.MacroNodesStack) : nullptr;
-
-			if (SourceStringPin)
+			if (const UEdGraphPin* ConnectedStringPin = FollowInputPin(*MorphTargetNamePin))
 			{
-				if (const UCustomizableObjectNodeStringConstant* StringNode = Cast<UCustomizableObjectNodeStringConstant>(SourceStringPin->GetOwningNode()))
+				if (const UEdGraphPin* SourceStringPin = GraphTraversal::FindStringSource(*ConnectedStringPin, &GenerationContext.MacroNodesStack))
+				{
+					if (const UCustomizableObjectNodeStaticString* StringNode = Cast<UCustomizableObjectNodeStaticString>(SourceStringPin->GetOwningNode()))
+					{
+						MorphTargetName = StringNode->Value;
+					}
+				}
+				else
 				{
-					MorphTargetName = StringNode->Value;
+					GenerationContext.Log(LOCTEXT("ModifierNodeTagError", "Could not find a linked String node."), Node);
 				}
 			}
-			else
-			{
-				GenerationContext.Log(LOCTEXT("ModifierNodeTagError", "Could not found a linked String node."), Node);
-			}
 		}
 
 		SurfNode->MeshMorph = MorphTargetName;
@@ -695,7 +639,7 @@ mu::Ptr<mu::NodeModifier> GenerateMutableSourceModifier(const UEdGraphPin * Pin,
 		}
 
 		TransformNode->MultipleTagsPolicy = TypedNodeTransformMesh->MultipleTagPolicy;
-		TransformNode->RequiredTags = GetNodeRequiredTags(*TypedNodeTransformMesh, GenerationContext);
+		TransformNode->RequiredTags = TypedNodeTransformMesh->GetNodeRequiredTags(&GenerationContext.MacroNodesStack);
 
 		GenerationContext.MeshGenerationFlags.Pop();
 	}
diff --git a/Engine/Plugins/Mutable/Source/CustomizableObjectEditor/Private/MuCOE/GenerateMutableSource/GenerateMutableSourceProjector.cpp b/Engine/Plugins/Mutable/Source/CustomizableObjectEditor/Private/MuCOE/GenerateMutableSource/GenerateMutableSourceProjector.cpp
--- a/Engine/Plugins/Mutable/Source/CustomizableObjectEditor/Private/MuCOE/GenerateMutableSource/GenerateMutableSourceProjector.cpp
+++ b/Engine/Plugins/Mutable/Source/CustomizableObjectEditor/Private/MuCOE/GenerateMutableSource/GenerateMutableSourceProjector.cpp
@@ -83,7 +83,7 @@ mu::Ptr<mu::NodeProjector> GenerateMutableSourceProjector(const UEdGraphPin* Pin
 		mu::Ptr<mu::NodeProjectorParameter> ProjectorNode = new mu::NodeProjectorParameter();
 		Result = ProjectorNode;
 
-		ProjectorNode->SetName(TypedNodeParam->ParameterName);
+		ProjectorNode->SetName(TypedNodeParam->GetParameterName(&GenerationContext.MacroNodesStack));
 		ProjectorNode->SetUid(GenerationContext.GetNodeIdUnique(Node).ToString());
 		switch ((int)TypedNodeParam->DefaultValue.ProjectionType)
 		{
@@ -128,7 +128,7 @@ mu::Ptr<mu::NodeProjector> GenerateMutableSourceProjector(const UEdGraphPin* Pin
 			check(false);
 		}
 
-		GenerationContext.ParameterUIDataMap.Add(TypedNodeParam->ParameterName, FMutableParameterData(
+		GenerationContext.ParameterUIDataMap.Add(TypedNodeParam->GetParameterName(&GenerationContext.MacroNodesStack), FMutableParameterData(
 			TypedNodeParam->ParamUIMetadata,
 			EMutableParameterType::Projector));
 	}
diff --git a/Engine/Plugins/Mutable/Source/CustomizableObjectEditor/Private/MuCOE/GenerateMutableSource/GenerateMutableSourceSurface.cpp b/Engine/Plugins/Mutable/Source/CustomizableObjectEditor/Private/MuCOE/GenerateMutableSource/GenerateMutableSourceSurface.cpp
--- a/Engine/Plugins/Mutable/Source/CustomizableObjectEditor/Private/MuCOE/GenerateMutableSource/GenerateMutableSourceSurface.cpp
+++ b/Engine/Plugins/Mutable/Source/CustomizableObjectEditor/Private/MuCOE/GenerateMutableSource/GenerateMutableSourceSurface.cpp
@@ -42,7 +42,7 @@
 #include "MuCOE/Nodes/CustomizableObjectNodeModifierRemoveMesh.h"
 #include "MuCOE/Nodes/CustomizableObjectNodeModifierRemoveMeshBlocks.h"
 #include "MuCOE/Nodes/CustomizableObjectNodeSkeletalMesh.h"
-#include "MuCOE/Nodes/CustomizableObjectNodeStringConstant.h"
+#include "MuCOE/Nodes/CustomizableObjectNodeStaticString.h"
 #include "MuCOE/Nodes/CustomizableObjectNodeTable.h"
 #include "MuCOE/Nodes/CustomizableObjectNodeTunnel.h"
 #include "MuT/NodeImageFormat.h"
@@ -125,6 +125,7 @@ mu::Ptr<mu::NodeSurface> GenerateMutableSourceSurface(const UEdGraphPin * Pin, F
 		Result = SurfNode;
 
 		// Add to the list of surfaces that could be reused between LODs for this NodeMaterial.
+		// TODO(Max)
 		TArray<FMutableGraphGenerationContext::FSharedSurface>& SharedSurfaces = GenerationContext.SharedSurfaceIds.FindOrAdd(TypedNodeMat, {});
 		FMutableGraphGenerationContext::FSharedSurface& SharedSurface = SharedSurfaces.Add_GetRef(FMutableGraphGenerationContext::FSharedSurface(GenerationContext.CurrentLOD, SurfNode));
 		SharedSurface.bMakeUnique = !TypedNodeMat->IsReuseMaterialBetweenLODs();
@@ -803,43 +804,13 @@ mu::Ptr<mu::NodeSurface> GenerateMutableSourceSurface(const UEdGraphPin * Pin, F
 				SurfNode->Scalars[MaterialIndex].Name = MaterialName;
 			}
 		}
-		
 
-		if (const TArray<FString>* EnableTags = TypedNodeMat->GetEnableTags())
+		for (const FString& Tag : TypedNodeMat->GetEnableTags(&GenerationContext.MacroNodesStack))
 		{
-			for (const FString& Tag : *EnableTags)
-			{
-				SurfNode->Tags.AddUnique(Tag);
-			}
-
-			SurfNode->Tags.AddUnique( TypedNodeMat->GetInternalTag() );
+			SurfNode->Tags.AddUnique(Tag);
 		}
 
-		// Add Macro linked Tags
-		if (TypedNodeMat->IsInMacro())
-		{
-			UEdGraphPin* TagsPin = TypedNodeMat->GetEnableTagsPin();
-			
-			if (TagsPin)
-			{
-				for (const UEdGraphPin* StringPin : FollowInputPinArray(*TagsPin))
-				{
-					const UEdGraphPin* SourceStringPin = GraphTraversal::FindStringSource(*StringPin, &GenerationContext.MacroNodesStack);
-
-					if (SourceStringPin)
-					{
-						if (const UCustomizableObjectNodeStringConstant* StringNode = Cast<UCustomizableObjectNodeStringConstant>(SourceStringPin->GetOwningNode()))
-						{
-							SurfNode->Tags.AddUnique(StringNode->Value);
-						}
-					}
-					else
-					{
-						GenerationContext.Log(LOCTEXT("SurfaceNodeTagError", "Could not found a linked String node."), Node);
-					}
-				}
-			}
-		}
+		SurfNode->Tags.AddUnique( TypedNodeMat->GetInternalTag() );
 	}
 
 	else if (const UCustomizableObjectNodeMaterialVariation* TypedNodeVar = Cast<UCustomizableObjectNodeMaterialVariation>(Node))
@@ -880,7 +851,7 @@ mu::Ptr<mu::NodeSurface> GenerateMutableSourceSurface(const UEdGraphPin * Pin, F
 
 			if (UEdGraphPin* VariationPin = TypedNodeVar->VariationPin(VariationIndex))
 			{
-				SurfNode->Variations[VariationIndex].Tag = GetVariationTag(*TypedNodeVar, VariationIndex, GenerationContext);
+				SurfNode->Variations[VariationIndex].Tag = TypedNodeVar->GetVariationTag(VariationIndex, &GenerationContext.MacroNodesStack);
 				for (const UEdGraphPin* ConnectedPin : FollowInputPinArray(*VariationPin))
 				{
 					// Is it a modifier?
diff --git a/Engine/Plugins/Mutable/Source/CustomizableObjectEditor/Private/MuCOE/GenerateMutableSource/GenerateMutableSourceTransform.cpp b/Engine/Plugins/Mutable/Source/CustomizableObjectEditor/Private/MuCOE/GenerateMutableSource/GenerateMutableSourceTransform.cpp
--- a/Engine/Plugins/Mutable/Source/CustomizableObjectEditor/Private/MuCOE/GenerateMutableSource/GenerateMutableSourceTransform.cpp
+++ b/Engine/Plugins/Mutable/Source/CustomizableObjectEditor/Private/MuCOE/GenerateMutableSource/GenerateMutableSourceTransform.cpp
@@ -44,11 +44,11 @@ mu::Ptr<mu::NodeMatrix> GenerateMutableSourceTransform(const UEdGraphPin* Pin, F
 		mu::Ptr<mu::NodeMatrixParameter> MatrixNode = new mu::NodeMatrixParameter();
 		Result = MatrixNode;
 		
-		MatrixNode->Name = TypedNodeTransformParameter->ParameterName;
+		MatrixNode->Name = TypedNodeTransformParameter->GetParameterName(&GenerationContext.MacroNodesStack);
 		MatrixNode->Uid = GenerationContext.GetNodeIdUnique(Node).ToString();
 		MatrixNode->DefaultValue = FMatrix44f(TypedNodeTransformParameter->DefaultValue.ToMatrixWithScale());
 
-		GenerationContext.ParameterUIDataMap.Add(TypedNodeTransformParameter->ParameterName, FMutableParameterData(
+		GenerationContext.ParameterUIDataMap.Add(TypedNodeTransformParameter->GetParameterName(&GenerationContext.MacroNodesStack), FMutableParameterData(
 			TypedNodeTransformParameter->ParamUIMetadata,
 			EMutableParameterType::Transform));
 	}
diff --git a/Engine/Plugins/Mutable/Source/CustomizableObjectEditor/Private/MuCOE/GraphTraversal.cpp b/Engine/Plugins/Mutable/Source/CustomizableObjectEditor/Private/MuCOE/GraphTraversal.cpp
--- a/Engine/Plugins/Mutable/Source/CustomizableObjectEditor/Private/MuCOE/GraphTraversal.cpp
+++ b/Engine/Plugins/Mutable/Source/CustomizableObjectEditor/Private/MuCOE/GraphTraversal.cpp
@@ -23,7 +23,7 @@
 #include "MuCOE/Nodes/CustomizableObjectNodeObject.h"
 #include "MuCOE/Nodes/CustomizableObjectNodeSkeletalMesh.h"
 #include "MuCOE/Nodes/CustomizableObjectNodeStaticMesh.h"
-#include "MuCOE/Nodes/CustomizableObjectNodeStringConstant.h"
+#include "MuCOE/Nodes/CustomizableObjectNodeStaticString.h"
 #include "MuCOE/Nodes/CustomizableObjectNodeTable.h"
 #include "MuCOE/Nodes/CustomizableObjectNodeTunnel.h"
 #include "MuCOE/Nodes/CustomizableObjectNodeAnimationPose.h"
@@ -468,14 +468,15 @@ void RecursiveVisitNodes(UCustomizableObjectNode& CurrentNode, const TFunction<v
 					}
 
 					// Going to the node linked to the OutputPin
-					const UEdGraphPin* OutputLinkedPin = FollowInputPin(*OutputPin);
-
-					// We go directly to the node linked to the output node to ensure that is linked and we need to visit it.
-					if (UCustomizableObjectNode* OutputLinkedNode = Cast<UCustomizableObjectNode>(OutputLinkedPin->GetOwningNode()))
+					if (const UEdGraphPin* OutputLinkedPin = FollowInputPin(*OutputPin))
 					{
-						MacroContext.Push(MacroInstanceNode);
-						RecursiveVisitNodes(*OutputLinkedNode, VisitFunction, MacroContext, VisitedNodes, ObjectGroupMap);
-						MacroContext.Pop();
+						// We go directly to the node linked to the output node to ensure that is linked and we need to visit it.
+						if (UCustomizableObjectNode* OutputLinkedNode = Cast<UCustomizableObjectNode>(OutputLinkedPin->GetOwningNode()))
+						{
+							MacroContext.Push(MacroInstanceNode);
+							RecursiveVisitNodes(*OutputLinkedNode, VisitFunction, MacroContext, VisitedNodes, ObjectGroupMap);
+							MacroContext.Pop();
+						}
 					}
 				}
 			}
@@ -502,11 +503,12 @@ void RecursiveVisitNodes(UCustomizableObjectNode& CurrentNode, const TFunction<v
 						if (const UEdGraphPin* InputPin = ParentMacroInstanceNode->FindPin(ConnectedPin->PinName, EEdGraphPinDirection::EGPD_Input))
 						{
 							// Going to the node linked to the Macro Instance
-							const UEdGraphPin* InputLinkedPin = FollowInputPin(*InputPin);
-
-							if (UCustomizableObjectNode* InputLinkedNode = Cast<UCustomizableObjectNode>(InputLinkedPin->GetOwningNode()))
+							if (const UEdGraphPin* InputLinkedPin = FollowInputPin(*InputPin))
 							{
-								RecursiveVisitNodes(*InputLinkedNode, VisitFunction, MacroContext, VisitedNodes, ObjectGroupMap);
+								if (UCustomizableObjectNode* InputLinkedNode = Cast<UCustomizableObjectNode>(InputLinkedPin->GetOwningNode()))
+								{
+									RecursiveVisitNodes(*InputLinkedNode, VisitFunction, MacroContext, VisitedNodes, ObjectGroupMap);
+								}
 							}
 						}
 					}
@@ -917,7 +919,7 @@ const UEdGraphPin* GraphTraversal::FindStringSource(const UEdGraphPin& Pin, TArr
 	const UEdGraphNode* Node = Pin.GetOwningNode();
 	check(Node);
 
-	if (Cast<UCustomizableObjectNodeStringConstant>(Node))
+	if (Cast<UCustomizableObjectNodeStaticString>(Node))
 	{
 		return &Pin;
 	}
@@ -947,24 +949,26 @@ const UEdGraphPin* GraphTraversal::FindStringSource(const UEdGraphPin& Pin, TArr
 	else if (const UCustomizableObjectNodeTunnel* NodeTunnel = Cast<UCustomizableObjectNodeTunnel>(Node))
 	{
 		check(NodeTunnel->bIsInputNode);
-		check(MacroContext->Num());
-
-		const UCustomizableObjectNodeMacroInstance* MacroInstanceNode = MacroContext->Pop();
 
-		if (MacroInstanceNode)
+		if (MacroContext && MacroContext->Num())
 		{
-			if (const UEdGraphPin* InputPin = MacroInstanceNode->FindPin(Pin.PinName, EEdGraphPinDirection::EGPD_Input))
+			const UCustomizableObjectNodeMacroInstance* MacroInstanceNode = MacroContext->Pop();
+
+			if (MacroInstanceNode)
 			{
-				if (const UEdGraphPin* ConnectedPin = FollowInputPin(*InputPin))
+				if (const UEdGraphPin* InputPin = MacroInstanceNode->FindPin(Pin.PinName, EEdGraphPinDirection::EGPD_Input))
 				{
-					ReturnPin = GraphTraversal::FindStringSource(*ConnectedPin, MacroContext);
+					if (const UEdGraphPin* ConnectedPin = FollowInputPin(*InputPin))
+					{
+						ReturnPin = GraphTraversal::FindStringSource(*ConnectedPin, MacroContext);
+					}
 				}
 			}
-		}
 
-		MacroContext->Push(MacroInstanceNode);
+			MacroContext->Push(MacroInstanceNode);
 
-		return ReturnPin;
+			return ReturnPin;
+		}
 	}
 
 	else
diff --git a/Engine/Plugins/Mutable/Source/CustomizableObjectEditor/Private/MuCOE/Nodes/CustomizableObjectNode.cpp b/Engine/Plugins/Mutable/Source/CustomizableObjectEditor/Private/MuCOE/Nodes/CustomizableObjectNode.cpp
--- a/Engine/Plugins/Mutable/Source/CustomizableObjectEditor/Private/MuCOE/Nodes/CustomizableObjectNode.cpp
+++ b/Engine/Plugins/Mutable/Source/CustomizableObjectEditor/Private/MuCOE/Nodes/CustomizableObjectNode.cpp
@@ -447,9 +447,15 @@ bool UCustomizableObjectNode::IsAffectedByLOD() const
 }
 
 
-TArray<FString>* UCustomizableObjectNode::GetEnableTags()
+TArray<FString> UCustomizableObjectNode::GetEnableTags(TArray<const UCustomizableObjectNodeMacroInstance*>* MacroContext)
 { 
-	return nullptr; 
+	return TArray<FString>(); 
+}
+
+
+TArray<FString>* UCustomizableObjectNode::GetEnableTagsArray()
+{
+	return nullptr;
 }
 
 
@@ -796,4 +802,10 @@ bool UCustomizableObjectNode::IsInMacro() const
 }
 
 
+bool UCustomizableObjectNode::CreateStaticStringPinWidget() const
+{
+	return true;
+}
+
+
 #undef LOCTEXT_NAMESPACE
diff --git a/Engine/Plugins/Mutable/Source/CustomizableObjectEditor/Private/MuCOE/Nodes/CustomizableObjectNodeColorParameter.cpp b/Engine/Plugins/Mutable/Source/CustomizableObjectEditor/Private/MuCOE/Nodes/CustomizableObjectNodeColorParameter.cpp
--- a/Engine/Plugins/Mutable/Source/CustomizableObjectEditor/Private/MuCOE/Nodes/CustomizableObjectNodeColorParameter.cpp
+++ b/Engine/Plugins/Mutable/Source/CustomizableObjectEditor/Private/MuCOE/Nodes/CustomizableObjectNodeColorParameter.cpp
@@ -4,6 +4,7 @@
 
 #include "MuCOE/EdGraphSchema_CustomizableObject.h"
 #include "MuCO/CustomizableObjectCustomVersion.h"
+#include "MuCOE/CustomizableObjectMacroLibrary/CustomizableObjectMacroLibrary.h"
 
 class UCustomizableObjectNodeRemapPins;
 
@@ -25,54 +26,9 @@ void UCustomizableObjectNodeColorParameter::BackwardsCompatibleFixup(int32 Custo
 }
 
 
-void UCustomizableObjectNodeColorParameter::AllocateDefaultPins(UCustomizableObjectNodeRemapPins* RemapPins)
+FName UCustomizableObjectNodeColorParameter::GetCategory() const
 {
-	const FName Type = UEdGraphSchema_CustomizableObject::PC_Color;
-	const FName PinName = UEdGraphSchema_CustomizableObject::GetPinCategoryName(Type);
-	const FText PinFriendlyName = UEdGraphSchema_CustomizableObject::GetPinCategoryFriendlyName(Type);
-	
-	UEdGraphPin* ValuePin = CustomCreatePin(EGPD_Output, Type, PinName);
-	ValuePin->PinFriendlyName = PinFriendlyName;
-	ValuePin->bDefaultValueIsIgnored = true;
-}
-
-
-FText UCustomizableObjectNodeColorParameter::GetNodeTitle(ENodeTitleType::Type TitleType) const
-{
-	if (TitleType == ENodeTitleType::ListView)
-	{
-		return LOCTEXT("Color_Parameter", "Color Parameter");
-	}
-	else if (TitleType == ENodeTitleType::EditableTitle)
-	{
-		return FText::Format(LOCTEXT("Color_Parameter_EditableTitle", "{0}"), FText::FromString(ParameterName));
-	}
-	else
-	{
-		return FText::Format(LOCTEXT("Color_Parameter_Title", "{0}\nColor Parameter"), FText::FromString(ParameterName));
-	}
-}
-
-
-FLinearColor UCustomizableObjectNodeColorParameter::GetNodeTitleColor() const
-{
-	const UEdGraphSchema_CustomizableObject* Schema = GetDefault<UEdGraphSchema_CustomizableObject>();
-	return Schema->GetPinTypeColor(Schema->PC_Color);
-}
-
-
-FText UCustomizableObjectNodeColorParameter::GetTooltipText() const
-{
-	return LOCTEXT("Color_Parameter_Tooltip", "Expose a runtime modifiable color parameter from the Customizable Object.");
-}
-
-
-void UCustomizableObjectNodeColorParameter::OnRenameNode(const FString& NewName)
-{
-	if (!NewName.IsEmpty())
-	{
-		ParameterName = NewName;
-	}
+	return UEdGraphSchema_CustomizableObject::PC_Color;
 }
 
 #undef LOCTEXT_NAMESPACE
diff --git a/Engine/Plugins/Mutable/Source/CustomizableObjectEditor/Private/MuCOE/Nodes/CustomizableObjectNodeColorParameter.h b/Engine/Plugins/Mutable/Source/CustomizableObjectEditor/Private/MuCOE/Nodes/CustomizableObjectNodeColorParameter.h
--- a/Engine/Plugins/Mutable/Source/CustomizableObjectEditor/Private/MuCOE/Nodes/CustomizableObjectNodeColorParameter.h
+++ b/Engine/Plugins/Mutable/Source/CustomizableObjectEditor/Private/MuCOE/Nodes/CustomizableObjectNodeColorParameter.h
@@ -2,7 +2,7 @@
 
 #pragma once
 
-#include "MuCOE/Nodes/CustomizableObjectNode.h"
+#include "MuCOE/Nodes/CustomizableObjectNodeParameter.h"
 
 #include "CustomizableObjectNodeColorParameter.generated.h"
 
@@ -13,7 +13,7 @@ class UObject;
 
 
 UCLASS()
-class CUSTOMIZABLEOBJECTEDITOR_API UCustomizableObjectNodeColorParameter : public UCustomizableObjectNode
+class CUSTOMIZABLEOBJECTEDITOR_API UCustomizableObjectNodeColorParameter : public UCustomizableObjectNodeParameter
 {
 public:
 	GENERATED_BODY()
@@ -21,22 +21,10 @@ public:
 	UPROPERTY(EditAnywhere, Category=CustomizableObject, meta = (DontUpdateWhileEditing))
 	FLinearColor DefaultValue = FLinearColor(1, 1, 1, 1);
 
-	UPROPERTY(EditAnywhere, Category=CustomizableObject)
-	FString ParameterName = "Default Name";
-
-	UPROPERTY(EditAnywhere, Category = UI, meta = (DisplayName = "Parameter UI Metadata"))
-	FMutableParamUIMetadata ParamUIMetadata;
-
 	// Begin EdGraphNode interface
-	virtual FText GetNodeTitle(ENodeTitleType::Type TitleType) const override;
-	virtual FLinearColor GetNodeTitleColor() const override;
-	virtual FText GetTooltipText() const override;
-	virtual void OnRenameNode(const FString& NewName) override;
-	virtual bool GetCanRenameNode() const override { return true; }
 	virtual void BackwardsCompatibleFixup(int32 CustomizableObjectCustomVersion) override;
 
-	// UCustomizableObjectNode interface
-	virtual void AllocateDefaultPins(UCustomizableObjectNodeRemapPins* RemapPins) override;
-	virtual bool IsAffectedByLOD() const override { return false; }
+	// CustomizableObjectNodeParameter interface
+	virtual FName GetCategory() const override;
 };
 
diff --git a/Engine/Plugins/Mutable/Source/CustomizableObjectEditor/Private/MuCOE/Nodes/CustomizableObjectNodeComponent.cpp b/Engine/Plugins/Mutable/Source/CustomizableObjectEditor/Private/MuCOE/Nodes/CustomizableObjectNodeComponent.cpp
--- a/Engine/Plugins/Mutable/Source/CustomizableObjectEditor/Private/MuCOE/Nodes/CustomizableObjectNodeComponent.cpp
+++ b/Engine/Plugins/Mutable/Source/CustomizableObjectEditor/Private/MuCOE/Nodes/CustomizableObjectNodeComponent.cpp
@@ -3,14 +3,18 @@
 #include "MuCOE/Nodes/CustomizableObjectNodeComponent.h"
 
 #include "MuCOE/EdGraphSchema_CustomizableObject.h"
+#include "MuCOE/GraphTraversal.h"
+#include "MuCOE/Nodes/CustomizableObjectNodeMacroInstance.h"
+#include "MuCOE/Nodes/CustomizableObjectNodeStaticString.h"
+
 
 void UCustomizableObjectNodeComponent::AllocateDefaultPins(UCustomizableObjectNodeRemapPins* RemapPins)
 {
 	OutputPin = CustomCreatePin(EGPD_Output, UEdGraphSchema_CustomizableObject::PC_Component, FName(TEXT("Component")));
 
-	if (IsInMacro())
+	if (CVarEnableMutableMacros.GetValueOnAnyThread())
 	{
-		ComponentNamePin = CustomCreatePin(EGPD_Input, UEdGraphSchema_CustomizableObject::PC_String, FName("Component Name"));
+		ComponentNamePin = CustomCreatePin(EGPD_Input, UEdGraphSchema_CustomizableObject::PC_String, FName("Name"));
 	}
 }
 
@@ -25,3 +29,72 @@ bool UCustomizableObjectNodeComponent::IsAffectedByLOD() const
 { 
 	return false; 
 }
+
+
+bool UCustomizableObjectNodeComponent::GetCanRenameNode() const
+{
+	const UEdGraphPin* NamePin = GetComponentNamePin();
+
+	if (NamePin && NamePin->LinkedTo.Num())
+	{
+		return false;
+	}
+
+	return true;
+}
+
+
+void UCustomizableObjectNodeComponent::OnRenameNode(const FString& NewName)
+{
+	if (!NewName.IsEmpty())
+	{
+		SetComponentName(FName(*NewName));
+	}
+}
+
+
+void UCustomizableObjectNodeComponent::PinConnectionListChanged(UEdGraphPin* Pin)
+{
+	if (Pin == GetComponentNamePin())
+	{
+		GetGraph()->NotifyGraphChanged();
+	}
+}
+
+
+FName UCustomizableObjectNodeComponent::GetComponentName(TArray<const UCustomizableObjectNodeMacroInstance*>* MacroContext) const
+{
+	const UEdGraphPin* NamePin = GetComponentNamePin();
+
+	if (NamePin)
+	{
+		if (const UEdGraphPin* LinkedPin = FollowInputPin(*NamePin))
+		{
+			const UEdGraphPin* StringPin = GraphTraversal::FindStringSource(*LinkedPin, MacroContext);
+
+			if (const UCustomizableObjectNodeStaticString* StringNode = StringPin ? Cast<UCustomizableObjectNodeStaticString>(StringPin->GetOwningNode()) : nullptr)
+			{
+				return FName(StringNode->Value);
+			}
+		}
+	}
+
+	return ComponentName;
+}
+
+
+void UCustomizableObjectNodeComponent::SetComponentName(const FName& InComponentName)
+{
+	ComponentName = InComponentName;
+}
+
+
+UEdGraphPin* UCustomizableObjectNodeComponent::GetComponentNamePin() const
+{
+	if (CVarEnableMutableMacros.GetValueOnAnyThread())
+	{
+		return ComponentNamePin.Get();
+	}
+
+	return nullptr;
+}
diff --git a/Engine/Plugins/Mutable/Source/CustomizableObjectEditor/Private/MuCOE/Nodes/CustomizableObjectNodeComponent.h b/Engine/Plugins/Mutable/Source/CustomizableObjectEditor/Private/MuCOE/Nodes/CustomizableObjectNodeComponent.h
--- a/Engine/Plugins/Mutable/Source/CustomizableObjectEditor/Private/MuCOE/Nodes/CustomizableObjectNodeComponent.h
+++ b/Engine/Plugins/Mutable/Source/CustomizableObjectEditor/Private/MuCOE/Nodes/CustomizableObjectNodeComponent.h
@@ -6,6 +6,7 @@
 
 #include "CustomizableObjectNodeComponent.generated.h"
 
+class UCustomizableObjectNodeMacroInstance;
 class UEdGraphPin;
 
 
@@ -18,18 +19,30 @@ public:
 	// UEdGraphNode interface
 	virtual void AllocateDefaultPins(UCustomizableObjectNodeRemapPins* RemapPins) override;
 	virtual FLinearColor GetNodeTitleColor() const override;
+	virtual void PinConnectionListChanged(UEdGraphPin* Pin) override;
+	virtual bool GetCanRenameNode() const override;
+	virtual void OnRenameNode(const FString& NewName) override;
 
 	// UCustomizableObjectNode interface
 	virtual bool IsAffectedByLOD() const override;
 
-public:
+	//Own Interface
+	FName GetComponentName(TArray<const UCustomizableObjectNodeMacroInstance*>* MacroContext = nullptr) const;
+	void SetComponentName(const FName& InComponentName);
+	UEdGraphPin* GetComponentNamePin() const;
+
+protected:
 
-	UPROPERTY(EditAnywhere, Category = PassthroughMesh)
+	UPROPERTY(EditAnywhere, Category = ComponentMesh)
 	FName ComponentName = "Default name";
 
+public:
+
 	UPROPERTY()
 	FEdGraphPinReference OutputPin;
 
+private:
+	
 	UPROPERTY()
 	FEdGraphPinReference ComponentNamePin;
 
diff --git a/Engine/Plugins/Mutable/Source/CustomizableObjectEditor/Private/MuCOE/Nodes/CustomizableObjectNodeComponentMesh.cpp b/Engine/Plugins/Mutable/Source/CustomizableObjectEditor/Private/MuCOE/Nodes/CustomizableObjectNodeComponentMesh.cpp
--- a/Engine/Plugins/Mutable/Source/CustomizableObjectEditor/Private/MuCOE/Nodes/CustomizableObjectNodeComponentMesh.cpp
+++ b/Engine/Plugins/Mutable/Source/CustomizableObjectEditor/Private/MuCOE/Nodes/CustomizableObjectNodeComponentMesh.cpp
@@ -86,7 +86,7 @@ FText UCustomizableObjectNodeComponentMesh::GetNodeTitle(ENodeTitleType::Type Ti
 	FFormatNamedArguments Args;
 	Args.Add(TEXT("ComponentName"), FText::FromName(ComponentName));
 
-	if (TitleType == ENodeTitleType::ListView)
+	if (TitleType == ENodeTitleType::ListView || (GetComponentNamePin() && GetComponentNamePin()->LinkedTo.Num()))
 	{
 		NodeTitle = LOCTEXT("Component_Mesh", "Mesh Component");
 	}
@@ -103,15 +103,6 @@ FText UCustomizableObjectNodeComponentMesh::GetNodeTitle(ENodeTitleType::Type Ti
 }
 
 
-void UCustomizableObjectNodeComponentMesh::OnRenameNode(const FString& NewName)
-{
-	if (!NewName.IsEmpty())
-	{
-		ComponentName = FName(*NewName);
-	}
-}
-
-
 TSoftObjectPtr<UMaterialInterface> UCustomizableObjectNodeComponentMesh::GetOverlayMaterial() const
 {
 	return OverlayMaterial;
diff --git a/Engine/Plugins/Mutable/Source/CustomizableObjectEditor/Private/MuCOE/Nodes/CustomizableObjectNodeComponentMesh.h b/Engine/Plugins/Mutable/Source/CustomizableObjectEditor/Private/MuCOE/Nodes/CustomizableObjectNodeComponentMesh.h
--- a/Engine/Plugins/Mutable/Source/CustomizableObjectEditor/Private/MuCOE/Nodes/CustomizableObjectNodeComponentMesh.h
+++ b/Engine/Plugins/Mutable/Source/CustomizableObjectEditor/Private/MuCOE/Nodes/CustomizableObjectNodeComponentMesh.h
@@ -2,6 +2,7 @@
 
 #pragma once
 
+#include "CustomizableObjectNodeComponent.h"
 #include "CustomizableObjectNodeComponentMeshBase.h"
 #include "MuR/System.h"
 
@@ -60,7 +61,6 @@ public:
 	// UEdGraphNode interface
 	virtual void AllocateDefaultPins(UCustomizableObjectNodeRemapPins* RemapPins) override;
 	virtual FText GetNodeTitle(ENodeTitleType::Type TitleType) const override;
-	virtual void OnRenameNode(const FString& NewName) override;
 
 	//UCustomizableObjectNode interface
 	bool IsSingleOutputNode() const;
diff --git a/Engine/Plugins/Mutable/Source/CustomizableObjectEditor/Private/MuCOE/Nodes/CustomizableObjectNodeComponentMeshAddTo.cpp b/Engine/Plugins/Mutable/Source/CustomizableObjectEditor/Private/MuCOE/Nodes/CustomizableObjectNodeComponentMeshAddTo.cpp
--- a/Engine/Plugins/Mutable/Source/CustomizableObjectEditor/Private/MuCOE/Nodes/CustomizableObjectNodeComponentMeshAddTo.cpp
+++ b/Engine/Plugins/Mutable/Source/CustomizableObjectEditor/Private/MuCOE/Nodes/CustomizableObjectNodeComponentMeshAddTo.cpp
@@ -4,7 +4,10 @@
 
 #include "MuCO/CustomizableObjectCustomVersion.h"
 #include "MuCOE/EdGraphSchema_CustomizableObject.h"
+#include "MuCOE/GraphTraversal.h"
 #include "MuCOE/Nodes/CustomizableObjectNodeComponentMesh.h"
+#include "MuCOE/Nodes/CustomizableObjectNodeMacroInstance.h"
+#include "MuCOE/Nodes/CustomizableObjectNodeStaticString.h"
 
 
 #define LOCTEXT_NAMESPACE "CustomizableObjectEditor"
@@ -22,9 +25,9 @@ void UCustomizableObjectNodeComponentMeshAddTo::AllocateDefaultPins(UCustomizabl
 		LODPins.Add(Pin);
 	}
 
-	if (IsInMacro())
+	if (CVarEnableMutableMacros.GetValueOnAnyThread())
 	{
-		ParentComponentNamePinRef = CustomCreatePin(EGPD_Input, UEdGraphSchema_CustomizableObject::PC_String, FName("Component Name"));
+		ParentComponentNamePin = CustomCreatePin(EGPD_Input, UEdGraphSchema_CustomizableObject::PC_String, FName("Component Name"));
 	}
 
 	OutputPin = CustomCreatePin(EGPD_Output, UEdGraphSchema_CustomizableObject::PC_Component, FName(TEXT("Component")));
@@ -64,11 +67,12 @@ bool UCustomizableObjectNodeComponentMeshAddTo::IsSingleOutputNode() const
 
 FText UCustomizableObjectNodeComponentMeshAddTo::GetNodeTitle(ENodeTitleType::Type TitleType) const
 {
+	const UEdGraphPin* ParentNamePin = GetParentComponentNamePin();
 	FText NodeTitle;
 	FFormatNamedArguments Args;
 	Args.Add(TEXT("ComponentName"), FText::FromName(ParentComponentName));
 
-	if (TitleType == ENodeTitleType::ListView)
+	if (TitleType == ENodeTitleType::ListView || (ParentNamePin && ParentNamePin->LinkedTo.Num()))
 	{
 		NodeTitle = LOCTEXT("ComponentMeshAdd", "Add To Mesh Component");
 	}
@@ -100,6 +104,15 @@ FLinearColor UCustomizableObjectNodeComponentMeshAddTo::GetNodeTitleColor() cons
 }
 
 
+void UCustomizableObjectNodeComponentMeshAddTo::PinConnectionListChanged(UEdGraphPin* Pin)
+{
+	if (Pin == GetParentComponentNamePin())
+	{
+		GetGraph()->NotifyGraphChanged();
+	}
+}
+
+
 int32 UCustomizableObjectNodeComponentMeshAddTo::GetNumLODs()
 {
 	return NumLODs;
@@ -136,5 +149,41 @@ const UCustomizableObjectNode* UCustomizableObjectNodeComponentMeshAddTo::GetOwn
 }
 
 
+FName UCustomizableObjectNodeComponentMeshAddTo::GetParentComponentName(TArray<const UCustomizableObjectNodeMacroInstance*>* MacroContext) const
+{
+	const UEdGraphPin* ParentNamePin = GetParentComponentNamePin();
+	if (ParentNamePin)
+	{
+		if (const UEdGraphPin* LinkedPin = FollowInputPin(*ParentNamePin))
+		{
+			const UEdGraphPin* StringPin = GraphTraversal::FindStringSource(*LinkedPin, MacroContext);
+
+			if (const UCustomizableObjectNodeStaticString* StringNode = StringPin ? Cast<UCustomizableObjectNodeStaticString>(StringPin->GetOwningNode()) : nullptr)
+			{
+				return FName(StringNode->Value);
+			}
+		}
+	}
+
+	return ParentComponentName;
+}
+
+
+void UCustomizableObjectNodeComponentMeshAddTo::SetParentComponentName(const FName& InComponentName)
+{
+	ParentComponentName = InComponentName;
+}
+
+
+UEdGraphPin* UCustomizableObjectNodeComponentMeshAddTo::GetParentComponentNamePin() const
+{
+	if (CVarEnableMutableMacros.GetValueOnAnyThread())
+	{
+		return ParentComponentNamePin.Get();
+	}
+
+	return nullptr;
+}
+
 #undef LOCTEXT_NAMESPACE
 
diff --git a/Engine/Plugins/Mutable/Source/CustomizableObjectEditor/Private/MuCOE/Nodes/CustomizableObjectNodeComponentMeshAddTo.h b/Engine/Plugins/Mutable/Source/CustomizableObjectEditor/Private/MuCOE/Nodes/CustomizableObjectNodeComponentMeshAddTo.h
--- a/Engine/Plugins/Mutable/Source/CustomizableObjectEditor/Private/MuCOE/Nodes/CustomizableObjectNodeComponentMeshAddTo.h
+++ b/Engine/Plugins/Mutable/Source/CustomizableObjectEditor/Private/MuCOE/Nodes/CustomizableObjectNodeComponentMeshAddTo.h
@@ -6,6 +6,8 @@
 
 #include "CustomizableObjectNodeComponentMeshAddTo.generated.h"
 
+class UCustomizableObjectNodeMacroInstance;
+
 UCLASS()
 class CUSTOMIZABLEOBJECTEDITOR_API UCustomizableObjectNodeComponentMeshAddTo : public UCustomizableObjectNode, public ICustomizableObjectNodeComponentMeshInterface
 {
@@ -19,6 +21,7 @@ public:
 	virtual FText GetNodeTitle(ENodeTitleType::Type TitleType) const override;
 	virtual void OnRenameNode(const FString& NewName) override;
 	virtual FLinearColor GetNodeTitleColor() const override;
+	virtual void PinConnectionListChanged(UEdGraphPin* Pin) override;
 
 	// UObject interface
 	virtual void PostEditChangeProperty(FPropertyChangedEvent& PropertyChangedEvent) override;
@@ -35,11 +38,18 @@ public:
 	virtual void SetOutputPin(const UEdGraphPin* Pin);
 	virtual const UCustomizableObjectNode* GetOwningNode() const;
 
-public:
+	// Own interface
+	FName GetParentComponentName(TArray<const UCustomizableObjectNodeMacroInstance*>* MacroContext = nullptr) const;
+	void SetParentComponentName(const FName& InComponentName);
+	UEdGraphPin* GetParentComponentNamePin() const;
+
+private:
 
 	UPROPERTY(EditAnywhere, Category = ComponentMesh)
 	FName ParentComponentName;
 
+public:
+
 	UPROPERTY(EditAnywhere, Category = ComponentMesh)
 	int32 NumLODs = 1;
 
@@ -51,8 +61,10 @@ public:
 
 	UPROPERTY()
 	FEdGraphPinReference OutputPin;
+
+private:
 	
 	UPROPERTY()
-	FEdGraphPinReference ParentComponentNamePinRef;
+	FEdGraphPinReference ParentComponentNamePin;
 
 };
\ No newline at end of file
diff --git a/Engine/Plugins/Mutable/Source/CustomizableObjectEditor/Private/MuCOE/Nodes/CustomizableObjectNodeComponentMeshBase.h b/Engine/Plugins/Mutable/Source/CustomizableObjectEditor/Private/MuCOE/Nodes/CustomizableObjectNodeComponentMeshBase.h
--- a/Engine/Plugins/Mutable/Source/CustomizableObjectEditor/Private/MuCOE/Nodes/CustomizableObjectNodeComponentMeshBase.h
+++ b/Engine/Plugins/Mutable/Source/CustomizableObjectEditor/Private/MuCOE/Nodes/CustomizableObjectNodeComponentMeshBase.h
@@ -2,7 +2,7 @@
 
 #pragma once
 
-#include "MuCOE/Nodes/CustomizableObjectNodeComponent.h"
+#include "MuCOE/Nodes/CustomizableObjectNode.h"
 
 #include "CustomizableObjectNodeComponentMeshBase.generated.h"
 
diff --git a/Engine/Plugins/Mutable/Source/CustomizableObjectEditor/Private/MuCOE/Nodes/CustomizableObjectNodeComponentPassthroughMesh.cpp b/Engine/Plugins/Mutable/Source/CustomizableObjectEditor/Private/MuCOE/Nodes/CustomizableObjectNodeComponentPassthroughMesh.cpp
--- a/Engine/Plugins/Mutable/Source/CustomizableObjectEditor/Private/MuCOE/Nodes/CustomizableObjectNodeComponentPassthroughMesh.cpp
+++ b/Engine/Plugins/Mutable/Source/CustomizableObjectEditor/Private/MuCOE/Nodes/CustomizableObjectNodeComponentPassthroughMesh.cpp
@@ -149,7 +149,7 @@ FText UCustomizableObjectNodeComponentPassthroughMesh::GetNodeTitle(ENodeTitleTy
 	FFormatNamedArguments Args;
 	Args.Add(TEXT("ComponentName"), FText::FromName(ComponentName));
 
-	if (TitleType == ENodeTitleType::ListView)
+	if (TitleType == ENodeTitleType::ListView || (GetComponentNamePin() && GetComponentNamePin()->LinkedTo.Num()))
 	{
 		NodeTitle = LOCTEXT("ComponentPassthroughMesh", "Passthrough Mesh Component");
 	}
@@ -211,15 +211,6 @@ FText UCustomizableObjectNodeComponentPassthroughMesh::GetTooltipText() const
 }
 
 
-void UCustomizableObjectNodeComponentPassthroughMesh::OnRenameNode(const FString& NewName)
-{
-	if (!NewName.IsEmpty())
-	{
-		ComponentName = FName(*NewName);
-	}
-}
-
-
 UMaterialInterface* UCustomizableObjectNodeComponentPassthroughMesh::GetMaterialInterfaceFor(const int32 LODIndex, const int32 SectionIndex, const FSkeletalMeshModel* ImportedModel) const
 {
 	if (FSkeletalMaterial* SkeletalMaterial = GetSkeletalMaterialFor(LODIndex, SectionIndex, ImportedModel))
diff --git a/Engine/Plugins/Mutable/Source/CustomizableObjectEditor/Private/MuCOE/Nodes/CustomizableObjectNodeComponentPassthroughMesh.h b/Engine/Plugins/Mutable/Source/CustomizableObjectEditor/Private/MuCOE/Nodes/CustomizableObjectNodeComponentPassthroughMesh.h
--- a/Engine/Plugins/Mutable/Source/CustomizableObjectEditor/Private/MuCOE/Nodes/CustomizableObjectNodeComponentPassthroughMesh.h
+++ b/Engine/Plugins/Mutable/Source/CustomizableObjectEditor/Private/MuCOE/Nodes/CustomizableObjectNodeComponentPassthroughMesh.h
@@ -77,8 +77,6 @@ public:
 	// UEdGraphNode interface
 	virtual FText GetNodeTitle(ENodeTitleType::Type TitleType) const override;
 	virtual FText GetTooltipText() const override;
-	virtual bool GetCanRenameNode() const override { return true; }
-	virtual void OnRenameNode(const FString& NewName) override;
 
 	// UCustomizableObjectNode interface
 	virtual void AllocateDefaultPins(UCustomizableObjectNodeRemapPins* RemapPins) override;
diff --git a/Engine/Plugins/Mutable/Source/CustomizableObjectEditor/Private/MuCOE/Nodes/CustomizableObjectNodeCopyMaterial.cpp b/Engine/Plugins/Mutable/Source/CustomizableObjectEditor/Private/MuCOE/Nodes/CustomizableObjectNodeCopyMaterial.cpp
--- a/Engine/Plugins/Mutable/Source/CustomizableObjectEditor/Private/MuCOE/Nodes/CustomizableObjectNodeCopyMaterial.cpp
+++ b/Engine/Plugins/Mutable/Source/CustomizableObjectEditor/Private/MuCOE/Nodes/CustomizableObjectNodeCopyMaterial.cpp
@@ -275,16 +275,25 @@ bool UCustomizableObjectNodeCopyMaterial::IsReuseMaterialBetweenLODs() const
 }
 
 
-TArray<FString>* UCustomizableObjectNodeCopyMaterial::GetEnableTags()
+TArray<FString> UCustomizableObjectNodeCopyMaterial::GetEnableTags(TArray<const UCustomizableObjectNodeMacroInstance*>* MacroContext)
 {
 	if (UCustomizableObjectNodeMaterial* NodeMaterial = GetMaterialNode())
 	{
-		return NodeMaterial->GetEnableTags();
+		return NodeMaterial->GetEnableTags(MacroContext);
 	}
-	else
+	
+	return TArray<FString>();
+}
+
+
+TArray<FString>* UCustomizableObjectNodeCopyMaterial::GetEnableTagsArray()
+{
+	if (UCustomizableObjectNodeMaterial* NodeMaterial = GetMaterialNode())
 	{
-		return nullptr;
+		return NodeMaterial->GetEnableTagsArray();
 	}
+
+	return nullptr;
 }
 
 
diff --git a/Engine/Plugins/Mutable/Source/CustomizableObjectEditor/Private/MuCOE/Nodes/CustomizableObjectNodeCopyMaterial.h b/Engine/Plugins/Mutable/Source/CustomizableObjectEditor/Private/MuCOE/Nodes/CustomizableObjectNodeCopyMaterial.h
--- a/Engine/Plugins/Mutable/Source/CustomizableObjectEditor/Private/MuCOE/Nodes/CustomizableObjectNodeCopyMaterial.h
+++ b/Engine/Plugins/Mutable/Source/CustomizableObjectEditor/Private/MuCOE/Nodes/CustomizableObjectNodeCopyMaterial.h
@@ -31,11 +31,12 @@ public:
 	virtual bool ProvidesCustomPinRelevancyTest() const override;
 	virtual bool IsPinRelevant(const UEdGraphPin* Pin) const override;
 	virtual void BackwardsCompatibleFixup(int32 CustomizableObjectCustomVersion) override;
+	virtual TArray<FString> GetEnableTags(TArray<const UCustomizableObjectNodeMacroInstance*>* MacroContext = nullptr) override;
+	virtual TArray<FString>* GetEnableTagsArray() override;
 
 	// UCustomizableObjectNodeMaterialBase interface
 	virtual UMaterialInterface* GetMaterial() const override;
 	virtual bool IsReuseMaterialBetweenLODs() const override;
-	virtual TArray<FString>* GetEnableTags() override;
 	virtual UEdGraphPin* GetMaterialAssetPin() const override;
 	virtual int32 GetNumParameters(EMaterialParameterType Type) const override;
 	virtual FNodeMaterialParameterId GetParameterId(EMaterialParameterType Type, int32 ParameterIndex) const override;
diff --git a/Engine/Plugins/Mutable/Source/CustomizableObjectEditor/Private/MuCOE/Nodes/CustomizableObjectNodeEnumParameter.cpp b/Engine/Plugins/Mutable/Source/CustomizableObjectEditor/Private/MuCOE/Nodes/CustomizableObjectNodeEnumParameter.cpp
--- a/Engine/Plugins/Mutable/Source/CustomizableObjectEditor/Private/MuCOE/Nodes/CustomizableObjectNodeEnumParameter.cpp
+++ b/Engine/Plugins/Mutable/Source/CustomizableObjectEditor/Private/MuCOE/Nodes/CustomizableObjectNodeEnumParameter.cpp
@@ -4,6 +4,8 @@
 
 #include "MuCOE/EdGraphSchema_CustomizableObject.h"
 #include "MuCO/CustomizableObjectCustomVersion.h"
+#include "MuCOE/CustomizableObjectMacroLibrary/CustomizableObjectMacroLibrary.h"
+
 
 class UCustomizableObjectNodeRemapPins;
 
@@ -37,42 +39,6 @@ void UCustomizableObjectNodeEnumParameter::BackwardsCompatibleFixup(int32 Custom
 }
 
 
-void UCustomizableObjectNodeEnumParameter::AllocateDefaultPins(UCustomizableObjectNodeRemapPins* RemapPins)
-{
-	const FName Type = UEdGraphSchema_CustomizableObject::PC_Enum;
-	const FName PinName = UEdGraphSchema_CustomizableObject::GetPinCategoryName(Type);
-	const FText PinFriendlyName = UEdGraphSchema_CustomizableObject::GetPinCategoryFriendlyName(Type);
-	
-	UEdGraphPin* ValuePin = CustomCreatePin(EGPD_Output, Type, PinName);
-	ValuePin->bDefaultValueIsIgnored = true;
-	ValuePin->PinFriendlyName = PinFriendlyName;
-}
-
-
-FText UCustomizableObjectNodeEnumParameter::GetNodeTitle(ENodeTitleType::Type TitleType) const
-{
-	if (TitleType == ENodeTitleType::ListView)
-	{
-		return LOCTEXT("Enum_Parameter", "Enum Parameter");
-	}
-	else if (TitleType == ENodeTitleType::EditableTitle)
-	{
-		return FText::Format(LOCTEXT("Enum_Parameter_EditableTitle", "{0}"), FText::FromString(ParameterName));
-	}
-	else
-	{
-		return FText::Format(LOCTEXT("Enum_Parameter_Title", "{0}\nEnum Parameter"), FText::FromString(ParameterName));
-	}
-}
-
-
-FLinearColor UCustomizableObjectNodeEnumParameter::GetNodeTitleColor() const
-{
-	const UEdGraphSchema_CustomizableObject* Schema = GetDefault<UEdGraphSchema_CustomizableObject>();
-	return Schema->GetPinTypeColor(Schema->PC_Enum);
-}
-
-
 FText UCustomizableObjectNodeEnumParameter::GetTooltipText() const
 {
 	return LOCTEXT("Enum_Parameter_Tooltip",
@@ -80,12 +46,9 @@ FText UCustomizableObjectNodeEnumParameter::GetTooltipText() const
 }
 
 
-void UCustomizableObjectNodeEnumParameter::OnRenameNode(const FString& NewName)
+FName UCustomizableObjectNodeEnumParameter::GetCategory() const
 {
-	if (!NewName.IsEmpty())
-	{
-		ParameterName = NewName;
-	}
+	return UEdGraphSchema_CustomizableObject::PC_Enum;
 }
 
 #undef LOCTEXT_NAMESPACE
diff --git a/Engine/Plugins/Mutable/Source/CustomizableObjectEditor/Private/MuCOE/Nodes/CustomizableObjectNodeEnumParameter.h b/Engine/Plugins/Mutable/Source/CustomizableObjectEditor/Private/MuCOE/Nodes/CustomizableObjectNodeEnumParameter.h
--- a/Engine/Plugins/Mutable/Source/CustomizableObjectEditor/Private/MuCOE/Nodes/CustomizableObjectNodeEnumParameter.h
+++ b/Engine/Plugins/Mutable/Source/CustomizableObjectEditor/Private/MuCOE/Nodes/CustomizableObjectNodeEnumParameter.h
@@ -2,7 +2,7 @@
 
 #pragma once
 
-#include "MuCOE/Nodes/CustomizableObjectNode.h"
+#include "MuCOE/Nodes/CustomizableObjectNodeParameter.h"
 
 #include "CustomizableObjectNodeEnumParameter.generated.h"
 
@@ -27,7 +27,7 @@ struct CUSTOMIZABLEOBJECTEDITOR_API FCustomizableObjectNodeEnumValue
 
 
 UCLASS()
-class CUSTOMIZABLEOBJECTEDITOR_API UCustomizableObjectNodeEnumParameter : public UCustomizableObjectNode
+class CUSTOMIZABLEOBJECTEDITOR_API UCustomizableObjectNodeEnumParameter : public UCustomizableObjectNodeParameter
 {
 public:
 	GENERATED_BODY()
@@ -35,29 +35,19 @@ public:
 	UPROPERTY(EditAnywhere, Category=CustomizableObject)
 	int32 DefaultIndex = 0;
 
-	UPROPERTY(EditAnywhere, Category=CustomizableObject)
-	FString ParameterName = "Default Name";
-
-	UPROPERTY(EditAnywhere, Category = UI, meta = (DisplayName = "Parameter UI Metadata"))
-	FMutableParamUIMetadata ParamUIMetadata;
-
 	UPROPERTY(EditAnywhere, Category=CustomizableObject)
 	TArray<FCustomizableObjectNodeEnumValue> Values;
 
+
 	// UObject interface.
 	virtual void PostEditChangeProperty(FPropertyChangedEvent& PropertyChangedEvent) override;
 
 	// Begin EdGraphNode interface
-	virtual FText GetNodeTitle(ENodeTitleType::Type TitleType) const override;
-	virtual FLinearColor GetNodeTitleColor() const override;
 	virtual FText GetTooltipText() const override;
-	virtual void OnRenameNode(const FString& NewName) override;
-	virtual bool GetCanRenameNode() const override { return true; }
 	virtual void BackwardsCompatibleFixup(int32 CustomizableObjectCustomVersion) override;
 
-	// UCustomizableObjectNode interface
-	virtual void AllocateDefaultPins(UCustomizableObjectNodeRemapPins* RemapPins) override;
-	virtual bool IsAffectedByLOD() const override { return false; }
+	// CustomizableObjectNodeParameter interface
+	virtual FName GetCategory() const override;
 
 };
 
diff --git a/Engine/Plugins/Mutable/Source/CustomizableObjectEditor/Private/MuCOE/Nodes/CustomizableObjectNodeFloatParameter.cpp b/Engine/Plugins/Mutable/Source/CustomizableObjectEditor/Private/MuCOE/Nodes/CustomizableObjectNodeFloatParameter.cpp
--- a/Engine/Plugins/Mutable/Source/CustomizableObjectEditor/Private/MuCOE/Nodes/CustomizableObjectNodeFloatParameter.cpp
+++ b/Engine/Plugins/Mutable/Source/CustomizableObjectEditor/Private/MuCOE/Nodes/CustomizableObjectNodeFloatParameter.cpp
@@ -4,6 +4,8 @@
 
 #include "MuCO/CustomizableObjectCustomVersion.h"
 #include "MuCOE/EdGraphSchema_CustomizableObject.h"
+#include "MuCOE/CustomizableObjectMacroLibrary/CustomizableObjectMacroLibrary.h"
+
 
 class UCustomizableObjectNodeRemapPins;
 
@@ -22,67 +24,6 @@ void UCustomizableObjectNodeFloatParameter::PostEditChangeProperty(FPropertyChan
 }
 
 
-void UCustomizableObjectNodeFloatParameter::AllocateDefaultPins(UCustomizableObjectNodeRemapPins* RemapPins)
-{
-	const FName Type = UEdGraphSchema_CustomizableObject::PC_Float;
-	const FName PinName = UEdGraphSchema_CustomizableObject::GetPinCategoryName(Type);
-	const FText PinFriendlyName = UEdGraphSchema_CustomizableObject::GetPinCategoryFriendlyName(Type);
-	
-	UEdGraphPin* ValuePin = CustomCreatePin(EGPD_Output, Type, PinName);
-	ValuePin->bDefaultValueIsIgnored = true;
-	ValuePin->PinFriendlyName = PinFriendlyName;
-}
-
-
-bool UCustomizableObjectNodeFloatParameter::IsAffectedByLOD() const
-{
-	return false;
-}
-
-
-FText UCustomizableObjectNodeFloatParameter::GetNodeTitle(ENodeTitleType::Type TitleType) const
-{
-	FText NodeTitle;
-
-	if (TitleType == ENodeTitleType::ListView)
-	{
-		NodeTitle = LOCTEXT("Float_Parameter", "Float Parameter");
-	}
-	else if (TitleType == ENodeTitleType::EditableTitle)
-	{
-		NodeTitle = FText::Format(LOCTEXT("Float_Parameter_EditableTitle", "{0}"), FText::FromString(ParameterName));
-	}
-	else
-	{
-		NodeTitle = FText::Format(LOCTEXT("Float_Parameter_Title", "{0}\nFloat Parameter"), FText::FromString(ParameterName));
-	}
-
-	return NodeTitle;
-}
-
-
-FLinearColor UCustomizableObjectNodeFloatParameter::GetNodeTitleColor() const
-{
-	const UEdGraphSchema_CustomizableObject* Schema = GetDefault<UEdGraphSchema_CustomizableObject>();
-	return Schema->GetPinTypeColor(Schema->PC_Float);
-}
-
-
-FText UCustomizableObjectNodeFloatParameter::GetTooltipText() const
-{
-	return LOCTEXT("Float_Parameter_Tooltip", "Expose a numeric parameter from the Customizable Object that can be modified at runtime.");
-}
-
-
-void UCustomizableObjectNodeFloatParameter::OnRenameNode(const FString& NewName)
-{
-	if (!NewName.IsEmpty())
-	{
-		ParameterName = NewName;
-	}
-}
-
-
 void UCustomizableObjectNodeFloatParameter::BackwardsCompatibleFixup(int32 CustomizableObjectCustomVersion)
 {
 	Super::BackwardsCompatibleFixup(CustomizableObjectCustomVersion);
@@ -102,4 +43,10 @@ void UCustomizableObjectNodeFloatParameter::BackwardsCompatibleFixup(int32 Custo
 }
 
 
+FName UCustomizableObjectNodeFloatParameter::GetCategory() const
+{
+	return UEdGraphSchema_CustomizableObject::PC_Float;
+}
+
+
 #undef LOCTEXT_NAMESPACE
diff --git a/Engine/Plugins/Mutable/Source/CustomizableObjectEditor/Private/MuCOE/Nodes/CustomizableObjectNodeFloatParameter.h b/Engine/Plugins/Mutable/Source/CustomizableObjectEditor/Private/MuCOE/Nodes/CustomizableObjectNodeFloatParameter.h
--- a/Engine/Plugins/Mutable/Source/CustomizableObjectEditor/Private/MuCOE/Nodes/CustomizableObjectNodeFloatParameter.h
+++ b/Engine/Plugins/Mutable/Source/CustomizableObjectEditor/Private/MuCOE/Nodes/CustomizableObjectNodeFloatParameter.h
@@ -2,7 +2,7 @@
 
 #pragma once
 
-#include "MuCOE/Nodes/CustomizableObjectNode.h"
+#include "MuCOE/Nodes/CustomizableObjectNodeParameter.h"
 
 #include "CustomizableObjectNodeFloatParameter.generated.h"
 
@@ -27,34 +27,22 @@ struct CUSTOMIZABLEOBJECTEDITOR_API FCustomizableObjectNodeFloatDescription
 
 
 UCLASS()
-class CUSTOMIZABLEOBJECTEDITOR_API UCustomizableObjectNodeFloatParameter : public UCustomizableObjectNode
+class CUSTOMIZABLEOBJECTEDITOR_API UCustomizableObjectNodeFloatParameter : public UCustomizableObjectNodeParameter
 {
 public:
 	GENERATED_BODY()
 
 	UPROPERTY(EditAnywhere, Category=CustomizableObject)
 	float DefaultValue = 1.0f;
-
-	UPROPERTY(EditAnywhere, Category=CustomizableObject)
-	FString ParameterName = "Default Name";
-
-	UPROPERTY(EditAnywhere, Category = UI, meta = (DisplayName = "Parameter UI Metadata"))
-	FMutableParamUIMetadata ParamUIMetadata;
 	
 	// UObject interface.
 	virtual void PostEditChangeProperty(FPropertyChangedEvent& PropertyChangedEvent) override;
 	
-	// Begin EdGraphNode interface
-	virtual FText GetNodeTitle(ENodeTitleType::Type TitleType) const override;
-	virtual FLinearColor GetNodeTitleColor() const override;
-	virtual FText GetTooltipText() const override;
-	virtual bool GetCanRenameNode() const override { return true; }
-
 	// UCustomizableObjectNode interface
-	virtual void AllocateDefaultPins(UCustomizableObjectNodeRemapPins* RemapPins) override;
-	virtual bool IsAffectedByLOD() const override;
 	virtual void BackwardsCompatibleFixup(int32 CustomizableObjectCustomVersion) override;
-	virtual void OnRenameNode(const FString& NewName) override;
+
+	// CustomizableObjectNodeParameter interface
+	virtual FName GetCategory() const override;
 
 };
 
diff --git a/Engine/Plugins/Mutable/Source/CustomizableObjectEditor/Private/MuCOE/Nodes/CustomizableObjectNodeGroupProjectorParameter.cpp b/Engine/Plugins/Mutable/Source/CustomizableObjectEditor/Private/MuCOE/Nodes/CustomizableObjectNodeGroupProjectorParameter.cpp
--- a/Engine/Plugins/Mutable/Source/CustomizableObjectEditor/Private/MuCOE/Nodes/CustomizableObjectNodeGroupProjectorParameter.cpp
+++ b/Engine/Plugins/Mutable/Source/CustomizableObjectEditor/Private/MuCOE/Nodes/CustomizableObjectNodeGroupProjectorParameter.cpp
@@ -4,6 +4,7 @@
 
 #include "MuCO/CustomizableObjectCustomVersion.h"
 #include "MuCO/LoadUtils.h"
+#include "MuCOE/CustomizableObjectMacroLibrary/CustomizableObjectMacroLibrary.h"
 #include "MuCOE/EdGraphSchema_CustomizableObject.h"
 
 
@@ -11,6 +12,52 @@ class UCustomizableObjectNodeRemapPins;
 
 #define LOCTEXT_NAMESPACE "CustomizableObjectEditor"
 
+FText UCustomizableObjectNodeGroupProjectorParameter::GetTooltipText() const
+{
+	return LOCTEXT("Group_Projector_Parameter_Tooltip", "Projects one or many textures to all children in the group it's connected to. It modifies only the materials that define a specific material asset texture parameter.");
+}
+
+
+void UCustomizableObjectNodeGroupProjectorParameter::BackwardsCompatibleFixup(int32 CustomizableObjectCustomVersion)
+{
+	Super::BackwardsCompatibleFixup(CustomizableObjectCustomVersion);
+	
+	if (CustomizableObjectCustomVersion == FCustomizableObjectCustomVersion::GroupProjectorPinTypeAdded)
+	{
+		if (UEdGraphPin* Pin = FindPin(TEXT("Value")))
+		{
+			Pin->PinType.PinCategory = UEdGraphSchema_CustomizableObject::PC_GroupProjector;
+		}
+	}
+	else if (CustomizableObjectCustomVersion == FCustomizableObjectCustomVersion::GroupProjectorImagePinRemoved)
+	{
+		ReconstructNode();
+	}
+	else if (CustomizableObjectCustomVersion == FCustomizableObjectCustomVersion::UpdatedNodesPinName2)
+	{
+		// Looking for the projector pin as it was renamed in the same compatibility fix but of the parent
+		if (UEdGraphPin* Pin = FindPin(TEXT("Projector")))
+		{
+			Pin->PinName = TEXT("Group Projector");
+			Pin->PinFriendlyName = LOCTEXT("Group_Projector_Pin_Category", "Group Projector");
+		}
+	}
+}
+
+
+FName UCustomizableObjectNodeGroupProjectorParameter::GetCategory() const
+{
+	return UEdGraphSchema_CustomizableObject::PC_GroupProjector;
+}
+
+
+UEdGraphPin& UCustomizableObjectNodeGroupProjectorParameter::OutputPin() const
+{
+	const FName Type = UEdGraphSchema_CustomizableObject::PC_GroupProjector;
+	const FName PinName = UEdGraphSchema_CustomizableObject::GetPinCategoryName(Type);
+	return *FindPin(PinName);
+}
+
 
 TArray<FGroupProjectorParameterImage> UCustomizableObjectNodeGroupProjectorParameter::GetOptionTexturesFromTable() const
 {
@@ -39,7 +86,7 @@ TArray<FGroupProjectorParameterImage> UCustomizableObjectNodeGroupProjectorParam
 		FString PropertyValue(TEXT(""));
 		PropertyTexturePath->ExportText_InContainer(0, PropertyValue, RowData, RowData, nullptr, PPF_None);
 		UTexture2D* Texture = MutablePrivate::LoadObject<UTexture2D>(nullptr, *PropertyValue);
-		
+
 		if (Texture == nullptr)
 		{
 			UE_LOG(LogMutable, Warning, TEXT("WARNING: Unable to load texture %s"), *PropertyValue);
@@ -84,82 +131,4 @@ TArray<FGroupProjectorParameterImage> UCustomizableObjectNodeGroupProjectorParam
 	return ArrayDataTable;
 }
 
-
-FText UCustomizableObjectNodeGroupProjectorParameter::GetNodeTitle(ENodeTitleType::Type TitleType) const
-{
-	if (TitleType == ENodeTitleType::ListView)
-	{
-		return LOCTEXT("Group_Projector_Parameter", "Group Projector Parameter");
-	}
-	else if (TitleType == ENodeTitleType::EditableTitle)
-	{
-		return FText::Format(LOCTEXT("Group_Projector_Parameter_EditableTitle", "{0}"), FText::FromString(ParameterName));
-	}
-	else
-	{
-		return FText::Format(LOCTEXT("Group_Projector_Parameter_Title", "{0}\nGroup Projector Parameter"), FText::FromString(ParameterName));
-	}
-}
-
-
-FLinearColor UCustomizableObjectNodeGroupProjectorParameter::GetNodeTitleColor() const
-{
-	const UEdGraphSchema_CustomizableObject* Schema = GetDefault<UEdGraphSchema_CustomizableObject>();
-	return Schema->GetPinTypeColor(Schema->PC_GroupProjector);
-}
-
-
-FText UCustomizableObjectNodeGroupProjectorParameter::GetTooltipText() const
-{
-	return LOCTEXT("Group_Projector_Parameter_Tooltip", "Projects one or many textures to all children in the group it's connected to. It modifies only the materials that define a specific material asset texture parameter.");
-}
-
-
-void UCustomizableObjectNodeGroupProjectorParameter::BackwardsCompatibleFixup(int32 CustomizableObjectCustomVersion)
-{
-	Super::BackwardsCompatibleFixup(CustomizableObjectCustomVersion);
-	
-	if (CustomizableObjectCustomVersion == FCustomizableObjectCustomVersion::GroupProjectorPinTypeAdded)
-	{
-		if (UEdGraphPin* Pin = FindPin(TEXT("Value")))
-		{
-			Pin->PinType.PinCategory = UEdGraphSchema_CustomizableObject::PC_GroupProjector;
-		}
-	}
-	else if (CustomizableObjectCustomVersion == FCustomizableObjectCustomVersion::GroupProjectorImagePinRemoved)
-	{
-		ReconstructNode();
-	}
-	else if (CustomizableObjectCustomVersion == FCustomizableObjectCustomVersion::UpdatedNodesPinName2)
-	{
-		// Looking for the projector pin as it was renamed in the same compatibility fix but of the parent
-		if (UEdGraphPin* Pin = FindPin(TEXT("Projector")))
-		{
-			Pin->PinName = TEXT("Group Projector");
-			Pin->PinFriendlyName = LOCTEXT("Group_Projector_Pin_Category", "Group Projector");
-		}
-	}
-}
-
-
-void UCustomizableObjectNodeGroupProjectorParameter::AllocateDefaultPins(UCustomizableObjectNodeRemapPins* RemapPins)
-{
-	const FName Type = UEdGraphSchema_CustomizableObject::PC_GroupProjector;
-	const FName PinName = UEdGraphSchema_CustomizableObject::GetPinCategoryName(Type);
-	const FText PinFriendlyName = UEdGraphSchema_CustomizableObject::GetPinCategoryFriendlyName(Type);
-	
-	UEdGraphPin* ValuePin = CustomCreatePin(EGPD_Output, Type, PinName);
-	ValuePin->bDefaultValueIsIgnored = true;
-	ValuePin->PinFriendlyName = PinFriendlyName;
-}
-
-
-UEdGraphPin& UCustomizableObjectNodeGroupProjectorParameter::OutputPin() const
-{
-	const FName Type = UEdGraphSchema_CustomizableObject::PC_GroupProjector;
-	const FName PinName = UEdGraphSchema_CustomizableObject::GetPinCategoryName(Type);
-	return *FindPin(PinName);
-}
-
-
 #undef LOCTEXT_NAMESPACE
diff --git a/Engine/Plugins/Mutable/Source/CustomizableObjectEditor/Private/MuCOE/Nodes/CustomizableObjectNodeGroupProjectorParameter.h b/Engine/Plugins/Mutable/Source/CustomizableObjectEditor/Private/MuCOE/Nodes/CustomizableObjectNodeGroupProjectorParameter.h
--- a/Engine/Plugins/Mutable/Source/CustomizableObjectEditor/Private/MuCOE/Nodes/CustomizableObjectNodeGroupProjectorParameter.h
+++ b/Engine/Plugins/Mutable/Source/CustomizableObjectEditor/Private/MuCOE/Nodes/CustomizableObjectNodeGroupProjectorParameter.h
@@ -45,13 +45,6 @@ class CUSTOMIZABLEOBJECTEDITOR_API UCustomizableObjectNodeGroupProjectorParamete
 public:
 	GENERATED_BODY()
 
-	/** Return array with the sticker name and UTexture2D for projection */
-	TArray<FGroupProjectorParameterImage> GetOptionTexturesFromTable() const;
-
-	/** Returns the final option images without repeated elements in the option names,
-	* the data table has preference over elements in the OptionTextures array */
-	TArray<FGroupProjectorParameterImage> GetFinalOptionTexturesNoRepeat() const;
-
 	UPROPERTY(EditAnywhere, Category = ProjectorGroup, Meta = (ToolTip = "Only used when connected to a Group node. Specifies which material channel in the Group node's child material nodes will be connected to the projection."))
 	FString MaterialChannelNameToConnect;
 
@@ -91,15 +84,22 @@ public:
 	int32 UVLayout = 0;
 
 	// EdGraphNode interface.
-	virtual FText GetNodeTitle(ENodeTitleType::Type TitleType) const override;
-	virtual FLinearColor GetNodeTitleColor() const override;
 	virtual FText GetTooltipText() const override;
 
 	// UCustomizableObjectNode interface.
 	virtual void BackwardsCompatibleFixup(int32 CustomizableObjectCustomVersion) override;
-	virtual void AllocateDefaultPins(UCustomizableObjectNodeRemapPins* RemapPins) override;
+
+	// CustomizableObjectNodeParameter interface
+	virtual FName GetCategory() const override;
 
 	// Own interface
 	UEdGraphPin& OutputPin() const;
+
+	/** Return array with the sticker name and UTexture2D for projection */
+	TArray<FGroupProjectorParameterImage> GetOptionTexturesFromTable() const;
+
+	/** Returns the final option images without repeated elements in the option names,
+	* the data table has preference over elements in the OptionTextures array */
+	TArray<FGroupProjectorParameterImage> GetFinalOptionTexturesNoRepeat() const;
 };
 
diff --git a/Engine/Plugins/Mutable/Source/CustomizableObjectEditor/Private/MuCOE/Nodes/CustomizableObjectNodeMacroInstance.cpp b/Engine/Plugins/Mutable/Source/CustomizableObjectEditor/Private/MuCOE/Nodes/CustomizableObjectNodeMacroInstance.cpp
--- a/Engine/Plugins/Mutable/Source/CustomizableObjectEditor/Private/MuCOE/Nodes/CustomizableObjectNodeMacroInstance.cpp
+++ b/Engine/Plugins/Mutable/Source/CustomizableObjectEditor/Private/MuCOE/Nodes/CustomizableObjectNodeMacroInstance.cpp
@@ -160,6 +160,19 @@ void UCustomizableObjectNodeMacroInstance::PostEditChangeProperty(FPropertyChang
 	}
 }
 
+
+bool UCustomizableObjectNodeMacroInstance::CreateStaticStringPinWidget() const
+{
+	return false;
+}
+
+
+bool UCustomizableObjectNodeMacroInstance::IsExperimental() const
+{
+	return true;
+}
+
+
 UEdGraphPin* UCustomizableObjectNodeMacroInstance::GetMacroIOPin(ECOMacroIOType IONodeType, const FName& PinName) const
 {
 	UEdGraphPin* MacroIOPin = nullptr;
diff --git a/Engine/Plugins/Mutable/Source/CustomizableObjectEditor/Private/MuCOE/Nodes/CustomizableObjectNodeMacroInstance.h b/Engine/Plugins/Mutable/Source/CustomizableObjectEditor/Private/MuCOE/Nodes/CustomizableObjectNodeMacroInstance.h
--- a/Engine/Plugins/Mutable/Source/CustomizableObjectEditor/Private/MuCOE/Nodes/CustomizableObjectNodeMacroInstance.h
+++ b/Engine/Plugins/Mutable/Source/CustomizableObjectEditor/Private/MuCOE/Nodes/CustomizableObjectNodeMacroInstance.h
@@ -69,6 +69,8 @@ public:
 	void AllocateDefaultPins(UCustomizableObjectNodeRemapPins* RemapPins) override;
 	virtual UCustomizableObjectNodeMacroInstanceRemapPins* CreateRemapPinsDefault() const override;
 	virtual void PostEditChangeProperty(FPropertyChangedEvent& PropertyChangedEvent) override;
+	virtual bool CreateStaticStringPinWidget() const override;
+	virtual bool IsExperimental() const override;
 
 	// Node MacroInstance Interface
 	/** Returns the pin of the Macro's input/output node with the same name */
diff --git a/Engine/Plugins/Mutable/Source/CustomizableObjectEditor/Private/MuCOE/Nodes/CustomizableObjectNodeMacroInstanceDetails.cpp b/Engine/Plugins/Mutable/Source/CustomizableObjectEditor/Private/MuCOE/Nodes/CustomizableObjectNodeMacroInstanceDetails.cpp
--- a/Engine/Plugins/Mutable/Source/CustomizableObjectEditor/Private/MuCOE/Nodes/CustomizableObjectNodeMacroInstanceDetails.cpp
+++ b/Engine/Plugins/Mutable/Source/CustomizableObjectEditor/Private/MuCOE/Nodes/CustomizableObjectNodeMacroInstanceDetails.cpp
@@ -52,6 +52,7 @@ void FCustomizableObjectNodeMacroInstanceDetails::CustomizeDetails(IDetailLayout
 	.NameContent()
 	[
 		SNew(STextBlock).Text(LOCTEXT("ParentMacroText", "Parent Macro"))
+		.Font(IDetailLayoutBuilder::GetDetailFont())
 	]
 	.ValueContent()
 	[
diff --git a/Engine/Plugins/Mutable/Source/CustomizableObjectEditor/Private/MuCOE/Nodes/CustomizableObjectNodeMaterial.cpp b/Engine/Plugins/Mutable/Source/CustomizableObjectEditor/Private/MuCOE/Nodes/CustomizableObjectNodeMaterial.cpp
--- a/Engine/Plugins/Mutable/Source/CustomizableObjectEditor/Private/MuCOE/Nodes/CustomizableObjectNodeMaterial.cpp
+++ b/Engine/Plugins/Mutable/Source/CustomizableObjectEditor/Private/MuCOE/Nodes/CustomizableObjectNodeMaterial.cpp
@@ -14,6 +14,7 @@
 #include "MuCOE/Nodes/CustomizableObjectNodeCopyMaterial.h"
 #include "MuCOE/Nodes/CustomizableObjectNodeSkeletalMesh.h"
 #include "MuCOE/Nodes/CustomizableObjectNodeStaticMesh.h"
+#include "MuCOE/Nodes/CustomizableObjectNodeStaticString.h"
 #include "MuCOE/Nodes/CustomizableObjectNodeTable.h"
 #include "MuCOE/Nodes/SCustomizableObjectNodeMaterial.h"
 #include "ObjectEditorUtils.h"
@@ -23,6 +24,7 @@
 #include "MuCOE/CustomizableObjectEditorLogger.h"
 #include "MuCOE/CustomizableObjectEditorUtilities.h"
 #include "MuCOE/Nodes/CustomizableObjectNodeExternalPin.h"
+#include "MuCOE/CustomizableObjectMacroLibrary/CustomizableObjectMacroLibrary.h"
 
 class SGraphNode;
 class SWidget;
@@ -335,12 +337,13 @@ void UCustomizableObjectNodeMaterial::AllocateDefaultPins(UCustomizableObjectNod
 		TableMaterialPin->PinToolTip = "Pin for a Material from a Table Node";
 	}
 
-	if (IsInMacro())
+	if (CVarEnableMutableMacros.GetValueOnAnyThread())
 	{
 		const FString PinFriendlyName = "Enable Tags";
 		const FString PinName = PinFriendlyName + FString(TEXT("_Input_Pin"));
 		UEdGraphPin* TagsPin = CustomCreatePin(EGPD_Input, Schema->PC_String, FName(PinName), true);
 		TagsPin->PinFriendlyName = FText::FromString(PinFriendlyName);
+		TagsPin->PinToolTip = "List of Tags that this node will Enable";
 		EnableTagsPinRef = TagsPin;
 	}
 
@@ -887,7 +890,12 @@ UEdGraphPin* UCustomizableObjectNodeMaterial::GetMaterialAssetPin() const
 
 UEdGraphPin* UCustomizableObjectNodeMaterial::GetEnableTagsPin() const
 {
-	return EnableTagsPinRef.Get();
+	if (CVarEnableMutableMacros.GetValueOnAnyThread())
+	{
+		return EnableTagsPinRef.Get();
+	}
+
+	return nullptr;
 }
 
 
@@ -1248,7 +1256,38 @@ bool UCustomizableObjectNodeMaterial::IsReuseMaterialBetweenLODs() const
 }
 
 
-TArray<FString>* UCustomizableObjectNodeMaterial::GetEnableTags()
+TArray<FString> UCustomizableObjectNodeMaterial::GetEnableTags(TArray<const UCustomizableObjectNodeMacroInstance*>* MacroContext)
+{
+	// Getting tags from linked pin
+	const UEdGraphPin* EnableTagsPin = GetEnableTagsPin();
+
+	if (EnableTagsPin && FollowInputPinArray(*EnableTagsPin).Num())
+	{
+		TArray<FString> OutTags;
+
+		for (const UEdGraphPin* StringPin : FollowInputPinArray(*EnableTagsPin))
+		{
+			const UEdGraphPin* SourceStringPin = GraphTraversal::FindStringSource(*StringPin, MacroContext);
+
+			if (SourceStringPin)
+			{
+				if (const UCustomizableObjectNodeStaticString* StringNode = Cast<UCustomizableObjectNodeStaticString>(SourceStringPin->GetOwningNode()))
+				{
+					OutTags.AddUnique(StringNode->Value);
+				}
+			}
+		}
+
+		return OutTags;
+	}
+	else
+	{
+		return Tags;
+	}
+}
+
+
+TArray<FString>* UCustomizableObjectNodeMaterial::GetEnableTagsArray()
 {
 	return &Tags;
 }
diff --git a/Engine/Plugins/Mutable/Source/CustomizableObjectEditor/Private/MuCOE/Nodes/CustomizableObjectNodeMaterial.h b/Engine/Plugins/Mutable/Source/CustomizableObjectEditor/Private/MuCOE/Nodes/CustomizableObjectNodeMaterial.h
--- a/Engine/Plugins/Mutable/Source/CustomizableObjectEditor/Private/MuCOE/Nodes/CustomizableObjectNodeMaterial.h
+++ b/Engine/Plugins/Mutable/Source/CustomizableObjectEditor/Private/MuCOE/Nodes/CustomizableObjectNodeMaterial.h
@@ -129,7 +129,8 @@ public:
 	virtual bool IsNodeOutDatedAndNeedsRefresh() override;
 	virtual FString GetRefreshMessage() const override;
 	virtual TSharedPtr<IDetailsView> CustomizePinDetails(const UEdGraphPin& Pin) const override;
-	virtual TArray<FString>* GetEnableTags() override;
+	virtual TArray<FString> GetEnableTags(TArray<const UCustomizableObjectNodeMacroInstance*>* MacroContext = nullptr) override;
+	virtual TArray<FString>* GetEnableTagsArray() override;
 	virtual FString GetInternalTagDisplayName() override;
 	
 	// UCustomizableObjectNodeMaterialBase interface
diff --git a/Engine/Plugins/Mutable/Source/CustomizableObjectEditor/Private/MuCOE/Nodes/CustomizableObjectNodeMeshClipWithMesh.cpp b/Engine/Plugins/Mutable/Source/CustomizableObjectEditor/Private/MuCOE/Nodes/CustomizableObjectNodeMeshClipWithMesh.cpp
--- a/Engine/Plugins/Mutable/Source/CustomizableObjectEditor/Private/MuCOE/Nodes/CustomizableObjectNodeMeshClipWithMesh.cpp
+++ b/Engine/Plugins/Mutable/Source/CustomizableObjectEditor/Private/MuCOE/Nodes/CustomizableObjectNodeMeshClipWithMesh.cpp
@@ -57,7 +57,7 @@ void UCustomizableObjectNodeModifierClipWithMesh::BackwardsCompatibleFixup(int32
 				FString NewLegacyTag = MakeNodeAutoTag(ParentNode);
 				RequiredTags.Add(NewLegacyTag);
 
-				TArray<FString>* NodeEnableTags = ParentNode->GetEnableTags();
+				TArray<FString>* NodeEnableTags = ParentNode->GetEnableTagsArray();
 				if (NodeEnableTags)
 				{
 					NodeEnableTags->AddUnique(NewLegacyTag);
diff --git a/Engine/Plugins/Mutable/Source/CustomizableObjectEditor/Private/MuCOE/Nodes/CustomizableObjectNodeMeshMorph.cpp b/Engine/Plugins/Mutable/Source/CustomizableObjectEditor/Private/MuCOE/Nodes/CustomizableObjectNodeMeshMorph.cpp
--- a/Engine/Plugins/Mutable/Source/CustomizableObjectEditor/Private/MuCOE/Nodes/CustomizableObjectNodeMeshMorph.cpp
+++ b/Engine/Plugins/Mutable/Source/CustomizableObjectEditor/Private/MuCOE/Nodes/CustomizableObjectNodeMeshMorph.cpp
@@ -8,6 +8,7 @@
 #include "MuCOE/EdGraphSchema_CustomizableObject.h"
 #include "MuCOE/GraphTraversal.h"
 #include "MuCOE/CustomizableObjectMacroLibrary/CustomizableObjectGraphEditorToolkit.h"
+#include "MuCOE/CustomizableObjectMacroLibrary/CustomizableObjectMacroLibrary.h"
 #include "MuCOE/Nodes/CustomizableObjectNodeMeshReshapeCommon.h"
 #include "MuCOE/Nodes/CustomizableObjectNodeSkeletalMesh.h"
 
@@ -40,7 +41,7 @@ void UCustomizableObjectNodeMeshMorph::AllocateDefaultPins(UCustomizableObjectNo
 	FactorPin->bDefaultValueIsIgnored = true;
 
 	// Create this pin only if we are in a Macro graph
-	if (IsInMacro())
+	if (CVarEnableMutableMacros.GetValueOnAnyThread())
 	{
 		PinName = TEXT("Morph Target Name");
 		MorphTargetNamePinRef = CustomCreatePin(EGPD_Input, Schema->PC_String, FName(*PinName));
@@ -145,7 +146,12 @@ FText UCustomizableObjectNodeMeshMorph::GetTooltipText() const
 
 UEdGraphPin* UCustomizableObjectNodeMeshMorph::MorphTargetNamePin() const
 {
-	return MorphTargetNamePinRef.Get();
+	if (CVarEnableMutableMacros.GetValueOnAnyThread())
+	{
+		return MorphTargetNamePinRef.Get();
+	}
+
+	return nullptr;
 }
 
 
@@ -166,4 +172,11 @@ void UCustomizableObjectNodeMeshMorph::Serialize(FArchive& Ar)
 }
 
 
+bool UCustomizableObjectNodeMeshMorph::CreateStaticStringPinWidget() const
+{
+	return false;
+}
+
+
+
 #undef LOCTEXT_NAMESPACE
diff --git a/Engine/Plugins/Mutable/Source/CustomizableObjectEditor/Private/MuCOE/Nodes/CustomizableObjectNodeMeshMorph.h b/Engine/Plugins/Mutable/Source/CustomizableObjectEditor/Private/MuCOE/Nodes/CustomizableObjectNodeMeshMorph.h
--- a/Engine/Plugins/Mutable/Source/CustomizableObjectEditor/Private/MuCOE/Nodes/CustomizableObjectNodeMeshMorph.h
+++ b/Engine/Plugins/Mutable/Source/CustomizableObjectEditor/Private/MuCOE/Nodes/CustomizableObjectNodeMeshMorph.h
@@ -31,6 +31,7 @@ public:
 	bool IsNodeOutDatedAndNeedsRefresh() override;
 	FString GetRefreshMessage() const override;
 	void PinConnectionListChanged(UEdGraphPin* Pin) override;
+	virtual bool CreateStaticStringPinWidget() const override;
 
 
 	// Own interface
diff --git a/Engine/Plugins/Mutable/Source/CustomizableObjectEditor/Private/MuCOE/Nodes/CustomizableObjectNodeMeshMorphDetails.cpp b/Engine/Plugins/Mutable/Source/CustomizableObjectEditor/Private/MuCOE/Nodes/CustomizableObjectNodeMeshMorphDetails.cpp
--- a/Engine/Plugins/Mutable/Source/CustomizableObjectEditor/Private/MuCOE/Nodes/CustomizableObjectNodeMeshMorphDetails.cpp
+++ b/Engine/Plugins/Mutable/Source/CustomizableObjectEditor/Private/MuCOE/Nodes/CustomizableObjectNodeMeshMorphDetails.cpp
@@ -80,21 +80,28 @@ void FCustomizableObjectNodeMeshMorphDetails::CustomizeDetails( IDetailLayoutBui
 
 				MorphTargetComboOptions.Sort(CompareNames);
 
-				BlocksCategory.AddCustomRow(LOCTEXT("MeshMorphDetails_Target", "Target"))
+				DetailBuilder.EditDefaultProperty(MorphTargetNameProperty)->
+				CustomWidget()
+				.NameContent()
 				[
-					SNew(SProperty, MorphTargetNameProperty)
-					.ShouldDisplayName(false)
-					.CustomWidget()
+					SNew(STextBlock)
+					.Text(LOCTEXT("MeshMorphDetails_MorphTargetName", "Morph Target Name")).Font(IDetailLayoutBuilder::GetDetailFont())
+					.IsEnabled(this, &FCustomizableObjectNodeMeshMorphDetails::IsMorphNameSelectorWidgetEnabled)
+					.ToolTipText(this, &FCustomizableObjectNodeMeshMorphDetails::MorphNameSelectorWidgetTooltip)
+				]
+				.ValueContent()
+				[
+					SNew(SBorder)
+					.BorderImage(UE_MUTABLE_GET_BRUSH("NoBorder"))
+					.Padding(FMargin(0.0f, 0.0f, 10.0f, 0.0f))
 					[
-						SNew(SBorder)
-						.BorderImage(UE_MUTABLE_GET_BRUSH("NoBorder"))
-						.Padding(FMargin(0.0f, 0.0f, 10.0f, 0.0f))
-						[
-							SNew(STextComboBox)
-							.OptionsSource(&MorphTargetComboOptions)
-							.InitiallySelectedItem(ItemToSelect)
-							.OnSelectionChanged(this, &FCustomizableObjectNodeMeshMorphDetails::OnMorphTargetComboBoxSelectionChanged, MorphTargetNameProperty)
-						]
+						SNew(STextComboBox)
+						.OptionsSource(&MorphTargetComboOptions)
+						.InitiallySelectedItem(ItemToSelect)
+						.OnSelectionChanged(this, &FCustomizableObjectNodeMeshMorphDetails::OnMorphTargetComboBoxSelectionChanged, MorphTargetNameProperty)
+						.Font(IDetailLayoutBuilder::GetDetailFont())
+						.IsEnabled(this, &FCustomizableObjectNodeMeshMorphDetails::IsMorphNameSelectorWidgetEnabled)
+						.ToolTipText(this, &FCustomizableObjectNodeMeshMorphDetails::MorphNameSelectorWidgetTooltip)
 					]
 				];
 			}
@@ -109,17 +116,19 @@ void FCustomizableObjectNodeMeshMorphDetails::CustomizeDetails( IDetailLayoutBui
 			CustomWidget()
 			.NameContent()
 			[
-				SNew(STextBlock).Text(LOCTEXT("MeshMorphDetails_MorphTargetName", "Morph Target Name")).Font(IDetailLayoutBuilder::GetDetailFont())
+				SNew(STextBlock)
+				.Text(LOCTEXT("MeshMorphDetails_MorphTargetName", "Morph Target Name")).Font(IDetailLayoutBuilder::GetDetailFont())
 			]
 			.ValueContent()
 			[
 				SNew(STextBlock)
-					.Text(Message)
-					.Font(IDetailLayoutBuilder::GetDetailFont())
+				.Text(Message)
+				.Font(IDetailLayoutBuilder::GetDetailFont())
 			];
 	}
 }
 
+
 void FCustomizableObjectNodeMeshMorphDetails::OnMorphTargetComboBoxSelectionChanged(TSharedPtr<FString> Selection, ESelectInfo::Type SelectInfo, TSharedRef<IPropertyHandle> ParentProperty)
 {	
 	if (Selection.IsValid())
@@ -128,5 +137,25 @@ void FCustomizableObjectNodeMeshMorphDetails::OnMorphTargetComboBoxSelectionChan
 	}
 }
 
+
+bool FCustomizableObjectNodeMeshMorphDetails::IsMorphNameSelectorWidgetEnabled() const
+{
+	// Disabled if there is a string node linked to the "Required Tags" pin.
+	return !(Node->MorphTargetNamePin() && FollowInputPin(*Node->MorphTargetNamePin()));
+}
+
+
+FText FCustomizableObjectNodeMeshMorphDetails::MorphNameSelectorWidgetTooltip() const
+{
+	if (Node->MorphTargetNamePin() && FollowInputPin(*Node->MorphTargetNamePin()))
+	{
+		return LOCTEXT("MorphTargetNameWidgetTooltip_Ignored", "Disabled. When there is a string node linked to the Morph Target Name pin the morph target name selected in this widget is ignored.");
+	}
+	else
+	{
+		return LOCTEXT("MorphTargetNameWidgetTooltip", "Select the morph target name.");
+	}
+}
+
 #undef LOCTEXT_NAMESPACE
 
diff --git a/Engine/Plugins/Mutable/Source/CustomizableObjectEditor/Private/MuCOE/Nodes/CustomizableObjectNodeMeshMorphDetails.h b/Engine/Plugins/Mutable/Source/CustomizableObjectEditor/Private/MuCOE/Nodes/CustomizableObjectNodeMeshMorphDetails.h
--- a/Engine/Plugins/Mutable/Source/CustomizableObjectEditor/Private/MuCOE/Nodes/CustomizableObjectNodeMeshMorphDetails.h
+++ b/Engine/Plugins/Mutable/Source/CustomizableObjectEditor/Private/MuCOE/Nodes/CustomizableObjectNodeMeshMorphDetails.h
@@ -20,6 +20,14 @@ public:
 	/** ILayoutDetails interface */
 	virtual void CustomizeDetails( IDetailLayoutBuilder& DetailBuilder ) override;
 
+private:
+
+	/** Disables/Enables the tags property widget. */
+	bool IsMorphNameSelectorWidgetEnabled() const;
+
+	/** Set the widget tooltip. */
+	FText MorphNameSelectorWidgetTooltip() const;
+
 private:
 
 	class UCustomizableObjectNodeMeshMorph* Node = nullptr;
diff --git a/Engine/Plugins/Mutable/Source/CustomizableObjectEditor/Private/MuCOE/Nodes/CustomizableObjectNodeMeshParameter.cpp b/Engine/Plugins/Mutable/Source/CustomizableObjectEditor/Private/MuCOE/Nodes/CustomizableObjectNodeMeshParameter.cpp
--- a/Engine/Plugins/Mutable/Source/CustomizableObjectEditor/Private/MuCOE/Nodes/CustomizableObjectNodeMeshParameter.cpp
+++ b/Engine/Plugins/Mutable/Source/CustomizableObjectEditor/Private/MuCOE/Nodes/CustomizableObjectNodeMeshParameter.cpp
@@ -4,6 +4,7 @@
 
 #include "MuCOE/EdGraphSchema_CustomizableObject.h"
 #include "MuCOE/CustomizableObjectLayout.h"
+#include "MuCOE/CustomizableObjectMacroLibrary/CustomizableObjectMacroLibrary.h"
 #include "MuCO/CustomizableObjectCustomVersion.h"
 #include "Rendering/SkeletalMeshModel.h"
 #include "Engine/SkeletalMesh.h"
@@ -12,7 +13,6 @@ class UCustomizableObjectNodeRemapPins;
 
 #define LOCTEXT_NAMESPACE "CustomizableObjectEditor"
 
-
 void UCustomizableObjectNodeMeshParameter::AllocateDefaultPins(UCustomizableObjectNodeRemapPins* RemapPins)
 {
 	USkeletalMesh* SkeletalMesh = DefaultValue.LoadSynchronous();
@@ -22,6 +22,11 @@ void UCustomizableObjectNodeMeshParameter::AllocateDefaultPins(UCustomizableObje
 		return;
 	}
 
+	if (CVarEnableMutableMacros.GetValueOnAnyThread())
+	{
+		NamePin = CustomCreatePin(EGPD_Input, UEdGraphSchema_CustomizableObject::PC_String, FName("Name"));
+	}
+
 	const UEdGraphSchema_CustomizableObject* Schema = GetDefault<UEdGraphSchema_CustomizableObject>();
 
 	if (const FSkeletalMeshModel* ImportedModel = SkeletalMesh->GetImportedModel())
@@ -82,7 +87,9 @@ bool UCustomizableObjectNodeMeshParameter::IsExperimental() const
 
 FText UCustomizableObjectNodeMeshParameter::GetNodeTitle(ENodeTitleType::Type TitleType) const
 {
-	if (TitleType == ENodeTitleType::ListView)
+	const UEdGraphPin* ParamNamePin = NamePin.Get();
+
+	if (TitleType == ENodeTitleType::ListView || (ParamNamePin && ParamNamePin->LinkedTo.Num()))
 	{
 		return LOCTEXT("Mesh_Parameter", "Mesh Parameter");
 	}
@@ -119,6 +126,28 @@ void UCustomizableObjectNodeMeshParameter::OnRenameNode(const FString& NewName)
 }
 
 
+bool UCustomizableObjectNodeMeshParameter::GetCanRenameNode() const
+{
+	const UEdGraphPin* ParamNamePin = NamePin.Get();
+
+	if (CVarEnableMutableMacros.GetValueOnAnyThread() && ParamNamePin && ParamNamePin->LinkedTo.Num())
+	{
+		return false;
+	}
+
+	return true;
+}
+
+
+void UCustomizableObjectNodeMeshParameter::PinConnectionListChanged(UEdGraphPin* Pin)
+{
+	if (Pin == NamePin)
+	{
+		GetGraph()->NotifyGraphChanged();
+	}
+}
+
+
 UTexture2D* UCustomizableObjectNodeMeshParameter::FindTextureForPin(const UEdGraphPin* Pin) const
 {
 	return nullptr;
diff --git a/Engine/Plugins/Mutable/Source/CustomizableObjectEditor/Private/MuCOE/Nodes/CustomizableObjectNodeMeshParameter.h b/Engine/Plugins/Mutable/Source/CustomizableObjectEditor/Private/MuCOE/Nodes/CustomizableObjectNodeMeshParameter.h
--- a/Engine/Plugins/Mutable/Source/CustomizableObjectEditor/Private/MuCOE/Nodes/CustomizableObjectNodeMeshParameter.h
+++ b/Engine/Plugins/Mutable/Source/CustomizableObjectEditor/Private/MuCOE/Nodes/CustomizableObjectNodeMeshParameter.h
@@ -56,6 +56,11 @@ public:
 	UPROPERTY(EditAnywhere, Category = UI, meta = (DisplayName = "Parameter UI Metadata"))
 	FMutableParamUIMetadata ParamUIMetadata;
 
+	UPROPERTY()
+	FEdGraphPinReference NamePin;
+
+public:
+
 	// UObject interface.
 	virtual void PostEditChangeProperty(FPropertyChangedEvent& PropertyChangedEvent) override;
 
@@ -64,7 +69,8 @@ public:
 	virtual FLinearColor GetNodeTitleColor() const override;
 	virtual FText GetTooltipText() const override;
 	virtual void OnRenameNode(const FString& NewName) override;
-	virtual bool GetCanRenameNode() const override { return true; }
+	virtual bool GetCanRenameNode() const override;
+	virtual void PinConnectionListChanged(UEdGraphPin* Pin) override;
 
 	// UCustomizableObjectNode interface
 	virtual void AllocateDefaultPins(UCustomizableObjectNodeRemapPins* RemapPins) override;
diff --git a/Engine/Plugins/Mutable/Source/CustomizableObjectEditor/Private/MuCOE/Nodes/CustomizableObjectNodeMeshSectionDetails.cpp b/Engine/Plugins/Mutable/Source/CustomizableObjectEditor/Private/MuCOE/Nodes/CustomizableObjectNodeMeshSectionDetails.cpp
--- a/Engine/Plugins/Mutable/Source/CustomizableObjectEditor/Private/MuCOE/Nodes/CustomizableObjectNodeMeshSectionDetails.cpp
+++ b/Engine/Plugins/Mutable/Source/CustomizableObjectEditor/Private/MuCOE/Nodes/CustomizableObjectNodeMeshSectionDetails.cpp
@@ -67,6 +67,8 @@ void FCustomizableObjectNodeMeshSectionDetails::CustomizeDetails( IDetailLayoutB
 					SNew(STextBlock)
 						.Text(LOCTEXT("MeshSectionDetails_Tags", "Enable Tags"))
 						.Font(IDetailLayoutBuilder::GetDetailFont())
+						.IsEnabled(this, &FCustomizableObjectNodeMeshSectionDetails::IsTagsPropertyWidgetEnabled)
+						.ToolTipText(this, &FCustomizableObjectNodeMeshSectionDetails::TagsPropertyWidgetTooltip)
 				]
 		]
 		.ValueContent()
@@ -78,6 +80,8 @@ void FCustomizableObjectNodeMeshSectionDetails::CustomizeDetails( IDetailLayoutB
 				.AllowInternalTags(false)
 				.EmptyListText(LOCTEXT("MeshSectionDetails_NoTags", "No tags enabled by this mesh section."))
 				.OnTagListChanged(this, &FCustomizableObjectNodeMeshSectionDetails::OnEnableTagsPropertyChanged)
+				.IsEnabled(this, &FCustomizableObjectNodeMeshSectionDetails::IsTagsPropertyWidgetEnabled)
+				.ToolTipText(this, &FCustomizableObjectNodeMeshSectionDetails::TagsPropertyWidgetTooltip)
 		];
 	}
 }
@@ -91,4 +95,28 @@ void FCustomizableObjectNodeMeshSectionDetails::OnEnableTagsPropertyChanged()
 }
 
 
+bool FCustomizableObjectNodeMeshSectionDetails::IsTagsPropertyWidgetEnabled() const
+{
+	const UEdGraphPin* EnableTagsPin = Node->GetEnableTagsPin();
+
+	// Disabled if there is a string node linked to the "Enable Tags" pin.
+	return !(EnableTagsPin && FollowInputPinArray(*EnableTagsPin).Num());
+}
+
+
+FText FCustomizableObjectNodeMeshSectionDetails::TagsPropertyWidgetTooltip() const
+{
+	const UEdGraphPin* EnableTagsPin = Node->GetEnableTagsPin();
+
+	if (EnableTagsPin && FollowInputPinArray(*EnableTagsPin).Num())
+	{
+		return LOCTEXT("EnableTagsWidgetTooltip_Ignored","Disabled. When there are string nodes linked to the Enable Tags pin, the tag list is ignored.");
+	}
+	else
+	{
+		return LOCTEXT("EnableTagsWidgetTooltip", "List of Tags that this node will enable.");
+	}
+}
+
+
 #undef LOCTEXT_NAMESPACE
diff --git a/Engine/Plugins/Mutable/Source/CustomizableObjectEditor/Private/MuCOE/Nodes/CustomizableObjectNodeMeshSectionDetails.h b/Engine/Plugins/Mutable/Source/CustomizableObjectEditor/Private/MuCOE/Nodes/CustomizableObjectNodeMeshSectionDetails.h
--- a/Engine/Plugins/Mutable/Source/CustomizableObjectEditor/Private/MuCOE/Nodes/CustomizableObjectNodeMeshSectionDetails.h
+++ b/Engine/Plugins/Mutable/Source/CustomizableObjectEditor/Private/MuCOE/Nodes/CustomizableObjectNodeMeshSectionDetails.h
@@ -35,6 +35,14 @@ protected:
 	/** */
 	virtual void OnEnableTagsPropertyChanged();
 
+private:
+
+	/** Disables/Enables the tags property widget. */
+	bool IsTagsPropertyWidgetEnabled() const;
+
+	/** Set the widget tooltip. */
+	FText TagsPropertyWidgetTooltip() const;
+
 private:
 
 	class UCustomizableObjectNodeMaterial* Node = nullptr;
diff --git a/Engine/Plugins/Mutable/Source/CustomizableObjectEditor/Private/MuCOE/Nodes/CustomizableObjectNodeModifierBase.cpp b/Engine/Plugins/Mutable/Source/CustomizableObjectEditor/Private/MuCOE/Nodes/CustomizableObjectNodeModifierBase.cpp
--- a/Engine/Plugins/Mutable/Source/CustomizableObjectEditor/Private/MuCOE/Nodes/CustomizableObjectNodeModifierBase.cpp
+++ b/Engine/Plugins/Mutable/Source/CustomizableObjectEditor/Private/MuCOE/Nodes/CustomizableObjectNodeModifierBase.cpp
@@ -4,7 +4,9 @@
 
 #include "MuCOE/EdGraphSchema_CustomizableObject.h"
 #include "MuCOE/CustomizableObjectEditorLogger.h"
+#include "MuCOE/CustomizableObjectMacroLibrary/CustomizableObjectMacroLibrary.h"
 #include "MuCOE/Nodes/CustomizableObjectNodeObject.h"
+#include "MuCOE/Nodes/CustomizableObjectNodeStaticString.h"
 #include "MuCOE/GraphTraversal.h"
 #include "MuCO/CustomizableObjectCustomVersion.h"
 #include "MuT/NodeModifier.h"
@@ -25,7 +27,40 @@ UEdGraphPin* UCustomizableObjectNodeModifierBase::OutputPin() const
 
 UEdGraphPin* UCustomizableObjectNodeModifierBase::RequiredTagsPin() const
 {
-	return RequiredTagsPinRef.Get();
+	if (CVarEnableMutableMacros.GetValueOnAnyThread())
+	{
+		return RequiredTagsPinRef.Get();
+	}
+
+	return nullptr;
+}
+
+
+TArray<FString> UCustomizableObjectNodeModifierBase::GetNodeRequiredTags(TArray<const UCustomizableObjectNodeMacroInstance*>* MacroContext) const
+{
+	const UEdGraphPin* ReqTagsPin = RequiredTagsPin();
+
+	if (ReqTagsPin && FollowInputPinArray(*ReqTagsPin).Num())
+	{
+		TArray<FString> OutTags;
+
+		for (const UEdGraphPin* StringPin : FollowInputPinArray(*ReqTagsPin))
+		{
+			const UEdGraphPin* SourceStringPin = GraphTraversal::FindStringSource(*StringPin, MacroContext);
+
+			if (SourceStringPin)
+			{
+				if (const UCustomizableObjectNodeStaticString* StringNode = Cast<UCustomizableObjectNodeStaticString>(SourceStringPin->GetOwningNode()))
+				{
+					OutTags.Add(StringNode->Value);
+				}
+			}
+		}
+
+		return OutTags;
+	}
+
+	return RequiredTags;
 }
 
 
@@ -36,8 +71,8 @@ bool UCustomizableObjectNodeModifierBase::IsApplicableTo(UCustomizableObjectNode
 		return false;
 	}
 
-	const TArray<FString>* EnabledTags = Candidate->GetEnableTags();
-	if (EnabledTags)
+	const TArray<FString>& EnabledTags = Candidate->GetEnableTags();
+	if (EnabledTags.Num())
 	{
 		FString InternalTag = Candidate->GetInternalTag();
 
@@ -47,7 +82,7 @@ bool UCustomizableObjectNodeModifierBase::IsApplicableTo(UCustomizableObjectNode
 		{
 			for (const FString& RequiredTag : RequiredTags)
 			{
-				if (InternalTag==RequiredTag || EnabledTags->Contains(RequiredTag))
+				if (InternalTag==RequiredTag || EnabledTags.Contains(RequiredTag))
 				{
 					return true;
 				}
@@ -59,7 +94,7 @@ bool UCustomizableObjectNodeModifierBase::IsApplicableTo(UCustomizableObjectNode
 		{
 			for (const FString& RequiredTag : RequiredTags)
 			{
-				if (InternalTag != RequiredTag && !EnabledTags->Contains(RequiredTag))
+				if (InternalTag != RequiredTag && !EnabledTags.Contains(RequiredTag))
 				{
 					return false;
 				}
@@ -278,7 +313,7 @@ void UCustomizableObjectNodeModifierBase::PostBackwardsCompatibleFixup()
 			UCustomizableObjectNode* ParentNode = GetCustomizableObjectExternalNode<UCustomizableObjectNode>(Tag.ParentObject.Get(), Tag.ParentNode);
 			if (ParentNode)
 			{
-				TArray<FString>* NodeEnableTags = ParentNode->GetEnableTags();
+				TArray<FString>* NodeEnableTags = ParentNode->GetEnableTagsArray();
 				if (NodeEnableTags)
 				{
 					NodeEnableTags->AddUnique(Tag.Tag);
@@ -300,7 +335,7 @@ FString UCustomizableObjectNodeModifierBase::MakeNodeAutoTag( UEdGraphNode* Node
 
 void UCustomizableObjectNodeModifierBase::AllocateDefaultPins(UCustomizableObjectNodeRemapPins* RemapPins)
 {
-	if (IsInMacro())
+	if (CVarEnableMutableMacros.GetValueOnAnyThread())
 	{
 		RequiredTagsPinRef = CustomCreatePin(EGPD_Input, UEdGraphSchema_CustomizableObject::PC_String, "Required Tags", true);
 	}
diff --git a/Engine/Plugins/Mutable/Source/CustomizableObjectEditor/Private/MuCOE/Nodes/CustomizableObjectNodeModifierBase.h b/Engine/Plugins/Mutable/Source/CustomizableObjectEditor/Private/MuCOE/Nodes/CustomizableObjectNodeModifierBase.h
--- a/Engine/Plugins/Mutable/Source/CustomizableObjectEditor/Private/MuCOE/Nodes/CustomizableObjectNodeModifierBase.h
+++ b/Engine/Plugins/Mutable/Source/CustomizableObjectEditor/Private/MuCOE/Nodes/CustomizableObjectNodeModifierBase.h
@@ -57,6 +57,7 @@ public:
 	// Own interface
 	virtual UEdGraphPin* OutputPin() const;
 	virtual UEdGraphPin* RequiredTagsPin() const;
+	TArray<FString> GetNodeRequiredTags(TArray<const UCustomizableObjectNodeMacroInstance*>* MacroContext = nullptr) const;
 
 	/** Return true if the tags and policy in this node would make this node modify the given parameter node. */
 	bool IsApplicableTo(UCustomizableObjectNode* Candidate);
diff --git a/Engine/Plugins/Mutable/Source/CustomizableObjectEditor/Private/MuCOE/Nodes/CustomizableObjectNodeModifierBaseDetails.cpp b/Engine/Plugins/Mutable/Source/CustomizableObjectEditor/Private/MuCOE/Nodes/CustomizableObjectNodeModifierBaseDetails.cpp
--- a/Engine/Plugins/Mutable/Source/CustomizableObjectEditor/Private/MuCOE/Nodes/CustomizableObjectNodeModifierBaseDetails.cpp
+++ b/Engine/Plugins/Mutable/Source/CustomizableObjectEditor/Private/MuCOE/Nodes/CustomizableObjectNodeModifierBaseDetails.cpp
@@ -70,6 +70,8 @@ void FCustomizableObjectNodeModifierBaseDetails::CustomizeDetails( IDetailLayout
 					SNew(STextBlock)
 						.Text(LOCTEXT("ModifierDetails_RequiredTags", "Required Tags"))
 						.Font(IDetailLayoutBuilder::GetDetailFont())
+						.IsEnabled(this, &FCustomizableObjectNodeModifierBaseDetails::IsTagsPropertyWidgetEnabled)
+						.ToolTipText(this, &FCustomizableObjectNodeModifierBaseDetails::TagsPropertyWidgetTooltip)
 				]
 		]
 		.ValueContent()
@@ -80,6 +82,8 @@ void FCustomizableObjectNodeModifierBaseDetails::CustomizeDetails( IDetailLayout
 				.TagArray(&Node->RequiredTags)
 				.EmptyListText(LOCTEXT("ModifierDetails_NoRequiredTagsWarning", "Warning: There are no required tags, so this modifier will not do anything."))
 				.OnTagListChanged( this, &FCustomizableObjectNodeModifierBaseDetails::OnRequiredTagsPropertyChanged )
+				.IsEnabled(this, &FCustomizableObjectNodeModifierBaseDetails::IsTagsPropertyWidgetEnabled)
+				.ToolTipText(this, &FCustomizableObjectNodeModifierBaseDetails::TagsPropertyWidgetTooltip)
 		];
 }
 
@@ -92,4 +96,27 @@ void FCustomizableObjectNodeModifierBaseDetails::OnRequiredTagsPropertyChanged()
 }
 
 
+bool FCustomizableObjectNodeModifierBaseDetails::IsTagsPropertyWidgetEnabled() const
+{
+	const UEdGraphPin* ReqTagsPin = Node->RequiredTagsPin();
+
+	// Disabled if there is a string node linked to the "Required Tags" pin.
+	return !(ReqTagsPin && FollowInputPinArray(*ReqTagsPin).Num());
+}
+
+
+FText FCustomizableObjectNodeModifierBaseDetails::TagsPropertyWidgetTooltip() const
+{
+	const UEdGraphPin* ReqTagsPin = Node->RequiredTagsPin();
+
+	if (ReqTagsPin && FollowInputPinArray(*ReqTagsPin).Num())
+	{
+		return LOCTEXT("RequiredTagsWidgetTooltip_Ignored", "Disabled. When there are string nodes linked to the Required Tags pin, the tag list is ignored.");
+	}
+	else
+	{
+		return LOCTEXT("RequiredTagsWidgetTooltip", "List of Tags that this node will Required.");
+	}
+}
+
 #undef LOCTEXT_NAMESPACE
diff --git a/Engine/Plugins/Mutable/Source/CustomizableObjectEditor/Private/MuCOE/Nodes/CustomizableObjectNodeModifierBaseDetails.h b/Engine/Plugins/Mutable/Source/CustomizableObjectEditor/Private/MuCOE/Nodes/CustomizableObjectNodeModifierBaseDetails.h
--- a/Engine/Plugins/Mutable/Source/CustomizableObjectEditor/Private/MuCOE/Nodes/CustomizableObjectNodeModifierBaseDetails.h
+++ b/Engine/Plugins/Mutable/Source/CustomizableObjectEditor/Private/MuCOE/Nodes/CustomizableObjectNodeModifierBaseDetails.h
@@ -35,6 +35,14 @@ protected:
 	/** Called when a required tags or policy property has changed, potentially changing the modified nodes. */
 	virtual void OnRequiredTagsPropertyChanged();
 
+private:
+
+	/** Disables/Enables the tags property widget. */
+	bool IsTagsPropertyWidgetEnabled() const;
+
+	/** Set the widget tooltip. */
+	FText TagsPropertyWidgetTooltip() const;
+
 private:
 
 	class UCustomizableObjectNodeModifierBase* Node = nullptr;
diff --git a/Engine/Plugins/Mutable/Source/CustomizableObjectEditor/Private/MuCOE/Nodes/CustomizableObjectNodeModifierEditMeshSectionBase.cpp b/Engine/Plugins/Mutable/Source/CustomizableObjectEditor/Private/MuCOE/Nodes/CustomizableObjectNodeModifierEditMeshSectionBase.cpp
--- a/Engine/Plugins/Mutable/Source/CustomizableObjectEditor/Private/MuCOE/Nodes/CustomizableObjectNodeModifierEditMeshSectionBase.cpp
+++ b/Engine/Plugins/Mutable/Source/CustomizableObjectEditor/Private/MuCOE/Nodes/CustomizableObjectNodeModifierEditMeshSectionBase.cpp
@@ -31,7 +31,7 @@ void UCustomizableObjectNodeModifierEditMeshSectionBase::BackwardsCompatibleFixu
 			FString NewLegacyTag = MakeNodeAutoTag(ParentNode);
 			RequiredTags.Add(NewLegacyTag);
 
-			TArray<FString>* NodeEnableTags = ParentNode->GetEnableTags();
+			TArray<FString>* NodeEnableTags = ParentNode->GetEnableTagsArray();
 			if (NodeEnableTags)
 			{
 				NodeEnableTags->AddUnique(NewLegacyTag);
diff --git a/Engine/Plugins/Mutable/Source/CustomizableObjectEditor/Private/MuCOE/Nodes/CustomizableObjectNodeModifierExtendMeshSection.cpp b/Engine/Plugins/Mutable/Source/CustomizableObjectEditor/Private/MuCOE/Nodes/CustomizableObjectNodeModifierExtendMeshSection.cpp
--- a/Engine/Plugins/Mutable/Source/CustomizableObjectEditor/Private/MuCOE/Nodes/CustomizableObjectNodeModifierExtendMeshSection.cpp
+++ b/Engine/Plugins/Mutable/Source/CustomizableObjectEditor/Private/MuCOE/Nodes/CustomizableObjectNodeModifierExtendMeshSection.cpp
@@ -10,8 +10,10 @@
 #include "MuCOE/GraphTraversal.h"
 #include "MuCOE/Nodes/CustomizableObjectNodeCopyMaterial.h"
 #include "MuCOE/Nodes/CustomizableObjectNodeSkeletalMesh.h"
+#include "MuCOE/Nodes/CustomizableObjectNodeStaticString.h"
 #include "MuCOE/Nodes/CustomizableObjectNodeTable.h"
 #include "MuCOE/CustomizableObjectPin.h"
+#include "MuCOE/CustomizableObjectMacroLibrary/CustomizableObjectMacroLibrary.h"
 
 class UCustomizableObjectNode;
 class UCustomizableObjectNodeRemapPins;
@@ -118,9 +120,11 @@ void UCustomizableObjectNodeModifierExtendMeshSection::AllocateDefaultPins(UCust
 	//Create Node Modifier Common Pins
 	Super::AllocateDefaultPins(RemapPins);
 
-	if (IsInMacro())
+	if (CVarEnableMutableMacros.GetValueOnAnyThread())
 	{
-		EnableTagsPinRef = CustomCreatePin(EGPD_Input, Schema->PC_String, FName("Enable Tags"), true);
+		UEdGraphPin* EnableTagsPin = CustomCreatePin(EGPD_Input, Schema->PC_String, FName("Enable Tags"), true);
+		EnableTagsPin->PinToolTip = "List of Tags that this node will Enable";
+		EnableTagsPinRef = EnableTagsPin;
 	}
 
 	// Begin texture pins
@@ -163,7 +167,12 @@ UEdGraphPin* UCustomizableObjectNodeModifierExtendMeshSection::AddMeshPin() cons
 
 UEdGraphPin* UCustomizableObjectNodeModifierExtendMeshSection::GetEnableTagsPin() const
 {
-	return EnableTagsPinRef.Get();
+	if (CVarEnableMutableMacros.GetValueOnAnyThread())
+	{
+		return EnableTagsPinRef.Get();
+	}
+
+	return nullptr;
 }
 
 
@@ -220,9 +229,41 @@ bool UCustomizableObjectNodeModifierExtendMeshSection::CustomRemovePin(UEdGraphP
 }
 
 
-TArray<FString>* UCustomizableObjectNodeModifierExtendMeshSection::GetEnableTags() 
+TArray<FString> UCustomizableObjectNodeModifierExtendMeshSection::GetEnableTags(TArray<const UCustomizableObjectNodeMacroInstance*>* MacroContext)
 { 
-	return &Tags; 
+	const UEdGraphPin* TagsPin = GetEnableTagsPin();
+
+	// Getting tags from linked pin
+	if (TagsPin && FollowInputPinArray(*TagsPin).Num())
+	{
+		TArray<FString> OutTags;
+
+		for (const UEdGraphPin* StringPin : FollowInputPinArray(*TagsPin))
+		{
+			const UEdGraphPin* SourceStringPin = GraphTraversal::FindStringSource(*StringPin, MacroContext);
+
+			if (SourceStringPin)
+			{
+				if (const UCustomizableObjectNodeStaticString* StringNode = Cast<UCustomizableObjectNodeStaticString>(SourceStringPin->GetOwningNode()))
+				{
+					OutTags.AddUnique(StringNode->Value);
+				}
+			}
+			
+		}
+
+		return OutTags;
+	}
+	else
+	{
+		return Tags;
+	}
+}
+
+
+TArray<FString>* UCustomizableObjectNodeModifierExtendMeshSection::GetEnableTagsArray()
+{
+	return &Tags;
 }
 
 
diff --git a/Engine/Plugins/Mutable/Source/CustomizableObjectEditor/Private/MuCOE/Nodes/CustomizableObjectNodeModifierExtendMeshSection.h b/Engine/Plugins/Mutable/Source/CustomizableObjectEditor/Private/MuCOE/Nodes/CustomizableObjectNodeModifierExtendMeshSection.h
--- a/Engine/Plugins/Mutable/Source/CustomizableObjectEditor/Private/MuCOE/Nodes/CustomizableObjectNodeModifierExtendMeshSection.h
+++ b/Engine/Plugins/Mutable/Source/CustomizableObjectEditor/Private/MuCOE/Nodes/CustomizableObjectNodeModifierExtendMeshSection.h
@@ -46,7 +46,8 @@ public:
 	virtual FString GetRefreshMessage() const override;
 	virtual bool IsSingleOutputNode() const override;
 	virtual bool CustomRemovePin(UEdGraphPin& Pin) override;
-	virtual TArray<FString>* GetEnableTags() override;
+	virtual TArray<FString> GetEnableTags(TArray<const UCustomizableObjectNodeMacroInstance*>* MacroContext = nullptr) override;
+	virtual TArray<FString>* GetEnableTagsArray() override;
 	virtual FString GetInternalTagDisplayName() override;
 
 	// Own interface
diff --git a/Engine/Plugins/Mutable/Source/CustomizableObjectEditor/Private/MuCOE/Nodes/CustomizableObjectNodeModifierExtendMeshSectionDetails.cpp b/Engine/Plugins/Mutable/Source/CustomizableObjectEditor/Private/MuCOE/Nodes/CustomizableObjectNodeModifierExtendMeshSectionDetails.cpp
--- a/Engine/Plugins/Mutable/Source/CustomizableObjectEditor/Private/MuCOE/Nodes/CustomizableObjectNodeModifierExtendMeshSectionDetails.cpp
+++ b/Engine/Plugins/Mutable/Source/CustomizableObjectEditor/Private/MuCOE/Nodes/CustomizableObjectNodeModifierExtendMeshSectionDetails.cpp
@@ -67,6 +67,8 @@ void FCustomizableObjectNodeModifierExtendMeshSectionDetails::CustomizeDetails(
 					SNew(STextBlock)
 						.Text(LOCTEXT("ExtendMeshSectionDetails_Tags", "Tags enabled for extended data"))
 						.Font(IDetailLayoutBuilder::GetDetailFont())
+						.IsEnabled(this, &FCustomizableObjectNodeModifierExtendMeshSectionDetails::IsTagsPropertyWidgetEnabled)
+						.ToolTipText(this, &FCustomizableObjectNodeModifierExtendMeshSectionDetails::TagsPropertyWidgetTooltip)
 				]
 		]
 		.ValueContent()
@@ -78,6 +80,8 @@ void FCustomizableObjectNodeModifierExtendMeshSectionDetails::CustomizeDetails(
 				.AllowInternalTags(false)
 				.EmptyListText(LOCTEXT("ExtendMeshSectionDetails_NoTags", "No tags enabled by this extended mesh section."))
 				.OnTagListChanged(this, &FCustomizableObjectNodeModifierExtendMeshSectionDetails::OnEnableTagsPropertyChanged)
+				.IsEnabled(this, &FCustomizableObjectNodeModifierExtendMeshSectionDetails::IsTagsPropertyWidgetEnabled)
+				.ToolTipText(this, &FCustomizableObjectNodeModifierExtendMeshSectionDetails::TagsPropertyWidgetTooltip)
 		];
 	}
 }
@@ -91,4 +95,28 @@ void FCustomizableObjectNodeModifierExtendMeshSectionDetails::OnEnableTagsProper
 }
 
 
+bool FCustomizableObjectNodeModifierExtendMeshSectionDetails::IsTagsPropertyWidgetEnabled() const
+{
+	const UEdGraphPin* EnableTagsPin = Node->GetEnableTagsPin();
+
+	// Disabled if there is a string node linked to the "Enable Tags" pin.
+	return !(EnableTagsPin && FollowInputPinArray(*EnableTagsPin).Num());
+}
+
+
+FText FCustomizableObjectNodeModifierExtendMeshSectionDetails::TagsPropertyWidgetTooltip() const
+{
+	const UEdGraphPin* EnableTagsPin = Node->GetEnableTagsPin();
+
+	if (EnableTagsPin && FollowInputPinArray(*EnableTagsPin).Num())
+	{
+		return LOCTEXT("EnableTagsWidgetTooltip_Ignored", "Disabled. When there are string nodes linked to the Enable Tags pin, the tag list is ignored.");
+	}
+	else
+	{
+		return LOCTEXT("EnableTagsWidgetTooltip", "List of Tags that this node will enable.");
+	}
+}
+
+
 #undef LOCTEXT_NAMESPACE
diff --git a/Engine/Plugins/Mutable/Source/CustomizableObjectEditor/Private/MuCOE/Nodes/CustomizableObjectNodeModifierExtendMeshSectionDetails.h b/Engine/Plugins/Mutable/Source/CustomizableObjectEditor/Private/MuCOE/Nodes/CustomizableObjectNodeModifierExtendMeshSectionDetails.h
--- a/Engine/Plugins/Mutable/Source/CustomizableObjectEditor/Private/MuCOE/Nodes/CustomizableObjectNodeModifierExtendMeshSectionDetails.h
+++ b/Engine/Plugins/Mutable/Source/CustomizableObjectEditor/Private/MuCOE/Nodes/CustomizableObjectNodeModifierExtendMeshSectionDetails.h
@@ -30,8 +30,15 @@ public:
 
 private:
 
-	class UCustomizableObjectNodeModifierExtendMeshSection* Node = nullptr;
+	/** Disables/Enables the tags property widget. */
+	bool IsTagsPropertyWidgetEnabled() const;
+
+	/** Set the widget tooltip. */
+	FText TagsPropertyWidgetTooltip() const;
 
+private:
+
+	class UCustomizableObjectNodeModifierExtendMeshSection* Node = nullptr;
 
 	/** */
 	TSharedPtr<IPropertyHandle> EnableTagsPropertyHandle;
diff --git a/Engine/Plugins/Mutable/Source/CustomizableObjectEditor/Private/MuCOE/Nodes/CustomizableObjectNodeModifierMorphMeshSection.cpp b/Engine/Plugins/Mutable/Source/CustomizableObjectEditor/Private/MuCOE/Nodes/CustomizableObjectNodeModifierMorphMeshSection.cpp
--- a/Engine/Plugins/Mutable/Source/CustomizableObjectEditor/Private/MuCOE/Nodes/CustomizableObjectNodeModifierMorphMeshSection.cpp
+++ b/Engine/Plugins/Mutable/Source/CustomizableObjectEditor/Private/MuCOE/Nodes/CustomizableObjectNodeModifierMorphMeshSection.cpp
@@ -8,6 +8,7 @@
 #include "MuCOE/GraphTraversal.h"
 #include "MuCOE/Nodes/CustomizableObjectNodeMaterial.h"
 #include "MuCOE/Nodes/CustomizableObjectNodeSkeletalMesh.h"
+#include "MuCOE/CustomizableObjectMacroLibrary/CustomizableObjectMacroLibrary.h"
 
 class UCustomizableObjectNodeRemapPins;
 struct FPropertyChangedEvent;
@@ -27,7 +28,7 @@ void UCustomizableObjectNodeModifierMorphMeshSection::AllocateDefaultPins(UCusto
 	UEdGraphPin* FactorPin = CustomCreatePin(EGPD_Input, Schema->PC_Float, FName(*PinName));
 	FactorPin->bDefaultValueIsIgnored = true;
 
-	if (IsInMacro())
+	if (CVarEnableMutableMacros.GetValueOnAnyThread())
 	{
 		PinName = TEXT("Morph Target Name");
 		MorphTargetNamePinRef = CustomCreatePin(EGPD_Input, Schema->PC_String, FName(*PinName));
@@ -64,7 +65,18 @@ bool UCustomizableObjectNodeModifierMorphMeshSection::IsSingleOutputNode() const
 
 UEdGraphPin* UCustomizableObjectNodeModifierMorphMeshSection::MorphTargetNamePin() const
 {
-	return MorphTargetNamePinRef.Get();
+	if (CVarEnableMutableMacros.GetValueOnAnyThread())
+	{
+		return MorphTargetNamePinRef.Get();
+	}
+
+	return nullptr;
+}
+
+
+bool UCustomizableObjectNodeModifierMorphMeshSection::CreateStaticStringPinWidget() const
+{
+	return false;
 }
 
 
diff --git a/Engine/Plugins/Mutable/Source/CustomizableObjectEditor/Private/MuCOE/Nodes/CustomizableObjectNodeModifierMorphMeshSection.h b/Engine/Plugins/Mutable/Source/CustomizableObjectEditor/Private/MuCOE/Nodes/CustomizableObjectNodeModifierMorphMeshSection.h
--- a/Engine/Plugins/Mutable/Source/CustomizableObjectEditor/Private/MuCOE/Nodes/CustomizableObjectNodeModifierMorphMeshSection.h
+++ b/Engine/Plugins/Mutable/Source/CustomizableObjectEditor/Private/MuCOE/Nodes/CustomizableObjectNodeModifierMorphMeshSection.h
@@ -33,6 +33,7 @@ public:
 	virtual void AllocateDefaultPins(UCustomizableObjectNodeRemapPins* RemapPins) override;
 	virtual FString GetRefreshMessage() const override;
 	virtual bool IsSingleOutputNode() const override;
+	virtual bool CreateStaticStringPinWidget() const override;
 
 	UEdGraphPin* FactorPin() const
 	{
diff --git a/Engine/Plugins/Mutable/Source/CustomizableObjectEditor/Private/MuCOE/Nodes/CustomizableObjectNodeModifierMorphMeshSectionDetails.cpp b/Engine/Plugins/Mutable/Source/CustomizableObjectEditor/Private/MuCOE/Nodes/CustomizableObjectNodeModifierMorphMeshSectionDetails.cpp
--- a/Engine/Plugins/Mutable/Source/CustomizableObjectEditor/Private/MuCOE/Nodes/CustomizableObjectNodeModifierMorphMeshSectionDetails.cpp
+++ b/Engine/Plugins/Mutable/Source/CustomizableObjectEditor/Private/MuCOE/Nodes/CustomizableObjectNodeModifierMorphMeshSectionDetails.cpp
@@ -51,7 +51,6 @@ void FCustomizableObjectNodeModifierMorphMeshSectionDetails::CustomizeDetails( I
 	DetailBuilder.HideProperty(GET_MEMBER_NAME_CHECKED(UCustomizableObjectNodeModifierWithMaterial, ReferenceMaterial), UCustomizableObjectNodeModifierWithMaterial::StaticClass());
 	
 	// Add a morph selection widget.
-	IDetailCategoryBuilder& MorphCategory = DetailBuilder.EditCategory( "Morph" );
 	TSharedRef<IPropertyHandle> MorphTargetNameProperty = DetailBuilder.GetProperty("MorphTargetName");
 
 	if (!Node->IsInMacro())
@@ -59,46 +58,30 @@ void FCustomizableObjectNodeModifierMorphMeshSectionDetails::CustomizeDetails( I
 		// Scan for hint morph names
 		RefreshMorphOptions();
 
-		MorphCategory.AddCustomRow(LOCTEXT("MorphMaterialDetails_Target", "Target"))
+		DetailBuilder.EditDefaultProperty(MorphTargetNameProperty)->
+		CustomWidget()
+		.NameContent()
 		[
-			SNew(SProperty, MorphTargetNameProperty)
-			.ShouldDisplayName(false)
-			.CustomWidget()
+			SNew(STextBlock)
+			.Text(LOCTEXT("MorphMaterialDetails_MorphTarget", "Morph Target Name")).Font(IDetailLayoutBuilder::GetDetailFont())
+			.IsEnabled(this, &FCustomizableObjectNodeModifierMorphMeshSectionDetails::IsMorphNameSelectorWidgetEnabled)
+			.ToolTipText(this, &FCustomizableObjectNodeModifierMorphMeshSectionDetails::MorphNameSelectorWidgetTooltip)
+		]
+		.ValueContent()
+		[
+			SAssignNew(this->MorphCombo, SMutableSearchComboBox)
+			.ButtonStyle(FAppStyle::Get(), "NoBorder")
+			.OptionsSource(&MorphOptionsSource)
+			.OnSelectionChanged(this, &FCustomizableObjectNodeModifierMorphMeshSectionDetails::OnMorphTargetComboBoxSelectionChanged)
+			.Content()
 			[
-				SNew(SBorder)
-				.BorderImage(UE_MUTABLE_GET_BRUSH("NoBorder"))
-				.Padding(FMargin(0.0f, 0.0f, 10.0f, 0.0f))
-				[
-					SNew(SHorizontalBox)
-
-					+ SHorizontalBox::Slot()
-					.FillWidth(10.0f)
-					.VAlign(VAlign_Center)
-					[
-						SNew(STextBlock)
-							.Text(LOCTEXT("MorphMaterialDetails_MorphTarget", "Morph Target"))
-							.Font(IDetailLayoutBuilder::GetDetailFont())
-					]
-
-					+ SHorizontalBox::Slot()
-					.FillWidth(10.0f)
-					.HAlign(HAlign_Fill)
-					.VAlign(VAlign_Center)
-					[
-						SAssignNew(this->MorphCombo , SMutableSearchComboBox)
-						.ButtonStyle(FAppStyle::Get(), "NoBorder")
-						.OptionsSource(&MorphOptionsSource)
-						.OnSelectionChanged(this, &FCustomizableObjectNodeModifierMorphMeshSectionDetails::OnMorphTargetComboBoxSelectionChanged)
-						.Content()
-						[
-							SNew(SEditableTextBox)
-							.Font(IDetailLayoutBuilder::GetDetailFont())
-							.Text_Lambda([&]() { return Node ? FText::FromString(Node->MorphTargetName) : FText();})
-							.OnTextChanged(this, &FCustomizableObjectNodeModifierMorphMeshSectionDetails::OnMorphTargetComboBoxSelectionChanged)
-						]
-					]
-				]
+				SNew(SEditableTextBox)
+				.Font(IDetailLayoutBuilder::GetDetailFont())
+				.Text_Lambda([&]() { return Node ? FText::FromString(Node->MorphTargetName) : FText(); })
+				.OnTextChanged(this, &FCustomizableObjectNodeModifierMorphMeshSectionDetails::OnMorphTargetComboBoxSelectionChanged)
 			]
+			.IsEnabled(this, &FCustomizableObjectNodeModifierMorphMeshSectionDetails::IsMorphNameSelectorWidgetEnabled)
+			.ToolTipText(this, &FCustomizableObjectNodeModifierMorphMeshSectionDetails::MorphNameSelectorWidgetTooltip)
 		];
 	}
 	else
@@ -107,7 +90,8 @@ void FCustomizableObjectNodeModifierMorphMeshSectionDetails::CustomizeDetails( I
 		CustomWidget()
 		.NameContent()
 		[
-			SNew(STextBlock).Text(LOCTEXT("MorphMaterialDetails_MorphTargetName", "Morph Target Name")).Font(IDetailLayoutBuilder::GetDetailFont())
+			SNew(STextBlock).Text(LOCTEXT("MorphMaterialDetails_MorphTargetName", "Morph Target Name"))
+			.Font(IDetailLayoutBuilder::GetDetailFont())
 		]
 		.ValueContent()
 		[
@@ -240,7 +224,7 @@ TSharedPtr<SMutableSearchComboBox::FFilteredOption> FCustomizableObjectNodeModif
 		{
 			Option = MakeShared<SMutableSearchComboBox::FFilteredOption>();
 			Option->Parent = ParentOption;
-			Option->DisplayOption = ObjectNode->ObjectName;
+			Option->DisplayOption = ObjectNode->GetObjectName();
 			if (Option->DisplayOption.IsEmpty())
 			{
 				Option->DisplayOption = "Unnamed Object";
@@ -252,7 +236,7 @@ TSharedPtr<SMutableSearchComboBox::FFilteredOption> FCustomizableObjectNodeModif
 		{
 			Option = MakeShared<SMutableSearchComboBox::FFilteredOption>();
 			Option->Parent = ParentOption;
-			Option->DisplayOption = GroupNode->GroupName;
+			Option->DisplayOption = GroupNode->GetGroupName();
 			if (Option->DisplayOption.IsEmpty())
 			{
 				Option->DisplayOption = "Unnamed Group";
@@ -340,4 +324,24 @@ void FCustomizableObjectNodeModifierMorphMeshSectionDetails::AddMorphsFromNode(
 }
 
 
+bool FCustomizableObjectNodeModifierMorphMeshSectionDetails::IsMorphNameSelectorWidgetEnabled() const
+{
+	// Disabled if there is a string node linked to the "Required Tags" pin.
+	return !(Node->MorphTargetNamePin() && FollowInputPin(*Node->MorphTargetNamePin()));
+}
+
+
+FText FCustomizableObjectNodeModifierMorphMeshSectionDetails::MorphNameSelectorWidgetTooltip() const
+{
+	if (Node->MorphTargetNamePin() && FollowInputPin(*Node->MorphTargetNamePin()))
+	{
+		return LOCTEXT("MorphTargetNameWidgetTooltip_Ignored", "Disabled. When there is a string node linked to the Morph Target Name pin, the morph target name selected in this widget is ignored.");
+	}
+	else
+	{
+		return LOCTEXT("MorphTargetNameWidgetTooltip", "Select the morph target name.");
+	}
+}
+
+
 #undef LOCTEXT_NAMESPACE
diff --git a/Engine/Plugins/Mutable/Source/CustomizableObjectEditor/Private/MuCOE/Nodes/CustomizableObjectNodeModifierMorphMeshSectionDetails.h b/Engine/Plugins/Mutable/Source/CustomizableObjectEditor/Private/MuCOE/Nodes/CustomizableObjectNodeModifierMorphMeshSectionDetails.h
--- a/Engine/Plugins/Mutable/Source/CustomizableObjectEditor/Private/MuCOE/Nodes/CustomizableObjectNodeModifierMorphMeshSectionDetails.h
+++ b/Engine/Plugins/Mutable/Source/CustomizableObjectEditor/Private/MuCOE/Nodes/CustomizableObjectNodeModifierMorphMeshSectionDetails.h
@@ -40,11 +40,20 @@ private:
 	/** */
 	TSharedPtr<SMutableSearchComboBox> MorphCombo;
 	TArray< TSharedRef<SMutableSearchComboBox::FFilteredOption> > MorphOptionsSource;
+	
+private:
+
 	void OnMorphTargetComboBoxSelectionChanged(const FText& NewText);
 
 	/** */
 	void AddMorphsFromNode(UEdGraphNode* Candidate, TMap<UEdGraphNode*, TSharedPtr<SMutableSearchComboBox::FFilteredOption>>& AddedOptions);
-		
+
+	/** Disables/Enables the tags property widget. */
+	bool IsMorphNameSelectorWidgetEnabled() const;
+
+	/** Set the widget tooltip. */
+	FText MorphNameSelectorWidgetTooltip() const;
+
 	/** */
 	TSharedPtr<SMutableSearchComboBox::FFilteredOption> AddNodeHierarchyOptions(UEdGraphNode* Node, TMap<UEdGraphNode*, TSharedPtr<SMutableSearchComboBox::FFilteredOption>>& AddedOptions);
 
diff --git a/Engine/Plugins/Mutable/Source/CustomizableObjectEditor/Private/MuCOE/Nodes/CustomizableObjectNodeObject.cpp b/Engine/Plugins/Mutable/Source/CustomizableObjectEditor/Private/MuCOE/Nodes/CustomizableObjectNodeObject.cpp
--- a/Engine/Plugins/Mutable/Source/CustomizableObjectEditor/Private/MuCOE/Nodes/CustomizableObjectNodeObject.cpp
+++ b/Engine/Plugins/Mutable/Source/CustomizableObjectEditor/Private/MuCOE/Nodes/CustomizableObjectNodeObject.cpp
@@ -21,6 +21,7 @@
 #include "MuCOE/Nodes/CustomizableObjectNodeComponentMeshAddTo.h"
 #include "MuCOE/Nodes/CustomizableObjectNodeCopyMaterial.h"
 #include "MuCOE/Nodes/CustomizableObjectNodeMaterialSwitch.h"
+#include "MuCOE/CustomizableObjectMacroLibrary/CustomizableObjectMacroLibrary.h"
 #include "Settings/EditorStyleSettings.h"
 
 class UCustomizableObjectNodeRemapPins;
@@ -262,11 +263,11 @@ void UCustomizableObjectNodeObject::BackwardsCompatibleFixup(int32 CustomizableO
 
 				++NodesCreated;
 				
-				NodeComponentMesh->ComponentName = MeshComponent.Name;
+				NodeComponentMesh->SetComponentName(MeshComponent.Name);
 				NodeComponentMesh->NumLODs = NumLODs_DEPRECATED;
 				if (FComponentSettings* Result = ComponentSettings_DEPRECATED.FindByPredicate([&](const FComponentSettings& Settings)
 				{
-					return Settings.ComponentName == NodeComponentMesh->ComponentName;
+					return Settings.ComponentName == MeshComponent.Name;
 				}))
 				{
 					NodeComponentMesh->LODReductionSettings = Result->LODReductionSettings;
@@ -313,7 +314,7 @@ void UCustomizableObjectNodeObject::BackwardsCompatibleFixup(int32 CustomizableO
 
 					++NodesCreated;
 
-					NodeComponentMeshAddTo->ParentComponentName = MeshComponentName;
+					NodeComponentMeshAddTo->SetParentComponentName(MeshComponentName);
 					NodeComponentMeshAddTo->NumLODs = NumLODs_DEPRECATED;
 					NodeComponentMeshAddTo->AutoLODStrategy = AutoLODStrategy_DEPRECATED;
 					
@@ -322,7 +323,7 @@ void UCustomizableObjectNodeObject::BackwardsCompatibleFixup(int32 CustomizableO
 
 					ComponentsPin()->MakeLinkTo(NodeComponentMeshAddTo->OutputPin.Get());
 
-					ExistingNodeComponentMeshAddTo.Add(NodeComponentMeshAddTo->ParentComponentName ,NodeComponentMeshAddTo);
+					ExistingNodeComponentMeshAddTo.Add(MeshComponentName, NodeComponentMeshAddTo);
 					
 					return NodeComponentMeshAddTo;
 				};
@@ -475,7 +476,7 @@ void UCustomizableObjectNodeObject::BackwardsCompatibleFixup(int32 CustomizableO
 			{
 				if (UCustomizableObjectNodeComponentMesh* NodeComponentMesh = Cast<UCustomizableObjectNodeComponentMesh>(Pin->GetOwningNode()))
 				{
-					NodeComponentMeshes.Emplace(NodeComponentMesh->ComponentName, NodeComponentMesh);
+					NodeComponentMeshes.Emplace(NodeComponentMesh->GetComponentName(), NodeComponentMesh);
 				}
 			}
 
@@ -486,7 +487,7 @@ void UCustomizableObjectNodeObject::BackwardsCompatibleFixup(int32 CustomizableO
 			{
 				if (UCustomizableObjectNodeComponentMeshAddTo* NodeComponentMeshAddTo = Cast<UCustomizableObjectNodeComponentMeshAddTo>(Pin->GetOwningNode()))
 				{
-					UCustomizableObjectNodeComponentMesh** Result = NodeComponentMeshes.Find(NodeComponentMeshAddTo->ParentComponentName);
+					UCustomizableObjectNodeComponentMesh** Result = NodeComponentMeshes.Find(NodeComponentMeshAddTo->GetParentComponentName());
 					if (!Result)
 					{
 						continue;
@@ -596,13 +597,6 @@ void UCustomizableObjectNodeObject::AllocateDefaultPins(UCustomizableObjectNodeR
 	{
 		OutputPin->bHidden = true;
 	}
-	else
-	{
-		if (IsInMacro())
-		{
-			NamePinRef = CustomCreatePin(EGPD_Input, Schema->PC_String, "Name");
-		}
-	}
 }
 
 
@@ -728,6 +722,7 @@ bool UCustomizableObjectNodeObject::IsSingleOutputNode() const
 	return true;
 }
 
+
 void UCustomizableObjectNodeObject::OnRenameNode(const FString& NewName)
 {
 	if (!NewName.IsEmpty())
@@ -737,4 +732,16 @@ void UCustomizableObjectNodeObject::OnRenameNode(const FString& NewName)
 }
 
 
+FString UCustomizableObjectNodeObject::GetObjectName(TArray<const UCustomizableObjectNodeMacroInstance*>* MacroContext ) const
+{
+	return ObjectName;
+}
+
+
+void UCustomizableObjectNodeObject::SetObjectName(const FString& Name)
+{
+	ObjectName = Name;
+}
+
+
 #undef LOCTEXT_NAMESPACE
diff --git a/Engine/Plugins/Mutable/Source/CustomizableObjectEditor/Private/MuCOE/Nodes/CustomizableObjectNodeObject.h b/Engine/Plugins/Mutable/Source/CustomizableObjectEditor/Private/MuCOE/Nodes/CustomizableObjectNodeObject.h
--- a/Engine/Plugins/Mutable/Source/CustomizableObjectEditor/Private/MuCOE/Nodes/CustomizableObjectNodeObject.h
+++ b/Engine/Plugins/Mutable/Source/CustomizableObjectEditor/Private/MuCOE/Nodes/CustomizableObjectNodeObject.h
@@ -111,9 +111,13 @@ public:
 
 	UCustomizableObjectNodeObject();
 
+protected:
+
 	UPROPERTY(EditAnywhere, Category=CustomizableObject)
 	FString ObjectName;
 
+public:
+
 	UPROPERTY(EditAnywhere, Category = UI, meta = (DisplayName = "Parameter UI Metadata"))
 	FMutableParamUIMetadata ParamUIMetadata;
 
@@ -160,11 +164,13 @@ public:
 	// UCustomizableObjectNode interface
 	virtual void BackwardsCompatibleFixup(int32 CustomizableObjectCustomVersion) override;
 	virtual void PostBackwardsCompatibleFixup() override;
-	void AllocateDefaultPins(UCustomizableObjectNodeRemapPins* RemapPins) override;
+	virtual void AllocateDefaultPins(UCustomizableObjectNodeRemapPins* RemapPins) override;
 	virtual bool IsNodeSupportedInMacros() const override;
 
 	// Own Interface
 	void SetParentObject(UCustomizableObject* CustomizableParentObject);
+	virtual FString GetObjectName(TArray<const UCustomizableObjectNodeMacroInstance*>* MacroContext = nullptr) const;
+	void SetObjectName(const FString& Name);
 
 	// Own interface
 	UPROPERTY()
@@ -205,7 +211,5 @@ private:
 	static const FName ModifiersPinName;
 	static const FName OutputPinName;
 
-	UPROPERTY()
-	FEdGraphPinReference NamePinRef;
 };
 
diff --git a/Engine/Plugins/Mutable/Source/CustomizableObjectEditor/Private/MuCOE/Nodes/CustomizableObjectNodeObjectChild.cpp b/Engine/Plugins/Mutable/Source/CustomizableObjectEditor/Private/MuCOE/Nodes/CustomizableObjectNodeObjectChild.cpp
--- a/Engine/Plugins/Mutable/Source/CustomizableObjectEditor/Private/MuCOE/Nodes/CustomizableObjectNodeObjectChild.cpp
+++ b/Engine/Plugins/Mutable/Source/CustomizableObjectEditor/Private/MuCOE/Nodes/CustomizableObjectNodeObjectChild.cpp
@@ -2,6 +2,10 @@
 
 #include "MuCOE/Nodes/CustomizableObjectNodeObjectChild.h"
 
+#include "MuCOE/CustomizableObjectMacroLibrary/CustomizableObjectMacroLibrary.h"
+#include "MuCOE/EdGraphSchema_CustomizableObject.h"
+#include "MuCOE/GraphTraversal.h"
+#include "MuCOE/Nodes/CustomizableObjectNodeStaticString.h"
 
 #define LOCTEXT_NAMESPACE "CustomizableObjectEditor"
 
@@ -15,12 +19,13 @@ UCustomizableObjectNodeObjectChild::UCustomizableObjectNodeObjectChild()
 
 FText UCustomizableObjectNodeObjectChild::GetNodeTitle(ENodeTitleType::Type TitleType) const
 {
+	const UEdGraphPin* ChildNamePin = NamePin.Get();
 	FFormatNamedArguments Args;
 	Args.Add(TEXT("ObjectName"), FText::FromString(ObjectName));
 
-	if (TitleType == ENodeTitleType::ListView)
+	if (TitleType == ENodeTitleType::ListView || (ChildNamePin && ChildNamePin->LinkedTo.Num()))
 	{
-		return FText::Format(LOCTEXT("Child_Object_Title_List", "{ObjectName} - Child Object"), Args);
+		return FText::Format(LOCTEXT("Child_Object_Title_List", "Child Object"), Args);
 	}
 	else if (TitleType == ENodeTitleType::EditableTitle)
 	{
@@ -51,8 +56,24 @@ bool UCustomizableObjectNodeObjectChild::CanDuplicateNode() const
 }
 
 
-bool UCustomizableObjectNodeObjectChild::IsNodeSupportedInMacros() const
+void UCustomizableObjectNodeObjectChild::PinConnectionListChanged(UEdGraphPin* Pin)
+{
+	if (Pin == NamePin)
+	{
+		GetGraph()->NotifyGraphChanged();
+	}
+}
+
+
+bool UCustomizableObjectNodeObjectChild::GetCanRenameNode() const
 {
+	const UEdGraphPin* ChildNamePin = NamePin.Get();
+
+	if (CVarEnableMutableMacros.GetValueOnAnyThread() && ChildNamePin && ChildNamePin->LinkedTo.Num())
+	{
+		return false;
+	}
+
 	return true;
 }
 
@@ -60,8 +81,45 @@ bool UCustomizableObjectNodeObjectChild::IsNodeSupportedInMacros() const
 FText UCustomizableObjectNodeObjectChild::GetTooltipText() const
 {
 	return LOCTEXT("Child_Object_Tooltip",
-	"Defines a customizable object children in the same asset as its parent, to ease the addition of small Customizable Objects directly into\ntheir parents asset. Functionally equivalent to the Base Object Node when it has a parent defined. It can be a children of the root\nobject or of any children, allowing arbitrary nesting of objects. Defines materials that can be added to its parent, modify it, remove\nparts of it or change any of its parameters. Also defines properties for others to use or modify.");
+		"Defines a customizable object children in the same asset as its parent, to ease the addition of small Customizable Objects directly into\ntheir parents asset. Functionally equivalent to the Base Object Node when it has a parent defined. It can be a children of the root\nobject or of any children, allowing arbitrary nesting of objects. Defines materials that can be added to its parent, modify it, remove\nparts of it or change any of its parameters. Also defines properties for others to use or modify.");
+
+}
+
+
+void UCustomizableObjectNodeObjectChild::AllocateDefaultPins(UCustomizableObjectNodeRemapPins* RemapPins)
+{
+	Super::AllocateDefaultPins(RemapPins);
 
+	if (CVarEnableMutableMacros.GetValueOnAnyThread())
+	{
+		NamePin = CustomCreatePin(EGPD_Input, UEdGraphSchema_CustomizableObject::PC_String, "Name");
+	}
+}
+
+
+bool UCustomizableObjectNodeObjectChild::IsNodeSupportedInMacros() const
+{
+	return true;
 }
 
+
+FString UCustomizableObjectNodeObjectChild::GetObjectName(TArray<const UCustomizableObjectNodeMacroInstance*>* MacroContext) const
+{
+	if (CVarEnableMutableMacros.GetValueOnAnyThread() && NamePin.Get())
+	{
+		if (const UEdGraphPin* LinkedPin = FollowInputPin(*NamePin.Get()))
+		{
+			const UEdGraphPin* StringPin = GraphTraversal::FindStringSource(*LinkedPin, MacroContext);
+
+			if (const UCustomizableObjectNodeStaticString* StringNode = StringPin ? Cast<UCustomizableObjectNodeStaticString>(StringPin->GetOwningNode()) : nullptr)
+			{
+				return StringNode->Value;
+			}
+		}
+	}
+
+	return ObjectName;
+}
+
+
 #undef LOCTEXT_NAMESPACE
diff --git a/Engine/Plugins/Mutable/Source/CustomizableObjectEditor/Private/MuCOE/Nodes/CustomizableObjectNodeObjectChild.h b/Engine/Plugins/Mutable/Source/CustomizableObjectEditor/Private/MuCOE/Nodes/CustomizableObjectNodeObjectChild.h
--- a/Engine/Plugins/Mutable/Source/CustomizableObjectEditor/Private/MuCOE/Nodes/CustomizableObjectNodeObjectChild.h
+++ b/Engine/Plugins/Mutable/Source/CustomizableObjectEditor/Private/MuCOE/Nodes/CustomizableObjectNodeObjectChild.h
@@ -17,15 +17,26 @@ public:
 
 	UCustomizableObjectNodeObjectChild();
 
+	UPROPERTY()
+	FEdGraphPinReference NamePin;
+
 public:
+
 	// UEdGraphNode interface
-	FText GetNodeTitle(ENodeTitleType::Type TitleType) const override;
-	void PrepareForCopying() override;
-	FText GetTooltipText() const override;
-	bool CanUserDeleteNode() const override;
-	bool CanDuplicateNode() const override;
+	virtual FText GetNodeTitle(ENodeTitleType::Type TitleType) const override;
+	virtual void PrepareForCopying() override;
+	virtual FText GetTooltipText() const override;
+	virtual bool CanUserDeleteNode() const override;
+	virtual bool CanDuplicateNode() const override;
+	virtual void PinConnectionListChanged(UEdGraphPin* Pin) override;
+	virtual bool GetCanRenameNode() const;
 
 	// UCustomizableObjectNode interface
 	virtual bool IsNodeSupportedInMacros() const override;
+	virtual void AllocateDefaultPins(UCustomizableObjectNodeRemapPins* RemapPins) override;
+
+	// UCustomizableObjectNodeObject interface
+	virtual FString GetObjectName(TArray<const UCustomizableObjectNodeMacroInstance*>* MacroContext = nullptr) const override;
+
 };
 
diff --git a/Engine/Plugins/Mutable/Source/CustomizableObjectEditor/Private/MuCOE/Nodes/CustomizableObjectNodeObjectDetails.cpp b/Engine/Plugins/Mutable/Source/CustomizableObjectEditor/Private/MuCOE/Nodes/CustomizableObjectNodeObjectDetails.cpp
--- a/Engine/Plugins/Mutable/Source/CustomizableObjectEditor/Private/MuCOE/Nodes/CustomizableObjectNodeObjectDetails.cpp
+++ b/Engine/Plugins/Mutable/Source/CustomizableObjectEditor/Private/MuCOE/Nodes/CustomizableObjectNodeObjectDetails.cpp
@@ -8,6 +8,7 @@
 #include "MuCO/CustomizableObjectPrivate.h"
 #include "MuCOE/CustomizableObjectEditorUtilities.h"
 #include "MuCOE/CustomizableObjectGraph.h"
+#include "MuCOE/CustomizableObjectMacroLibrary/CustomizableObjectMacroLibrary.h"
 #include "MuCOE/GraphTraversal.h"
 #include "MuCOE/Nodes/CustomizableObjectNode.h"
 #include "MuCOE/Nodes/CustomizableObjectNodeObject.h"
@@ -211,7 +212,7 @@ void FCustomizableObjectNodeObjectDetails::CustomizeDetails( IDetailLayoutBuilde
 
 				for (UCustomizableObjectNodeObjectGroup* GroupNode : GroupNodes)
 				{
-					GroupNodeComboOptions.Add(MakeShareable(new FString(GroupNode->GroupName)));
+					GroupNodeComboOptions.Add(MakeShareable(new FString(GroupNode->GetGroupName())));
 
 					if (BaseObjectNode->ParentObjectGroupId == GroupNode->NodeGuid)
 					{
@@ -252,6 +253,11 @@ void FCustomizableObjectNodeObjectDetails::CustomizeDetails( IDetailLayoutBuilde
 		else
 		{
 			DetailBuilder.HideProperty("NumMeshComponents");
+
+			if (CVarEnableMutableMacros.GetValueOnAnyThread())
+			{
+				DetailBuilder.HideProperty("ObjectName");
+			}
 		}
 
 		
@@ -335,7 +341,7 @@ void FCustomizableObjectNodeObjectDetails::OnGroupNodeComboBoxSelectionChanged(T
 
 		for (UCustomizableObjectNodeObjectGroup* GroupNode : GroupNodes)
 		{
-			if (*Selection == GroupNode->GroupName)
+			if (*Selection == GroupNode->GetGroupName())
 			{
 				const FScopedTransaction Transaction(LOCTEXT("ChangedAttachedToExternalObjectTransaction", "Changed Attached to External Object"));
 				BaseObjectNode->Modify();
@@ -399,7 +405,7 @@ void FCustomizableObjectNodeObjectDetails::FillParameterNamesArray()
 		{
 			if (ColorParameterNode)
 			{
-				BaseObjectNode->ParameterNames.Add(ColorParameterNode->ParameterName);
+				BaseObjectNode->ParameterNames.Add(ColorParameterNode->GetParameterName());
 			}
 		}
 
@@ -410,7 +416,7 @@ void FCustomizableObjectNodeObjectDetails::FillParameterNamesArray()
 		{
 			if (FloatParameterNode)
 			{
-				BaseObjectNode->ParameterNames.Add(FloatParameterNode->ParameterName);
+				BaseObjectNode->ParameterNames.Add(FloatParameterNode->GetParameterName());
 			}
 		}
 
@@ -421,7 +427,7 @@ void FCustomizableObjectNodeObjectDetails::FillParameterNamesArray()
 		{
 			if (EnumParameterNode)
 			{
-				BaseObjectNode->ParameterNames.Add(EnumParameterNode->ParameterName);
+				BaseObjectNode->ParameterNames.Add(EnumParameterNode->GetParameterName());
 			}
 		}
 
@@ -432,7 +438,7 @@ void FCustomizableObjectNodeObjectDetails::FillParameterNamesArray()
 		{
 			if (GroupProjectorParameterNode)
 			{
-				BaseObjectNode->ParameterNames.Add(GroupProjectorParameterNode->ParameterName);
+				BaseObjectNode->ParameterNames.Add(GroupProjectorParameterNode->GetParameterName());
 			}
 		}
 
@@ -443,7 +449,7 @@ void FCustomizableObjectNodeObjectDetails::FillParameterNamesArray()
 		{
 			if (ProjectorParameterNode)
 			{
-				BaseObjectNode->ParameterNames.Add(ProjectorParameterNode->ParameterName);
+				BaseObjectNode->ParameterNames.Add(ProjectorParameterNode->GetParameterName());
 			}
 		}
 
@@ -454,7 +460,7 @@ void FCustomizableObjectNodeObjectDetails::FillParameterNamesArray()
 		{
 			if (TextureParameterNode)
 			{
-				BaseObjectNode->ParameterNames.Add(TextureParameterNode->ParameterName);
+				BaseObjectNode->ParameterNames.Add(TextureParameterNode->GetParameterName());
 			}
 		}
 
@@ -494,7 +500,7 @@ void FCustomizableObjectNodeObjectDetails::FillParameterNamesArray()
 				}
 				else
 				{
-					BaseObjectNode->ParameterNames.Add(GroupNode->GroupName);
+					BaseObjectNode->ParameterNames.Add(GroupNode->GetGroupName());
 				}
 			}
 		}
@@ -536,7 +542,7 @@ void FCustomizableObjectNodeObjectDetails::FillParameterNamesArray()
 
 		if (ToggleGroupObjectIds.Contains(ParentObjectGroupId))
 		{
-			BaseObjectNode->ParameterNames.Add(ObjectNode->ObjectName);
+			BaseObjectNode->ParameterNames.Add(ObjectNode->GetObjectName());
 		}
 	}
 }
diff --git a/Engine/Plugins/Mutable/Source/CustomizableObjectEditor/Private/MuCOE/Nodes/CustomizableObjectNodeObjectGroup.cpp b/Engine/Plugins/Mutable/Source/CustomizableObjectEditor/Private/MuCOE/Nodes/CustomizableObjectNodeObjectGroup.cpp
--- a/Engine/Plugins/Mutable/Source/CustomizableObjectEditor/Private/MuCOE/Nodes/CustomizableObjectNodeObjectGroup.cpp
+++ b/Engine/Plugins/Mutable/Source/CustomizableObjectEditor/Private/MuCOE/Nodes/CustomizableObjectNodeObjectGroup.cpp
@@ -3,9 +3,11 @@
 #include "MuCOE/Nodes/CustomizableObjectNodeObjectGroup.h"
 
 #include "MuCO/CustomizableObjectCustomVersion.h"
+#include "MuCOE/CustomizableObjectMacroLibrary/CustomizableObjectMacroLibrary.h"
 #include "MuCOE/EdGraphSchema_CustomizableObject.h"
 #include "MuCOE/GraphTraversal.h"
 #include "MuCOE/Nodes/CustomizableObjectNodeObject.h"
+#include "MuCOE/Nodes/CustomizableObjectNodeStaticString.h"
 
 class UCustomizableObjectNodeRemapPins;
 struct FPropertyChangedEvent;
@@ -86,11 +88,12 @@ void UCustomizableObjectNodeObjectGroup::PostEditChangeProperty(FPropertyChanged
 
 FText UCustomizableObjectNodeObjectGroup::GetNodeTitle(ENodeTitleType::Type TitleType) const
 {
+	const UEdGraphPin* GroupNamePin = NamePin.Get();
 	FText NodeTitle;
 	FFormatNamedArguments Args;
 	Args.Add(TEXT("GroupName"), FText::FromString(GroupName));
 
-	if (TitleType == ENodeTitleType::ListView)
+	if (TitleType == ENodeTitleType::ListView || (GroupNamePin && GroupNamePin->LinkedTo.Num()))
 	{
 		NodeTitle = LOCTEXT("Object_Group", "Object Group");
 	}
@@ -129,6 +132,28 @@ void UCustomizableObjectNodeObjectGroup::OnRenameNode(const FString& NewName)
 }
 
 
+bool UCustomizableObjectNodeObjectGroup::GetCanRenameNode() const
+{
+	const UEdGraphPin* GroupNamePin = NamePin.Get();
+
+	if (CVarEnableMutableMacros.GetValueOnAnyThread() && GroupNamePin && GroupNamePin->LinkedTo.Num())
+	{
+		return false;
+	}
+
+	return true;
+}
+
+
+void UCustomizableObjectNodeObjectGroup::PinConnectionListChanged(UEdGraphPin* Pin)
+{
+	if (Pin == NamePin)
+	{
+		GetGraph()->NotifyGraphChanged();
+	}
+}
+
+
 void UCustomizableObjectNodeObjectGroup::AllocateDefaultPins(UCustomizableObjectNodeRemapPins* RemapPins)
 {
 	const UEdGraphSchema_CustomizableObject* Schema = GetDefault<UEdGraphSchema_CustomizableObject>();
@@ -140,6 +165,11 @@ void UCustomizableObjectNodeObjectGroup::AllocateDefaultPins(UCustomizableObject
 	ObjectsPin->bDefaultValueIsIgnored = true;
 
 	CustomCreatePin(EGPD_Output, Schema->PC_Object, FName("Group"));
+
+	if (CVarEnableMutableMacros.GetValueOnAnyThread())
+	{
+		NamePin = CustomCreatePin(EGPD_Input, UEdGraphSchema_CustomizableObject::PC_String, FName("Name"));
+	}
 }
 
 
@@ -167,4 +197,28 @@ UEdGraphPin* UCustomizableObjectNodeObjectGroup::GroupPin() const
 }
 
 
+FString UCustomizableObjectNodeObjectGroup::GetGroupName(TArray<const UCustomizableObjectNodeMacroInstance*>* MacroContext) const
+{
+	if (CVarEnableMutableMacros.GetValueOnAnyThread() && NamePin.Get())
+	{
+		if (const UEdGraphPin* LinkedPin = FollowInputPin(*NamePin.Get()))
+		{
+			const UEdGraphPin* StringPin = GraphTraversal::FindStringSource(*LinkedPin, MacroContext);
+
+			if (const UCustomizableObjectNodeStaticString* StringNode = StringPin ? Cast<UCustomizableObjectNodeStaticString>(StringPin->GetOwningNode()) : nullptr)
+			{
+				return StringNode->Value;
+			}
+		}
+	}
+
+	return GroupName;
+}
+
+
+void UCustomizableObjectNodeObjectGroup::SetGroupName(const FString& Name)
+{
+	GroupName = Name;
+}
+
 #undef LOCTEXT_NAMESPACE
diff --git a/Engine/Plugins/Mutable/Source/CustomizableObjectEditor/Private/MuCOE/Nodes/CustomizableObjectNodeObjectGroup.h b/Engine/Plugins/Mutable/Source/CustomizableObjectEditor/Private/MuCOE/Nodes/CustomizableObjectNodeObjectGroup.h
--- a/Engine/Plugins/Mutable/Source/CustomizableObjectEditor/Private/MuCOE/Nodes/CustomizableObjectNodeObjectGroup.h
+++ b/Engine/Plugins/Mutable/Source/CustomizableObjectEditor/Private/MuCOE/Nodes/CustomizableObjectNodeObjectGroup.h
@@ -24,9 +24,6 @@ public:
 
 	UCustomizableObjectNodeObjectGroup();
 
-	UPROPERTY(EditAnywhere, Category=GroupInfo, meta = (DisplayName= "Name"))
-	FString GroupName;
-
 	UPROPERTY(EditAnywhere, Category=GroupInfo, meta = (DisplayName = "Type"))
 	ECustomizableObjectGroupType GroupType;
 	
@@ -40,6 +37,9 @@ public:
 	// are clashes with socket names, the one with higher priority will be kept and the other discarded.
 	UPROPERTY(EditAnywhere, Category = MeshSockets)
 	int32 SocketPriority = 0;
+	
+	UPROPERTY()
+	FEdGraphPinReference NamePin;
 
 	// UObject interface.
 	virtual void Serialize(FArchive& Ar) override;
@@ -49,8 +49,9 @@ public:
 	virtual FText GetNodeTitle(ENodeTitleType::Type TitleType) const override;
 	virtual FLinearColor GetNodeTitleColor() const override;
 	virtual FText GetTooltipText() const override;
-	virtual bool GetCanRenameNode() const override { return true; }
+	virtual bool GetCanRenameNode() const override;
 	virtual void OnRenameNode(const FString& NewName) override;
+	virtual void PinConnectionListChanged(UEdGraphPin* Pin) override;
 
 	// UCustomizableObjectNode interface
 	virtual void AllocateDefaultPins(UCustomizableObjectNodeRemapPins* RemapPins) override;
@@ -58,12 +59,17 @@ public:
 
 	// Own interface
 	UEdGraphPin* ObjectsPin() const;
-
 	UEdGraphPin* GroupProjectorsPin() const;
-
 	UEdGraphPin* GroupPin() const;
+	FString GetGroupName(TArray<const UCustomizableObjectNodeMacroInstance*>* MacroContext = nullptr) const;
+	void SetGroupName(const FString& Name);
 
 private:
+
+	UPROPERTY(EditAnywhere, Category = GroupInfo, meta = (DisplayName = "Name"))
+	FString GroupName;
+
 	FString LastGroupName;
+
 };
 
diff --git a/Engine/Plugins/Mutable/Source/CustomizableObjectEditor/Private/MuCOE/Nodes/CustomizableObjectNodeParameter.cpp b/Engine/Plugins/Mutable/Source/CustomizableObjectEditor/Private/MuCOE/Nodes/CustomizableObjectNodeParameter.cpp
new file mode 100644
--- /dev/null
+++ b/Engine/Plugins/Mutable/Source/CustomizableObjectEditor/Private/MuCOE/Nodes/CustomizableObjectNodeParameter.cpp
@@ -0,0 +1,126 @@
+// Copyright Epic Games, Inc. All Rights Reserved.
+
+#include "MuCOE/Nodes/CustomizableObjectNodeParameter.h"
+
+#include "MuCOE/EdGraphSchema_CustomizableObject.h"
+#include "MuCOE/CustomizableObjectMacroLibrary/CustomizableObjectMacroLibrary.h"
+#include "MuCOE/GraphTraversal.h"
+#include "MuCOE/Nodes/CustomizableObjectNodeMacroInstance.h"
+#include "MuCOE/Nodes/CustomizableObjectNodeStaticString.h"
+
+class UCustomizableObjectNodeRemapPins;
+
+#define LOCTEXT_NAMESPACE "CustomizableObjectEditor"
+
+
+void UCustomizableObjectNodeParameter::AllocateDefaultPins(UCustomizableObjectNodeRemapPins* RemapPins)
+{
+	const FName Type = GetCategory();
+	const FName PinName = UEdGraphSchema_CustomizableObject::GetPinCategoryName(Type);
+	const FText PinFriendlyName = UEdGraphSchema_CustomizableObject::GetPinCategoryFriendlyName(Type);
+	
+	UEdGraphPin* ValuePin = CustomCreatePin(EGPD_Output, Type, PinName);
+	ValuePin->PinFriendlyName = PinFriendlyName;
+	ValuePin->bDefaultValueIsIgnored = true;
+
+	if (CVarEnableMutableMacros.GetValueOnAnyThread())
+	{
+		NamePin = CustomCreatePin(EGPD_Input, UEdGraphSchema_CustomizableObject::PC_String, FName("Name"));
+	}
+}
+
+
+FText UCustomizableObjectNodeParameter::GetNodeTitle(ENodeTitleType::Type TitleType) const
+{
+	FText NodeTitle;
+	const UEdGraphPin* ParamNamePin = NamePin.Get();
+	FFormatNamedArguments Args;
+	Args.Add(TEXT("ParameterName"), FText::FromString(ParameterName));
+	Args.Add(TEXT("Type"), UEdGraphSchema_CustomizableObject::GetPinCategoryFriendlyName(GetCategory()));
+
+	if (TitleType == ENodeTitleType::ListView || (ParamNamePin && ParamNamePin->LinkedTo.Num()))
+	{
+		NodeTitle = LOCTEXT("ParameterTitle_ListView", "{Type} Parameter");
+	}
+	else if (TitleType == ENodeTitleType::EditableTitle)
+	{
+		NodeTitle = LOCTEXT("ParameterTitle_EditableTitle", "{ParameterName}");
+	}
+	else
+	{
+		NodeTitle = LOCTEXT("ParameterTitle_Title", "{ParameterName}\n{Type} Parameter");
+	}
+
+	return FText::Format(NodeTitle, Args);
+}
+
+
+FLinearColor UCustomizableObjectNodeParameter::GetNodeTitleColor() const
+{
+	return UEdGraphSchema_CustomizableObject::GetPinTypeColor(GetCategory());
+}
+
+
+FText UCustomizableObjectNodeParameter::GetTooltipText() const
+{
+	return FText::Format(LOCTEXT("Parameter_Tooltip", "Expose a runtime modifiable {0} parameter from the Customizable Object."), FText::FromName(GetCategory()));
+}
+
+
+void UCustomizableObjectNodeParameter::OnRenameNode(const FString& NewName)
+{
+	if (!NewName.IsEmpty())
+	{
+		ParameterName = NewName;
+	}
+}
+
+
+bool UCustomizableObjectNodeParameter::GetCanRenameNode() const
+{
+	const UEdGraphPin* ParamNamePin = NamePin.Get();
+
+	if (CVarEnableMutableMacros.GetValueOnAnyThread() && ParamNamePin && ParamNamePin->LinkedTo.Num())
+	{
+		return false;
+	}
+
+	return true;
+}
+
+
+void UCustomizableObjectNodeParameter::PinConnectionListChanged(UEdGraphPin* Pin)
+{
+	if (Pin == NamePin)
+	{
+		GetGraph()->NotifyGraphChanged();
+	}
+}
+
+
+FString UCustomizableObjectNodeParameter::GetParameterName(TArray<const UCustomizableObjectNodeMacroInstance*>* MacroContext) const
+{
+	if (CVarEnableMutableMacros.GetValueOnAnyThread() && NamePin.Get())
+	{
+		if (const UEdGraphPin* LinkedPin = FollowInputPin(*NamePin.Get()))
+		{
+			const UEdGraphPin* StringPin = GraphTraversal::FindStringSource(*LinkedPin, MacroContext);
+
+			if (const UCustomizableObjectNodeStaticString* StringNode = StringPin ? Cast<UCustomizableObjectNodeStaticString>(StringPin->GetOwningNode()) : nullptr)
+			{
+				return StringNode->Value;
+			}
+		}
+	}
+
+	return ParameterName;
+}
+
+
+void UCustomizableObjectNodeParameter::SetParameterName(const FString& Name)
+{
+	ParameterName = Name;
+}
+
+#undef LOCTEXT_NAMESPACE
+
diff --git a/Engine/Plugins/Mutable/Source/CustomizableObjectEditor/Private/MuCOE/Nodes/CustomizableObjectNodeParameter.h b/Engine/Plugins/Mutable/Source/CustomizableObjectEditor/Private/MuCOE/Nodes/CustomizableObjectNodeParameter.h
new file mode 100644
--- /dev/null
+++ b/Engine/Plugins/Mutable/Source/CustomizableObjectEditor/Private/MuCOE/Nodes/CustomizableObjectNodeParameter.h
@@ -0,0 +1,52 @@
+// Copyright Epic Games, Inc. All Rights Reserved.
+
+#pragma once
+
+#include "MuCOE/Nodes/CustomizableObjectNode.h"
+
+#include "CustomizableObjectNodeParameter.generated.h"
+
+namespace ENodeTitleType { enum Type : int; }
+
+class UCustomizableObjectNodeRemapPins;
+class UObject;
+
+
+UCLASS(Abstract)
+class CUSTOMIZABLEOBJECTEDITOR_API UCustomizableObjectNodeParameter : public UCustomizableObjectNode
+{
+public:
+	GENERATED_BODY()
+
+	UPROPERTY(EditAnywhere, Category = UI, meta = (DisplayName = "Parameter UI Metadata"))
+	FMutableParamUIMetadata ParamUIMetadata;
+
+	UPROPERTY()
+	FEdGraphPinReference NamePin;
+
+private:
+
+	UPROPERTY(EditAnywhere, Category = CustomizableObject)
+	FString ParameterName = "Default Name";
+
+public:
+
+	// Begin EdGraphNode interface
+	virtual FText GetNodeTitle(ENodeTitleType::Type TitleType) const override;
+	virtual FLinearColor GetNodeTitleColor() const override;
+	virtual FText GetTooltipText() const override;
+	virtual void OnRenameNode(const FString& NewName) override;
+	virtual bool GetCanRenameNode() const override;
+	virtual void PinConnectionListChanged(UEdGraphPin* Pin) override;
+
+	// UCustomizableObjectNode interface
+	virtual void AllocateDefaultPins(UCustomizableObjectNodeRemapPins* RemapPins) override;
+	virtual bool IsAffectedByLOD() const override { return false; }
+
+	// Own interface
+	/** Return the pin category of this node. */
+	virtual FName GetCategory() const PURE_VIRTUAL(UCustomizableObjectNodeParameter::GetCategory, return {}; );
+	virtual FString GetParameterName(TArray<const UCustomizableObjectNodeMacroInstance*>* MacroContext = nullptr) const;
+	virtual void SetParameterName(const FString& Name);
+};
+
diff --git a/Engine/Plugins/Mutable/Source/CustomizableObjectEditor/Private/MuCOE/Nodes/CustomizableObjectNodeParameterDetails.cpp b/Engine/Plugins/Mutable/Source/CustomizableObjectEditor/Private/MuCOE/Nodes/CustomizableObjectNodeParameterDetails.cpp
new file mode 100644
--- /dev/null
+++ b/Engine/Plugins/Mutable/Source/CustomizableObjectEditor/Private/MuCOE/Nodes/CustomizableObjectNodeParameterDetails.cpp
@@ -0,0 +1,22 @@
+// Copyright Epic Games, Inc. All Rights Reserved.
+
+#include "MuCOE/Nodes/CustomizableObjectNodeParameterDetails.h"
+
+#include "DetailLayoutBuilder.h"
+#include "MuCOE/CustomizableObjectMacroLibrary/CustomizableObjectMacroLibrary.h"
+
+
+TSharedRef<IDetailCustomization> FCustomizableObjectNodeParameterDetails::MakeInstance()
+{
+	return MakeShareable(new FCustomizableObjectNodeParameterDetails);
+}
+
+void FCustomizableObjectNodeParameterDetails::CustomizeDetails(IDetailLayoutBuilder& DetailBuilder)
+{
+	FCustomizableObjectNodeDetails::CustomizeDetails(DetailBuilder);
+	
+	if (CVarEnableMutableMacros.GetValueOnAnyThread())
+	{
+		DetailBuilder.HideProperty("ParameterName");
+	}
+}
\ No newline at end of file
diff --git a/Engine/Plugins/Mutable/Source/CustomizableObjectEditor/Private/MuCOE/Nodes/CustomizableObjectNodeParameterDetails.h b/Engine/Plugins/Mutable/Source/CustomizableObjectEditor/Private/MuCOE/Nodes/CustomizableObjectNodeParameterDetails.h
new file mode 100644
--- /dev/null
+++ b/Engine/Plugins/Mutable/Source/CustomizableObjectEditor/Private/MuCOE/Nodes/CustomizableObjectNodeParameterDetails.h
@@ -0,0 +1,16 @@
+// Copyright Epic Games, Inc. All Rights Reserved.
+
+#pragma once
+
+#include "CustomizableObjectNodeDetails.h"
+#include "IDetailCustomization.h"
+
+class FCustomizableObjectNodeParameterDetails : public FCustomizableObjectNodeDetails
+{
+public:
+	/** Makes a new instance of this detail layout class for a specific detail view requesting it */
+	static TSharedRef<IDetailCustomization> MakeInstance();
+
+	/** ILayoutDetails interface */
+	virtual void CustomizeDetails(IDetailLayoutBuilder& DetailBuilder) override;
+};
diff --git a/Engine/Plugins/Mutable/Source/CustomizableObjectEditor/Private/MuCOE/Nodes/CustomizableObjectNodeProjectorParameter.cpp b/Engine/Plugins/Mutable/Source/CustomizableObjectEditor/Private/MuCOE/Nodes/CustomizableObjectNodeProjectorParameter.cpp
--- a/Engine/Plugins/Mutable/Source/CustomizableObjectEditor/Private/MuCOE/Nodes/CustomizableObjectNodeProjectorParameter.cpp
+++ b/Engine/Plugins/Mutable/Source/CustomizableObjectEditor/Private/MuCOE/Nodes/CustomizableObjectNodeProjectorParameter.cpp
@@ -3,6 +3,7 @@
 #include "MuCOE/Nodes/CustomizableObjectNodeProjectorParameter.h"
 
 #include "MuCO/CustomizableObjectCustomVersion.h"
+#include "MuCOE/CustomizableObjectMacroLibrary/CustomizableObjectMacroLibrary.h"
 #include "MuCOE/EdGraphSchema_CustomizableObject.h"
 
 class UCustomizableObjectNodeRemapPins;
@@ -29,18 +30,6 @@ void UCustomizableObjectNodeProjectorParameter::PostEditChangeProperty(FProperty
 }
 
 
-void UCustomizableObjectNodeProjectorParameter::AllocateDefaultPins(UCustomizableObjectNodeRemapPins* RemapPins)
-{
-	const FName Type = UEdGraphSchema_CustomizableObject::PC_Projector;
-	const FName PinName = UEdGraphSchema_CustomizableObject::GetPinCategoryName(Type);
-	const FText PinFriendlyName = UEdGraphSchema_CustomizableObject::GetPinCategoryFriendlyName(Type);
-	
-	UEdGraphPin* ValuePin = CustomCreatePin(EGPD_Output, Type, PinName);
-	ValuePin->bDefaultValueIsIgnored = true;
-	ValuePin->PinFriendlyName = PinFriendlyName;
-}
-
-
 void UCustomizableObjectNodeProjectorParameter::BackwardsCompatibleFixup(int32 CustomizableObjectCustomVersion)
 {
 	Super::BackwardsCompatibleFixup(CustomizableObjectCustomVersion);
@@ -64,6 +53,12 @@ void UCustomizableObjectNodeProjectorParameter::BackwardsCompatibleFixup(int32 C
 }
 
 
+FName UCustomizableObjectNodeProjectorParameter::GetCategory() const
+{
+	return UEdGraphSchema_CustomizableObject::PC_Projector;
+}
+
+
 ECustomizableObjectProjectorType UCustomizableObjectNodeProjectorParameter::GetProjectorType() const
 {
 	return DefaultValue.ProjectionType;
@@ -130,42 +125,4 @@ void UCustomizableObjectNodeProjectorParameter::SetProjectorDefaultAngle(float A
 }
 
 
-FText UCustomizableObjectNodeProjectorParameter::GetNodeTitle(ENodeTitleType::Type TitleType) const
-{
-	if (TitleType == ENodeTitleType::ListView)
-	{
-		return LOCTEXT("Projector_Parameter", "Projector Parameter");
-	}
-	else if (TitleType == ENodeTitleType::EditableTitle)
-	{
-		return FText::Format(LOCTEXT("Projector_Parameter_EditableTitle", "{0}"), FText::FromString(ParameterName));
-	}
-	else
-	{
-		return FText::Format(LOCTEXT("Projector_Parameter_Title", "{0}\nProjector Parameter"), FText::FromString(ParameterName));
-	}
-}
-
-
-FLinearColor UCustomizableObjectNodeProjectorParameter::GetNodeTitleColor() const
-{
-	const UEdGraphSchema_CustomizableObject* Schema = GetDefault<UEdGraphSchema_CustomizableObject>();
-	return Schema->GetPinTypeColor(Schema->PC_Projector);
-}
-
-
-FText UCustomizableObjectNodeProjectorParameter::GetTooltipText() const
-{
-	return LOCTEXT("Projector_Parameter_Tooltip", "Exposes a runtime modifiable projector parameter from the Customizable Object.");
-}
-
-
-void UCustomizableObjectNodeProjectorParameter::OnRenameNode(const FString& NewName)
-{
-	if (!NewName.IsEmpty())
-	{
-		ParameterName = NewName;
-	}
-}
-
 #undef LOCTEXT_NAMESPACE
diff --git a/Engine/Plugins/Mutable/Source/CustomizableObjectEditor/Private/MuCOE/Nodes/CustomizableObjectNodeProjectorParameter.h b/Engine/Plugins/Mutable/Source/CustomizableObjectEditor/Private/MuCOE/Nodes/CustomizableObjectNodeProjectorParameter.h
--- a/Engine/Plugins/Mutable/Source/CustomizableObjectEditor/Private/MuCOE/Nodes/CustomizableObjectNodeProjectorParameter.h
+++ b/Engine/Plugins/Mutable/Source/CustomizableObjectEditor/Private/MuCOE/Nodes/CustomizableObjectNodeProjectorParameter.h
@@ -2,7 +2,7 @@
 
 #pragma once
 
-#include "MuCOE/Nodes/CustomizableObjectNode.h"
+#include "MuCOE/Nodes/CustomizableObjectNodeParameter.h"
 
 #include "CustomizableObjectNodeProjectorParameter.generated.h"
 
@@ -15,7 +15,7 @@ struct FPropertyChangedEvent;
 
 
 UCLASS()
-class CUSTOMIZABLEOBJECTEDITOR_API UCustomizableObjectNodeProjectorParameter : public UCustomizableObjectNode
+class CUSTOMIZABLEOBJECTEDITOR_API UCustomizableObjectNodeProjectorParameter : public UCustomizableObjectNodeParameter
 {
 public:
 	GENERATED_BODY()
@@ -23,12 +23,6 @@ public:
 	UPROPERTY(EditAnywhere, Category=CustomizableObject, meta = (ShowOnlyInnerProperties))
 	FCustomizableObjectProjector DefaultValue;
 
-	UPROPERTY(EditAnywhere, Category=CustomizableObject)
-	FString ParameterName = "Default Name";
-
-	UPROPERTY(EditAnywhere, Category = UI, meta = (DisplayName = "Parameter UI Metadata"))
-	FMutableParamUIMetadata ParamUIMetadata;
-
 	UPROPERTY(EditAnywhere, Category = CustomizableObject, Meta = (DisplayName = "Projection Angle (degrees)"))
 	float ProjectionAngle = 360.0f;
 
@@ -52,19 +46,13 @@ public:
 
 	// UObject interface.
 	virtual void PostEditChangeProperty(FPropertyChangedEvent& PropertyChangedEvent) override;
-
-	// Begin EdGraphNode interface
-	virtual FText GetNodeTitle(ENodeTitleType::Type TitleType) const override;
-	virtual FLinearColor GetNodeTitleColor() const override;
-	virtual FText GetTooltipText() const override;
-	virtual bool GetCanRenameNode() const override { return true; }
-	virtual void OnRenameNode(const FString& NewName) override;
-
+	
 	// UCustomizableObjectNode interface
-	virtual void AllocateDefaultPins(UCustomizableObjectNodeRemapPins* RemapPins) override;
-	virtual bool IsAffectedByLOD() const override { return false; }
 	virtual void BackwardsCompatibleFixup(int32 CustomizableObjectCustomVersion) override;
 
+	// CustomizableObjectNodeParameter interface
+	virtual FName GetCategory() const override;
+
 	// Own interface
 
 	ECustomizableObjectProjectorType GetProjectorType() const;
diff --git a/Engine/Plugins/Mutable/Source/CustomizableObjectEditor/Private/MuCOE/Nodes/CustomizableObjectNodeStaticString.cpp b/Engine/Plugins/Mutable/Source/CustomizableObjectEditor/Private/MuCOE/Nodes/CustomizableObjectNodeStaticString.cpp
new file mode 100644
--- /dev/null
+++ b/Engine/Plugins/Mutable/Source/CustomizableObjectEditor/Private/MuCOE/Nodes/CustomizableObjectNodeStaticString.cpp
@@ -0,0 +1,45 @@
+// Copyright Epic Games, Inc. All Rights Reserved.
+
+#include "MuCOE/Nodes/CustomizableObjectNodeStaticString.h"
+
+#include "MuCOE/EdGraphSchema_CustomizableObject.h"
+#include "MuCOE/Nodes/SCustomizableObjectNode.h"
+
+#define LOCTEXT_NAMESPACE "CustomizableObjectNodeStringConstant"
+
+
+FText UCustomizableObjectNodeStaticString::GetNodeTitle(ENodeTitleType::Type TitleType) const
+{
+	return LOCTEXT("StaticStringNodeTitle", "Static String");
+}
+
+
+FText UCustomizableObjectNodeStaticString::GetTooltipText() const
+{
+	return LOCTEXT("StaticStringNodeTooltip", "Static String Node");
+}
+
+
+FLinearColor UCustomizableObjectNodeStaticString::GetNodeTitleColor() const
+{
+	return UEdGraphSchema_CustomizableObject::GetPinTypeColor(UEdGraphSchema_CustomizableObject::PC_String);
+}
+
+
+bool UCustomizableObjectNodeStaticString::GetCanRenameNode() const
+{
+	return false;
+}
+
+
+void UCustomizableObjectNodeStaticString::AllocateDefaultPins(UCustomizableObjectNodeRemapPins* RemapPins)
+{
+	const UEdGraphSchema_CustomizableObject* Schema = GetDefault<UEdGraphSchema_CustomizableObject>();
+
+	UEdGraphPin* ValuePin = CustomCreatePin(EGPD_Output, Schema->PC_String, FName("Value"));
+	UEdGraphPin* StringPin = CustomCreatePin(EGPD_Input, Schema->PC_String, FName("String"));
+	StringPin->bNotConnectable = true;
+}
+
+
+#undef LOCTEXT_NAMESPACE
diff --git a/Engine/Plugins/Mutable/Source/CustomizableObjectEditor/Private/MuCOE/Nodes/CustomizableObjectNodeStaticString.h b/Engine/Plugins/Mutable/Source/CustomizableObjectEditor/Private/MuCOE/Nodes/CustomizableObjectNodeStaticString.h
new file mode 100644
--- /dev/null
+++ b/Engine/Plugins/Mutable/Source/CustomizableObjectEditor/Private/MuCOE/Nodes/CustomizableObjectNodeStaticString.h
@@ -0,0 +1,31 @@
+// Copyright Epic Games, Inc. All Rights Reserved.
+
+#pragma once
+
+#include "MuCOE/Nodes/CustomizableObjectNode.h"
+
+#include "CustomizableObjectNodeStaticString.generated.h"
+
+
+UCLASS()
+class CUSTOMIZABLEOBJECTEDITOR_API UCustomizableObjectNodeStaticString : public UCustomizableObjectNode
+{
+public:
+	GENERATED_BODY()
+
+	/** Macro Library where the macro to instantiate belongs. */
+	UPROPERTY(EditAnywhere, Category = CustomizableObject)
+	FString Value;
+
+public:
+
+	// EdGraphNode interface
+	virtual FText GetNodeTitle(ENodeTitleType::Type TitleType) const override;
+	virtual FText GetTooltipText() const override;
+	virtual FLinearColor GetNodeTitleColor() const override;
+	virtual bool GetCanRenameNode() const override;
+
+	// UCustomizableObjectNode interface
+	void AllocateDefaultPins(UCustomizableObjectNodeRemapPins* RemapPins) override;
+
+};
\ No newline at end of file
diff --git a/Engine/Plugins/Mutable/Source/CustomizableObjectEditor/Private/MuCOE/Nodes/CustomizableObjectNodeStringConstant.cpp b/Engine/Plugins/Mutable/Source/CustomizableObjectEditor/Private/MuCOE/Nodes/CustomizableObjectNodeStringConstant.cpp
deleted file mode 100644
diff --git a/Engine/Plugins/Mutable/Source/CustomizableObjectEditor/Private/MuCOE/Nodes/CustomizableObjectNodeStringConstant.h b/Engine/Plugins/Mutable/Source/CustomizableObjectEditor/Private/MuCOE/Nodes/CustomizableObjectNodeStringConstant.h
deleted file mode 100644
diff --git a/Engine/Plugins/Mutable/Source/CustomizableObjectEditor/Private/MuCOE/Nodes/CustomizableObjectNodeTextureParameter.cpp b/Engine/Plugins/Mutable/Source/CustomizableObjectEditor/Private/MuCOE/Nodes/CustomizableObjectNodeTextureParameter.cpp
--- a/Engine/Plugins/Mutable/Source/CustomizableObjectEditor/Private/MuCOE/Nodes/CustomizableObjectNodeTextureParameter.cpp
+++ b/Engine/Plugins/Mutable/Source/CustomizableObjectEditor/Private/MuCOE/Nodes/CustomizableObjectNodeTextureParameter.cpp
@@ -3,6 +3,7 @@
 #include "MuCOE/Nodes/CustomizableObjectNodeTextureParameter.h"
 
 #include "MuCO/CustomizableObjectCustomVersion.h"
+#include "MuCOE/CustomizableObjectMacroLibrary/CustomizableObjectMacroLibrary.h"
 #include "MuCOE/EdGraphSchema_CustomizableObject.h"
 
 class UCustomizableObjectNodeRemapPins;
@@ -10,63 +11,12 @@ class UCustomizableObjectNodeRemapPins;
 #define LOCTEXT_NAMESPACE "CustomizableObjectEditor"
 
 
-void UCustomizableObjectNodeTextureParameter::AllocateDefaultPins(UCustomizableObjectNodeRemapPins* RemapPins)
-{
-	const FName Type = UEdGraphSchema_CustomizableObject::PC_Image;
-	const FName PinName = UEdGraphSchema_CustomizableObject::GetPinCategoryName(Type);
-	const FText PinFriendlyName = UEdGraphSchema_CustomizableObject::GetPinCategoryFriendlyName(Type);
-	
-	UEdGraphPin* ValuePin = CustomCreatePin(EGPD_Output, Type, PinName);
-	ValuePin->bDefaultValueIsIgnored = true;
-	ValuePin->PinFriendlyName = PinFriendlyName;
-}
-
-
 bool UCustomizableObjectNodeTextureParameter::IsExperimental() const
 {
 	return true;
 }
 
 
-FText UCustomizableObjectNodeTextureParameter::GetNodeTitle(ENodeTitleType::Type TitleType) const
-{
-	if (TitleType == ENodeTitleType::ListView)
-	{
-		return LOCTEXT("Texture_Parameter", "Texture Parameter");
-	}
-	else if (TitleType == ENodeTitleType::EditableTitle)
-	{
-		return FText::Format(LOCTEXT("Texture_Parameter_EditableTitle", "{0}"), FText::FromString(ParameterName));
-	}
-	else
-	{
-		return FText::Format(LOCTEXT("Texture_Parameter_Title", "{0}\nTexture Parameter"), FText::FromString(ParameterName));
-	}
-}
-
-
-FLinearColor UCustomizableObjectNodeTextureParameter::GetNodeTitleColor() const
-{
-	const UEdGraphSchema_CustomizableObject* Schema = GetDefault<UEdGraphSchema_CustomizableObject>();
-	return Schema->GetPinTypeColor(Schema->PC_Image);
-}
-
-
-FText UCustomizableObjectNodeTextureParameter::GetTooltipText() const
-{
-	return LOCTEXT("Texture_Parameter_Tooltip", "Expose a runtime modifiable texture parameter from the Customizable Object.");
-}
-
-
-void UCustomizableObjectNodeTextureParameter::OnRenameNode(const FString& NewName)
-{
-	if (!NewName.IsEmpty())
-	{
-		ParameterName = NewName;
-	}
-}
-
-
 void UCustomizableObjectNodeTextureParameter::BackwardsCompatibleFixup(int32 CustomizableObjectCustomVersion)
 {
 	Super::BackwardsCompatibleFixup(CustomizableObjectCustomVersion);
@@ -86,5 +36,11 @@ void UCustomizableObjectNodeTextureParameter::BackwardsCompatibleFixup(int32 Cus
 	}
 }
 
+
+FName UCustomizableObjectNodeTextureParameter::GetCategory() const
+{
+	return UEdGraphSchema_CustomizableObject::PC_Image;
+}
+
 #undef LOCTEXT_NAMESPACE
 
diff --git a/Engine/Plugins/Mutable/Source/CustomizableObjectEditor/Private/MuCOE/Nodes/CustomizableObjectNodeTextureParameter.h b/Engine/Plugins/Mutable/Source/CustomizableObjectEditor/Private/MuCOE/Nodes/CustomizableObjectNodeTextureParameter.h
--- a/Engine/Plugins/Mutable/Source/CustomizableObjectEditor/Private/MuCOE/Nodes/CustomizableObjectNodeTextureParameter.h
+++ b/Engine/Plugins/Mutable/Source/CustomizableObjectEditor/Private/MuCOE/Nodes/CustomizableObjectNodeTextureParameter.h
@@ -2,7 +2,7 @@
 
 #pragma once
 
-#include "MuCOE/Nodes/CustomizableObjectNode.h"
+#include "MuCOE/Nodes/CustomizableObjectNodeParameter.h"
 
 #include "CustomizableObjectNodeTextureParameter.generated.h"
 
@@ -14,7 +14,7 @@ class UTexture2D;
 
 
 UCLASS()
-class CUSTOMIZABLEOBJECTEDITOR_API UCustomizableObjectNodeTextureParameter : public UCustomizableObjectNode
+class CUSTOMIZABLEOBJECTEDITOR_API UCustomizableObjectNodeTextureParameter : public UCustomizableObjectNodeParameter
 {
 public:
 	GENERATED_BODY()
@@ -27,12 +27,6 @@ public:
 	UPROPERTY(EditAnywhere, Category=CustomizableObject)
 	TObjectPtr<UTexture2D> ReferenceValue;
 
-	UPROPERTY(EditAnywhere, Category=CustomizableObject)
-	FString ParameterName = "Default Name";
-
-	UPROPERTY(EditAnywhere, Category = UI, meta = (DisplayName = "Parameter UI Metadata"))
-	FMutableParamUIMetadata ParamUIMetadata;
-
 	/** Set the width of the Texture when there is no texture reference.*/
 	UPROPERTY(EditAnywhere, Category = CustomizableObject)
 	int32 TextureSizeX = 0;
@@ -41,17 +35,11 @@ public:
 	UPROPERTY(EditAnywhere, Category = CustomizableObject)
 	int32 TextureSizeY = 0;
 
-	// EdGraphNode interface
-	virtual FText GetNodeTitle(ENodeTitleType::Type TitleType) const override;
-	virtual FLinearColor GetNodeTitleColor() const override;
-	virtual FText GetTooltipText() const override;
-	virtual void OnRenameNode(const FString& NewName) override;
-	virtual bool GetCanRenameNode() const override { return true; }
-
 	// UCustomizableObjectNode interface
 	virtual void BackwardsCompatibleFixup(int32 CustomizableObjectCustomVersion) override;
-	virtual void AllocateDefaultPins(UCustomizableObjectNodeRemapPins* RemapPins) override;
-	virtual bool IsAffectedByLOD() const override { return false; }
 	virtual bool IsExperimental() const override;
+
+	// CustomizableObjectNodeParameter interface
+	virtual FName GetCategory() const override;
 };
 
diff --git a/Engine/Plugins/Mutable/Source/CustomizableObjectEditor/Private/MuCOE/Nodes/CustomizableObjectNodeTransformParameter.cpp b/Engine/Plugins/Mutable/Source/CustomizableObjectEditor/Private/MuCOE/Nodes/CustomizableObjectNodeTransformParameter.cpp
--- a/Engine/Plugins/Mutable/Source/CustomizableObjectEditor/Private/MuCOE/Nodes/CustomizableObjectNodeTransformParameter.cpp
+++ b/Engine/Plugins/Mutable/Source/CustomizableObjectEditor/Private/MuCOE/Nodes/CustomizableObjectNodeTransformParameter.cpp
@@ -3,51 +3,13 @@
 #include "MuCOE/Nodes/CustomizableObjectNodeTransformParameter.h"
 
 #include "MuCO/CustomizableObjectCustomVersion.h"
+#include "MuCOE/CustomizableObjectMacroLibrary/CustomizableObjectMacroLibrary.h"
 #include "MuCOE/EdGraphSchema_CustomizableObject.h"
 
 
 #define LOCTEXT_NAMESPACE "CustomizableObjectEditor"
 
 
-FText UCustomizableObjectNodeTransformParameter::GetNodeTitle(ENodeTitleType::Type TitleType) const
-{
-	if (TitleType == ENodeTitleType::ListView)
-	{
-		return LOCTEXT("Transform_Parameter", "Transform Parameter");
-	}
-	else if (TitleType == ENodeTitleType::EditableTitle)
-	{
-		return FText::Format(LOCTEXT("Transform_Parameter_EditableTitle", "{0}"), FText::FromString(ParameterName));
-	}
-	else
-	{
-		return FText::Format(LOCTEXT("Transform_Parameter_Title", "{0}\nTransform Parameter"), FText::FromString(ParameterName));
-	}
-}
-
-
-FLinearColor UCustomizableObjectNodeTransformParameter::GetNodeTitleColor() const
-{
-	const UEdGraphSchema_CustomizableObject* Schema = GetDefault<UEdGraphSchema_CustomizableObject>();
-	return Schema->GetPinTypeColor(Schema->PC_Transform);
-}
-
-
-FText UCustomizableObjectNodeTransformParameter::GetTooltipText() const
-{
-	return LOCTEXT("Transform_Parameter_Tooltip", "Expose a runtime modifiable transform parameter from the Customizable Object.");
-}
-
-
-void UCustomizableObjectNodeTransformParameter::OnRenameNode(const FString& NewName)
-{
-	if (!NewName.IsEmpty())
-	{
-		ParameterName = NewName;
-	}
-}
-
-
 void UCustomizableObjectNodeTransformParameter::BackwardsCompatibleFixup(int32 CustomizableObjectCustomVersion)
 {
 	Super::BackwardsCompatibleFixup(CustomizableObjectCustomVersion);
@@ -63,15 +25,9 @@ void UCustomizableObjectNodeTransformParameter::BackwardsCompatibleFixup(int32 C
 }
 
 
-void UCustomizableObjectNodeTransformParameter::AllocateDefaultPins(UCustomizableObjectNodeRemapPins* RemapPins)
+FName UCustomizableObjectNodeTransformParameter::GetCategory() const
 {
-	const FName Type = UEdGraphSchema_CustomizableObject::PC_Transform;
-	const FName PinName = UEdGraphSchema_CustomizableObject::GetPinCategoryName(Type);
-	const FText PinFriendlyName = UEdGraphSchema_CustomizableObject::GetPinCategoryFriendlyName(Type);
-	
-	UEdGraphPin* ValuePin = CustomCreatePin(EGPD_Output, Type, PinName);
-	ValuePin->bDefaultValueIsIgnored = true;
-	ValuePin->PinFriendlyName = PinFriendlyName;
+	return UEdGraphSchema_CustomizableObject::PC_Transform;
 }
 
 #undef LOCTEXT_NAMESPACE
diff --git a/Engine/Plugins/Mutable/Source/CustomizableObjectEditor/Private/MuCOE/Nodes/CustomizableObjectNodeTransformParameter.h b/Engine/Plugins/Mutable/Source/CustomizableObjectEditor/Private/MuCOE/Nodes/CustomizableObjectNodeTransformParameter.h
--- a/Engine/Plugins/Mutable/Source/CustomizableObjectEditor/Private/MuCOE/Nodes/CustomizableObjectNodeTransformParameter.h
+++ b/Engine/Plugins/Mutable/Source/CustomizableObjectEditor/Private/MuCOE/Nodes/CustomizableObjectNodeTransformParameter.h
@@ -2,7 +2,7 @@
 
 #pragma once
 
-#include "MuCOE/Nodes/CustomizableObjectNode.h"
+#include "MuCOE/Nodes/CustomizableObjectNodeParameter.h"
 
 #include "CustomizableObjectNodeTransformParameter.generated.h"
 
@@ -10,7 +10,7 @@ namespace ENodeTitleType { enum Type : int; }
 
 
 UCLASS()
-class CUSTOMIZABLEOBJECTEDITOR_API UCustomizableObjectNodeTransformParameter : public UCustomizableObjectNode
+class CUSTOMIZABLEOBJECTEDITOR_API UCustomizableObjectNodeTransformParameter : public UCustomizableObjectNodeParameter
 {
 	GENERATED_BODY()
 
@@ -18,22 +18,10 @@ public:
 	UPROPERTY(EditAnywhere, Category=CustomizableObject, meta = (DontUpdateWhileEditing))
 	FTransform DefaultValue = FTransform::Identity;
 
-	UPROPERTY(EditAnywhere, Category=CustomizableObject)
-	FString ParameterName = "Default Name";
-
-	UPROPERTY(EditAnywhere, Category = UI, meta = (DisplayName = "Parameter UI Metadata"))
-	FMutableParamUIMetadata ParamUIMetadata;
-
 	// Begin EdGraphNode interface
-	virtual FText GetNodeTitle(ENodeTitleType::Type TitleType) const override;
-	virtual FLinearColor GetNodeTitleColor() const override;
-	virtual FText GetTooltipText() const override;
-	virtual void OnRenameNode(const FString& NewName) override;
-	virtual bool GetCanRenameNode() const override { return true; }
 	virtual void BackwardsCompatibleFixup(int32 CustomizableObjectCustomVersion) override;
 
-	// UCustomizableObjectNode interface
-	virtual void AllocateDefaultPins(UCustomizableObjectNodeRemapPins* RemapPins) override;
-	virtual bool IsAffectedByLOD() const override { return false; }
+	// CustomizableObjectNodeParameter interface
+	virtual FName GetCategory() const override;
 	
 };
diff --git a/Engine/Plugins/Mutable/Source/CustomizableObjectEditor/Private/MuCOE/Nodes/CustomizableObjectNodeTunnel.cpp b/Engine/Plugins/Mutable/Source/CustomizableObjectEditor/Private/MuCOE/Nodes/CustomizableObjectNodeTunnel.cpp
--- a/Engine/Plugins/Mutable/Source/CustomizableObjectEditor/Private/MuCOE/Nodes/CustomizableObjectNodeTunnel.cpp
+++ b/Engine/Plugins/Mutable/Source/CustomizableObjectEditor/Private/MuCOE/Nodes/CustomizableObjectNodeTunnel.cpp
@@ -84,6 +84,12 @@ bool UCustomizableObjectNodeTunnel::CanUserDeleteNode() const
 }
 
 
+bool UCustomizableObjectNodeTunnel::CanDuplicateNode() const
+{
+	return false;
+}
+
+
 void UCustomizableObjectNodeTunnel::AllocateDefaultPins(UCustomizableObjectNodeRemapPins* RemapPins)
 {
 	for (const UCustomizableObjectMacroInputOutput* Variable : ParentMacro->InputOutputs)
@@ -102,4 +108,10 @@ void UCustomizableObjectNodeTunnel::AllocateDefaultPins(UCustomizableObjectNodeR
 	}
 }
 
+
+bool UCustomizableObjectNodeTunnel::CreateStaticStringPinWidget() const
+{
+	return false;
+}
+
 #undef LOCTEXT_NAMESPACE
\ No newline at end of file
diff --git a/Engine/Plugins/Mutable/Source/CustomizableObjectEditor/Private/MuCOE/Nodes/CustomizableObjectNodeTunnel.h b/Engine/Plugins/Mutable/Source/CustomizableObjectEditor/Private/MuCOE/Nodes/CustomizableObjectNodeTunnel.h
--- a/Engine/Plugins/Mutable/Source/CustomizableObjectEditor/Private/MuCOE/Nodes/CustomizableObjectNodeTunnel.h
+++ b/Engine/Plugins/Mutable/Source/CustomizableObjectEditor/Private/MuCOE/Nodes/CustomizableObjectNodeTunnel.h
@@ -61,10 +61,13 @@ public:
 	virtual FText GetTooltipText() const override;
 	virtual FLinearColor GetNodeTitleColor() const override;
 	virtual bool GetCanRenameNode() const override;
-	virtual bool CanUserDeleteNode() const;
+	virtual bool CanUserDeleteNode() const override;
+	virtual bool CanDuplicateNode() const override;
+
 
 	// UCustomizableObjectNode interface
 	void AllocateDefaultPins(UCustomizableObjectNodeRemapPins* RemapPins) override;
 	virtual UCustomizableObjectNodeTunnelRemapPins* CreateRemapPinsDefault() const override;
+	virtual bool CreateStaticStringPinWidget() const override;
 
 };
\ No newline at end of file
diff --git a/Engine/Plugins/Mutable/Source/CustomizableObjectEditor/Private/MuCOE/Nodes/CustomizableObjectNodeVariation.cpp b/Engine/Plugins/Mutable/Source/CustomizableObjectEditor/Private/MuCOE/Nodes/CustomizableObjectNodeVariation.cpp
--- a/Engine/Plugins/Mutable/Source/CustomizableObjectEditor/Private/MuCOE/Nodes/CustomizableObjectNodeVariation.cpp
+++ b/Engine/Plugins/Mutable/Source/CustomizableObjectEditor/Private/MuCOE/Nodes/CustomizableObjectNodeVariation.cpp
@@ -3,7 +3,10 @@
 #include "MuCOE/Nodes/CustomizableObjectNodeVariation.h"
 
 #include "MuCO/CustomizableObjectCustomVersion.h"
+#include "MuCOE/CustomizableObjectMacroLibrary/CustomizableObjectMacroLibrary.h"
 #include "MuCOE/EdGraphSchema_CustomizableObject.h"
+#include "MuCOE/GraphTraversal.h"
+#include "MuCOE/Nodes/CustomizableObjectNodeStaticString.h"
 
 class UCustomizableObjectNodeRemapPins;
 
@@ -35,7 +38,7 @@ void UCustomizableObjectNodeVariation::AllocateDefaultPins(UCustomizableObjectNo
 	
 	VariationsPins.SetNum(VariationsData.Num());
 
-	if (IsInMacro())
+	if (CVarEnableMutableMacros.GetValueOnAnyThread())
 	{
 		VariationTagPins.SetNum(VariationsData.Num());
 	}
@@ -50,7 +53,7 @@ void UCustomizableObjectNodeVariation::AllocateDefaultPins(UCustomizableObjectNo
 		
 		VariationsPins[VariationIndex] = VariationPin;
 
-		if (IsInMacro())
+		if (CVarEnableMutableMacros.GetValueOnAnyThread())
 		{
 			FText VariationTagSufix = LOCTEXT("TagSufix", " - Tag");
 			const FName TagPinName = FName(PinName.ToString() + VariationTagSufix.ToString());
@@ -82,6 +85,28 @@ const FCustomizableObjectVariation& UCustomizableObjectNodeVariation::GetVariati
 }
 
 
+FString UCustomizableObjectNodeVariation::GetVariationTag(int32 Index, TArray<const UCustomizableObjectNodeMacroInstance*>* MacroContext) const
+{
+	if (const UEdGraphPin* TagPin = VariationTagPin(Index))
+	{
+		if (const UEdGraphPin* ConnectedStringPin = FollowInputPin(*TagPin))
+		{
+			const UEdGraphPin* SourceStringPin = GraphTraversal::FindStringSource(*ConnectedStringPin, MacroContext);
+
+			if (SourceStringPin)
+			{
+				if (const UCustomizableObjectNodeStaticString* StringNode = Cast<UCustomizableObjectNodeStaticString>(SourceStringPin->GetOwningNode()))
+				{
+					return StringNode->Value;
+				}
+			}
+		}
+	}
+
+	return GetVariation(Index).Tag;
+}
+
+
 UEdGraphPin* UCustomizableObjectNodeVariation::DefaultPin() const
 {
 	return FindPin(TEXT("Default"));
@@ -101,7 +126,7 @@ UEdGraphPin* UCustomizableObjectNodeVariation::VariationPin(int32 Index) const
 
 UEdGraphPin* UCustomizableObjectNodeVariation::VariationTagPin(int32 Index) const
 {
-	if (VariationTagPins.IsValidIndex(Index))
+	if (VariationTagPins.IsValidIndex(Index) && CVarEnableMutableMacros.GetValueOnAnyThread())
 	{
 		return VariationTagPins[Index].Get();
 	}
diff --git a/Engine/Plugins/Mutable/Source/CustomizableObjectEditor/Private/MuCOE/Nodes/CustomizableObjectNodeVariation.h b/Engine/Plugins/Mutable/Source/CustomizableObjectEditor/Private/MuCOE/Nodes/CustomizableObjectNodeVariation.h
--- a/Engine/Plugins/Mutable/Source/CustomizableObjectEditor/Private/MuCOE/Nodes/CustomizableObjectNodeVariation.h
+++ b/Engine/Plugins/Mutable/Source/CustomizableObjectEditor/Private/MuCOE/Nodes/CustomizableObjectNodeVariation.h
@@ -72,6 +72,9 @@ public:
 	/** Get the variation at the given index. */
 	const FCustomizableObjectVariation& GetVariation(int32 Index) const;
 
+	/** Get the variation tag at the given index */
+	FString GetVariationTag(int32 Index, TArray<const UCustomizableObjectNodeMacroInstance*>* MacroContext = nullptr) const;
+
 	/** Get the Default Input Pin. */
 	UEdGraphPin* DefaultPin() const;
 
diff --git a/Engine/Plugins/Mutable/Source/CustomizableObjectEditor/Private/MuCOE/Nodes/SCustomizableObjectNode.cpp b/Engine/Plugins/Mutable/Source/CustomizableObjectEditor/Private/MuCOE/Nodes/SCustomizableObjectNode.cpp
--- a/Engine/Plugins/Mutable/Source/CustomizableObjectEditor/Private/MuCOE/Nodes/SCustomizableObjectNode.cpp
+++ b/Engine/Plugins/Mutable/Source/CustomizableObjectEditor/Private/MuCOE/Nodes/SCustomizableObjectNode.cpp
@@ -1,9 +1,15 @@
-﻿// Copyright Epic Games, Inc. All Rights Reserved.
+// Copyright Epic Games, Inc. All Rights Reserved.
 
 #include "SCustomizableObjectNode.h"
 
 #include "SCustomizableObjectNodePin.h"
+#include "MuCOE/EdGraphSchema_CustomizableObject.h"
+
 #include "MuCOE/Nodes/CustomizableObjectNode.h"
+#include "MuCOE/Nodes/CustomizableObjectNodeStaticString.h"
+
+#include "MuCOE/Nodes/SCustomizableObjectNodeStaticStringPin.h"
+
 #include "Widgets/SBoxPanel.h"
 
 #define LOCTEXT_NAMESPACE "CustomizableObjectEditor"
@@ -19,6 +25,15 @@ void SCustomizableObjectNode::Construct(const FArguments& InArgs, UEdGraphNode*
 
 TSharedPtr<SGraphPin> SCustomizableObjectNode::CreatePinWidget(UEdGraphPin* Pin) const
 {
+	check(Pin->GetOwningNode());
+
+	const UCustomizableObjectNode* Node = Cast<UCustomizableObjectNode>(Pin->GetOwningNode());
+
+	if (Pin->PinType.PinCategory == UEdGraphSchema_CustomizableObject::PC_String && Pin->PinType.ContainerType != EPinContainerType::Array && Node->CreateStaticStringPinWidget())
+	{
+		return SNew(SCustomizableObjectNodeStaticStringPin, Pin);
+	}
+
 	return SNew(SCustomizableObjectNodePin, Pin);
 }
 
diff --git a/Engine/Plugins/Mutable/Source/CustomizableObjectEditor/Private/MuCOE/Nodes/SCustomizableObjectNodeStaticStringPin.cpp b/Engine/Plugins/Mutable/Source/CustomizableObjectEditor/Private/MuCOE/Nodes/SCustomizableObjectNodeStaticStringPin.cpp
new file mode 100644
--- /dev/null
+++ b/Engine/Plugins/Mutable/Source/CustomizableObjectEditor/Private/MuCOE/Nodes/SCustomizableObjectNodeStaticStringPin.cpp
@@ -0,0 +1,114 @@
+// Copyright Epic Games, Inc. All Rights Reserved.
+
+#include "MuCOE/Nodes/SCustomizableObjectNodeStaticStringPin.h"
+
+#include "MuCOE/Nodes/CustomizableObjectNodeStaticString.h"
+#include "MuCOE/Nodes/CustomizableObjectNodeComponentMesh.h"
+#include "MuCOE/Nodes/CustomizableObjectNodeComponentMeshAddTo.h"
+#include "MuCOE/Nodes/CustomizableObjectNodeMeshMorph.h"
+#include "MuCOE/Nodes/CustomizableObjectNodeMeshParameter.h"
+#include "MuCOE/Nodes/CustomizableObjectNodeObjectChild.h"
+#include "MuCOE/Nodes/CustomizableObjectNodeObjectGroup.h"
+#include "MuCOE/Nodes/CustomizableObjectNodeParameter.h"
+#include "MuCOE/Nodes/CustomizableObjectNodeTunnel.h"
+#include "Widgets/Input/SEditableTextBox.h"
+
+
+void SCustomizableObjectNodeStaticStringPin::Construct(const FArguments& InArgs, UEdGraphPin* InGraphPinObj)
+{
+	SCustomizableObjectNodePin::Construct(SCustomizableObjectNodePin::FArguments(), InGraphPinObj);
+
+	// Hide pin Icon and Label
+	if (InGraphPinObj->bNotConnectable)
+	{
+		SetShowLabel(false);
+		PinImage->SetVisibility(EVisibility::Collapsed);
+	}
+}
+
+
+TSharedRef<SWidget>	SCustomizableObjectNodeStaticStringPin::GetDefaultValueWidget()
+{
+	return SNew(SEditableTextBox)
+		.Text(this, &SCustomizableObjectNodeStaticStringPin::GetNodeStringValue)
+		.OnTextCommitted(this, &SCustomizableObjectNodeStaticStringPin::OnTextCommited)
+		.Visibility(this, &SCustomizableObjectNodeStaticStringPin::GetWidgetVisibility);
+}
+
+
+FText SCustomizableObjectNodeStaticStringPin::GetNodeStringValue() const
+{
+	const UEdGraphNode* Node = GraphPinObj->GetOwningNode();
+
+	if (const UCustomizableObjectNodeStaticString* StringNode = Cast<UCustomizableObjectNodeStaticString>(Node))
+	{
+		return FText::FromString(StringNode->Value);
+	}
+	else if (const UCustomizableObjectNodeComponent* ComponentNode = Cast<UCustomizableObjectNodeComponent>(Node))
+	{
+		return FText::FromName(ComponentNode->GetComponentName());
+	}
+	else if (const UCustomizableObjectNodeComponentMeshAddTo* AddComponentNode = Cast<UCustomizableObjectNodeComponentMeshAddTo>(Node))
+	{
+		return FText::FromName(AddComponentNode->GetParentComponentName());
+	}
+	else if (const UCustomizableObjectNodeMeshParameter* MeshParamNode = Cast<UCustomizableObjectNodeMeshParameter>(Node))
+	{
+		return FText::FromString(MeshParamNode->ParameterName);
+	}
+	else if (const UCustomizableObjectNodeObjectChild* ChildObjectNode = Cast<UCustomizableObjectNodeObjectChild>(Node))
+	{
+		return FText::FromString(ChildObjectNode->GetObjectName());
+	}
+	else if (const UCustomizableObjectNodeObjectGroup* GroupNode = Cast<UCustomizableObjectNodeObjectGroup>(Node))
+	{
+		return FText::FromString(GroupNode->GetGroupName());
+	}
+	else if (const UCustomizableObjectNodeParameter* ParameterNode = Cast< UCustomizableObjectNodeParameter>(Node))
+	{
+		return FText::FromString(ParameterNode->GetParameterName());
+	}
+
+	return FText();
+}
+
+
+void SCustomizableObjectNodeStaticStringPin::OnTextCommited(const FText& InValue, ETextCommit::Type InCommitInfo)
+{
+	UEdGraphNode* Node = GraphPinObj->GetOwningNode();
+
+	if (UCustomizableObjectNodeStaticString* StringNode = Cast<UCustomizableObjectNodeStaticString>(Node))
+	{
+		StringNode->Value = InValue.ToString();
+	}
+	else if (UCustomizableObjectNodeComponent* ComponentNode = Cast<UCustomizableObjectNodeComponent>(Node))
+	{
+		ComponentNode->SetComponentName(FName(*InValue.ToString()));
+	}
+	else if (UCustomizableObjectNodeComponentMeshAddTo* AddComponentNode = Cast<UCustomizableObjectNodeComponentMeshAddTo>(Node))
+	{
+		AddComponentNode->SetParentComponentName(FName(*InValue.ToString()));
+	}
+	else if (UCustomizableObjectNodeMeshParameter* MeshParamNode = Cast<UCustomizableObjectNodeMeshParameter>(Node))
+	{
+		MeshParamNode->ParameterName = InValue.ToString();
+	}
+	else if (UCustomizableObjectNodeObjectChild* ChildObjectNode = Cast<UCustomizableObjectNodeObjectChild>(Node))
+	{
+		ChildObjectNode->SetObjectName(InValue.ToString());
+	}
+	else if (UCustomizableObjectNodeObjectGroup* GroupNode = Cast<UCustomizableObjectNodeObjectGroup>(Node))
+	{
+		GroupNode->SetGroupName(InValue.ToString());
+	}
+	else if (UCustomizableObjectNodeParameter* ParameterNode = Cast<UCustomizableObjectNodeParameter>(Node))
+	{
+		ParameterNode->SetParameterName(InValue.ToString());
+	}
+}
+
+
+EVisibility SCustomizableObjectNodeStaticStringPin::GetWidgetVisibility() const
+{
+	return GraphPinObj->LinkedTo.Num() ? EVisibility::Collapsed : EVisibility::Visible;
+}
diff --git a/Engine/Plugins/Mutable/Source/CustomizableObjectEditor/Private/MuCOE/Nodes/SCustomizableObjectNodeStaticStringPin.h b/Engine/Plugins/Mutable/Source/CustomizableObjectEditor/Private/MuCOE/Nodes/SCustomizableObjectNodeStaticStringPin.h
new file mode 100644
--- /dev/null
+++ b/Engine/Plugins/Mutable/Source/CustomizableObjectEditor/Private/MuCOE/Nodes/SCustomizableObjectNodeStaticStringPin.h
@@ -0,0 +1,28 @@
+// Copyright Epic Games, Inc. All Rights Reserved.
+
+#pragma once
+
+#include "MuCOE/Nodes/SCustomizableObjectNodePin.h"
+
+class SCustomizableObjectNodeStaticStringPin : public SCustomizableObjectNodePin
+{
+public:
+
+	SLATE_BEGIN_ARGS(SCustomizableObjectNodeStaticStringPin) {}
+	SLATE_END_ARGS()
+
+	void Construct(const FArguments& InArgs, UEdGraphPin* InGraphPinObj);
+
+	TSharedRef<SWidget>	GetDefaultValueWidget() override;
+
+private:
+
+	/** Return the value of the node */
+	FText GetNodeStringValue() const;
+
+	/** On Text Commited Callback */
+	void OnTextCommited(const FText& InValue, ETextCommit::Type InCommitInfo);
+
+	/** Returns the visibility of the pin widget */
+	EVisibility GetWidgetVisibility() const;
+};
diff --git a/Engine/Plugins/Mutable/Source/CustomizableObjectEditor/Private/MuCOE/SCustomizableObjectEditorTagExplorer.cpp b/Engine/Plugins/Mutable/Source/CustomizableObjectEditor/Private/MuCOE/SCustomizableObjectEditorTagExplorer.cpp
--- a/Engine/Plugins/Mutable/Source/CustomizableObjectEditor/Private/MuCOE/SCustomizableObjectEditorTagExplorer.cpp
+++ b/Engine/Plugins/Mutable/Source/CustomizableObjectEditor/Private/MuCOE/SCustomizableObjectEditorTagExplorer.cpp
@@ -159,11 +159,11 @@ void SCustomizableObjectEditorTagExplorer::FillTagInformation(const UCustomizabl
 			// Gather tags from nodes that activate them
 			if (UCustomizableObjectNode* TypedNodeMat = Cast<UCustomizableObjectNode>(Node))
 			{
-				const TArray<FString>* MaterialTags = TypedNodeMat->GetEnableTags();
+				const TArray<FString> MaterialTags = TypedNodeMat->GetEnableTags();
 				
-				if (MaterialTags)
+				if (MaterialTags.Num())
 				{
-					for (const FString& CurrentTag : *MaterialTags)
+					for (const FString& CurrentTag : MaterialTags)
 					{
 						NodeTags.Add(CurrentTag, TypedNodeMat);
 						Tags.AddUnique(CurrentTag);
diff --git a/Engine/Plugins/Mutable/Source/CustomizableObjectEditor/Private/MuCOE/SMutableTagListWidget.cpp b/Engine/Plugins/Mutable/Source/CustomizableObjectEditor/Private/MuCOE/SMutableTagListWidget.cpp
--- a/Engine/Plugins/Mutable/Source/CustomizableObjectEditor/Private/MuCOE/SMutableTagListWidget.cpp
+++ b/Engine/Plugins/Mutable/Source/CustomizableObjectEditor/Private/MuCOE/SMutableTagListWidget.cpp
@@ -134,7 +134,7 @@ TSharedPtr<SMutableSearchComboBox::FFilteredOption> SMutableTagComboBox::AddNode
 		{
 			Option = MakeShared<SMutableSearchComboBox::FFilteredOption>();
 			Option->Parent = ParentOption;
-			Option->DisplayOption = ObjectNode->ObjectName;
+			Option->DisplayOption = ObjectNode->GetObjectName();
 			if (Option->DisplayOption.IsEmpty())
 			{
 				Option->DisplayOption = "Unnamed Object";
@@ -146,7 +146,7 @@ TSharedPtr<SMutableSearchComboBox::FFilteredOption> SMutableTagComboBox::AddNode
 		{
 			Option = MakeShared<SMutableSearchComboBox::FFilteredOption>();
 			Option->Parent = ParentOption;
-			Option->DisplayOption = GroupNode->GroupName;
+			Option->DisplayOption = GroupNode->GetGroupName();
 			if (Option->DisplayOption.IsEmpty())
 			{
 				Option->DisplayOption = "Unnamed Group";
@@ -204,12 +204,12 @@ void SMutableTagComboBox::RefreshOptions()
 				continue;
 			}
 
-			TArray<FString>* EnableTags = Typed->GetEnableTags();
-			if (EnableTags)
+			const TArray<FString>& EnableTags = Typed->GetEnableTags();
+			if (EnableTags.Num())
 			{
 				TSharedPtr<SMutableSearchComboBox::FFilteredOption> NodeOption = AddNodeHierarchyOptions(Typed, AddedOptions);
 
-				for (const FString& OneTag : *EnableTags)
+				for (const FString& OneTag : EnableTags)
 				{
 					if (OneTag.IsEmpty())
 					{
diff --git a/Engine/Plugins/Mutable/Source/CustomizableObjectEditor/Public/MuCOE/Nodes/CustomizableObjectNode.h b/Engine/Plugins/Mutable/Source/CustomizableObjectEditor/Public/MuCOE/Nodes/CustomizableObjectNode.h
--- a/Engine/Plugins/Mutable/Source/CustomizableObjectEditor/Public/MuCOE/Nodes/CustomizableObjectNode.h
+++ b/Engine/Plugins/Mutable/Source/CustomizableObjectEditor/Public/MuCOE/Nodes/CustomizableObjectNode.h
@@ -12,6 +12,7 @@
 class FArchive;
 class FCustomizableObjectGraphEditorToolkit;
 class SWidget;
+class UCustomizableObjectNodeMacroInstance;
 class UCustomizableObjectNodeObject;
 class UCustomizableObjectNodeRemapPins;
 class UCustomizableObjectNodeRemapPinsByName;
@@ -177,8 +178,11 @@ public:
 	// Used during compilation process to cache the node for all LODs, or generate it specifically for each of them.
 	virtual bool IsAffectedByLOD() const;
 
-	/** Return an array of tags that this node will enable and apply to its data. Null if none, or of it doesn't apply. */
-	virtual TArray<FString>* GetEnableTags();
+	/** Return an array of tags that this node will enable and apply to its data. This tags can come from the node details or string nodes linked to the "Enable tags" pin. */
+	virtual TArray<FString> GetEnableTags(TArray<const UCustomizableObjectNodeMacroInstance*>* MacroContext = nullptr);
+
+	/** Return an array with all tags defined in the node's details panel. Null if none, or of it doesn't apply. */
+	virtual TArray<FString>* GetEnableTagsArray();
 
 	/** Return the unique internal tag that can identify this node across all objects.
 	* Internal tags are automatically added to some nodes that can be then referred to from other nodes.
@@ -288,6 +292,9 @@ public:
 
 	/** Returns true if this node is inside a Macro Graph */
 	bool IsInMacro() const;
+
+	/** If a node has a static string pin, this function determines if it should construct an input text widget next to the pin or not. */
+	virtual bool CreateStaticStringPinWidget() const;
 	
 	UPROPERTY()
 	FPostEditChangePropertyDelegate PostEditChangePropertyDelegate;
