[mutable] Remove unused compiler code.
#rnx

diff --git a/Engine/Plugins/Mutable/Source/MutableTools/Private/MuT/CodeOptimiser.cpp b/Engine/Plugins/Mutable/Source/MutableTools/Private/MuT/CodeOptimiser.cpp
--- a/Engine/Plugins/Mutable/Source/MutableTools/Private/MuT/CodeOptimiser.cpp
+++ b/Engine/Plugins/Mutable/Source/MutableTools/Private/MuT/CodeOptimiser.cpp
@@ -348,99 +348,6 @@ namespace mu
 	}
 
 
-	FFindOpTypeVisitor::FFindOpTypeVisitor(const TArray<EOpType>& Types)
-		: TypesToFind(Types)
-	{
-	}
-
-	bool FFindOpTypeVisitor::Apply( FProgram& program, OP::ADDRESS rootAt )
-	{
-		bool found = false;
-
-		// If the program added new operations, we haven't visited them.
-		Visited.SetNumZeroed( program.OpAddress.Num() );
-
-		Pending.Reset();
-		Pending.Reserve( program.OpAddress.Num()/4 );
-		Pending.Add({ false,rootAt });
-
-		// Don't early out to be able to complete parent op cached flags
-		while ( !Pending.IsEmpty() )
-		{
-			TPair<bool,int32> item = Pending.Pop();
-			OP::ADDRESS at = item.Value;
-
-			// Not cached?
-			if (Visited[at]<=1)
-			{
-				if (item.Key)
-				{
-					// Item indicating we finished with all the children of a parent
-					check(Visited[at]==1);
-
-					bool bSubtreeFound = false;
-					ForEachReference( program, at, [&](OP::ADDRESS ref)
-					{
-						bSubtreeFound = bSubtreeFound || Visited[ref]==3;
-					});
-
-					Visited[at] = bSubtreeFound ? 3 : 2;
-				}
-				else if (!found)
-				{
-					if (TypesToFind.Find(program.GetOpType(at)) != INDEX_NONE )
-					{
-						Visited[at] = 3; // visited, ops found
-						found = true;
-					}
-					else
-					{
-						check(Visited[at]==0);
-
-						Visited[at] = 1;
-						Pending.Add({ true,at });
-
-						ForEachReference( program, at, [&](OP::ADDRESS ref)
-						{
-							if (ref && Visited[ref]==0)
-							{
-								Pending.Add({ false,ref });
-							}
-						});
-					}
-				}
-				else
-				{
-					// We won't process it.
-					Visited[at] = 0;
-				}
-			}
-		}
-
-		return Visited[rootAt]==3;
-	}
-
-
-	IsConstantVisitor::IsConstantVisitor()
-	{
-		TArray<EOpType> ParameterTypes;
-		ParameterTypes.Add( EOpType::BO_PARAMETER );
-		ParameterTypes.Add( EOpType::NU_PARAMETER );
-		ParameterTypes.Add( EOpType::SC_PARAMETER );
-		ParameterTypes.Add( EOpType::CO_PARAMETER );
-		ParameterTypes.Add( EOpType::PR_PARAMETER );
-		ParameterTypes.Add( EOpType::IM_PARAMETER );
-		ParameterTypes.Add( EOpType::ME_PARAMETER );
-		FindOpTypeVisitor = MakeUnique<FFindOpTypeVisitor>(ParameterTypes);
-	}
-
-
-	bool IsConstantVisitor::Apply( FProgram& program, OP::ADDRESS at )
-	{
-		return !FindOpTypeVisitor->Apply(program,at);
-	}
-
-
 	class FConstantTask 
 	{
 	public:
diff --git a/Engine/Plugins/Mutable/Source/MutableTools/Private/MuT/CodeOptimiser.h b/Engine/Plugins/Mutable/Source/MutableTools/Private/MuT/CodeOptimiser.h
--- a/Engine/Plugins/Mutable/Source/MutableTools/Private/MuT/CodeOptimiser.h
+++ b/Engine/Plugins/Mutable/Source/MutableTools/Private/MuT/CodeOptimiser.h
@@ -35,48 +35,6 @@ namespace mu
     };
 
 
-    /** Scan a subtree and return true if any op of a type in the provided subset is found.
-    * Intermediate data is used between calls to apply, so don't remove program code or directly
-    * change the instructions. Adding new instructions is ok.
-    */
-    class FFindOpTypeVisitor
-    {
-    public:
-
-		FFindOpTypeVisitor(const TArray<EOpType>&);
-        bool Apply( FProgram&, OP::ADDRESS );
-
-    private:
-
-        TArray<EOpType> TypesToFind;
-
-		TArray< TPair<bool,OP::ADDRESS> > Pending;
-
-        // 0 not visited
-        // 1 children pending
-        // 2 visited and not found
-        // 3 visited and found
-		TArray<uint8> Visited;
-    };
-
-
-    /** Scan a subtree and return true if the expression is constant.
-    * It has the same restrictions that FindOpTypeVisitor
-    */
-    class IsConstantVisitor
-    {
-    public:
-
-        IsConstantVisitor();
-
-        bool Apply( FProgram&, OP::ADDRESS );
-
-    private:
-
-        TUniquePtr<FFindOpTypeVisitor> FindOpTypeVisitor;
-    };
-
-
     /** ConstantGenerator replaces constant subtrees of operations with an equivalent single */
     extern bool ConstantGenerator( const CompilerOptions::Private*, Ptr<ASTOp>& Root, int32 Pass );
 
diff --git a/Engine/Plugins/Mutable/Source/MutableTools/Private/MuT/DataPacker.cpp b/Engine/Plugins/Mutable/Source/MutableTools/Private/MuT/DataPacker.cpp
--- a/Engine/Plugins/Mutable/Source/MutableTools/Private/MuT/DataPacker.cpp
+++ b/Engine/Plugins/Mutable/Source/MutableTools/Private/MuT/DataPacker.cpp
@@ -39,262 +39,8 @@
 namespace mu
 {
 
-    //---------------------------------------------------------------------------------------------
-    //---------------------------------------------------------------------------------------------
-    //---------------------------------------------------------------------------------------------
-    SubtreeSearchConstantVisitor::SubtreeSearchConstantVisitor( FProgram& program,
-                                                                OP::ADDRESS constant,
-                                                                EOpType optype )
-        : Program( program ),
-          Constant( constant ),
-          OpType( optype )
-    {
-        Visited.SetNum( program.OpAddress.Num() );
-        if ( Visited.Num() )
-        {
-            FMemory::Memset( &Visited[0], 0, Visited.Num() );
-        }
-    }
-
-
-    //---------------------------------------------------------------------------------------------
-    bool SubtreeSearchConstantVisitor::Run( OP::ADDRESS root )
-    {
-        // too spammy
-        // \todo: called too many times!
-        // MUTABLE_CPUPROFILER_SCOPE(SubtreeSearchConstantVisitor);
-
-        bool found = false;
-
-        TArray< TPair<bool,int> > pending;
-        pending.Reserve(Program.OpAddress.Num()/4);
-
-		pending.Add({ false,root });
-
-        while( pending.Num() )
-        {
-			TPair<bool,int> item = pending.Pop();
-
-            if (item.Key)
-            {
-                // Item indicating we finished with all the children of a parent.
-
-                // Propagate the "found" state upwards
-                Visited[item.Value]=1;
-
-                //OP cop = Program.m_code[item.Value];
-                ForEachReference( Program, (OP::ADDRESS)item.Value, [&](OP::ADDRESS ref)
-                {
-                    if (ref && Visited[ref]==2)
-                    {
-                        Visited[item.Value]=2;
-                    }
-                });
-            }
-            else
-            {
-                if ( !Visited[item.Value] )
-                {
-                    EOpType thisOpType = Program.GetOpType(item.Value);
-                    if ( OpType == thisOpType )
-                    {
-						OP::ResourceConstantArgs args = Program.GetOpArgs<OP::ResourceConstantArgs>(item.Value);
-                        if ( Constant == args.value )
-                        {
-                            found = true;
-                            Visited[item.Value] = 2;
-                        }
-                    }
-
-                    if (!found)
-                    {
-						pending.Add({ true,item.Value });
-
-                        //OP cop = Program.m_code[item.Value];
-                        ForEachReference( Program, item.Value, [&](OP::ADDRESS ref)
-                        {
-                            if (ref && !Visited[ref])
-                            {
-								pending.Add({ false,ref });
-                            }
-                        });
-                    }
-                }
-                else
-                {
-                    found = (Visited[item.Value]==2);
-                }
-            }
-        }
-
-        return found;
-    }
-
-
-
-    //---------------------------------------------------------------------------------------------
-    //! Get all the parameters that affect the constant.
-    //! "Affect" means that the constant may be used or not depending on the parameter.
-    //---------------------------------------------------------------------------------------------
-    class GatherParametersVisitor : public UniqueConstCodeVisitorIterative< TArray<int> >
-    {
-    public:
-
-        GatherParametersVisitor( FProgram& program, OP::ADDRESS constant, EOpType opType )
-            : m_constSearch( program, constant, opType )
-        {
-            MUTABLE_CPUPROFILER_SCOPE(GatherParametersVisitor);
-
-			TArray<int> currentParams;
-            currentParams.SetNumZeroed( program.Parameters.Num() );
-            SetDefaultState(currentParams);
-
-            // TODO: we can optimize by precalculating what ops have the required optype below
-            OpType = opType;
-            Constant = constant;
-
-            m_allParams.SetNumZeroed( program.Parameters.Num() );
-            FullTraverse( program );
-
-            for (int32 i=0; i<m_allParams.Num(); ++i )
-            {
-                if ( m_allParams[i] )
-                {
-                    m_sortedParams.Add( int(i) );
-                }
-            }
-        }
-
-
-        virtual bool Visit( OP::ADDRESS at, FProgram& program )
-        {
-            bool recurse = true;
-
-            if ( m_constSearch.Run( at ) )
-            {
-                switch ( program.GetOpType(at) )
-                {
-                case EOpType::NU_CONDITIONAL:
-                case EOpType::SC_CONDITIONAL:
-                case EOpType::CO_CONDITIONAL:
-                case EOpType::IM_CONDITIONAL:
-                case EOpType::ME_CONDITIONAL:
-                case EOpType::LA_CONDITIONAL:
-                case EOpType::IN_CONDITIONAL:
-				case EOpType::ED_CONDITIONAL:
-                {
-					OP::ConditionalArgs args = program.GetOpArgs<OP::ConditionalArgs>(at);
-
-                    // If the constant is present in only one of the 2 branches
-                    bool foundYes = m_constSearch.Run( args.yes  );
-                    bool foundNo = m_constSearch.Run( args.no );
-
-                    if ( foundYes!=foundNo )
-                    {
-                        m_conditionVisitor.Run( args.condition, program );
-
-                        TArray<int32> currentParams = GetCurrentState();
-                        for (int32 p=0; p<m_conditionVisitor.RelevantParams.Num(); ++p )
-                        {
-                            currentParams[ m_conditionVisitor.RelevantParams[p] ]++;
-                        }
-
-                        // Set state for child recursion
-                        SetCurrentState( currentParams );
-                    }
-
-                    break;
-                }
-
-                case EOpType::NU_SWITCH:
-                case EOpType::SC_SWITCH:
-                case EOpType::CO_SWITCH:
-                case EOpType::IM_SWITCH:
-                case EOpType::ME_SWITCH:
-                case EOpType::LA_SWITCH:
-				case EOpType::ED_SWITCH:
-                {
-					const uint8* data = program.GetOpArgsPointer(at);
-					
-					OP::ADDRESS VarAddress;
-					FMemory::Memcpy( &VarAddress, data, sizeof(OP::ADDRESS) );
-
-                    m_conditionVisitor.Run( VarAddress, program );
-
-					TArray<int32> currentParams = GetCurrentState();
-                    for (int32 p=0; p<m_conditionVisitor.RelevantParams.Num(); ++p )
-                    {
-                        currentParams[ m_conditionVisitor.RelevantParams[p] ]++;
-                    }
-
-                    // Set state for child recursion
-                    SetCurrentState( currentParams );
-
-                    break;
-                }
-
-                default:
-                {
-                    // If we hit the constant we are analysing
-                    if (OpType == EOpType::ME_CONSTANT && program.GetOpType(at) == OpType)
-                    {
-						OP::MeshConstantArgs Args = program.GetOpArgs<OP::MeshConstantArgs>(at);
-                        if (Args.Value == Constant)
-                        {
-                            // Accumulate the currently relevant parameters
-                            const TArray<int32>& CurrentParams = GetCurrentState();
-                            for (int32 ParamIndex = 0; ParamIndex < CurrentParams.Num(); ++ParamIndex)
-                            {
-                                m_allParams[ParamIndex] += CurrentParams[ParamIndex];
-                            }
-                        }
-                    }
-                    else if ( program.GetOpType(at) == OpType )
-                    {
-						OP::ResourceConstantArgs args = program.GetOpArgs<OP::ResourceConstantArgs>(at);
-                        if ( args.value == Constant )
-                        {
-                            // Accumulate the currently relevant parameters
-                            const TArray<int>& currentParams = GetCurrentState();
-                            for (int32 i=0; i<currentParams.Num(); ++i )
-                            {
-                                m_allParams[i] += currentParams[i];
-                            }
-                        }
-                    }
-
-                    break;
-                }
-
-                }
-
-            }
-
-            return recurse;
-        }
-
-    public:
-
-        //! Output result
-		TArray<int> m_sortedParams;
-
-    private:
-
-        OP::ADDRESS Constant;
-        EOpType OpType;
-
-		TArray<int> m_allParams;
-
-        SubtreeParametersVisitor m_conditionVisitor;
-        SubtreeSearchConstantVisitor m_constSearch;
-
-    };
-
-
-    //---------------------------------------------------------------------------------------------
-    //---------------------------------------------------------------------------------------------
-    //---------------------------------------------------------------------------------------------
-    class AccumulateImageFormatsAST : public Visitor_TopDown_Unique_Const< TArray<bool> >
+	
+	class AccumulateImageFormatsAST : public Visitor_TopDown_Unique_Const< TArray<bool> >
     {
     public:
 
@@ -302,10 +48,6 @@ namespace mu
         {
             MUTABLE_CPUPROFILER_SCOPE(AccumulateImageFormatsAST);
 
-            // TODO: Use statically sized vectors (std::array)
-            //vector<bool> initial( Count, true );
-            //std::fill( SupportedFormats.begin(), SupportedFormats.end(), initial );
-
 			TArray<bool> defaultState;
 			defaultState.SetNumZeroed(int32(EImageFormat::Count));
 
diff --git a/Engine/Plugins/Mutable/Source/MutableTools/Private/MuT/DataPacker.h b/Engine/Plugins/Mutable/Source/MutableTools/Private/MuT/DataPacker.h
--- a/Engine/Plugins/Mutable/Source/MutableTools/Private/MuT/DataPacker.h
+++ b/Engine/Plugins/Mutable/Source/MutableTools/Private/MuT/DataPacker.h
@@ -14,28 +14,4 @@ namespace mu
     /** Convert constant data to different formats, based on their usage. */
     extern void DataOptimise( const CompilerOptions*, ASTOpList& roots );
 
-    //---------------------------------------------------------------------------------------------
-    //! Find the given constant in a subtree
-    //---------------------------------------------------------------------------------------------
-    class SubtreeSearchConstantVisitor
-    {
-    public:
-
-        SubtreeSearchConstantVisitor( FProgram&, OP::ADDRESS Constant, EOpType );
-
-        bool Run( OP::ADDRESS root );
-
-
-    private:
-
-        FProgram& Program;
-        OP::ADDRESS Constant;
-        EOpType OpType;
-
-        //! 0: not visited
-        //! 1: visited and not found
-        //! 2: visited and found
-        TArray<uint8> Visited;
-    };
-
 }
