[Mutable] Fix IndexBuffer re-initialization for mesh streaming.

- Every time the buffer was re-initialized it inserted more elements to it. Not only was allocating more data but it triggered a possible bug with RHI buffer allocation tracking when the re-initialized buffer needed an allocation larger than the maximum allowed by the pool.

#rnx

diff --git a/Engine/Plugins/Mutable/Source/CustomizableObject/Private/MuCO/UnrealConversionUtils.cpp b/Engine/Plugins/Mutable/Source/CustomizableObject/Private/MuCO/UnrealConversionUtils.cpp
--- a/Engine/Plugins/Mutable/Source/CustomizableObject/Private/MuCO/UnrealConversionUtils.cpp
+++ b/Engine/Plugins/Mutable/Source/CustomizableObject/Private/MuCO/UnrealConversionUtils.cpp
@@ -494,6 +494,9 @@ namespace UnrealConversionUtils
 		check(LODResource.MultiSizeIndexContainer.GetDataTypeSize() == ElementSize);
 		
 		const int32 MaxNumIndices = LODResource.RenderSections.Last().BaseIndex + (LODResource.RenderSections.Last().NumTriangles * 3);
+
+		// Don't assume the buffer is empty, otherwise we may add extra elements using Insert(). 
+		LODResource.MultiSizeIndexContainer.GetIndexBuffer()->Empty(MaxNumIndices);
 		LODResource.MultiSizeIndexContainer.GetIndexBuffer()->Insert(0, MaxNumIndices);
 
 		if (IndexCount == MaxNumIndices)
