[mutable] Don't treat curve constant as if it was an operation.
#rnx

diff --git a/Engine/Plugins/Mutable/Source/MutableRuntime/Internal/MuR/Operations.h b/Engine/Plugins/Mutable/Source/MutableRuntime/Internal/MuR/Operations.h
--- a/Engine/Plugins/Mutable/Source/MutableRuntime/Internal/MuR/Operations.h
+++ b/Engine/Plugins/Mutable/Source/MutableRuntime/Internal/MuR/Operations.h
@@ -464,7 +464,7 @@ namespace mu
         struct ScalarCurveArgs
         {
             ADDRESS time;   // Operation generating the time value to sample the curve
-            ADDRESS curve;  // Constant curve
+            ADDRESS curve;  // Constant curve (not an op)
         };
 
         struct ColourSampleImageArgs
diff --git a/Engine/Plugins/Mutable/Source/MutableRuntime/Private/MuR/Operations.cpp b/Engine/Plugins/Mutable/Source/MutableRuntime/Private/MuR/Operations.cpp
--- a/Engine/Plugins/Mutable/Source/MutableRuntime/Private/MuR/Operations.cpp
+++ b/Engine/Plugins/Mutable/Source/MutableRuntime/Private/MuR/Operations.cpp
@@ -185,7 +185,6 @@ namespace mu
         {
 			OP::ScalarCurveArgs args = program.GetOpArgs<OP::ScalarCurveArgs>(at);
             f(args.time );
-        	f(args.curve);
             break;
         }
 
