[mutable] Fixed dragged&dropped CO Instances to a level not loading during PIE, remaining as the reference mesh.
#tests PIE
#rnx

diff --git a/Engine/Plugins/Mutable/Source/CustomizableObject/Private/MuCO/CustomizableObjectSystem.cpp b/Engine/Plugins/Mutable/Source/CustomizableObject/Private/MuCO/CustomizableObjectSystem.cpp
--- a/Engine/Plugins/Mutable/Source/CustomizableObject/Private/MuCO/CustomizableObjectSystem.cpp
+++ b/Engine/Plugins/Mutable/Source/CustomizableObject/Private/MuCO/CustomizableObjectSystem.cpp
@@ -4020,14 +4020,14 @@ int32 UCustomizableObjectSystem::TickInternal(const bool bBlocking)
 		}
 		
 		if (!Usage->GetSkipSetReferenceSkeletalMesh() &&
-			Object->bEnableUseRefSkeletalMeshAsPlaceholder)
+			 Object->bEnableUseRefSkeletalMeshAsPlaceholder)
 		{
 			const FName& ComponentName = Usage->GetComponentName();
 			
-			if (USkeletalMesh* SkeletalMesh = Object->GetComponentMeshReferenceSkeletalMesh(ComponentName))
+			if (USkeletalMesh* ReferenceSkeletalMesh = Object->GetComponentMeshReferenceSkeletalMesh(ComponentName))
 			{
 				Parent->EmptyOverrideMaterials();
-				Parent->SetSkeletalMesh(SkeletalMesh);
+				Parent->SetSkeletalMesh(ReferenceSkeletalMesh);
 			}
 		}
 		
@@ -4139,7 +4139,7 @@ int32 UCustomizableObjectSystem::TickInternal(const bool bBlocking)
 							LODUpdateDescriptorHash.FirstRequestedLOD = LODUpdateCandidate.FirstRequestedLOD;
 
 							if (CustomizableInstancePrivate->MinSquareDistFromComponentToPlayer < MaxSquareDistanceFound &&
-								(IsStreamingEnabled(*Object) || !LODUpdateDescriptorHash.IsSubset(CustomizableInstancePrivate->CommittedDescriptorHash)))
+								(IsStreamingEnabled(*Object) || !LODUpdateDescriptorHash.IsSubset(CustomizableInstancePrivate->CommittedDescriptorHash) || CustomizableInstancePrivate->bAutomaticUpdateRequired))
 							{
 								MaxPriorityFound = LODUpdateCandidate.Priority;
 								MaxSquareDistanceFound = CustomizableInstancePrivate->MinSquareDistFromComponentToPlayer;
