[Mutable] Fix PhysicsAssets and Clothing bake.

- Transient physics assets could have references to private objects in other packages preventing them to be saved.
- Clothing data was not copied correctly in RegenerateImportedModel. Now the data is taken from the render data.
- Now baked PhisicsAssets are named following similar conventions to what is done to skeletons.

#rnx

diff --git a/Engine/Plugins/Mutable/Source/CustomizableObject/Private/MuCO/CustomizableObjectInstance.cpp b/Engine/Plugins/Mutable/Source/CustomizableObject/Private/MuCO/CustomizableObjectInstance.cpp
--- a/Engine/Plugins/Mutable/Source/CustomizableObject/Private/MuCO/CustomizableObjectInstance.cpp
+++ b/Engine/Plugins/Mutable/Source/CustomizableObject/Private/MuCO/CustomizableObjectInstance.cpp
@@ -1357,6 +1357,23 @@ USkeleton* UCustomizableInstancePrivate::MergeSkeletons(UCustomizableObject& Cus
 
 namespace
 {
+	FORCEINLINE TObjectPtr<UPhysicsConstraintTemplate> ClonePhysicsConstraintTemplate(
+			const TObjectPtr<UPhysicsConstraintTemplate>& From, 
+			const TObjectPtr<UObject>& Outer,
+			FName Name = NAME_None)
+	{
+		// We don't use DuplicateObject here beacuse it is too slow.
+		TObjectPtr<UPhysicsConstraintTemplate> Result = NewObject<UPhysicsConstraintTemplate>(Outer, Name);
+		
+		Result->DefaultInstance = From->DefaultInstance;
+		Result->ProfileHandles = From->ProfileHandles;
+#if WITH_EDITOR
+		Result->SetDefaultProfile(From->DefaultInstance);
+#endif
+
+		return Result;
+	}
+
 	FKAggregateGeom MakeAggGeomFromMutablePhysics(int32 BodyIndex, const mu::FPhysicsBody* MutablePhysicsBody)
 	{
 		FKAggregateGeom BodyAggGeom;
@@ -1581,7 +1598,21 @@ namespace
 		if (!bTemplateBodyNotUsedFound)
 		{
 			Result->CollisionDisableTable = TemplateAsset->CollisionDisableTable;
-			Result->ConstraintSetup = TemplateAsset->ConstraintSetup;
+
+			const int32 NumConstraints = TemplateAsset->ConstraintSetup.Num();
+			Result->ConstraintSetup.SetNum(NumConstraints);	
+	
+			for (int32 ConstraintIndex = 0; ConstraintIndex < NumConstraints; ++ConstraintIndex)
+			{
+				const TObjectPtr<UPhysicsConstraintTemplate>& TemplateConstraint = TemplateAsset->ConstraintSetup[ConstraintIndex];
+
+				if (!TemplateConstraint)
+				{
+					continue;
+				}
+
+				Result->ConstraintSetup[ConstraintIndex] = ClonePhysicsConstraintTemplate(TemplateConstraint, Result);
+			}
 		}
 		else
 		{
@@ -1599,28 +1630,30 @@ namespace
 			}
 
 			// Only add constraints that are part of the bones used for the mutable physics volumes description.
-			for (const TObjectPtr<UPhysicsConstraintTemplate>& Constrain : TemplateAsset->ConstraintSetup)
+			Result->ConstraintSetup.Reserve(TemplateAsset->ConstraintSetup.Num());
+			for (const TObjectPtr<UPhysicsConstraintTemplate>& Constraint : TemplateAsset->ConstraintSetup)
 			{
-				const FName BoneA = Constrain->DefaultInstance.ConstraintBone1;
-				const FName BoneB = Constrain->DefaultInstance.ConstraintBone2;
+				if (!Constraint)
+				{
+					continue;
+				}
+
+				const FName BoneA = Constraint->DefaultInstance.ConstraintBone1;
+				const FName BoneB = Constraint->DefaultInstance.ConstraintBone2;
 
 				if (BonesInUse.Find(BoneA) && BonesInUse.Find(BoneB))
 				{
-					Result->ConstraintSetup.Add(Constrain);
+					Result->ConstraintSetup.AddDefaulted_GetRef() = ClonePhysicsConstraintTemplate(Constraint, Result);
 				}	
 			}
 		}
-		
-		Result->ConstraintSetup = TemplateAsset->ConstraintSetup;
-		
+
 		Result->UpdateBodySetupIndexMap();
 		Result->UpdateBoundsBodiesArray();
 
-
-
-	#if WITH_EDITORONLY_DATA
+#if WITH_EDITORONLY_DATA
 		Result->ConstraintProfiles = TemplateAsset->ConstraintProfiles;
-	#endif
+#endif
 
 		return Result;
 	}
@@ -1813,7 +1846,6 @@ UPhysicsAsset* UCustomizableInstancePrivate::GetOrBuildMainPhysicsAsset(
 		}
 	
 		// Remap collision indices removing invalid ones.
-
 		CollisionDisableTable.Reserve(CollisionDisableTable.Num() + ValidAssets[CollisionSetIndex]->CollisionDisableTable.Num());
 		for (const TPair<FRigidBodyIndexPair, bool>& DisabledCollision : ValidAssets[CollisionSetIndex]->CollisionDisableTable)
 		{
@@ -1832,14 +1864,20 @@ UPhysicsAsset* UCustomizableInstancePrivate::GetOrBuildMainPhysicsAsset(
 		}
 
 		// Only add constraints that are part of the bones used for the mutable physics volumes description.
-		for (const TObjectPtr<UPhysicsConstraintTemplate>& Constrain : ValidAssets[CollisionSetIndex]->ConstraintSetup)
+		Result->ConstraintSetup.Reserve(Result->ConstraintSetup.Num() + ValidAssets[CollisionSetIndex]->ConstraintSetup.Num());
+		for (const TObjectPtr<UPhysicsConstraintTemplate>& Constraint : ValidAssets[CollisionSetIndex]->ConstraintSetup)
 		{
-			FName BoneA = Constrain->DefaultInstance.ConstraintBone1;
-			FName BoneB = Constrain->DefaultInstance.ConstraintBone2;
+			if (!Constraint)
+			{
+				continue;
+			}
+
+			FName BoneA = Constraint->DefaultInstance.ConstraintBone1;
+			FName BoneB = Constraint->DefaultInstance.ConstraintBone2;
 
 			if (BonesInUse.Find(BoneA) && BonesInUse.Find(BoneB))
 			{
-				Result->ConstraintSetup.Add(Constrain);
+				Result->ConstraintSetup.AddDefaulted_GetRef() = ClonePhysicsConstraintTemplate(Constraint, Result); 
 			}
 		}
 
@@ -7356,8 +7394,15 @@ void UCustomizableInstancePrivate::RegenerateImportedModels()
 
 				if (RenderSection.ClothMappingDataLODs.Num())
 				{
-					ImportedSection.ClothMappingDataLODs.SetNum(1);
-					ImportedSection.ClothMappingDataLODs[0] = RenderSection.ClothMappingDataLODs[0];
+					TArray<FMeshToMeshVertData>& ImportedClothMappingData = ImportedSection.ClothMappingDataLODs.AddDefaulted_GetRef();
+					
+					const int32 NumClothVerts = LODRenderData.ClothVertexBuffer.GetNumVertices();
+					ImportedClothMappingData.SetNumUninitialized(NumClothVerts);
+				
+					for (int32 ClothVertDataIndex = 0; ClothVertDataIndex < NumClothVerts; ++ClothVertDataIndex)
+					{
+						ImportedClothMappingData[ClothVertDataIndex] = LODRenderData.ClothVertexBuffer.MappingData(ClothVertDataIndex);
+					}
 				}
 
 				// Vertices
diff --git a/Engine/Plugins/Mutable/Source/CustomizableObjectEditor/Private/MuCOE/CustomizableObjectInstanceBakingUtils.cpp b/Engine/Plugins/Mutable/Source/CustomizableObjectEditor/Private/MuCOE/CustomizableObjectInstanceBakingUtils.cpp
--- a/Engine/Plugins/Mutable/Source/CustomizableObjectEditor/Private/MuCOE/CustomizableObjectInstanceBakingUtils.cpp
+++ b/Engine/Plugins/Mutable/Source/CustomizableObjectEditor/Private/MuCOE/CustomizableObjectInstanceBakingUtils.cpp
@@ -34,6 +34,7 @@ const FString MaterialInstanceAssetPrefix = TEXT("MI_");
 const FString MaterialDynamicInstanceAssetPrefix = TEXT("MID_");
 const FString MaterialConstantInstanceAssetPrefix = TEXT("MIC_");
 const FString SkeletonAssetPrefix = TEXT("SKEL_");
+const FString PhysicsAssetPrefix = TEXT("PHYS_");
 
 
 /**
@@ -810,7 +811,7 @@ bool BakeCustomizableObjectInstance(
 		{
 			const bool bTransient = Skeleton->GetPackage() == GetTransientPackage();
 
-			// Don't duplicate if not transient or export all assets.
+			// Duplicate only if transient or export all assets.
 			if (bTransient || bExportAllResources)
 			{
 				FString SkeletonName = Mesh->GetName();
@@ -835,6 +836,40 @@ bool BakeCustomizableObjectInstance(
 			}
 		}
 
+		// Skeletal Mesh's Physics Asset
+		bool bNewPhysicsAssetCreated = false;
+		if (UPhysicsAsset* PhysicsAsset = Mesh->GetPhysicsAsset())
+		{
+			const bool bTransient = PhysicsAsset->GetPackage() == GetTransientPackage();
+
+			// Duplicate only if transient or export all assets.
+			if (bTransient || bExportAllResources)
+			{
+				FString PhysicsAssetName = Mesh->GetName();
+				RemovePrefix(PhysicsAssetName, {SkeletalMeshAssetPrefix});
+				ComposeResourceName(PhysicsAssetPrefix, ObjectBaseName, PhysicsAssetName);
+				MakeResourceNameUnique(PhysicsAssetName, CachedResourceNames);
+				
+				EPackageSaveResolutionType SaveType = EPackageSaveResolutionType::None;
+				if (!CreateAssetAtPath(AssetPath, PhysicsAssetName, bUsedGrantedOverridingRights, bIsUnattendedExecution, SaveType))
+				{
+					return false;
+				}
+
+				FString PhysicsAssetPkgName = AssetPath + FString("/") + PhysicsAssetName;
+				UObject* DuplicatedPhysicsAsset = FUnrealBakeHelpers::BakeHelper_DuplicateAsset(PhysicsAsset, PhysicsAssetName, 
+					PhysicsAssetPkgName, false, ReplacementMap, bUsedGrantedOverridingRights, false);
+
+				CachedObjects.Add(DuplicatedPhysicsAsset);
+				TPair<EPackageSaveResolutionType, UPackage*> PackageToSave {SaveType, DuplicatedPhysicsAsset->GetPackage()};
+				OutSavedPackages.Add(PackageToSave);
+				ReplacementMap.Add(PhysicsAsset, DuplicatedPhysicsAsset);
+
+				bNewPhysicsAssetCreated = true;
+			}
+		}
+
+
 		// Skeletal Mesh
 		EPackageSaveResolutionType SaveType = EPackageSaveResolutionType::None;
 		if (!CreateAssetAtPath(AssetPath, SkeletalMeshName, bUsedGrantedOverridingRights, bIsUnattendedExecution, SaveType))
@@ -900,16 +935,6 @@ bool BakeCustomizableObjectInstance(
 				}
 			}
 
-			// Duplicate PhysicsAsset
-			{
-				if (const UPhysicsAsset* PhysicsAsset = Mesh->GetPhysicsAsset())
-				{
-					// Duplicate to change from the Transient Package to the baked mesh one
-					UPhysicsAsset* NewPhysicsAsset = Cast<UPhysicsAsset>(StaticDuplicateObject(PhysicsAsset, SkeletalMesh));
-					SkeletalMesh->SetPhysicsAsset(NewPhysicsAsset);
-				}
-			}
-
 			// Copy LODSettings from the Reference Skeletal Mesh
 			{
 				if (ModelResources->ReferenceSkeletalMeshesData.IsValidIndex(ObjectComponentIndex))
@@ -919,6 +944,12 @@ bool BakeCustomizableObjectInstance(
 				}
 			}
 
+			// Set the physics asset preview mesh if the SkeletalMesh physics assets has been generated as part of the bake.
+			if (SkeletalMesh->GetPhysicsAsset() && bNewPhysicsAssetCreated)
+			{
+				SkeletalMesh->GetPhysicsAsset()->SetPreviewMesh(SkeletalMesh);
+			}
+
 			// Generate render data
 			SkeletalMesh->Build();
 		}
