[Mutable] Mutable Macros Pre-QA improvements:
- Added static string pin widget logic.
- Removed unnecessary warnings.

#rnx

diff --git a/Engine/Plugins/Mutable/Source/CustomizableObjectEditor/Private/MuCOE/GenerateMutableSource/GenerateMutableSourceImage.cpp b/Engine/Plugins/Mutable/Source/CustomizableObjectEditor/Private/MuCOE/GenerateMutableSource/GenerateMutableSourceImage.cpp
--- a/Engine/Plugins/Mutable/Source/CustomizableObjectEditor/Private/MuCOE/GenerateMutableSource/GenerateMutableSourceImage.cpp
+++ b/Engine/Plugins/Mutable/Source/CustomizableObjectEditor/Private/MuCOE/GenerateMutableSource/GenerateMutableSourceImage.cpp
@@ -1202,11 +1202,6 @@ mu::Ptr<mu::NodeImage> GenerateMutableSourceImage(const UEdGraphPin* Pin, FMutab
 			{
 				Result = GenerateMutableSourceImage(FollowPin, GenerationContext, ReferenceTextureSize);
 			}
-			else
-			{
-				FText Msg = FText::Format(LOCTEXT("MacroTunnelError_PinNotLinked_Image", "Macro Instance Pin {0} is not linked."), FText::FromName(Pin->PinName));
-				GenerationContext.Log(Msg, Node);
-			}
 		}
 		else
 		{
diff --git a/Engine/Plugins/Mutable/Source/CustomizableObjectEditor/Private/MuCOE/GenerateMutableSource/GenerateMutableSourceMacro.h b/Engine/Plugins/Mutable/Source/CustomizableObjectEditor/Private/MuCOE/GenerateMutableSource/GenerateMutableSourceMacro.h
--- a/Engine/Plugins/Mutable/Source/CustomizableObjectEditor/Private/MuCOE/GenerateMutableSource/GenerateMutableSourceMacro.h
+++ b/Engine/Plugins/Mutable/Source/CustomizableObjectEditor/Private/MuCOE/GenerateMutableSource/GenerateMutableSourceMacro.h
@@ -52,11 +52,6 @@ mu::Ptr<T> GenerateMutableSourceMacro(const UEdGraphPin& Pin, FMutableGraphGener
 			{
 				Result = GenerationFunction(FollowPin, GenerationContext);
 			}
-			else
-			{
-				FText Msg = FText::Format(LOCTEXT("MacroTunnelError_PinNotLinked", "Macro Instance Pin {0} is not linked."), FText::FromName(Pin.PinName));
-				GenerationContext.Log(Msg, Node);
-			}
 		}
 		else
 		{
diff --git a/Engine/Plugins/Mutable/Source/CustomizableObjectEditor/Private/MuCOE/GenerateMutableSource/GenerateMutableSourceMesh.cpp b/Engine/Plugins/Mutable/Source/CustomizableObjectEditor/Private/MuCOE/GenerateMutableSource/GenerateMutableSourceMesh.cpp
--- a/Engine/Plugins/Mutable/Source/CustomizableObjectEditor/Private/MuCOE/GenerateMutableSource/GenerateMutableSourceMesh.cpp
+++ b/Engine/Plugins/Mutable/Source/CustomizableObjectEditor/Private/MuCOE/GenerateMutableSource/GenerateMutableSourceMesh.cpp
@@ -4640,11 +4640,6 @@ mu::Ptr<mu::NodeMesh> GenerateMutableSourceMesh(const UEdGraphPin* Pin,
 			{
 				Result = GenerateMutableSourceMesh(FollowPin, GenerationContext, BaseMeshData, bLinkedToExtendMaterial, bOnlyConnectedLOD);
 			}
-			else
-			{
-				FText Msg = FText::Format(LOCTEXT("MacroTunnelError_PinNotLinked_Mesh", "Macro Instance Pin {0} is not linked."), FText::FromName(Pin->PinName));
-				GenerationContext.Log(Msg, Node);
-			}
 		}
 		else
 		{
diff --git a/Engine/Plugins/Mutable/Source/CustomizableObjectEditor/Private/MuCOE/Nodes/CustomizableObjectNodeVariation.cpp b/Engine/Plugins/Mutable/Source/CustomizableObjectEditor/Private/MuCOE/Nodes/CustomizableObjectNodeVariation.cpp
--- a/Engine/Plugins/Mutable/Source/CustomizableObjectEditor/Private/MuCOE/Nodes/CustomizableObjectNodeVariation.cpp
+++ b/Engine/Plugins/Mutable/Source/CustomizableObjectEditor/Private/MuCOE/Nodes/CustomizableObjectNodeVariation.cpp
@@ -162,6 +162,8 @@ void UCustomizableObjectNodeVariation::BackwardsCompatibleFixup(int32 Customizab
 	{
 		ReconstructNode();
 	}
+
+	// TODO(Max UE-268615): Remove the tag name from the variation pins since now there is a pin with a textbox for it.
 }
 
 
diff --git a/Engine/Plugins/Mutable/Source/CustomizableObjectEditor/Private/MuCOE/Nodes/SCustomizableObjectNodeStaticStringPin.cpp b/Engine/Plugins/Mutable/Source/CustomizableObjectEditor/Private/MuCOE/Nodes/SCustomizableObjectNodeStaticStringPin.cpp
--- a/Engine/Plugins/Mutable/Source/CustomizableObjectEditor/Private/MuCOE/Nodes/SCustomizableObjectNodeStaticStringPin.cpp
+++ b/Engine/Plugins/Mutable/Source/CustomizableObjectEditor/Private/MuCOE/Nodes/SCustomizableObjectNodeStaticStringPin.cpp
@@ -11,6 +11,7 @@
 #include "MuCOE/Nodes/CustomizableObjectNodeObjectGroup.h"
 #include "MuCOE/Nodes/CustomizableObjectNodeParameter.h"
 #include "MuCOE/Nodes/CustomizableObjectNodeTunnel.h"
+#include "MuCOE/Nodes/CustomizableObjectNodeVariation.h"
 #include "Widgets/Input/SEditableTextBox.h"
 
 
@@ -68,6 +69,16 @@ FText SCustomizableObjectNodeStaticStringPin::GetNodeStringValue() const
 	{
 		return FText::FromString(ParameterNode->GetParameterName());
 	}
+	else if (const UCustomizableObjectNodeVariation* VariationNode = Cast<UCustomizableObjectNodeVariation>(Node))
+	{
+		for (int32 VariationIndex = 0; VariationIndex < VariationNode->GetNumVariations(); ++VariationIndex)
+		{
+			if (VariationNode->VariationTagPin(VariationIndex) == GraphPinObj)
+			{
+				return FText::FromString(VariationNode->GetVariationTag(VariationIndex));
+			}
+		}
+	}
 
 	return FText();
 }
@@ -105,6 +116,18 @@ void SCustomizableObjectNodeStaticStringPin::OnTextCommited(const FText& InValue
 	{
 		ParameterNode->SetParameterName(InValue.ToString());
 	}
+	else if (UCustomizableObjectNodeVariation* VariationNode = Cast<UCustomizableObjectNodeVariation>(Node))
+	{
+		for (int32 VariationIndex = 0; VariationIndex < VariationNode->GetNumVariations(); ++VariationIndex)
+		{
+			if (VariationNode->VariationTagPin(VariationIndex) == GraphPinObj && VariationNode->VariationsData.IsValidIndex(VariationIndex))
+			{
+				VariationNode->VariationsData[VariationIndex].Tag = InValue.ToString();
+
+				break;
+			}
+		}
+	}
 }
 
 
