[Mutable] Fix Morphs not working with streaming.

Names were not copied on streamed lods.

#rnx

diff --git a/Engine/Plugins/Mutable/Source/CustomizableObject/Internal/MuCO/CustomizableObjectSystemPrivate.h b/Engine/Plugins/Mutable/Source/CustomizableObject/Internal/MuCO/CustomizableObjectSystemPrivate.h
--- a/Engine/Plugins/Mutable/Source/CustomizableObject/Internal/MuCO/CustomizableObjectSystemPrivate.h
+++ b/Engine/Plugins/Mutable/Source/CustomizableObject/Internal/MuCO/CustomizableObjectSystemPrivate.h
@@ -90,7 +90,9 @@ void ReconstructMorphTargets(const mu::FMesh& Mesh, const TArray<FName>& GlobalN
  *  @param MutableDataStreamer Streamer to request the blocks.
  *  @param Mesh Mutable mesh to obtain which blocks to load.
  *  @param StreamingResult Key is the block. Value is the container where the data will be deserialized once streamed. */
-void LoadMorphTargets(FMutableStreamRequest& MutableDataStreamer, const TSharedRef<const mu::FMesh>& Mesh, TMap<uint32, FMorphTargetMeshData>& StreamingResult);
+void LoadMorphTargetsData(FMutableStreamRequest& MutableDataStreamer, const TSharedRef<const mu::FMesh>& Mesh, TMap<uint32, FMorphTargetMeshData>& StreamingResult);
+
+void LoadMorphTargetsMetadata(const FMutableStreamRequest& MutableDataStreamer, const TSharedRef<const mu::FMesh>& Mesh, TMap<uint32, FMorphTargetMeshData>& StreamingResult);
 
 
 /** Request the Mutable Data Streamer to load Cloth blocks. Does not stream them.
diff --git a/Engine/Plugins/Mutable/Source/CustomizableObject/Private/MuCO/CustomizableObjectInstance.cpp b/Engine/Plugins/Mutable/Source/CustomizableObject/Private/MuCO/CustomizableObjectInstance.cpp
--- a/Engine/Plugins/Mutable/Source/CustomizableObject/Private/MuCO/CustomizableObjectInstance.cpp
+++ b/Engine/Plugins/Mutable/Source/CustomizableObject/Private/MuCO/CustomizableObjectInstance.cpp
@@ -5182,7 +5182,7 @@ UE::Tasks::FTask UCustomizableInstancePrivate::LoadAdditionalAssetsAndData(const
 							continue;
 						}
 
-						LoadMorphTargets(StreamRequest, MutableMesh.ToSharedRef(), OperationData->InstanceUpdateData.RealTimeMorphTargetMeshData);
+						LoadMorphTargetsData(StreamRequest, MutableMesh.ToSharedRef(), OperationData->InstanceUpdateData.RealTimeMorphTargetMeshData);
 					}
 
 					// Metadata
@@ -5199,24 +5199,8 @@ UE::Tasks::FTask UCustomizableInstancePrivate::LoadAdditionalAssetsAndData(const
 						{
 							continue;
 						}
-					
-						const TArray<uint64>& StreamedResources = MutableMesh->GetStreamedResources();
 
-						for (uint64 ResourceId : StreamedResources)
-						{
-							FCustomizableObjectStreameableResourceId TypedResourceId = BitCast<FCustomizableObjectStreameableResourceId>(ResourceId);
-						
-							if (TypedResourceId.Type == static_cast<uint8>(FCustomizableObjectStreameableResourceId::EType::RealTimeMorphTarget))
-							{
-								check(TypedResourceId.Id != 0 && TypedResourceId.Id <= TNumericLimits<uint32>::Max());
-
-								if (ModelStreamableBulkData->RealTimeMorphStreamables.Contains(static_cast<uint32>(TypedResourceId.Id)))
-								{
-									FMorphTargetMeshData& ReadDestData = OperationData->InstanceUpdateData.RealTimeMorphTargetMeshData.FindOrAdd(TypedResourceId.Id);
-									ReadDestData.NameResolutionMap = ModelStreamableBulkData->RealTimeMorphStreamables[TypedResourceId.Id].NameResolutionMap;
-								}
-							}
-						}
+						LoadMorphTargetsMetadata(StreamRequest, MutableMesh.ToSharedRef(), OperationData->InstanceUpdateData.RealTimeMorphTargetMeshData);
 					}
 				}
 
diff --git a/Engine/Plugins/Mutable/Source/CustomizableObject/Private/MuCO/CustomizableObjectMeshUpdate.cpp b/Engine/Plugins/Mutable/Source/CustomizableObject/Private/MuCO/CustomizableObjectMeshUpdate.cpp
--- a/Engine/Plugins/Mutable/Source/CustomizableObject/Private/MuCO/CustomizableObjectMeshUpdate.cpp
+++ b/Engine/Plugins/Mutable/Source/CustomizableObject/Private/MuCO/CustomizableObjectMeshUpdate.cpp
@@ -303,10 +303,8 @@ void FCustomizableObjectMeshStreamIn::RequestMeshUpdate(const FContext& Context)
 					continue;
 				}
 
-				if (CVarEnableRealTimeMorphTargets.GetValueOnAnyThread())
-				{
-					LoadMorphTargets(StreamRequest, Mesh.ToSharedRef(), RefThis->OperationData->MorphTargetMeshData);
-				}
+				LoadMorphTargetsData(StreamRequest, Mesh.ToSharedRef(), RefThis->OperationData->MorphTargetMeshData);
+				LoadMorphTargetsMetadata(StreamRequest, Mesh.ToSharedRef(), RefThis->OperationData->MorphTargetMeshData);
 				
 				LoadClothing(StreamRequest, Mesh.ToSharedRef(), RefThis->OperationData->ClothingMeshData);
 			}
diff --git a/Engine/Plugins/Mutable/Source/CustomizableObject/Private/MuCO/CustomizableObjectSystem.cpp b/Engine/Plugins/Mutable/Source/CustomizableObject/Private/MuCO/CustomizableObjectSystem.cpp
--- a/Engine/Plugins/Mutable/Source/CustomizableObject/Private/MuCO/CustomizableObjectSystem.cpp
+++ b/Engine/Plugins/Mutable/Source/CustomizableObject/Private/MuCO/CustomizableObjectSystem.cpp
@@ -4863,9 +4863,9 @@ int32 UCustomizableObjectSystemPrivate::BlockTillAllRequestsFinished(float TimeL
 }
 
 
-void LoadMorphTargets(FMutableStreamRequest& MutableDataStreamer, const TSharedRef<const mu::FMesh>& Mesh, TMap<uint32, FMorphTargetMeshData>& StreamingResult)
+void LoadMorphTargetsData(FMutableStreamRequest& MutableDataStreamer, const TSharedRef<const mu::FMesh>& Mesh, TMap<uint32, FMorphTargetMeshData>& StreamingResult)
 {
-	MUTABLE_CPUPROFILER_SCOPE(LoadMorphTargets);
+	MUTABLE_CPUPROFILER_SCOPE(LoadMorphTargetsData);
 
 	if (!CVarEnableRealTimeMorphTargets.GetValueOnAnyThread())
 	{
@@ -4890,7 +4890,7 @@ void LoadMorphTargets(FMutableStreamRequest& MutableDataStreamer, const TSharedR
 			}
 			else
 			{
-				UE_LOG(LogMutable, Error, TEXT("Invalid streamed real time morph target data block [%d] found."), TypedResourceId.Id);
+				UE_LOG(LogMutable, Error, TEXT("Invalid streamed real time morph target data block [%llu] found."), TypedResourceId.Id);
 			}
 		}
 	}
@@ -4923,6 +4923,41 @@ void LoadMorphTargets(FMutableStreamRequest& MutableDataStreamer, const TSharedR
 }
 
 
+void LoadMorphTargetsMetadata(const FMutableStreamRequest& MutableDataStreamer, const TSharedRef<const mu::FMesh>& Mesh, TMap<uint32, FMorphTargetMeshData>& StreamingResult)
+{
+	MUTABLE_CPUPROFILER_SCOPE(LoadMorphTargetsMetadata);
+
+	if (!CVarEnableRealTimeMorphTargets.GetValueOnAnyThread())
+	{
+		return;
+	}
+
+	const TSharedPtr<FModelStreamableBulkData>& ModelStreamableBulkData = MutableDataStreamer.GetModelStreamableBulkData();
+
+	const TArray<uint64>& StreamedResources = Mesh->GetStreamedResources();
+
+	for (uint64 ResourceId : StreamedResources)
+	{
+		FCustomizableObjectStreameableResourceId TypedResourceId = BitCast<FCustomizableObjectStreameableResourceId>(ResourceId);
+						
+		if (TypedResourceId.Type == static_cast<uint8>(FCustomizableObjectStreameableResourceId::EType::RealTimeMorphTarget))
+		{
+			check(TypedResourceId.Id != 0 && TypedResourceId.Id <= TNumericLimits<uint32>::Max());
+
+			if (ModelStreamableBulkData->RealTimeMorphStreamables.Contains(static_cast<uint32>(TypedResourceId.Id)))
+			{
+				FMorphTargetMeshData& ReadDestData = StreamingResult.FindOrAdd(TypedResourceId.Id);
+				ReadDestData.NameResolutionMap = ModelStreamableBulkData->RealTimeMorphStreamables[TypedResourceId.Id].NameResolutionMap;
+			}
+			else
+			{
+				UE_LOG(LogMutable, Error, TEXT("Invalid streamed real time morph target data block [%llu] found."), TypedResourceId.Id);
+			}
+		}
+	}
+}
+
+
 void LoadClothing(FMutableStreamRequest& MutableDataStreamer, const TSharedRef<const mu::FMesh>& Mesh, TMap<uint32, FClothingMeshData>& StreamingResult)
 {
 	MUTABLE_CPUPROFILER_SCOPE(LoadClothing);
