[Mutable] Allow Mesh Streming to disable sections.

Recreate render state if any metadata from the initial generation has changed.

#rnx

diff --git a/Engine/Plugins/Mutable/Source/CustomizableObject/Internal/MuCO/UnrealConversionUtils.h b/Engine/Plugins/Mutable/Source/CustomizableObject/Internal/MuCO/UnrealConversionUtils.h
--- a/Engine/Plugins/Mutable/Source/CustomizableObject/Internal/MuCO/UnrealConversionUtils.h
+++ b/Engine/Plugins/Mutable/Source/CustomizableObject/Internal/MuCO/UnrealConversionUtils.h
@@ -108,7 +108,8 @@ namespace UnrealConversionUtils
 	CUSTOMIZABLEOBJECT_API bool CopyMutableIndexBuffers(
 		FSkeletalMeshLODRenderData& LODResource,
 		const mu::FMesh* InMutableMesh,
-		const TArray<uint32>& SurfaceIDs);
+		const TArray<uint32>& SurfaceIDs,
+		bool& bOutMarkRenderStateDirty);
 	
 
 	/**
diff --git a/Engine/Plugins/Mutable/Source/CustomizableObject/Private/MuCO/CustomizableObjectInstance.cpp b/Engine/Plugins/Mutable/Source/CustomizableObject/Private/MuCO/CustomizableObjectInstance.cpp
--- a/Engine/Plugins/Mutable/Source/CustomizableObject/Private/MuCO/CustomizableObjectInstance.cpp
+++ b/Engine/Plugins/Mutable/Source/CustomizableObject/Private/MuCO/CustomizableObjectInstance.cpp
@@ -4941,7 +4941,8 @@ bool UCustomizableInstancePrivate::BuildOrCopyRenderData(const TSharedRef<FUpdat
 			}
 
 			// Copy indices.
-			if (!UnrealConversionUtils::CopyMutableIndexBuffers(LODResource, LOD.Mesh.Get(), SurfaceIDs))
+			bool bMarkRenderStateDirty = false;
+			if (!UnrealConversionUtils::CopyMutableIndexBuffers(LODResource, LOD.Mesh.Get(), SurfaceIDs, bMarkRenderStateDirty))
 			{
 				// End with failure
 				return false;
diff --git a/Engine/Plugins/Mutable/Source/CustomizableObject/Private/MuCO/CustomizableObjectMeshUpdate.cpp b/Engine/Plugins/Mutable/Source/CustomizableObject/Private/MuCO/CustomizableObjectMeshUpdate.cpp
--- a/Engine/Plugins/Mutable/Source/CustomizableObject/Private/MuCO/CustomizableObjectMeshUpdate.cpp
+++ b/Engine/Plugins/Mutable/Source/CustomizableObject/Private/MuCO/CustomizableObjectMeshUpdate.cpp
@@ -6,6 +6,7 @@ CustomizableObjectMeshUpdate.cpp: Helpers to stream in CustomizableObject skelet
 
 
 #include "MuCO/CustomizableObjectMeshUpdate.h"
+#include "Components/SkinnedMeshComponent.h"
 
 #include "MutableStreamRequest.h"
 #include "MuCO/CustomizableObjectSystem.h"
@@ -35,11 +36,9 @@ FAutoConsoleVariableRef CVarMutableEnableGCHangFix(
 
 FCustomizableObjectMeshStreamIn::FCustomizableObjectMeshStreamIn(
 	const UCustomizableObjectSkeletalMesh* InMesh,
-	bool bInHighPrio,
 	EThreadType CreateResourcesThread,
 	const TSharedPtr<FModelStreamableBulkData>& ModelStreamableBulkData) :
-	FSkeletalMeshStreamIn(InMesh, CreateResourcesThread),
-	bHighPriority(bInHighPrio)
+	FSkeletalMeshStreamIn(InMesh, CreateResourcesThread)
 {
 	OperationData = MakeShared<FMutableMeshOperationData>();
 
@@ -126,11 +125,17 @@ void FCustomizableObjectMeshStreamIn::DoConvertResources(const FContext& Context
 
 	MUTABLE_CPUPROFILER_SCOPE(FCustomizableObjectMeshStreamIn::DoConvertResources)
 
-	ConvertMesh(Context);
+	bool bMarkRenderStateDirty = false;
+	ConvertMesh(Context, bMarkRenderStateDirty);
 
-	PushTask(Context
-		, CreateResourcesThread, SRA_UPDATE_CALLBACK(DoCreateBuffers)
-		, (EThreadType)Context.CurrentThread, SRA_UPDATE_CALLBACK(DoCancel));
+	if (bMarkRenderStateDirty)
+	{
+		PushTask(Context, TT_GameThread, SRA_UPDATE_CALLBACK(MarkRenderStateDirty), TT_None, SRA_UPDATE_CALLBACK(DoCancel));
+	}
+	else
+	{
+		PushTask(Context, CreateResourcesThread, SRA_UPDATE_CALLBACK(DoCreateBuffers), static_cast<EThreadType>(Context.CurrentThread), SRA_UPDATE_CALLBACK(DoCancel));
+	}
 }
 
 void FCustomizableObjectMeshStreamIn::DoCreateBuffers(const FContext& Context)
@@ -362,7 +367,7 @@ void FCustomizableObjectMeshStreamIn::CancelMeshUpdate(const FContext& Context)
 }
 
 
-void FCustomizableObjectMeshStreamIn::ConvertMesh(const FContext& Context)
+void FCustomizableObjectMeshStreamIn::ConvertMesh(const FContext& Context, bool& bOutMarkRenderStateDirty)
 {
 	MUTABLE_CPUPROFILER_SCOPE(FCustomizableObjectMeshStreamIn::ConvertMesh);
 
@@ -397,7 +402,7 @@ void FCustomizableObjectMeshStreamIn::ConvertMesh(const FContext& Context)
 
 		FSkeletalMeshLODRenderData& LODResource = *Context.LODResourcesView[LODIndex];
 		UnrealConversionUtils::CopyMutableVertexBuffers(LODResource, MutableMesh.Get(), bNeedsCPUAccess);
-		UnrealConversionUtils::CopyMutableIndexBuffers(LODResource, MutableMesh.Get(), Mesh->SurfaceIDs[LODIndex]);
+		UnrealConversionUtils::CopyMutableIndexBuffers(LODResource, MutableMesh.Get(), Mesh->SurfaceIDs[LODIndex], bOutMarkRenderStateDirty);
 		UnrealConversionUtils::CopyMutableSkinWeightProfilesBuffers(LODResource, *const_cast<UCustomizableObjectSkeletalMesh*>(Mesh), LODIndex, MutableMesh.Get(), Mesh->SkinWeightProfileIDs);
 		UnrealConversionUtils::MorphTargetVertexInfoBuffers(LODResource, *Mesh, *MutableMesh.Get(), OperationData->MorphTargetMeshData, LODIndex);
 		UnrealConversionUtils::ClothVertexBuffers(LODResource, *MutableMesh.Get(), OperationData->ClothingMeshData, LODIndex);
@@ -408,3 +413,32 @@ void FCustomizableObjectMeshStreamIn::ConvertMesh(const FContext& Context)
 	// Clear MeshUpdate data
 	OperationData = nullptr;
 }
+
+
+void FCustomizableObjectMeshStreamIn::MarkRenderStateDirty(const FContext& Context)
+{
+	MUTABLE_CPUPROFILER_SCOPE(FCustomizableObjectMeshStreamIn::ModifyRenderData);
+
+	check(Context.CurrentThread == TT_GameThread);
+	
+	const USkeletalMesh* Mesh = Context.Mesh;
+	FSkeletalMeshRenderData* RenderData = Context.RenderData;
+	
+	if (!IsCancelled() && Mesh && RenderData)
+	{
+		TArray<const UPrimitiveComponent*> Components;
+		IStreamingManager::Get().GetRenderAssetStreamingManager().GetAssetComponents(Mesh, Components);
+
+		for (const UPrimitiveComponent* ConstComponent : Components)
+		{
+			UPrimitiveComponent* Component = const_cast<UPrimitiveComponent*>(ConstComponent);
+			CastChecked<USkinnedMeshComponent>(Component)->MarkRenderStateDirty();
+		}
+	}
+	else
+	{
+		Abort();
+	}
+	
+	PushTask(Context, CreateResourcesThread, SRA_UPDATE_CALLBACK(DoCreateBuffers), (EThreadType)Context.CurrentThread, SRA_UPDATE_CALLBACK(DoCancel));
+}
\ No newline at end of file
diff --git a/Engine/Plugins/Mutable/Source/CustomizableObject/Private/MuCO/CustomizableObjectMeshUpdate.h b/Engine/Plugins/Mutable/Source/CustomizableObject/Private/MuCO/CustomizableObjectMeshUpdate.h
--- a/Engine/Plugins/Mutable/Source/CustomizableObject/Private/MuCO/CustomizableObjectMeshUpdate.h
+++ b/Engine/Plugins/Mutable/Source/CustomizableObject/Private/MuCO/CustomizableObjectMeshUpdate.h
@@ -68,7 +68,6 @@ class FCustomizableObjectMeshStreamIn : public FSkeletalMeshStreamIn
 public:
 	FCustomizableObjectMeshStreamIn(
 		const UCustomizableObjectSkeletalMesh* InMesh,
-		bool bInHighPrio,
 		EThreadType CreateResourcesThread,
 		const TSharedPtr<FModelStreamableBulkData>& ModelStreamableBulkData);
 
@@ -94,10 +93,9 @@ private:
 	void CancelMeshUpdate(const FContext& Context);
 
 	/** Converts from mu::FMesh to FSkeletalMeshLODRenderData. */
-	void ConvertMesh(const FContext& Context);
+	void ConvertMesh(const FContext& Context, bool& bOutMarkRenderStateDirty);
 
-	/** TODO: Stream priority. */
-	bool bHighPriority = false;
+	void MarkRenderStateDirty(const FContext& Context);
 
 	/** Context of the MeshUpdate. */
 	TSharedPtr<FMutableMeshOperationData> OperationData;
diff --git a/Engine/Plugins/Mutable/Source/CustomizableObject/Private/MuCO/CustomizableObjectSkeletalMesh.cpp b/Engine/Plugins/Mutable/Source/CustomizableObject/Private/MuCO/CustomizableObjectSkeletalMesh.cpp
--- a/Engine/Plugins/Mutable/Source/CustomizableObject/Private/MuCO/CustomizableObjectSkeletalMesh.cpp
+++ b/Engine/Plugins/Mutable/Source/CustomizableObject/Private/MuCO/CustomizableObjectSkeletalMesh.cpp
@@ -79,7 +79,7 @@ bool UCustomizableObjectSkeletalMesh::StreamIn(int32 NewMipCount, bool bHighPrio
 			? FRenderAssetUpdate::TT_Async
 			: FRenderAssetUpdate::TT_Render;
 
-		PendingUpdate = new FCustomizableObjectMeshStreamIn(this, bHighPrio, CreateResourcesThread, ModelStreamableBulkData);
+		PendingUpdate = new FCustomizableObjectMeshStreamIn(this, CreateResourcesThread, ModelStreamableBulkData);
 
 		return !PendingUpdate->IsCancelled();
 	}
diff --git a/Engine/Plugins/Mutable/Source/CustomizableObject/Private/MuCO/UnrealConversionUtils.cpp b/Engine/Plugins/Mutable/Source/CustomizableObject/Private/MuCO/UnrealConversionUtils.cpp
--- a/Engine/Plugins/Mutable/Source/CustomizableObject/Private/MuCO/UnrealConversionUtils.cpp
+++ b/Engine/Plugins/Mutable/Source/CustomizableObject/Private/MuCO/UnrealConversionUtils.cpp
@@ -464,7 +464,8 @@ namespace UnrealConversionUtils
 	bool CopyMutableIndexBuffers(
 		FSkeletalMeshLODRenderData& LODResource,
 		const mu::FMesh* InMutableMesh,
-		const TArray<uint32>& SurfaceIDs)
+		const TArray<uint32>& SurfaceIDs,
+		bool& bOutMarkRenderStateDirty)
 	{
 		MUTABLE_CPUPROFILER_SCOPE(CopyMutableIndexBuffers);
 
@@ -492,31 +493,60 @@ namespace UnrealConversionUtils
 		}
 		
 		check(LODResource.MultiSizeIndexContainer.GetDataTypeSize() == ElementSize);
-		
-		const int32 MaxNumIndices = LODResource.RenderSections.Last().BaseIndex + (LODResource.RenderSections.Last().NumTriangles * 3);
 
-		// Don't assume the buffer is empty, otherwise we may add extra elements using Insert(). 
-		LODResource.MultiSizeIndexContainer.GetIndexBuffer()->Empty(MaxNumIndices);
-		LODResource.MultiSizeIndexContainer.GetIndexBuffer()->Insert(0, MaxNumIndices);
+		const int32 NumIndices = LODResource.RenderSections.Last().BaseIndex + (LODResource.RenderSections.Last().NumTriangles * 3);
 
-		if (IndexCount == MaxNumIndices)
-		{
-			FMemory::Memcpy(LODResource.MultiSizeIndexContainer.GetIndexBuffer()->GetPointerTo(0), DataPtr, IndexCount * ElementSize);
-		}
-		else
+		// Don't assume the buffer is empty, otherwise we may add extra elements using Insert(). 
+		LODResource.MultiSizeIndexContainer.GetIndexBuffer()->Empty(NumIndices);
+		LODResource.MultiSizeIndexContainer.GetIndexBuffer()->Insert(0, NumIndices);
+		
+		FMemory::Memcpy(LODResource.MultiSizeIndexContainer.GetIndexBuffer()->GetPointerTo(0), DataPtr, IndexCount * ElementSize);
+		
+		for (int32 SectionIndex = 0; SectionIndex < LODResource.RenderSections.Num(); ++SectionIndex)
 		{
-			for (const mu::FMeshSurface& Surface : InMutableMesh->Surfaces)
+			FSkelMeshRenderSection& Section = LODResource.RenderSections[SectionIndex];
+			
+			const int32 SurfaceID = SurfaceIDs[SectionIndex];
+			const mu::FMeshSurface* Surface = InMutableMesh->Surfaces.FindByPredicate([SurfaceID](const mu::FMeshSurface& Surface)
 			{
-				const int32 SectionIndex = SurfaceIDs.Find(Surface.Id);
-				if (LODResource.RenderSections.IsValidIndex(SectionIndex) && !Surface.SubMeshes.IsEmpty())
-				{
-					void* DestPtr = LODResource.MultiSizeIndexContainer.GetIndexBuffer()->GetPointerTo(LODResource.RenderSections[SectionIndex].BaseIndex);
-					const int32 DataOffset = Surface.SubMeshes[0].IndexBegin * ElementSize;
-					const int32 DataSize = Surface.SubMeshes.Last().IndexEnd * ElementSize - DataOffset;
-					FMemory::Memcpy(DestPtr, DataPtr + DataOffset, DataSize);
-				}
+				return Surface.Id == SurfaceID;
+			});
+			
+			if (Surface)
+			{
+				const int32 NumVertices = Surface->SubMeshes.Last().VertexEnd - Surface->SubMeshes[0].VertexBegin;
+
+				bOutMarkRenderStateDirty |= Section.NumVertices != NumVertices;
+				Section.NumVertices = NumVertices;
+
+				const int32 IndexBegin = Surface->SubMeshes[0].IndexBegin; 
+				const int32 IndexEnd = Surface->SubMeshes.Last().IndexEnd; 
+				
+				const int32 NumTriangles = (IndexEnd - IndexBegin) / 3;
+
+				bOutMarkRenderStateDirty |= Section.NumTriangles != NumTriangles;
+				Section.NumTriangles = NumTriangles;
+
+				bOutMarkRenderStateDirty |= Section.BaseIndex != IndexBegin;
+				Section.BaseIndex = IndexBegin;
+
+				const int32 VertexBegin = Surface->SubMeshes[0].VertexBegin;
+
+				bOutMarkRenderStateDirty |= Section.BaseVertexIndex != VertexBegin;
+				Section.BaseVertexIndex = VertexBegin;
+			}
+			else
+			{
+				Section.bDisabled = true;
+				Section.NumTriangles = 0;
+				Section.NumVertices = 0;
+				Section.BaseIndex = 0;
+				Section.BaseVertexIndex = 0;
+
+				bOutMarkRenderStateDirty = true;	
 			}
 		}
+		
 		return true;
 	}
 
diff --git a/Engine/Plugins/Mutable/Source/CustomizableObjectEditor/Private/MuCOE/MutableMeshPreviewUtils.cpp b/Engine/Plugins/Mutable/Source/CustomizableObjectEditor/Private/MuCOE/MutableMeshPreviewUtils.cpp
--- a/Engine/Plugins/Mutable/Source/CustomizableObjectEditor/Private/MuCOE/MutableMeshPreviewUtils.cpp
+++ b/Engine/Plugins/Mutable/Source/CustomizableObjectEditor/Private/MuCOE/MutableMeshPreviewUtils.cpp
@@ -439,7 +439,8 @@ namespace MutableMeshPreviewUtils
 			}
 
 			// Copy index buffers or fail to generate the mesh
-			if (!UnrealConversionUtils::CopyMutableIndexBuffers(LODResource, InMutableMesh, SurfaceIDs))
+			bool bMarkRenderStateDirty = false;
+			if (!UnrealConversionUtils::CopyMutableIndexBuffers(LODResource, InMutableMesh, SurfaceIDs, bMarkRenderStateDirty))
 			{
 				return false;
 			}
