[mutable] Fix potential bug with code runner requesting different type of execution for some ops.
Remove some unused code.
#rnx

diff --git a/Engine/Plugins/Mutable/Source/MutableRuntime/Internal/MuR/CodeVisitor.h b/Engine/Plugins/Mutable/Source/MutableRuntime/Internal/MuR/CodeVisitor.h
--- a/Engine/Plugins/Mutable/Source/MutableRuntime/Internal/MuR/CodeVisitor.h
+++ b/Engine/Plugins/Mutable/Source/MutableRuntime/Internal/MuR/CodeVisitor.h
@@ -463,30 +463,5 @@ namespace mu
 
     };
 
-
-
-    /** Calculate all the parameters found under a particular operation
-    * It has an internal cache, so don't reuse objects of this class if the Program changes.
-	*/
-    class MUTABLERUNTIME_API SubtreeParametersVisitor
-    {
-    public:
-
-        void Run( OP::ADDRESS Root, const FProgram& );
-
-        //! After Run, list of relevant parameters.
-		TArray<int32> RelevantParams;
-
-    private:
-
-		TArray<int32> CurrentParams;
-		TArray<uint8> Visited;
-		TArray<OP::ADDRESS> Pending;
-
-        // Result cache
-        TMap< OP::ADDRESS, TArray<int32> > ResultCache;
-    };
-
-
 }
 
diff --git a/Engine/Plugins/Mutable/Source/MutableRuntime/Private/MuR/CodeRunner.cpp b/Engine/Plugins/Mutable/Source/MutableRuntime/Private/MuR/CodeRunner.cpp
--- a/Engine/Plugins/Mutable/Source/MutableRuntime/Private/MuR/CodeRunner.cpp
+++ b/Engine/Plugins/Mutable/Source/MutableRuntime/Private/MuR/CodeRunner.cpp
@@ -6472,7 +6472,7 @@ namespace mu
 			case 1:
 			{
 				// Get the variable result
-				int32 Var = LoadInt(FCacheAddress(VarAddress, Item));
+				int32 Var = LoadInt(FCacheAddress(VarAddress, Item.ExecutionIndex, Item.ExecutionOptions, FScheduledOp::EType::Full));
 
 				OP::ADDRESS ValueAt = DefAddress;
 				for (uint32 C = 0; C < CaseCount; ++C)
@@ -6850,7 +6850,7 @@ namespace mu
 			{
 				FExtendedImageDesc& Result = ImageDescResults.FindOrAdd(Item.At);
 				
-				TSharedPtr<const FLayout> Layout = LoadLayout(FCacheAddress(Args.Layout, Item));
+				TSharedPtr<const FLayout> Layout = LoadLayout(FCacheAddress(Args.Layout, Item.ExecutionIndex, Item.ExecutionOptions, FScheduledOp::EType::Full));
 
 				FIntPoint SizeInBlocks = Layout->GetGridSize();
 				FIntPoint BlockSizeInPixels(Args.BlockSize[0], Args.BlockSize[1]);
diff --git a/Engine/Plugins/Mutable/Source/MutableRuntime/Private/MuR/CodeVisitor.cpp b/Engine/Plugins/Mutable/Source/MutableRuntime/Private/MuR/CodeVisitor.cpp
deleted file mode 100644
