[mutable] Added handling of AssetRegistry search as part of the CO compilation.
- While the CO compilation still requires the AR scan to be correctly performed now said scan is handled by the compilation code itself.
- Fixed some typos.

diff --git a/Engine/Plugins/Mutable/Source/CustomizableObject/Private/MuCO/CustomizableObject.cpp b/Engine/Plugins/Mutable/Source/CustomizableObject/Private/MuCO/CustomizableObject.cpp
--- a/Engine/Plugins/Mutable/Source/CustomizableObject/Private/MuCO/CustomizableObject.cpp
+++ b/Engine/Plugins/Mutable/Source/CustomizableObject/Private/MuCO/CustomizableObject.cpp
@@ -664,7 +664,7 @@ void UCustomizableObjectPrivate::LoadCompiledDataFromDisk()
 			
 			TArray<uint8> HeaderBytes;
 
-			const int32 HeaderSize = sizeof(MutableCompiledDataStreamHeader);
+			constexpr int32 HeaderSize = sizeof(MutableCompiledDataStreamHeader);
 			HeaderBytes.SetNum(HeaderSize);
 			FileHandle->Read(HeaderBytes.GetData(), HeaderSize);
 			
diff --git a/Engine/Plugins/Mutable/Source/CustomizableObject/Private/MuCO/CustomizableObjectSystem.cpp b/Engine/Plugins/Mutable/Source/CustomizableObject/Private/MuCO/CustomizableObjectSystem.cpp
--- a/Engine/Plugins/Mutable/Source/CustomizableObject/Private/MuCO/CustomizableObjectSystem.cpp
+++ b/Engine/Plugins/Mutable/Source/CustomizableObject/Private/MuCO/CustomizableObjectSystem.cpp
@@ -3932,12 +3932,6 @@ int32 UCustomizableObjectSystem::TickInternal(const bool bBlocking)
 
 	// \TODO: Review: We should never compile an object from this tick, so this could be removed
 #if WITH_EDITOR
-	const FAssetRegistryModule& AssetRegistryModule = FModuleManager::LoadModuleChecked<FAssetRegistryModule>("AssetRegistry");
-
-	if (AssetRegistryModule.Get().IsLoadingAssets())
-	{
-		return 0; // Assets are still being loaded, so subobjects won't be found, compiled objects incomplete and thus updates wrong
-	}
 
 	// See if any COs pending to load can be completed.
 	for ( TArray<TObjectPtr<UCustomizableObject>>::TIterator It = GetPrivate()->ObjectsPendingLoad.CreateIterator(); It; ++It )
diff --git a/Engine/Plugins/Mutable/Source/CustomizableObjectEditor/Private/MuCOE/CustomizableObjectCompiler.cpp b/Engine/Plugins/Mutable/Source/CustomizableObjectEditor/Private/MuCOE/CustomizableObjectCompiler.cpp
--- a/Engine/Plugins/Mutable/Source/CustomizableObjectEditor/Private/MuCOE/CustomizableObjectCompiler.cpp
+++ b/Engine/Plugins/Mutable/Source/CustomizableObjectEditor/Private/MuCOE/CustomizableObjectCompiler.cpp
@@ -137,6 +137,13 @@ void FCustomizableObjectCompiler::Compile(const TSharedRef<FCompilationRequest>&
 	
 	check(IsInGameThread());
 
+	{
+		const FAssetRegistryModule& AssetRegistryModule = FModuleManager::LoadModuleChecked<FAssetRegistryModule>("AssetRegistry");
+		IAssetRegistry& AssetRegistry = AssetRegistryModule.Get();
+		check (AssetRegistry.IsSearchAllAssets());	// At this point the search of all assets should have already been triggered
+		check (!AssetRegistry.IsGathering());		// And it should have also been completed as we wait for it to be done before asking for the compilation
+	}
+	
 	UCustomizableObject* Object = InCompileRequest->GetCustomizableObject();
 	if (!Object)
 	{
@@ -359,6 +366,30 @@ void FCustomizableObjectCompiler::EnqueueCompileRequest(const TSharedRef<FCompil
 		return;
 	}
 		
+	{
+		const FAssetRegistryModule& AssetRegistryModule = FModuleManager::LoadModuleChecked<FAssetRegistryModule>("AssetRegistry");
+		IAssetRegistry& AssetRegistry = AssetRegistryModule.Get();
+
+		if (AssetRegistry.IsSearchAllAssets())
+		{
+			// The full search of the AR has already been requested, but we do not know if it has ended by this point
+			
+			// Wait for the search to be completed if this new compilation is SYNC
+			if (AssetRegistry.IsGathering() && !CompileRequest->bAsync)
+			{
+				UE_LOG(LogMutable, Log, TEXT("Waiting for the AR scan to be completed before enqueueing the compilation of the %s CO."), *CustomizableObject->GetName());
+				AssetRegistry.WaitForCompletion();
+			}
+		}
+		else
+		{
+			// Scan for all the assets async/sync based on the compilation options as the AR has not yet been scanned
+			UE_LOG(LogMutable, Display, TEXT("Performing full Asset Registry search as required by the CO compilation process."));
+			AssetRegistry.SearchAllAssets(!CompileRequest->bAsync);
+		}
+	}
+
+	
 	if (!CompileRequest->bAsync)
 	{
 		MUTABLE_CPUPROFILER_SCOPE(SyncCompile)
@@ -1152,6 +1183,15 @@ bool FCustomizableObjectCompiler::TryPopCompileRequest()
 		return false;
 	}
 
+	const FAssetRegistryModule& AssetRegistryModule = FModuleManager::LoadModuleChecked<FAssetRegistryModule>("AssetRegistry");
+	IAssetRegistry& AssetRegistry = AssetRegistryModule.Get();
+	check(AssetRegistry.IsSearchAllAssets());	// in all cases this should return true as we are taking care of calling the searching of the AR before being able to add a request
+	if (AssetRegistry.IsGathering())
+	{
+		// AR search is in progress, skip the compilation operation for now until it is ready.
+		return false;
+	}
+	
 	UCustomizableObjectSystemPrivate* SystemPrivate = UCustomizableObjectSystem::GetInstance()->GetPrivate();
 	if (SystemPrivate->CurrentMutableOperation)
 	{
diff --git a/Engine/Plugins/Mutable/Source/CustomizableObjectEditor/Private/MuCOE/CustomizableObjectInstanceBaker.cpp b/Engine/Plugins/Mutable/Source/CustomizableObjectEditor/Private/MuCOE/CustomizableObjectInstanceBaker.cpp
--- a/Engine/Plugins/Mutable/Source/CustomizableObjectEditor/Private/MuCOE/CustomizableObjectInstanceBaker.cpp
+++ b/Engine/Plugins/Mutable/Source/CustomizableObjectEditor/Private/MuCOE/CustomizableObjectInstanceBaker.cpp
@@ -100,7 +100,7 @@ void UCustomizableObjectInstanceBaker::Bake() const
 	TArray<FBakedResourceData> SavedPackagesData;
 	{
 		// Display a list with the paths of the assets saved
-		UE_LOG(LogMutable, Log, TEXT("Saved asets (%d) :"), SavedPackages.Num());
+		UE_LOG(LogMutable, Log, TEXT("Saved assets (%d) :"), SavedPackages.Num());
 		SavedPackagesData.Reserve(SavedPackages.Num());
 		for (const TPair<EPackageSaveResolutionType, UPackage*>& SavedPackagePair : SavedPackages)
 		{
@@ -120,7 +120,7 @@ void UCustomizableObjectInstanceBaker::Bake() const
 	}
 
 	// Report that the baking operation has been completed
-	UE_LOG(LogMutable, Display, TEXT("Finsishing COI baking procedure."));
+	UE_LOG(LogMutable, Display, TEXT("Finishhing COI baking procedure."));
 	FinishBakingOperation(bBakeWasSuccessful, &SavedPackagesData);
 	
 	// Reset the GIsRunningUnattendedScript flag 
diff --git a/Engine/Plugins/Mutable/Source/CustomizableObjectEditor/Private/MuCOE/SCustomizableObjectEditorViewport.cpp b/Engine/Plugins/Mutable/Source/CustomizableObjectEditor/Private/MuCOE/SCustomizableObjectEditorViewport.cpp
--- a/Engine/Plugins/Mutable/Source/CustomizableObjectEditor/Private/MuCOE/SCustomizableObjectEditorViewport.cpp
+++ b/Engine/Plugins/Mutable/Source/CustomizableObjectEditor/Private/MuCOE/SCustomizableObjectEditorViewport.cpp
@@ -171,7 +171,7 @@ FText SCustomizableObjectEditorViewport::GetCompileErrorOverlayText() const
 	}
 
 	const FAssetRegistryModule& AssetRegistryModule = FModuleManager::LoadModuleChecked<FAssetRegistryModule>("AssetRegistry");
-	if (AssetRegistryModule.Get().IsLoadingAssets())
+	if (AssetRegistryModule.Get().IsGathering())
 	{
 		return LOCTEXT("LoadingAssetRegistry", "Loading Asset Registry...");
 	}
