[Mutable] Fix OutOfDateWarning_Async crash due to invalid GEditor.

#rnx

diff --git a/Engine/Plugins/Mutable/Source/CustomizableObject/Private/MuCO/CustomizableObjectSystem.cpp b/Engine/Plugins/Mutable/Source/CustomizableObject/Private/MuCO/CustomizableObjectSystem.cpp
--- a/Engine/Plugins/Mutable/Source/CustomizableObject/Private/MuCO/CustomizableObjectSystem.cpp
+++ b/Engine/Plugins/Mutable/Source/CustomizableObject/Private/MuCO/CustomizableObjectSystem.cpp
@@ -366,12 +366,15 @@ void OutOfDateWarning_Async(const TSharedRef<FOutOfDateWarningContext>& Context)
 	if (FPlatformTime::Seconds() - Context->StartTime >= MaxTime)
 	{
 		// Time limit reached. Reschedule itself.
-		GEditor->GetTimerManager()->SetTimerForNextTick([=]()
+		if (GEditor)
 		{
-			Context->StartTime = FPlatformTime::Seconds();
-			OutOfDateWarning_Async(Context);
-		});
-
+			GEditor->GetTimerManager()->SetTimerForNextTick([=]()
+			{
+				Context->StartTime = FPlatformTime::Seconds();
+				OutOfDateWarning_Async(Context);
+			});
+		}
+		
 		return;
 	}
 
