[mutable] Fixed text conflict from LOCTEXT macro for namespace "CustomizableObjectDetails".
#tests Editor
#rnx

diff --git a/Engine/Plugins/Mutable/Source/CustomizableObjectEditor/Private/MuCOE/Nodes/CustomizableObjectNodeModifierMorphMeshSectionDetails.cpp b/Engine/Plugins/Mutable/Source/CustomizableObjectEditor/Private/MuCOE/Nodes/CustomizableObjectNodeModifierMorphMeshSectionDetails.cpp
--- a/Engine/Plugins/Mutable/Source/CustomizableObjectEditor/Private/MuCOE/Nodes/CustomizableObjectNodeModifierMorphMeshSectionDetails.cpp
+++ b/Engine/Plugins/Mutable/Source/CustomizableObjectEditor/Private/MuCOE/Nodes/CustomizableObjectNodeModifierMorphMeshSectionDetails.cpp
@@ -335,11 +335,11 @@ FText FCustomizableObjectNodeModifierMorphMeshSectionDetails::MorphNameSelectorW
 {
 	if (Node->MorphTargetNamePin() && FollowInputPin(*Node->MorphTargetNamePin()))
 	{
-		return LOCTEXT("MorphTargetNameWidgetTooltip_Ignored", "Disabled. When there is a string node linked to the Morph Target Name pin, the morph target name selected in this widget is ignored.");
+		return LOCTEXT("MorphMeshSectionTargetNameWidgetTooltip_Ignored", "Disabled. When there is a string node linked to the Morph Target Name pin, the morph target name selected in this widget is ignored.");
 	}
 	else
 	{
-		return LOCTEXT("MorphTargetNameWidgetTooltip", "Select the morph target name.");
+		return LOCTEXT("MorphMeshSectionTargetNameWidgetTooltip", "Select the morph target name.");
 	}
 }
 
