[mutable] Fixed missing parameters in the CO Editor Instance Details when they are children of Group nodes in toggle mode.
#tests Editor
#rnx

diff --git a/Engine/Plugins/Mutable/Source/CustomizableObjectEditor/Private/MuCOE/CustomizableInstanceDetails.cpp b/Engine/Plugins/Mutable/Source/CustomizableObjectEditor/Private/MuCOE/CustomizableInstanceDetails.cpp
--- a/Engine/Plugins/Mutable/Source/CustomizableObjectEditor/Private/MuCOE/CustomizableInstanceDetails.cpp
+++ b/Engine/Plugins/Mutable/Source/CustomizableObjectEditor/Private/MuCOE/CustomizableInstanceDetails.cpp
@@ -862,7 +862,13 @@ bool FCustomizableInstanceDetails::IsVisible(int32 ParamIndexInObject)
 	{
 		const FString* Value = UIMetadata.ExtraInformation.Find(FString("__DisplayWhenParentValueEquals"));
 
-		if (Value && CustomizableObject->GetPrivate()->FindParameter(*ParentName) != INDEX_NONE && CustomInstance->GetIntParameterSelectedOption(*ParentName) != *Value)
+		bool bParentIsBoolParam = CustomizableObject->GetParameterTypeByName(*ParentName) == EMutableParameterType::Bool;
+
+		FString SelectedOption = bParentIsBoolParam ?
+								 (CustomInstance->GetBoolParameterSelectedOption(*ParentName) ? FString("1") : FString("0")) :
+								 CustomInstance->GetIntParameterSelectedOption(*ParentName);
+
+		if (Value && CustomizableObject->GetPrivate()->FindParameter(*ParentName) != INDEX_NONE && SelectedOption != *Value)
 		{
 			return false;
 		}
diff --git a/Engine/Plugins/Mutable/Source/CustomizableObjectEditor/Private/MuCOE/GenerateMutableSource/GenerateMutableSource.cpp b/Engine/Plugins/Mutable/Source/CustomizableObjectEditor/Private/MuCOE/GenerateMutableSource/GenerateMutableSource.cpp
--- a/Engine/Plugins/Mutable/Source/CustomizableObjectEditor/Private/MuCOE/GenerateMutableSource/GenerateMutableSource.cpp
+++ b/Engine/Plugins/Mutable/Source/CustomizableObjectEditor/Private/MuCOE/GenerateMutableSource/GenerateMutableSource.cpp
@@ -1087,7 +1087,9 @@ mu::Ptr<mu::NodeObject> GenerateMutableSource(const UEdGraphPin * Pin, FMutableG
 		// External children
 		TArray<UCustomizableObjectNodeObject*> ExternalChildNodes;
 		GenerationContext.GroupIdToExternalNodeMap.MultiFind(TypedNodeGroup->NodeGuid, ExternalChildNodes);
-		GenerationContext.GuidToParamNameMap.Add(TypedNodeGroup->NodeGuid, TypedNodeGroup->GetGroupName(&GenerationContext.MacroNodesStack));
+		FMutableGraphGenerationContext::FParamInfo ParamInfo(TypedNodeGroup->GetGroupName(&GenerationContext.MacroNodesStack),
+															 TypedNodeGroup->GroupType == ECustomizableObjectGroupType::COGT_TOGGLE);
+		GenerationContext.GuidToParamNameMap.Add(TypedNodeGroup->NodeGuid, ParamInfo);
 
 		// Children
 		const TArray<UEdGraphPin*> ConnectedChildrenPins = FollowInputPinArray(*TypedNodeGroup->ObjectsPin());
@@ -1194,15 +1196,19 @@ mu::Ptr<mu::NodeObject> GenerateMutableSource(const UEdGraphPin * Pin, FMutableG
 
 				if (ParentId)
 				{
-					FString* ParentParamName = GenerationContext.GuidToParamNameMap.Find(*ParentId);
+					FMutableGraphGenerationContext::FParamInfo* ParentParamInfo = GenerationContext.GuidToParamNameMap.Find(*ParentId);
 
-					if (ParentParamName)
+					if (ParentParamInfo)
 					{
-						ParameterUIData.ParamUIMetadata.ExtraInformation.Add(FString("__ParentParamName"), *ParentParamName);
+						FString ParentParamName = ParentParamInfo->ParamName;
+
+						ParameterUIData.ParamUIMetadata.ExtraInformation.Add(FString("__ParentParamName"), 
+							ParentParamInfo->bIsToggle ? ValueName : *ParentParamName);
 
 						if (bHideWhenNotSelected)
 						{
-							ParameterUIData.ParamUIMetadata.ExtraInformation.Add(FString("__DisplayWhenParentValueEquals"), ValueName);
+							ParameterUIData.ParamUIMetadata.ExtraInformation.Add(FString("__DisplayWhenParentValueEquals"), 
+								ParentParamInfo->bIsToggle ? FString("1") : ValueName);
 						}
 
 						if (bCollapseUnderParent)
diff --git a/Engine/Plugins/Mutable/Source/CustomizableObjectEditor/Private/MuCOE/GenerateMutableSource/GenerateMutableSource.h b/Engine/Plugins/Mutable/Source/CustomizableObjectEditor/Private/MuCOE/GenerateMutableSource/GenerateMutableSource.h
--- a/Engine/Plugins/Mutable/Source/CustomizableObjectEditor/Private/MuCOE/GenerateMutableSource/GenerateMutableSource.h
+++ b/Engine/Plugins/Mutable/Source/CustomizableObjectEditor/Private/MuCOE/GenerateMutableSource/GenerateMutableSource.h
@@ -633,6 +633,14 @@ struct FMutableGraphGenerationContext
 		uint32 ID;
 	};
 
+	struct FParamInfo
+	{
+		FString ParamName;
+		bool bIsToggle = false;
+
+		FParamInfo(const FString& InParamName, bool bInIsToggle) : ParamName(InParamName), bIsToggle(bInIsToggle) {}
+	};
+
 	// Cache of runtime pass-through meshes and their IDs used in the core to identify them.
 	// These meshes will remain as external references even in optimized models.
 	TMap<TSoftObjectPtr<UStreamableRenderAsset>, FGeneratedReferencedMesh> PassthroughMeshMap;
@@ -805,7 +813,7 @@ struct FMutableGraphGenerationContext
 	TMultiMap<FGuid, UCustomizableObjectNodeObject*> GroupIdToExternalNodeMap;
 
 	// Easily retrieve a parameter name from its node guid
-	TMap<FGuid, FString> GuidToParamNameMap;
+	TMap<FGuid, FParamInfo> GuidToParamNameMap;
 
 	// Graph cycle detection
 	/** Visited nodes during the DAC recursion traversal.
