[Mutable] Fix Instance Usage crash when not attached to a Customizable Component.

#rnx

diff --git a/Engine/Plugins/Mutable/Source/CustomizableObject/Private/MuCO/CustomizableObjectInstanceUsage.cpp b/Engine/Plugins/Mutable/Source/CustomizableObject/Private/MuCO/CustomizableObjectInstanceUsage.cpp
--- a/Engine/Plugins/Mutable/Source/CustomizableObject/Private/MuCO/CustomizableObjectInstanceUsage.cpp
+++ b/Engine/Plugins/Mutable/Source/CustomizableObject/Private/MuCO/CustomizableObjectInstanceUsage.cpp
@@ -138,7 +138,13 @@ FName UCustomizableObjectInstanceUsage::GetComponentName() const
 
 UCustomizableSkeletalComponent* UCustomizableObjectInstanceUsagePrivate::GetCustomizableSkeletalComponent() const
 {
-	return Cast<UCustomizableSkeletalComponent>(GetPublic()->GetOuter()->GetOuter());
+	UObject* Outer = GetPublic()->GetOuter(); // UCustomizableSkeletalComponentPrivate
+	if (!Outer)
+	{
+		return nullptr;
+	}
+
+	return Cast<UCustomizableSkeletalComponent>(Outer->GetOuter());
 }
 
 UCustomizableObjectInstanceUsage* UCustomizableObjectInstanceUsagePrivate::GetPublic()
