[mutable] Added new compilation optimization level that allows the compiler to get the compilation optimization setting from the CO itself.

diff --git a/Engine/Plugins/Mutable/Source/CustomizableObject/Public/MuCO/CustomizableObjectSystem.h b/Engine/Plugins/Mutable/Source/CustomizableObject/Public/MuCO/CustomizableObjectSystem.h
--- a/Engine/Plugins/Mutable/Source/CustomizableObject/Public/MuCO/CustomizableObjectSystem.h
+++ b/Engine/Plugins/Mutable/Source/CustomizableObject/Public/MuCO/CustomizableObjectSystem.h
@@ -50,7 +50,8 @@ enum class ECustomizableObjectOptimizationLevel : uint8
 {
 	None = 0, 
 	Minimal UE_DEPRECATED(5.6, "Converted internally to None. Use None instead.") = 1,
-	Maximum = 2
+	Maximum = 2,
+	FromCustomizableObject  = 3					// Grab the optimization settings from the CO
 };
 
 
diff --git a/Engine/Plugins/Mutable/Source/CustomizableObjectEditor/Private/MuCOE/CustomizableObjectEditorModule.cpp b/Engine/Plugins/Mutable/Source/CustomizableObjectEditor/Private/MuCOE/CustomizableObjectEditorModule.cpp
--- a/Engine/Plugins/Mutable/Source/CustomizableObjectEditor/Private/MuCOE/CustomizableObjectEditorModule.cpp
+++ b/Engine/Plugins/Mutable/Source/CustomizableObjectEditor/Private/MuCOE/CustomizableObjectEditorModule.cpp
@@ -176,13 +176,13 @@ int32 ConvertOptimizationLevel(ECustomizableObjectOptimizationLevel Optimization
 	{
 	case ECustomizableObjectOptimizationLevel::None:
 		PRAGMA_DISABLE_DEPRECATION_WARNINGS
-			case ECustomizableObjectOptimizationLevel::Minimal:
-				PRAGMA_ENABLE_DEPRECATION_WARNINGS
-						return 0;
+	case ECustomizableObjectOptimizationLevel::Minimal:
+		PRAGMA_ENABLE_DEPRECATION_WARNINGS
+		return 0;
 
 	case ECustomizableObjectOptimizationLevel::Maximum:
 		return UE_MUTABLE_MAX_OPTIMIZATION;
-	
+		
 	default:
 		unimplemented();
 		return 0;
@@ -1091,10 +1091,19 @@ void FCustomizableObjectEditorModule::CompileCustomizableObject(UCustomizableObj
 		CompileRequest->bSkipIfCompiled = Params->bSkipIfCompiled;
 		CompileRequest->bSkipIfOutOfDate = Params->bSkipIfOutOfDate;
 		CompileRequest->Options.TextureCompression = Params->TextureCompression;
-		CompileRequest->Options.OptimizationLevel = ConvertOptimizationLevel(Params->OptimizationLevel);
 		CompileRequest->Options.bGatherReferences = Params->bGatherReferneces;
 		CompileRequest->Callback = Params->Callback;
 		CompileRequest->CallbackNative = Params->CallbackNative;
+		
+		// Override the optimization level provided in the request with the one set in the CO
+		if (Params->OptimizationLevel == ECustomizableObjectOptimizationLevel::FromCustomizableObject)
+		{
+			CompileRequest->Options.OptimizationLevel = Object.GetPrivate()->OptimizationLevel;
+		}
+		else
+		{
+			CompileRequest->Options.OptimizationLevel = ConvertOptimizationLevel(Params->OptimizationLevel);
+		}
 
 		if (Params->CompileOnlySelectedInstance)
 		{
