[mutable] Added log to report the compilation level being used for the compilation of a CO.

#rnx

diff --git a/Engine/Plugins/Mutable/Source/CustomizableObjectEditor/Private/MuCOE/CustomizableObjectCompileRunnable.cpp b/Engine/Plugins/Mutable/Source/CustomizableObjectEditor/Private/MuCOE/CustomizableObjectCompileRunnable.cpp
--- a/Engine/Plugins/Mutable/Source/CustomizableObjectEditor/Private/MuCOE/CustomizableObjectCompileRunnable.cpp
+++ b/Engine/Plugins/Mutable/Source/CustomizableObjectEditor/Private/MuCOE/CustomizableObjectCompileRunnable.cpp
@@ -108,10 +108,10 @@ uint32 FCustomizableObjectCompileRunnable::Run()
 
 	CompilerOptions->SetUseDiskCache(bUseDiskCache);
 
-	if (Options.OptimizationLevel > 2)
+	if (Options.OptimizationLevel > UE_MUTABLE_MAX_OPTIMIZATION)
 	{
 		UE_LOG(LogMutable, Log, TEXT("Mutable compile optimization level out of range. Clamping to maximum."));
-		Options.OptimizationLevel = 2;
+		Options.OptimizationLevel = UE_MUTABLE_MAX_OPTIMIZATION;
 	}
 
 	switch (Options.OptimizationLevel)
diff --git a/Engine/Plugins/Mutable/Source/CustomizableObjectEditor/Private/MuCOE/CustomizableObjectCompiler.cpp b/Engine/Plugins/Mutable/Source/CustomizableObjectEditor/Private/MuCOE/CustomizableObjectCompiler.cpp
--- a/Engine/Plugins/Mutable/Source/CustomizableObjectEditor/Private/MuCOE/CustomizableObjectCompiler.cpp
+++ b/Engine/Plugins/Mutable/Source/CustomizableObjectEditor/Private/MuCOE/CustomizableObjectCompiler.cpp
@@ -266,8 +266,33 @@ void FCustomizableObjectCompiler::Compile(const TSharedRef<FCompilationRequest>&
 		}
 	}
 
-	UE_LOG(LogMutable, Display, TEXT("Compiling Customizable Object %s for platform %s."), *CurrentObject->GetName(), *CurrentOptions.TargetPlatform->PlatformName());
+	{
+		const UEnum* OptimizationLevelEnum = StaticEnum<ECustomizableObjectOptimizationLevel>();
+		check(OptimizationLevelEnum);
+
+		const FString CurrentOptimizationLevelName = OptimizationLevelEnum->GetNameStringByIndex(CurrentOptions.OptimizationLevel);
+		const FString MutableCompilationStartMessage = FString::Printf(TEXT("Compiling Customizable Object %s for platform %s and optimization level \"%s\"."), *CurrentObject->GetName(), *CurrentOptions.TargetPlatform->PlatformName(), *CurrentOptimizationLevelName);
+
+		if (IsRunningCommandlet())
+		{
+			UE_LOG(LogMutable, Display, TEXT("%s"), *MutableCompilationStartMessage);
+		}
+		else
+		{
+			const int32 MaxOptimizationLevelValue = ConvertOptimizationLevel(ECustomizableObjectOptimizationLevel::Maximum);
+			if (CurrentOptions.OptimizationLevel == MaxOptimizationLevelValue)
+			{
+				UE_LOG(LogMutable, Display, TEXT("%s The Compilation will take more time to run due to the chosen optimization level."), *MutableCompilationStartMessage);
+			}
+			else
+			{
+				UE_LOG(LogMutable, Display, TEXT("%s"), *MutableCompilationStartMessage);
+			}
+		}
+		
+	}
 
+	
 	if (CurrentOptions.bForceLargeLODBias)
 	{
 		UE_LOG(LogMutable, Display, TEXT("Compiling Customizable Object with %d LODBias."), CurrentOptions.DebugBias);
diff --git a/Engine/Plugins/Mutable/Source/CustomizableObjectEditor/Private/MuCOE/CustomizableObjectInstanceBaker.cpp b/Engine/Plugins/Mutable/Source/CustomizableObjectEditor/Private/MuCOE/CustomizableObjectInstanceBaker.cpp
--- a/Engine/Plugins/Mutable/Source/CustomizableObjectEditor/Private/MuCOE/CustomizableObjectInstanceBaker.cpp
+++ b/Engine/Plugins/Mutable/Source/CustomizableObjectEditor/Private/MuCOE/CustomizableObjectInstanceBaker.cpp
@@ -120,7 +120,7 @@ void UCustomizableObjectInstanceBaker::Bake() const
 	}
 
 	// Report that the baking operation has been completed
-	UE_LOG(LogMutable, Display, TEXT("Finishhing COI baking procedure."));
+	UE_LOG(LogMutable, Display, TEXT("Finishing COI baking procedure."));
 	FinishBakingOperation(bBakeWasSuccessful, &SavedPackagesData);
 	
 	// Reset the GIsRunningUnattendedScript flag 
