Set default CO compile optimization level value in compile API to FromCustomizableObject

#rnx

diff --git a/Engine/Plugins/Mutable/Source/CustomizableObject/Public/MuCO/CustomizableObject.h b/Engine/Plugins/Mutable/Source/CustomizableObject/Public/MuCO/CustomizableObject.h
--- a/Engine/Plugins/Mutable/Source/CustomizableObject/Public/MuCO/CustomizableObject.h
+++ b/Engine/Plugins/Mutable/Source/CustomizableObject/Public/MuCO/CustomizableObject.h
@@ -201,7 +201,7 @@ struct FCompileParams
 	TObjectPtr<UCustomizableObjectInstance> CompileOnlySelectedInstance = nullptr;
 	
 	UPROPERTY(BlueprintReadWrite, Category = Compile)
-	ECustomizableObjectOptimizationLevel OptimizationLevel = ECustomizableObjectOptimizationLevel::Maximum;
+	ECustomizableObjectOptimizationLevel OptimizationLevel = ECustomizableObjectOptimizationLevel::FromCustomizableObject;
 
 	UPROPERTY(BlueprintReadWrite, Category = Compile)
 	ECustomizableObjectTextureCompression TextureCompression = ECustomizableObjectTextureCompression::Fast;
