[mutable] Updated override material setup for the parent Skeletal Mesh Component.
- Now, the parent override materials will only be removed from the component if new override materials are found to be set in their place.
- A log will now appear only when the override materials set in the parent component has changed, reducing this way the log spam produced by the previous log.

#rnx

diff --git a/Engine/Plugins/Mutable/Source/CustomizableObject/Private/MuCO/CustomizableObjectSystem.cpp b/Engine/Plugins/Mutable/Source/CustomizableObject/Private/MuCO/CustomizableObjectSystem.cpp
--- a/Engine/Plugins/Mutable/Source/CustomizableObject/Private/MuCO/CustomizableObjectSystem.cpp
+++ b/Engine/Plugins/Mutable/Source/CustomizableObject/Private/MuCO/CustomizableObjectSystem.cpp
@@ -1436,31 +1436,39 @@ void UpdateSkeletalMesh(const TSharedRef<FUpdateContextPrivate>& Context)
 			bAttachedParentUpdated = true;
 		}
 
-		TArray<TObjectPtr<UMaterialInterface>> OldOverridenMaterials = Parent->OverrideMaterials;
-
-		if (Parent->HasOverrideMaterials())
-		{
-			UE_LOG(LogMutable, Log, TEXT("Attaching Customizable Skeletal Component to Skeletal Mesh Component with overriden materials! Deleting overrides."));
-			Parent->EmptyOverrideMaterials();
-		}
 
+		TArray<TObjectPtr<UMaterialInterface>> NewOverridenMaterials;
+		
 		const bool bIsTransientMesh = SkeletalMesh ? SkeletalMesh->HasAllFlags(EObjectFlags::RF_Transient) : false;
-		const bool bUseOverrideMaterials = !bIsTransientMesh || (CustomizableObject->bEnableMeshCache && UCustomizableObjectSystem::IsMeshCacheEnabled());
+		const bool bIsUsingMeshCache = CustomizableObject->bEnableMeshCache && UCustomizableObjectSystem::IsMeshCacheEnabled();
+		const bool bUseOverrideMaterials = !bIsTransientMesh || bIsUsingMeshCache;
 		if (bUseOverrideMaterials)
 		{
 			if (FCustomizableInstanceComponentData* ComponentData = CustomizableObjectInstance->GetPrivate()->GetComponentData(Usage->GetComponentName()))
 			{
-				for (int32 Index = 0; Index < ComponentData->OverrideMaterials.Num(); ++Index)
-				{
-					Parent->SetMaterial(Index, ComponentData->OverrideMaterials[Index]);
-				}
+				NewOverridenMaterials = ComponentData->OverrideMaterials;
 			}
 		}
+		
+		if (Parent->OverrideMaterials != NewOverridenMaterials)
+		{
+			// Before setting the new override materials clear the old ones set in the Skeletal Mesh Component
+			// Note : this also resets the array of overlay materials used for each one of the override materials.
+			Parent->EmptyOverrideMaterials();
+
+			// Set the new override materials
+			for (int32 Index = 0; Index < NewOverridenMaterials.Num(); ++Index)
+			{
+				Parent->SetMaterial(Index, NewOverridenMaterials[Index]);
+			}
 
-		bAttachedParentUpdated |= OldOverridenMaterials != Parent->OverrideMaterials;
+			UE_LOG(LogMutable, Log, TEXT("Skeletal Mesh Component override materials have been updated."));
+			bAttachedParentUpdated = true;
+		}
+		
 
 		UMaterialInterface* OverlayMaterial = nullptr;
-		if (bUseOverrideMaterials /* bUseOverlayMaterials*/)
+		if (bUseOverrideMaterials)
 		{
 			if (FCustomizableInstanceComponentData* ComponentData = CustomizableObjectInstance->GetPrivate()->GetComponentData(Usage->GetComponentName()))
 			{
