[Mutable] Fix MeshReshape rigid clusters.

- Cluster attachament point was not used to compute displacemnt and so the vertices went to the attachment point.

#rnx

diff --git a/Engine/Plugins/Mutable/Source/MutableRuntime/Private/MuR/OpMeshReshape.h b/Engine/Plugins/Mutable/Source/MutableRuntime/Private/MuR/OpMeshReshape.h
--- a/Engine/Plugins/Mutable/Source/MutableRuntime/Private/MuR/OpMeshReshape.h
+++ b/Engine/Plugins/Mutable/Source/MutableRuntime/Private/MuR/OpMeshReshape.h
@@ -417,9 +417,17 @@ namespace mu
 				GetDeform(Shape, Binding, NewPosition, NewNormalPosition);
 
 				const FVector3f OldPosition = (PositionIter + VertexIndex).GetAsVec3f();
+  
+                FVector3f PositionDisplacement;
+                if ((OldPosition - Binding.AttachmentPoint).IsNearlyZero())
+                {
+				    PositionDisplacement = (NewPosition - OldPosition) * Binding.Weight;
+				}
+                else
+                {
+				    PositionDisplacement = (NewPosition - Binding.AttachmentPoint) * Binding.Weight;
+                }
 
-				FVector3f PositionDisplacement = (NewPosition - OldPosition) * Binding.Weight;
-				
 				// TODO: Separeat rigid from not rigid processing.
 				(PositionIter + VertexIndex).SetFromVec3f(OldPosition + PositionDisplacement);
 				
