[Mutable] Fix possible crash when a mesh without skeleton is used in a morph operation.

#rnx

diff --git a/Engine/Plugins/Mutable/Source/CustomizableObjectEditor/Private/MuCOE/GenerateMutableSource/GenerateMutableSourceMesh.cpp b/Engine/Plugins/Mutable/Source/CustomizableObjectEditor/Private/MuCOE/GenerateMutableSource/GenerateMutableSourceMesh.cpp
--- a/Engine/Plugins/Mutable/Source/CustomizableObjectEditor/Private/MuCOE/GenerateMutableSource/GenerateMutableSourceMesh.cpp
+++ b/Engine/Plugins/Mutable/Source/CustomizableObjectEditor/Private/MuCOE/GenerateMutableSource/GenerateMutableSourceMesh.cpp
@@ -2948,10 +2948,15 @@ TSharedPtr<mu::FMesh> BuildMorphedMutableMesh(const UEdGraphPin* BaseSourcePin,
 		Source.Mesh = SkeletalMesh;
 		TSharedPtr<mu::FMesh> BaseSourceMesh = GenerateMutableSkeletalMesh(Source, LODIndexConnected, SectionIndexConnected,
 			LODIndex, SectionIndex, FString(), GenerationContext, Node);
-		check(BaseSourceMesh->IsReference());
-		// The mesh will be modified and it may come from a cache, so we need to clone it.
-		MorphedSourceMesh = BaseSourceMesh->Clone(); 
-		MorphedSourceMesh->SetReferencedMorph(MorphTargetName);
+
+		if (BaseSourceMesh)
+		{
+			check(BaseSourceMesh->IsReference());
+		
+			// The mesh will be modified and it may come from a cache, so we need to clone it.
+			MorphedSourceMesh = BaseSourceMesh->Clone(); 
+			MorphedSourceMesh->SetReferencedMorph(MorphTargetName);
+		}
 	}
 
 	return MorphedSourceMesh;
