Unshelved from pending changelist '41920182':

[Mutable] Fix meshes incorrectly reusing poses.

diff --git a/Engine/Plugins/Mutable/Source/CustomizableObject/Internal/MuCO/CustomizableObjectPrivate.h b/Engine/Plugins/Mutable/Source/CustomizableObject/Internal/MuCO/CustomizableObjectPrivate.h
--- a/Engine/Plugins/Mutable/Source/CustomizableObject/Internal/MuCO/CustomizableObjectPrivate.h
+++ b/Engine/Plugins/Mutable/Source/CustomizableObject/Internal/MuCO/CustomizableObjectPrivate.h
@@ -1608,6 +1608,8 @@ enum class ECustomizableObjectVersions : int32
 
 	FixMeshReusalDueToLayouts,
 
+	IncorrectBonePoseMerging,
+	
 	// -----<new versions can be added above this line>--------
 	LastCustomizableObjectVersion
 };
diff --git a/Engine/Plugins/Mutable/Source/MutableRuntime/Public/MuR/Mesh.h b/Engine/Plugins/Mutable/Source/MutableRuntime/Public/MuR/Mesh.h
--- a/Engine/Plugins/Mutable/Source/MutableRuntime/Public/MuR/Mesh.h
+++ b/Engine/Plugins/Mutable/Source/MutableRuntime/Public/MuR/Mesh.h
@@ -480,7 +480,9 @@ namespace mu
 
 			inline bool operator==(const FBonePose& Other) const
 			{
-				return BoneUsageFlags == Other.BoneUsageFlags && BoneId == Other.BoneId;
+				return BoneUsageFlags == Other.BoneUsageFlags &&
+				BoneId == Other.BoneId &&
+				BoneTransform.Equals(Other.BoneTransform);
 			}
 		};
 
