[Mutable] Fix tracking allocation memory counters.

- Now it uses a singleton pattern.

#rnx

diff --git a/Engine/Plugins/Mutable/Source/CustomizableObject/Internal/MuCO/CustomizableObjectMipDataProvider.h b/Engine/Plugins/Mutable/Source/CustomizableObject/Internal/MuCO/CustomizableObjectMipDataProvider.h
--- a/Engine/Plugins/Mutable/Source/CustomizableObject/Internal/MuCO/CustomizableObjectMipDataProvider.h
+++ b/Engine/Plugins/Mutable/Source/CustomizableObject/Internal/MuCO/CustomizableObjectMipDataProvider.h
@@ -88,6 +88,17 @@ struct FMutableMipUpdateLevel
 	EPixelFormat Format;
 };
 
+namespace mu::MemoryCounters
+{
+	struct FPrefetchMemoryCounter
+	{
+		CUSTOMIZABLEOBJECT_API static std::atomic<SSIZE_T>& Get()
+		{
+			static std::atomic<SSIZE_T> Counter{0};
+			return Counter;
+		}
+	};
+}
 
 /** Runtime data used during a mutable image mipmap update */
 struct FMutableImageOperationData
@@ -104,12 +115,8 @@ struct FMutableImageOperationData
 
 	uint32 MutableTaskId = FMutableTaskGraph::INVALID_ID;
 
-	struct FPrefetchMemoryCounter
-	{
-		alignas(8) static inline std::atomic<SSIZE_T> Counter{ 0 };
-	};
-	using ArrayDataContainerType = TArray<uint8, mu::FDefaultMemoryTrackingAllocator<FPrefetchMemoryCounter>>;
-	ArrayDataContainerType AllocatedMemory;
+	using FPrefetchArray = TArray<uint8, mu::FDefaultMemoryTrackingAllocator<mu::MemoryCounters::FPrefetchMemoryCounter>>;
+	FPrefetchArray AllocatedMemory;
 
 	// Used to sync with the FMutableTextureMipDataProvider and FRenderAssetUpdate::Tick
 	FThreadSafeCounter* Counter;
diff --git a/Engine/Plugins/Mutable/Source/MutableRuntime/Internal/MuR/SystemPrivate.h b/Engine/Plugins/Mutable/Source/MutableRuntime/Internal/MuR/SystemPrivate.h
--- a/Engine/Plugins/Mutable/Source/MutableRuntime/Internal/MuR/SystemPrivate.h
+++ b/Engine/Plugins/Mutable/Source/MutableRuntime/Internal/MuR/SystemPrivate.h
@@ -29,7 +29,11 @@ namespace mu::MemoryCounters
 {
 	struct FInternalMemoryCounter
 	{
-		alignas(8) static inline std::atomic<SSIZE_T> Counter {0};
+		MUTABLERUNTIME_API static std::atomic<SSIZE_T>& Get()
+		{
+			static std::atomic<SSIZE_T> Counter{0};
+			return Counter;
+		}
 	};
 }
 
diff --git a/Engine/Plugins/Mutable/Source/MutableRuntime/Private/MuR/CodeRunner.h b/Engine/Plugins/Mutable/Source/MutableRuntime/Private/MuR/CodeRunner.h
--- a/Engine/Plugins/Mutable/Source/MutableRuntime/Private/MuR/CodeRunner.h
+++ b/Engine/Plugins/Mutable/Source/MutableRuntime/Private/MuR/CodeRunner.h
@@ -25,7 +25,11 @@ namespace mu::MemoryCounters
 {
 	struct FStreamingMemoryCounter
 	{
-		alignas(8) static inline std::atomic<SSIZE_T> Counter {0};
+		MUTABLERUNTIME_API static std::atomic<SSIZE_T>& Get()
+		{
+			static std::atomic<SSIZE_T> Counter{0};
+			return Counter;
+		}
 	};
 }
 
diff --git a/Engine/Plugins/Mutable/Source/MutableRuntime/Private/MuR/System.cpp b/Engine/Plugins/Mutable/Source/MutableRuntime/Private/MuR/System.cpp
--- a/Engine/Plugins/Mutable/Source/MutableRuntime/Private/MuR/System.cpp
+++ b/Engine/Plugins/Mutable/Source/MutableRuntime/Private/MuR/System.cpp
@@ -1149,10 +1149,10 @@ namespace mu
 		const uint32 TempBytes = GetTempBytes();
 
 		// Get allocator counters.
-		const SSIZE_T ImageAllocBytes	 = MemoryCounters::FImageMemoryCounter::Counter.load(std::memory_order_relaxed);
-		const SSIZE_T MeshAllocBytes     = MemoryCounters::FMeshMemoryCounter::Counter.load(std::memory_order_relaxed);
-		const SSIZE_T StreamAllocBytes   = MemoryCounters::FStreamingMemoryCounter::Counter.load(std::memory_order_relaxed);
-		const SSIZE_T InternalAllocBytes = MemoryCounters::FInternalMemoryCounter::Counter.load(std::memory_order_relaxed);
+		const SSIZE_T ImageAllocBytes	 = MemoryCounters::FImageMemoryCounter::Get().load(std::memory_order_relaxed);
+		const SSIZE_T MeshAllocBytes     = MemoryCounters::FMeshMemoryCounter::Get().load(std::memory_order_relaxed);
+		const SSIZE_T StreamAllocBytes   = MemoryCounters::FStreamingMemoryCounter::Get().load(std::memory_order_relaxed);
+		const SSIZE_T InternalAllocBytes = MemoryCounters::FInternalMemoryCounter::Get().load(std::memory_order_relaxed);
 
 		SSIZE_T TotalBytes = ImageAllocBytes + MeshAllocBytes + StreamAllocBytes + InternalAllocBytes;
 
@@ -1237,10 +1237,10 @@ namespace mu
 	{
 		MUTABLE_CPUPROFILER_SCOPE(GetCurrentMemoryBytes);
 
-		SSIZE_T TotalBytes = MemoryCounters::FImageMemoryCounter::Counter.load(std::memory_order_relaxed) + 
-						     MemoryCounters::FMeshMemoryCounter::Counter.load(std::memory_order_relaxed) +
-							 MemoryCounters::FStreamingMemoryCounter::Counter.load(std::memory_order_relaxed) +
-							 MemoryCounters::FInternalMemoryCounter::Counter.load(std::memory_order_relaxed);
+		SSIZE_T TotalBytes = MemoryCounters::FImageMemoryCounter::Get().load(std::memory_order_relaxed) + 
+						     MemoryCounters::FMeshMemoryCounter::Get().load(std::memory_order_relaxed) +
+							 MemoryCounters::FStreamingMemoryCounter::Get().load(std::memory_order_relaxed) +
+							 MemoryCounters::FInternalMemoryCounter::Get().load(std::memory_order_relaxed);
 
 		return TotalBytes;
 	}
diff --git a/Engine/Plugins/Mutable/Source/MutableRuntime/Public/MuR/ImageDataStorage.h b/Engine/Plugins/Mutable/Source/MutableRuntime/Public/MuR/ImageDataStorage.h
--- a/Engine/Plugins/Mutable/Source/MutableRuntime/Public/MuR/ImageDataStorage.h
+++ b/Engine/Plugins/Mutable/Source/MutableRuntime/Public/MuR/ImageDataStorage.h
@@ -15,12 +15,17 @@
 
 namespace mu::MemoryCounters
 {
-	struct FImageMemoryCounter
+	struct  FImageMemoryCounter
 	{
-		alignas(8) static inline std::atomic<SSIZE_T> Counter {0};
+		MUTABLERUNTIME_API static std::atomic<SSIZE_T>& Get()
+		{
+			static std::atomic<SSIZE_T> Counter{0};
+			return Counter;
+		}
 	};
 }
 
+
 namespace mu
 {
 	using FImageArray = TArray<uint8, FDefaultMemoryTrackingAllocator<MemoryCounters::FImageMemoryCounter>>;
diff --git a/Engine/Plugins/Mutable/Source/MutableRuntime/Public/MuR/MemoryTrackingAllocationPolicy.h b/Engine/Plugins/Mutable/Source/MutableRuntime/Public/MuR/MemoryTrackingAllocationPolicy.h
--- a/Engine/Plugins/Mutable/Source/MutableRuntime/Public/MuR/MemoryTrackingAllocationPolicy.h
+++ b/Engine/Plugins/Mutable/Source/MutableRuntime/Public/MuR/MemoryTrackingAllocationPolicy.h
@@ -12,19 +12,22 @@
 namespace mu
 {
 	
-	/**
-	 * CounterType is expected to be of the following form, the alingas is needed to prevent 
-     * type clashes between modules. 
-	 * struct FCounterTypeName
-	 * {
-	 *     alignas(8) static inline std::atomic<SSIZE_T> Counter {0}; 		
-	 * };
-	 */
+	
+	// CounterType is expected to be of the following form. 
+	//struct FCounterTypeName
+	//{
+	//	static std::atomic<SSIZE_T>& Get()
+	//	{
+	//		static std::atomic<SSIZE_T> Counter{0};
+	//		return Counter;
+	//	}
+	//};
+
 
 	template<typename BaseAlloc, typename CounterType>
 	class TMemoryTrackingAllocatorWrapper
 	{
-		static_assert(std::is_same_v<decltype(CounterType::Counter.load()), SSIZE_T>, "CounterType::Counter must be signed.");
+		static_assert(std::is_same_v<decltype(CounterType::Get().load()), SSIZE_T>, "CounterType must be signed.");
 
 	public:
 		using SizeType = typename BaseAlloc::SizeType;
@@ -49,7 +52,7 @@ namespace mu
 			/** Destructor. */
 			FORCEINLINE ~ForAnyElementType()
 			{
-				CounterType::Counter.fetch_sub(AllocSize, std::memory_order_relaxed);
+				CounterType::Get().fetch_sub(AllocSize, std::memory_order_relaxed);
 
 #if UE_MUTABLE_TRACK_ALLOCATOR_MEMORY_PEAK
 				FGlobalMemoryCounter::Update(-AllocSize);
@@ -65,7 +68,7 @@ namespace mu
 			{
 				BaseAlloc::ForAnyElementType::MoveToEmpty(Other);
 
-				CounterType::Counter.fetch_sub(AllocSize, std::memory_order_relaxed);
+				CounterType::Get().fetch_sub(AllocSize, std::memory_order_relaxed);
 
 #if UE_MUTABLE_TRACK_ALLOCATOR_MEMORY_PEAK
 				FGlobalMemoryCounter::Update(-AllocSize);
@@ -85,7 +88,7 @@ namespace mu
 
 				const SSIZE_T AllocatedSize = (SSIZE_T)BaseAlloc::ForAnyElementType::GetAllocatedSize(NewMax, NumBytesPerElement); 
 				const SSIZE_T Differential = AllocatedSize - AllocSize;
-				const SSIZE_T PrevCounterValue = CounterType::Counter.fetch_add(Differential, std::memory_order_relaxed);
+				const SSIZE_T PrevCounterValue = CounterType::Get().fetch_add(Differential, std::memory_order_relaxed);
 				check(PrevCounterValue >= AllocSize);
 
 #if UE_MUTABLE_TRACK_ALLOCATOR_MEMORY_PEAK
@@ -95,8 +98,11 @@ namespace mu
 				AllocSize = AllocatedSize;
 			}
 
-			template<class T = BaseAlloc>
-			FORCEINLINE typename TEnableIf<TAllocatorTraits<T>::SupportsElementAlignment, void>::Type ResizeAllocation(
+			template <
+				class T = BaseAlloc
+				UE_REQUIRES(TAllocatorTraits<T>::SupportsElementAlignment)
+			>
+			FORCEINLINE void ResizeAllocation(
 				SizeType CurrentNum,
 				SizeType NewMax,
 				SIZE_T NumBytesPerElement,
@@ -107,7 +113,7 @@ namespace mu
 
 				const SSIZE_T AllocatedSize = (SSIZE_T)BaseAlloc::ForAnyElementType::GetAllocatedSize(NewMax, NumBytesPerElement); 
 				const SSIZE_T Differential = AllocatedSize - AllocSize;
-				const SSIZE_T PrevCounterValue = CounterType::Counter.fetch_add(Differential, std::memory_order_relaxed);
+				const SSIZE_T PrevCounterValue = CounterType::Get().fetch_add(Differential, std::memory_order_relaxed);
 				check(PrevCounterValue >= AllocSize);
 
 #if UE_MUTABLE_TRACK_ALLOCATOR_MEMORY_PEAK
diff --git a/Engine/Plugins/Mutable/Source/MutableRuntime/Public/MuR/MeshBufferSet.h b/Engine/Plugins/Mutable/Source/MutableRuntime/Public/MuR/MeshBufferSet.h
--- a/Engine/Plugins/Mutable/Source/MutableRuntime/Public/MuR/MeshBufferSet.h
+++ b/Engine/Plugins/Mutable/Source/MutableRuntime/Public/MuR/MeshBufferSet.h
@@ -14,7 +14,11 @@ namespace mu::MemoryCounters
 {
 	struct FMeshMemoryCounter
 	{
-		alignas(8) static inline std::atomic<SSIZE_T> Counter{0};
+		MUTABLERUNTIME_API static std::atomic<SSIZE_T>& Get()
+		{
+			static std::atomic<SSIZE_T> Counter{0};
+			return Counter;
+		}
 	};
 }
 
