[Mutable] Allow use of DDC for editor compilations. Add plugin settings to define cook and editor DDC policies. Add CVar to turn on/off use of DDC.

#rnx

diff --git a/Engine/Plugins/Mutable/Source/CustomizableObject/Internal/MuCO/CustomizableObjectPrivate.h b/Engine/Plugins/Mutable/Source/CustomizableObject/Internal/MuCO/CustomizableObjectPrivate.h
--- a/Engine/Plugins/Mutable/Source/CustomizableObject/Internal/MuCO/CustomizableObjectPrivate.h
+++ b/Engine/Plugins/Mutable/Source/CustomizableObject/Internal/MuCO/CustomizableObjectPrivate.h
@@ -1097,6 +1097,13 @@ struct FModelStreamableBulkData
 
 	TArray<FByteBulkData> StreamableBulkData;
 
+#if WITH_EDITORONLY_DATA
+	// Used to know if roms and other resources must be streamed from the DDC.
+	bool bIsStoredInDDC = false;
+	UE::DerivedData::FCacheKey DDCKey = UE::DerivedData::FCacheKey::Empty;
+	UE::DerivedData::ECachePolicy DDCDefaultPolicy = UE::DerivedData::ECachePolicy::Default;
+#endif
+
 	// File path to stream resources from when not using FByteBulkData or DDC.
 	FString FullFilePath;
 
@@ -1280,11 +1287,6 @@ public:
 	UPROPERTY()
 	TMap<FName, FGuid> ParticipatingObjects;
 
-	// Used to know if roms and other resources must be streamed from the DDC.
-	bool bIsStoredInDDC = false;
-	UE::DerivedData::FCacheKey DDCKey = UE::DerivedData::FCacheKey::Empty;
-	UE::DerivedData::ECachePolicy DDCDefaultPolicy = UE::DerivedData::ECachePolicy::Default;
-
 	UPROPERTY()
 	TArray<FCustomizableObjectResourceData> StreamedResourceDataEditor;
 
@@ -1706,7 +1708,7 @@ public:
 	CUSTOMIZABLEOBJECT_API UE::DerivedData::FCacheKey GetDerivedDataCacheKeyForOptions(FCompilationOptions InOptions);
 
 	/** Attempts to load the compiled data from DDC. Builds key if not supplied. */
-	CUSTOMIZABLEOBJECT_API void LoadCompiledDataFromDDC(FCompilationOptions Options, UE::DerivedData::ECachePolicy DefaultPolicy, UE::DerivedData::FCacheKey* DDCKey);
+	CUSTOMIZABLEOBJECT_API bool LoadCompiledDataFromDDC(FCompilationOptions Options, UE::DerivedData::ECachePolicy DefaultPolicy, UE::DerivedData::FCacheKey* DDCKey, MutablePrivate::FMutableCachedPlatformData& InOutPlatformData);
 
 	/** Log data for debugging purposes */
 	CUSTOMIZABLEOBJECT_API void LogMemory();
@@ -1888,6 +1890,15 @@ public:
 
 #endif
 	
-	static constexpr int32 DerivedDataVersion = 0x46799a1c;
+	static constexpr int32 DerivedDataVersion = 0x34789a1c;
 };
 
+#if WITH_EDITOR
+
+/** Returns the DDC ValueId of the file owning a streamable resource.
+	* @ param StreamableDataType - UE level data type
+	* @ param ResourceType - mu::EDataType
+	*/
+CUSTOMIZABLEOBJECT_API UE::DerivedData::FValueId GetDerivedDataValueIdForResource(MutablePrivate::EStreamableDataType StreamableDataType, uint32 FileId, uint16 ResourceType, uint16 Flags);
+
+#endif
\ No newline at end of file
diff --git a/Engine/Plugins/Mutable/Source/CustomizableObject/Internal/MuCO/CustomizableObjectSystemPrivate.h b/Engine/Plugins/Mutable/Source/CustomizableObject/Internal/MuCO/CustomizableObjectSystemPrivate.h
--- a/Engine/Plugins/Mutable/Source/CustomizableObject/Internal/MuCO/CustomizableObjectSystemPrivate.h
+++ b/Engine/Plugins/Mutable/Source/CustomizableObject/Internal/MuCO/CustomizableObjectSystemPrivate.h
@@ -765,6 +765,35 @@ struct FPendingReleaseSkeletalMeshInfo
 };
 
 
+#if WITH_EDITORONLY_DATA
+
+UENUM()
+enum class ECustomizableObjectDDCPolicy : uint8
+{
+	None = 0,
+	Local,
+	Default
+};
+
+
+// Struct used to keep a copy of the EditorSettings needed to compile Customizable Objects.
+struct FEditorCompileSettings
+{
+	// General case
+	bool bIsMutableEnabled = true;
+
+	// Auto Compile 
+	bool bEnableAutomaticCompilation = true;
+	bool bCompileObjectsSynchronously = true;
+	bool bCompileRootObjectsOnStartPIE = false;
+
+	// DDC settings
+	ECustomizableObjectDDCPolicy EditorDerivedDataCachePolicy = ECustomizableObjectDDCPolicy::Default;
+	ECustomizableObjectDDCPolicy CookDerivedDataCachePolicy = ECustomizableObjectDDCPolicy::Default;
+};
+
+#endif
+
 /** Private part, hidden from outside the plugin.
  *
  * UE STREAMING:
@@ -995,6 +1024,11 @@ public:
 	
 	TSharedRef<FMutableStreamableManager> StreamableManager = MakeShared<FMutableStreamableManager>();
 
+#if WITH_EDITOR
+	// Copy of the Mutable Editor Settings tied to CO compilation. They are updated whenever changed
+	FEditorCompileSettings EditorSettings;
+#endif
+
 #if WITH_EDITORONLY_DATA
 	// Array to keep track of cached objects
 	TArray<FGuid> UncompiledCustomizableObjectIds;
diff --git a/Engine/Plugins/Mutable/Source/CustomizableObject/Internal/MuCO/UnrealMutableModelDiskStreamer.h b/Engine/Plugins/Mutable/Source/CustomizableObject/Internal/MuCO/UnrealMutableModelDiskStreamer.h
--- a/Engine/Plugins/Mutable/Source/CustomizableObject/Internal/MuCO/UnrealMutableModelDiskStreamer.h
+++ b/Engine/Plugins/Mutable/Source/CustomizableObject/Internal/MuCO/UnrealMutableModelDiskStreamer.h
@@ -14,7 +14,10 @@ class IAsyncReadFileHandle;
 class IAsyncReadRequest;
 class IBulkDataIORequest;
 class UCustomizableObject;
-namespace mu { class FModel; }
+namespace mu {
+	class FModel;
+	enum class EDataType : uint8;
+}
 struct FMutableStreamableBlock;
 
 class FUnrealMutableInputStream : public mu::FInputStream
@@ -61,7 +64,7 @@ public:
 
 	// mu::ModelReader interface
 	CUSTOMIZABLEOBJECT_API bool DoesBlockExist(const mu::FModel*, uint32 BlockKey) override;
-	CUSTOMIZABLEOBJECT_API FOperationID BeginReadBlock(const mu::FModel*, uint32 BlockKey, void* pBuffer, uint64 size, TFunction<void(bool bSuccess)>* CompletionCallback) override;
+	CUSTOMIZABLEOBJECT_API FOperationID BeginReadBlock(const mu::FModel*, uint32 BlockKey, void* pBuffer, uint64 size, mu::EDataType ResourceType, TFunction<void(bool bSuccess)>* CompletionCallback) override;
 	CUSTOMIZABLEOBJECT_API bool IsReadCompleted(FOperationID) override;
 	CUSTOMIZABLEOBJECT_API bool EndRead(FOperationID) override;
 
@@ -88,13 +91,6 @@ protected:
 		TMap<uint32, TSharedPtr<IAsyncReadFileHandle>> ReadFileHandles;
 
 		TWeakPtr<FModelStreamableBulkData> ModelStreamableBulkData;
-
-#if WITH_EDITORONLY_DATA
-		// DDC files streaming
-		bool bIsStoredInDDC = false;
-		UE::DerivedData::FCacheKey DDCKey;
-		UE::DerivedData::FCacheRecordPolicy DDCPolicy;
-#endif
 	};
 
 	TArray<FObjectData> Objects;
diff --git a/Engine/Plugins/Mutable/Source/CustomizableObject/Private/MuCO/CustomizableObject.cpp b/Engine/Plugins/Mutable/Source/CustomizableObject/Private/MuCO/CustomizableObject.cpp
--- a/Engine/Plugins/Mutable/Source/CustomizableObject/Private/MuCO/CustomizableObject.cpp
+++ b/Engine/Plugins/Mutable/Source/CustomizableObject/Private/MuCO/CustomizableObject.cpp
@@ -76,16 +76,6 @@ TAutoConsoleVariable<bool> CVarMutableUseBulkData(
 	TEXT("False - Use Mut files to store streamable data\n"));
 
 
-TAutoConsoleVariable<int32> CVarMutableDerivedDataCacheUsage(
-	TEXT("mutable.DerivedDataCacheUsage"),
-	2,
-	TEXT("Derived data cache access for cooked data.")
-	TEXT("0 - None. Disables access to the cache.")
-	TEXT("1 - Local. Allow cache requests to query and store records and values in local caches.")
-	TEXT("2 - Default. Allow cache requests to query and store records and values in any caches."),
-	ECVF_Default);
-
-
 TAutoConsoleVariable<bool> CVarMutableAsyncCook(
 	TEXT("Mutable.CookAsync"),
 	false,
@@ -3131,7 +3121,7 @@ UE::DerivedData::FCacheKey UCustomizableObjectPrivate::GetDerivedDataCacheKeyFor
 }
 
 
-void UCustomizableObjectPrivate::LoadCompiledDataFromDDC(FCompilationOptions Options, UE::DerivedData::ECachePolicy DefaultPolicy, UE::DerivedData::FCacheKey* DDCKey)
+bool UCustomizableObjectPrivate::LoadCompiledDataFromDDC(FCompilationOptions Options, UE::DerivedData::ECachePolicy DefaultPolicy, UE::DerivedData::FCacheKey* DDCKey, MutablePrivate::FMutableCachedPlatformData& PlatformData)
 {
 	MUTABLE_CPUPROFILER_SCOPE(UCustomizableObjectPrivate::LoadCompiledDataFromDDC);
 
@@ -3195,28 +3185,26 @@ void UCustomizableObjectPrivate::LoadCompiledDataFromDDC(FCompilationOptions Opt
 	// Check if all the requested buffers were found.
 	if (!ModelBytesDDC.IsNull() && !ModelResourcesBytesDDC.IsNull() && !BulkDataFilesBytesDDC.IsNull()  && !ModelStreamablesBytesDDC.IsNull())
 	{
-		MutablePrivate::FMutableCachedPlatformData CachedData;
-
 		// Load the compiled data to validate it.
 		FMemoryReaderView ModelResourcesReader(ModelResourcesBytesDDC.GetView());
 
 		if (TObjectPtr<UModelResources> LocalModelResources = LoadModelResources_Internal(ModelResourcesReader, Options.TargetPlatform, Options.bIsCooking))
 		{
-			CachedData.ModelResources = TStrongObjectPtr<UModelResources>(LocalModelResources);
-			CachedData.ModelResources->bIsStoredInDDC = true;
-			CachedData.ModelResources->DDCKey = CacheKey;
-			CachedData.ModelResources->DDCDefaultPolicy = ECachePolicy::Default;
+			PlatformData.ModelResources = TStrongObjectPtr<UModelResources>(LocalModelResources);
 
 			FMemoryReaderView ModelStreamablesReader(ModelStreamablesBytesDDC.GetView());
-			CachedData.ModelStreamableBulkData = LoadModelStreamableBulk_Internal(ModelStreamablesReader);
+			PlatformData.ModelStreamableBulkData = LoadModelStreamableBulk_Internal(ModelStreamablesReader);
+			PlatformData.ModelStreamableBulkData->bIsStoredInDDC = true;
+			PlatformData.ModelStreamableBulkData->DDCKey = CacheKey;
+			PlatformData.ModelStreamableBulkData->DDCDefaultPolicy = ECachePolicy::Default;
 
 			FMemoryReaderView ModelReader(ModelBytesDDC.GetView());
-			CachedData.Model = LoadModel_Internal(ModelReader);
+			PlatformData.Model = LoadModel_Internal(ModelReader);
 		}
 
-		if (!CachedData.Model || !CachedData.ModelResources || !CachedData.ModelStreamableBulkData)
+		if (!PlatformData.Model || !PlatformData.ModelResources || !PlatformData.ModelStreamableBulkData)
 		{
-			return;
+			return false;
 		}
 
 		// Cache CookedPlatfomData. 
@@ -3226,7 +3214,7 @@ void UCustomizableObjectPrivate::LoadCompiledDataFromDDC(FCompilationOptions Opt
 			// TODO UE-220138: Sync -> Async
 			MUTABLE_CPUPROFILER_SCOPE(CacheCookedPlatformData);
 
-			CachedData.ModelResources->InitCookData(*GetPublic());
+			PlatformData.ModelResources->InitCookData(*GetPublic());
 
 			// Value Id to file mapping to reconstruct the cached data
 			TMap<FValueId, MutablePrivate::FFile> ValueIdToFile;
@@ -3239,20 +3227,13 @@ void UCustomizableObjectPrivate::LoadCompiledDataFromDDC(FCompilationOptions Opt
 
 				ValueIdToFile.Reserve(BulkDataFiles.Num());
 
-				FValueId::ByteArray ValueIdBytes = {};
+				uint32 FileIndex = 0;
 				for (MutablePrivate::FFile& File : BulkDataFiles)
 				{
-					int8 ValueIdOffset = 0;
-					FMemory::Memcpy(&ValueIdBytes[ValueIdOffset], &File.DataType, sizeof(File.DataType));
-					ValueIdOffset += sizeof(File.DataType);
-					FMemory::Memcpy(&ValueIdBytes[ValueIdOffset], &File.Id, sizeof(File.Id));
-					ValueIdOffset += sizeof(File.Id);
-					FMemory::Memcpy(&ValueIdBytes[ValueIdOffset], &File.ResourceType, sizeof(File.ResourceType));
-					ValueIdOffset += sizeof(File.ResourceType);
-					FMemory::Memcpy(&ValueIdBytes[ValueIdOffset], &File.Flags, sizeof(File.Flags));
-
-					MutablePrivate::FFile& DestFile = ValueIdToFile.Add(FValueId(ValueIdBytes));
+					MutablePrivate::FFile& DestFile = ValueIdToFile.Add(GetDerivedDataValueIdForResource(File.DataType, FileIndex, File.ResourceType, File.Flags));
 					DestFile = MoveTemp(File);
+
+					++FileIndex;
 				}
 				BulkDataFiles.Empty();
 			}
@@ -3265,7 +3246,7 @@ void UCustomizableObjectPrivate::LoadCompiledDataFromDDC(FCompilationOptions Opt
 
 			FRequestOwner RequestOwner(EPriority::Blocking);
 			GetCache().Get(MakeArrayView(&Request, 1), RequestOwner,
-				[&CachedData, &ValueIdToFile](FCacheGetResponse&& Response)
+				[&PlatformData, &ValueIdToFile](FCacheGetResponse&& Response)
 				{
 					MUTABLE_CPUPROFILER_SCOPE(CacheBulkDataFromDDC);
 
@@ -3299,7 +3280,7 @@ void UCustomizableObjectPrivate::LoadCompiledDataFromDDC(FCompilationOptions Opt
 							{
 								for (const MutablePrivate::FBlock& Block : File->Blocks)
 								{
-									CachedData.ModelStreamableData.Set(Block.Id, TempData.GetData() + Block.Offset, Block.Size);
+									PlatformData.ModelStreamableData.Set(Block.Id, TempData.GetData() + Block.Offset, Block.Size);
 								}
 								break;
 							}
@@ -3307,7 +3288,7 @@ void UCustomizableObjectPrivate::LoadCompiledDataFromDDC(FCompilationOptions Opt
 							{
 								for (const MutablePrivate::FBlock& Block : File->Blocks)
 								{
-									CachedData.MorphStreamableData.Set(Block.Id, TempData.GetData() + Block.Offset, Block.Size);
+									PlatformData.MorphStreamableData.Set(Block.Id, TempData.GetData() + Block.Offset, Block.Size);
 								}
 								break;
 							}
@@ -3315,7 +3296,7 @@ void UCustomizableObjectPrivate::LoadCompiledDataFromDDC(FCompilationOptions Opt
 							{
 								for (const MutablePrivate::FBlock& Block : File->Blocks)
 								{
-									CachedData.ClothingStreamableData.Set(Block.Id, TempData.GetData() + Block.Offset, Block.Size);
+									PlatformData.ClothingStreamableData.Set(Block.Id, TempData.GetData() + Block.Offset, Block.Size);
 								}
 								break;
 							}
@@ -3333,25 +3314,38 @@ void UCustomizableObjectPrivate::LoadCompiledDataFromDDC(FCompilationOptions Opt
 			if (CVarMutableUseBulkData.GetValueOnAnyThread())
 			{
 				const uint32 NumBulkDataFilesPerBucket = MAX_uint8;
-				MutablePrivate::GenerateBulkDataFilesListWithFileLimit(CachedData.Model, *CachedData.ModelStreamableBulkData.Get(), NumBulkDataFilesPerBucket, CachedData.BulkDataFiles);
+				MutablePrivate::GenerateBulkDataFilesListWithFileLimit(PlatformData.Model, *PlatformData.ModelStreamableBulkData.Get(), NumBulkDataFilesPerBucket, PlatformData.BulkDataFiles);
 			}
 			else
 			{
-				MutablePrivate::GenerateBulkDataFilesListWithSizeLimit(CachedData.Model, *CachedData.ModelStreamableBulkData.Get(), Options.TargetPlatform, Options.PackagedDataBytesLimit, CachedData.BulkDataFiles);
+				MutablePrivate::GenerateBulkDataFilesListWithSizeLimit(PlatformData.Model, *PlatformData.ModelStreamableBulkData.Get(), Options.TargetPlatform, Options.PackagedDataBytesLimit, PlatformData.BulkDataFiles);
 			}
-
-			MutablePrivate::FMutableCachedPlatformData& CachedPlatformData = CachedPlatformsData.Add(Options.TargetPlatform->PlatformName(), {});
-			CachedPlatformData = MoveTemp(CachedData);
-		}
-		else // Editor DDC
-		{
-			SetModel(CachedData.Model, GenerateIdentifier(*GetPublic()));
-			SetModelResources(CachedData.ModelResources.Get(), Options.bIsCooking);
-			SetModelStreamableBulkData(CachedData.ModelStreamableBulkData, Options.bIsCooking);
 		}
+
+		return true;
 	}
 
-	return;
+	return false;
+}
+
+
+UE::DerivedData::FValueId GetDerivedDataValueIdForResource(MutablePrivate::EStreamableDataType StreamableDataType, uint32 FileId, uint16 ResourceType, uint16 Flags)
+{
+	UE::DerivedData::FValueId::ByteArray ValueIdBytes = {};
+	
+	static_assert(sizeof(UE::DerivedData::FValueId::ByteArray) >= sizeof(MutablePrivate::EStreamableDataType) + sizeof(uint32) + sizeof(uint16) + sizeof(uint16));
+
+	int8 ValueIdOffset = 0;
+	FMemory::Memcpy(&ValueIdBytes[ValueIdOffset], &StreamableDataType, sizeof(StreamableDataType));
+	ValueIdOffset += sizeof(StreamableDataType);
+	FMemory::Memcpy(&ValueIdBytes[ValueIdOffset], &FileId, sizeof(FileId));
+	ValueIdOffset += sizeof(FileId);
+	uint16 DataType = (uint16)ResourceType;
+	FMemory::Memcpy(&ValueIdBytes[ValueIdOffset], &DataType, sizeof(DataType));
+	ValueIdOffset += sizeof(DataType);
+	FMemory::Memcpy(&ValueIdBytes[ValueIdOffset], &Flags, sizeof(Flags));
+	UE::DerivedData::FValueId ResourceId = UE::DerivedData::FValueId(ValueIdBytes);
+	return ResourceId;
 }
 
 
diff --git a/Engine/Plugins/Mutable/Source/CustomizableObject/Private/MuCO/CustomizableObjectMipDataProvider.cpp b/Engine/Plugins/Mutable/Source/CustomizableObject/Private/MuCO/CustomizableObjectMipDataProvider.cpp
--- a/Engine/Plugins/Mutable/Source/CustomizableObject/Private/MuCO/CustomizableObjectMipDataProvider.cpp
+++ b/Engine/Plugins/Mutable/Source/CustomizableObject/Private/MuCO/CustomizableObjectMipDataProvider.cpp
@@ -7,6 +7,7 @@
 #include "MuCO/CustomizableObjectSystemPrivate.h"
 #include "MuCO/MutableStreamRequest.h"
 #include "MuR/Model.h"
+#include "MuR/Operations.h"
 #include "TextureResource.h"
 #include "UnrealMutableImageProvider.h"
 
@@ -471,6 +472,7 @@ int32 FMutableTextureMipDataProvider::GetMips(const FTextureUpdateContext& Conte
 				PrefetchRequest->AddBlock(
 					ModelStreamableBulkData->ModelStreamables[RomId],
 					MutablePrivate::EStreamableDataType::Model,
+					(uint16)mu::EDataType::Image,
 					TArrayView<uint8>(AllocatedMemory, Model->GetRomSize(RomId)));
 			}
 		}
diff --git a/Engine/Plugins/Mutable/Source/CustomizableObject/Private/MuCO/CustomizableObjectSystem.cpp b/Engine/Plugins/Mutable/Source/CustomizableObject/Private/MuCO/CustomizableObjectSystem.cpp
--- a/Engine/Plugins/Mutable/Source/CustomizableObject/Private/MuCO/CustomizableObjectSystem.cpp
+++ b/Engine/Plugins/Mutable/Source/CustomizableObject/Private/MuCO/CustomizableObjectSystem.cpp
@@ -2020,14 +2020,14 @@ bool UCustomizableObjectSystem::CheckIfDiskOrMipUpdateOperationsPending(const UC
 
 void UCustomizableObjectSystem::EditorSettingsChanged(const FEditorCompileSettings& InEditorSettings)
 {
-	EditorSettings = InEditorSettings;
+	GetPrivate()->EditorSettings = InEditorSettings;
 
 	CVarMutableEnabled->Set(InEditorSettings.bIsMutableEnabled);
 }
 
 bool UCustomizableObjectSystem::IsAutoCompileEnabled() const
 {
-	return EditorSettings.bEnableAutomaticCompilation;
+	return GetPrivate()->EditorSettings.bEnableAutomaticCompilation;
 }
 
 
@@ -2045,7 +2045,7 @@ void UCustomizableObjectSystem::SetAutoCompileCommandletEnabled(bool bValue)
 
 bool UCustomizableObjectSystem::IsAutoCompilationSync() const
 {
-	return EditorSettings.bCompileObjectsSynchronously;
+	return GetPrivate()->EditorSettings.bCompileObjectsSynchronously;
 }
 
 #endif
@@ -4923,6 +4923,7 @@ void LoadMorphTargetsData(FMutableStreamRequest& MutableDataStreamer, const TSha
 		MutableDataStreamer.AddBlock(
 			Block,
 			MutablePrivate::EStreamableDataType::RealTimeMorph,
+			0,
 			MakeArrayView(reinterpret_cast<uint8*>(ReadDestData.Data.GetData()), ReadDestData.Data.Num() * sizeof(FMorphTargetVertexData)));
 	}
 }
@@ -5017,6 +5018,7 @@ void LoadClothing(FMutableStreamRequest& MutableDataStreamer, const TSharedRef<c
 		MutableDataStreamer.AddBlock(
 			Block,
 			MutablePrivate::EStreamableDataType::Clothing,
+			0,
 			MakeArrayView(reinterpret_cast<uint8*>(ReadDestData.Data.GetData()), ReadDestData.Data.Num() * sizeof(FCustomizableObjectMeshToMeshVertData)));
 	}
 }
diff --git a/Engine/Plugins/Mutable/Source/CustomizableObject/Private/MuCO/MutableStreamRequest.cpp b/Engine/Plugins/Mutable/Source/CustomizableObject/Private/MuCO/MutableStreamRequest.cpp
--- a/Engine/Plugins/Mutable/Source/CustomizableObject/Private/MuCO/MutableStreamRequest.cpp
+++ b/Engine/Plugins/Mutable/Source/CustomizableObject/Private/MuCO/MutableStreamRequest.cpp
@@ -6,6 +6,12 @@
 #include "MuCO/CustomizableObjectSystemPrivate.h"
 #include "HAL/PlatformFileManager.h"
 
+#if WITH_EDITOR
+#include "DerivedDataCache.h"
+#include "DerivedDataCacheInterface.h"
+#include "DerivedDataCacheKey.h"
+#include "DerivedDataRequestOwner.h"
+#endif
 
 FMutableStreamRequest::FMutableStreamRequest(const TSharedPtr<FModelStreamableBulkData>& InModelStreamableBulkData) :
 	ModelStreamableBulkData(InModelStreamableBulkData)
@@ -19,7 +25,7 @@ const TSharedPtr<FModelStreamableBulkData>& FMutableStreamRequest::GetModelStrea
 }
 
 
-void FMutableStreamRequest::AddBlock(const FMutableStreamableBlock& Block, MutablePrivate::EStreamableDataType DataType, TArrayView<uint8> AllocatedMemoryView)
+void FMutableStreamRequest::AddBlock(const FMutableStreamableBlock& Block, MutablePrivate::EStreamableDataType DataType, uint16 ResourceType, TArrayView<uint8> AllocatedMemoryView)
 {
 	if (bIsStreaming)
 	{
@@ -30,24 +36,27 @@ void FMutableStreamRequest::AddBlock(const FMutableStreamableBlock& Block, Mutab
 	IAsyncReadFileHandle* FileHandle = nullptr;
 
 #if WITH_EDITOR
-	int32 FileHandleIndex = OpenFilesIds.Find(Block.FileId);
-	if (FileHandleIndex == INDEX_NONE)
+	if (!ModelStreamableBulkData->bIsStoredInDDC)
 	{
-		const FString& FullFileName = ModelStreamableBulkData->FullFilePath + UCustomizableObjectPrivate::GetDataTypeExtension(DataType);
-		IAsyncReadFileHandle* ReadFileHandle = FPlatformFileManager::Get().GetPlatformFile().OpenAsyncRead(*FullFileName);
-		OpenFileHandles.Emplace(ReadFileHandle);
+		int32 FileHandleIndex = OpenFilesIds.Find(Block.FileId);
+		if (FileHandleIndex == INDEX_NONE)
+		{
+			const FString& FullFileName = ModelStreamableBulkData->FullFilePath + UCustomizableObjectPrivate::GetDataTypeExtension(DataType);
+			IAsyncReadFileHandle* ReadFileHandle = FPlatformFileManager::Get().GetPlatformFile().OpenAsyncRead(*FullFileName);
+			OpenFileHandles.Emplace(ReadFileHandle);
 
-		FileHandleIndex = OpenFilesIds.Add(Block.FileId);
-		check(OpenFileHandles.Num() == OpenFilesIds.Num());
-	}
+			FileHandleIndex = OpenFilesIds.Add(Block.FileId);
+			check(OpenFileHandles.Num() == OpenFilesIds.Num());
+		}
 
-	if (OpenFileHandles.IsValidIndex(FileHandleIndex))
-	{
-		FileHandle = OpenFileHandles[FileHandleIndex].Get();
+		if (OpenFileHandles.IsValidIndex(FileHandleIndex))
+		{
+			FileHandle = OpenFileHandles[FileHandleIndex].Get();
+		}
 	}
 #endif
-		
-	BlockReadInfos.Emplace(Block.Offset, FileHandle, AllocatedMemoryView, Block.FileId);
+	
+	BlockReadInfos.Emplace(Block.Offset, FileHandle, AllocatedMemoryView, Block.FileId, DataType, ResourceType, Block.Flags);
 }
 
 
@@ -71,7 +80,7 @@ UE::Tasks::FTask FMutableStreamRequest::Stream()
 			ModelStreamableBulkData = ModelStreamableBulkData,
 			BlockReadInfos = MoveTemp(BlockReadInfos),
 			HeapMemory = this->HeapMemory
-		]()
+		]() mutable
 		{
 			MUTABLE_CPUPROFILER_SCOPE(CustomizableInstanceLoadBlocksAsyncRead_Request);
 			FScopeLock Lock(&HeapMemory->ReadRequestLock);
@@ -88,6 +97,78 @@ UE::Tasks::FTask FMutableStreamRequest::Stream()
 				return;
 			}
 
+#if WITH_EDITOR
+			if (ModelStreamableBulkData->bIsStoredInDDC)
+			{
+				using namespace UE::DerivedData;
+
+				// Skip loading values by default
+				FCacheRecordPolicyBuilder PolicyBuilder(ECachePolicy::Default | ECachePolicy::SkipData);
+
+				TArray<FValueId> ResourceIds;
+				ResourceIds.Reserve(BlockReadInfos.Num());
+
+				// Override policy for the resources to load.
+				for (const FBlockReadInfo& Block : BlockReadInfos)
+				{
+					FValueId ResourceId = GetDerivedDataValueIdForResource(Block.DataType, Block.FileId, Block.ResourceType, Block.ResourceFlags);
+
+					if (ResourceIds.Find(ResourceId) == INDEX_NONE)
+					{
+						// Add value policy once.
+						PolicyBuilder.AddValuePolicy(ResourceId, ECachePolicy::Default);
+					}
+
+					ResourceIds.Add(ResourceId);
+				}
+
+				FCacheGetRequest Request;
+				Request.Name = ModelStreamableBulkData->FullFilePath;
+				Request.Key = ModelStreamableBulkData->DDCKey;
+				Request.Policy = PolicyBuilder.Build();
+
+				TSharedPtr<UE::DerivedData::FRequestOwner>& DDCRequest = HeapMemory->DDCReadRequest.Add_GetRef(MakeShared<FRequestOwner>(EPriority::High));
+
+				GetCache().Get(MakeArrayView(&Request, 1), *DDCRequest,
+					[HeapMemory, BlockReadInfos = MoveTemp(BlockReadInfos), CapturedResourceIds = MoveTemp(ResourceIds)](FCacheGetResponse&& Response) mutable
+					{
+						bool bSuccess = Response.Status == EStatus::Ok;
+
+						if (ensure(bSuccess))
+						{
+							const int32 NumResources = CapturedResourceIds.Num();
+							for (int32 Index = 0; Index < NumResources; ++Index)
+							{
+								const FCompressedBuffer& CompressedBuffer = Response.Record.GetValue(CapturedResourceIds[Index]).GetData();
+								if (ensure(!CompressedBuffer.IsNull()))
+								{
+									const FBlockReadInfo& Block = BlockReadInfos[Index];
+									const uint64 Size = Block.AllocatedMemoryView.Num();
+									if (Size < CompressedBuffer.GetRawSize())
+									{
+										check(CompressedBuffer.GetRawSize() >= Block.Offset + Size);
+										FSharedBuffer DecompressedBuffer = CompressedBuffer.Decompress();
+										FMemory::Memcpy(Block.AllocatedMemoryView.GetData(), reinterpret_cast<const uint8*>(DecompressedBuffer.GetData()) + Block.Offset, Size);
+									}
+									else
+									{
+										check(CompressedBuffer.TryDecompressTo(MakeMemoryView(Block.AllocatedMemoryView.GetData(), Size)));
+									}
+								}
+							}
+						}
+
+						for (UE::Tasks::FTaskEvent& CompletionEvent : HeapMemory->CompletionEvents)
+						{
+							CompletionEvent.Trigger();
+						}
+
+					});
+
+				return;
+			}
+#endif
+
 			const bool bUseFBulkData = !ModelStreamableBulkData->StreamableBulkData.IsEmpty();
 			const EAsyncIOPriorityAndFlags Priority = CVarMutableHighPriorityLoading.GetValueOnAnyThread() ? AIOP_High : AIOP_Normal;
 
@@ -203,10 +284,14 @@ void FMutableStreamRequest::Cancel()
 }
 
 
-FMutableStreamRequest::FBlockReadInfo::FBlockReadInfo(uint64 InOffset, IAsyncReadFileHandle* InFileHandle, TArrayView<uint8>& InAllocatedMemoryView, uint32 InFileId)
+FMutableStreamRequest::FBlockReadInfo::FBlockReadInfo(uint64 InOffset, IAsyncReadFileHandle* InFileHandle, TArrayView<uint8>& InAllocatedMemoryView,
+	uint32 InFileId, MutablePrivate::EStreamableDataType InDataType, uint16 InResourceType, uint16 InResourceFlags)
 	: Offset(InOffset)
 	, FileHandle(InFileHandle)
 	, AllocatedMemoryView(InAllocatedMemoryView)
 	, FileId(InFileId)
+	, ResourceType(InResourceType)
+	, ResourceFlags(InResourceFlags)
+	, DataType(InDataType)
 {
 }
diff --git a/Engine/Plugins/Mutable/Source/CustomizableObject/Private/MuCO/MutableStreamRequest.h b/Engine/Plugins/Mutable/Source/CustomizableObject/Private/MuCO/MutableStreamRequest.h
--- a/Engine/Plugins/Mutable/Source/CustomizableObject/Private/MuCO/MutableStreamRequest.h
+++ b/Engine/Plugins/Mutable/Source/CustomizableObject/Private/MuCO/MutableStreamRequest.h
@@ -7,6 +7,9 @@ struct FBlockReadInfo;
 struct FMutableStreamableBlock;
 struct FModelStreamableBulkData;
 
+#if WITH_EDITOR
+#include "DerivedDataRequestOwner.h"
+#endif
 
 constexpr int32 BulkDataFileOffset =
 #if WITH_EDITOR
@@ -27,7 +30,7 @@ public:
 	const TSharedPtr<FModelStreamableBulkData>& GetModelStreamableBulkData() const;
 
 	/** Add a block to stream. */
-	void AddBlock(const FMutableStreamableBlock& Block, MutablePrivate::EStreamableDataType DataType, TArrayView<uint8> AllocatedMemoryView);
+	void AddBlock(const FMutableStreamableBlock& Block, MutablePrivate::EStreamableDataType DataType, uint16 InResourceType, TArrayView<uint8> AllocatedMemoryView);
 
 	/** Stream the previously requested data. */
 	UE::Tasks::FTask Stream();
@@ -39,12 +42,16 @@ private:
 	struct FBlockReadInfo
 	{
 		FBlockReadInfo() = default;
-		FBlockReadInfo(uint64 InOffset, IAsyncReadFileHandle* InFileHandle, TArrayView<uint8>& InAllocatedMemoryView, uint32 InFileId);
+		FBlockReadInfo(uint64 InOffset, IAsyncReadFileHandle* InFileHandle, TArrayView<uint8>& InAllocatedMemoryView, uint32 InFileId, MutablePrivate::EStreamableDataType InDataType, uint16 InResourceType, uint16 InResourceFlags);
 
 		uint64 Offset;
 		IAsyncReadFileHandle* FileHandle;
 		TArrayView<uint8> AllocatedMemoryView;
 		uint32 FileId;
+		uint16 ResourceType;
+		uint16 ResourceFlags;
+
+		MutablePrivate::EStreamableDataType DataType;
 	};
 	
 	const TSharedPtr<FModelStreamableBulkData> ModelStreamableBulkData;
@@ -58,6 +65,10 @@ private:
 		TArray<TSharedPtr<IAsyncReadRequest>> ReadRequests;
 		TArray<TSharedPtr<IBulkDataIORequest>> BulkReadRequests;
 
+#if WITH_EDITORONLY_DATA
+		TArray<TSharedPtr<UE::DerivedData::FRequestOwner>> DDCReadRequest;
+#endif
+
 		bool bIsCancelled = false;
 		FCriticalSection ReadRequestLock;
 	};
diff --git a/Engine/Plugins/Mutable/Source/CustomizableObject/Private/MuCO/UnrealMutableModelDiskStreamer.cpp b/Engine/Plugins/Mutable/Source/CustomizableObject/Private/MuCO/UnrealMutableModelDiskStreamer.cpp
--- a/Engine/Plugins/Mutable/Source/CustomizableObject/Private/MuCO/UnrealMutableModelDiskStreamer.cpp
+++ b/Engine/Plugins/Mutable/Source/CustomizableObject/Private/MuCO/UnrealMutableModelDiskStreamer.cpp
@@ -88,8 +88,6 @@ bool FUnrealMutableModelBulkReader::PrepareStreamingForObject(UCustomizableObjec
 		!=
 		nullptr;
 
-	const UModelResources& ModelResources = CustomizableObject->GetPrivate()->GetModelResourcesChecked();
-
 	if (!bAlreadyStreaming)
 	{
 		FObjectData NewData;
@@ -111,10 +109,6 @@ bool FUnrealMutableModelBulkReader::PrepareStreamingForObject(UCustomizableObjec
 #if WITH_EDITOR
 		FString FullFileName = CustomizableObject->GetPrivate()->GetCompiledDataFileName(nullptr, true);
 		NewData.BulkFilePrefix = *FullFileName + UCustomizableObjectPrivate::GetDataTypeExtension(MutablePrivate::EStreamableDataType::Model);
-
-		NewData.bIsStoredInDDC = ModelResources.bIsStoredInDDC;
-		NewData.DDCKey = ModelResources.DDCKey;
-		NewData.DDCPolicy = ModelResources.DDCDefaultPolicy;
 #else
 		if (ModelStreamableBulkData->StreamableBulkData.IsEmpty())
 		{
@@ -267,7 +261,7 @@ bool FUnrealMutableModelBulkReader::DoesBlockExist(const mu::FModel* Model, uint
 	else
 	{
 #if WITH_EDITOR
-		if (ObjectData->bIsStoredInDDC)
+		if (PinnedModelStreamableBulkData->bIsStoredInDDC)
 		{
 			MUTABLE_CPUPROFILER_SCOPE(FUnrealMutableModelBulkReader::DoesBlockExist_DDC)
 			return true;
@@ -284,7 +278,7 @@ bool FUnrealMutableModelBulkReader::DoesBlockExist(const mu::FModel* Model, uint
 }
 
 
-mu::FModelReader::FOperationID FUnrealMutableModelBulkReader::BeginReadBlock(const mu::FModel* Model, uint32 Key, void* Buffer, uint64 Size, TFunction<void(bool bSuccess)>* CompletionCallback)
+mu::FModelReader::FOperationID FUnrealMutableModelBulkReader::BeginReadBlock(const mu::FModel* Model, uint32 Key, void* Buffer, uint64 Size, mu::EDataType ResourceType, TFunction<void(bool bSuccess)>* CompletionCallback)
 {
 	MUTABLE_CPUPROFILER_SCOPE(FUnrealMutableModelBulkStreamer::OpenReadFile);
 
@@ -359,28 +353,23 @@ mu::FModelReader::FOperationID FUnrealMutableModelBulkReader::BeginReadBlock(con
 	else
 	{
 #if WITH_EDITOR
-		if (ObjectData->bIsStoredInDDC)
+		if (PinnedModelStreamableBulkData->bIsStoredInDDC)
 		{
 			MUTABLE_CPUPROFILER_SCOPE(FUnrealMutableModelBulkStreamer::OpenReadFileDDC);
 
 			using namespace UE::DerivedData;
 
-			FValueId::ByteArray ValueIdBytes = {};
-			check(sizeof(FValueId::ByteArray) >= 4 + sizeof(Block->FileId));
-
-			MutablePrivate::EStreamableDataType DataType = MutablePrivate::EStreamableDataType::Model;
-			FMemory::Memcpy(&ValueIdBytes, &DataType, sizeof(DataType));
-			FMemory::Memcpy(&ValueIdBytes[4], &Block->FileId, sizeof(Block->FileId));
-			FValueId ResourceId = FValueId(ValueIdBytes);
 
 			// Skip loading values by default
 			FCacheRecordPolicyBuilder PolicyBuilder(ECachePolicy::Default | ECachePolicy::SkipData);
 
-			PolicyBuilder.AddValuePolicy(ResourceId, ECachePolicy::Default); // Override policy for the Model. Always load the compiled model data.
+			// Override policy for the resource to load.
+			FValueId ResourceId = GetDerivedDataValueIdForResource(MutablePrivate::EStreamableDataType::Model, Block->FileId, (uint16)ResourceType, Block->Flags);
+			PolicyBuilder.AddValuePolicy(ResourceId, ECachePolicy::Default);
 
 			FCacheGetRequest Request;
 			Request.Name = ObjectData->BulkFilePrefix;
-			Request.Key = ObjectData->DDCKey;
+			Request.Key = PinnedModelStreamableBulkData->DDCKey;
 			Request.Policy = PolicyBuilder.Build();
 
 			FReadRequest ReadRequest;
@@ -390,24 +379,23 @@ mu::FModelReader::FOperationID FUnrealMutableModelBulkReader::BeginReadBlock(con
 				[CompletionCallbackCapture = *CompletionCallback, Block, ResourceId, Buffer, Size](FCacheGetResponse&& Response)
 				{
 					bool bSuccess = Response.Status == EStatus::Ok;
-					if (bSuccess)
+					if (ensure(bSuccess))
 					{
 						const FCompressedBuffer& CompressedBuffer = Response.Record.GetValue(ResourceId).GetData();
-						if (Size < CompressedBuffer.GetRawSize())
-						{
-							check(CompressedBuffer.GetRawSize() >= Block->Offset + Size);
-							FSharedBuffer DecompressedBuffer = CompressedBuffer.Decompress();
-							FMemory::Memcpy(Buffer, reinterpret_cast<const uint8*>(DecompressedBuffer.GetData()) + Block->Offset, Size);
-						}
-						else
+						if (ensure(!CompressedBuffer.IsNull()))
 						{
-							check(CompressedBuffer.TryDecompressTo(MakeMemoryView(Buffer, Size)));
+							if (Size < CompressedBuffer.GetRawSize())
+							{
+								check(CompressedBuffer.GetRawSize() >= Block->Offset + Size);
+								FSharedBuffer DecompressedBuffer = CompressedBuffer.Decompress();
+								FMemory::Memcpy(Buffer, reinterpret_cast<const uint8*>(DecompressedBuffer.GetData()) + Block->Offset, Size);
+							}
+							else
+							{
+								check(CompressedBuffer.TryDecompressTo(MakeMemoryView(Buffer, Size)));
+							}
 						}
 					}
-					else
-					{
-						check(false);
-					}
 
 					CompletionCallbackCapture(bSuccess);
 				});
diff --git a/Engine/Plugins/Mutable/Source/CustomizableObject/Public/MuCO/CustomizableObject.h b/Engine/Plugins/Mutable/Source/CustomizableObject/Public/MuCO/CustomizableObject.h
--- a/Engine/Plugins/Mutable/Source/CustomizableObject/Public/MuCO/CustomizableObject.h
+++ b/Engine/Plugins/Mutable/Source/CustomizableObject/Public/MuCO/CustomizableObject.h
@@ -37,8 +37,6 @@ DECLARE_MULTICAST_DELEGATE(FPostCompileDelegate)
 
 CUSTOMIZABLEOBJECT_API extern TAutoConsoleVariable<bool> CVarMutableUseBulkData;
 
-CUSTOMIZABLEOBJECT_API extern TAutoConsoleVariable<int32> CVarMutableDerivedDataCacheUsage;
-
 CUSTOMIZABLEOBJECT_API extern TAutoConsoleVariable<bool> CVarMutableAsyncCook;
 
 
diff --git a/Engine/Plugins/Mutable/Source/CustomizableObject/Public/MuCO/CustomizableObjectSystem.h b/Engine/Plugins/Mutable/Source/CustomizableObject/Public/MuCO/CustomizableObjectSystem.h
--- a/Engine/Plugins/Mutable/Source/CustomizableObject/Public/MuCO/CustomizableObjectSystem.h
+++ b/Engine/Plugins/Mutable/Source/CustomizableObject/Public/MuCO/CustomizableObjectSystem.h
@@ -28,6 +28,7 @@ class UCustomizableObjectSystemPrivate; // This is used to hide Mutable SDK memb
 class FUpdateContextPrivate;
 struct FFrame;
 struct FGuid;
+struct FEditorCompileSettings;
 
 
 extern TAutoConsoleVariable<bool> CVarClearWorkingMemoryOnUpdateEnd;
@@ -81,22 +82,6 @@ enum class ECustomizableObjectNumBoneInfluences : uint8
 };
 
 
-#if WITH_EDITOR
-
-// Struct used to keep a copy of the EditorSettings needed to compile Customizable Objects.
-struct FEditorCompileSettings 
-{
-	// General case
-	bool bIsMutableEnabled = true;
-
-	// Auto Compile 
-	bool bEnableAutomaticCompilation = true;
-	bool bCompileObjectsSynchronously = true;
-	bool bCompileRootObjectsOnStartPIE = false;
-};
-
-#endif
-
 namespace EMutableProfileMetric
 {
 	typedef uint8 Type;
@@ -402,12 +387,6 @@ public:
 	 * @return True if mutable will try to reuse them, false otherwise.
 	 */
 	static UE_API bool ShouldReuseTexturesBetweenInstances();
-
-
-#if WITH_EDITOR
-	// Copy of the Mutable Editor Settings tied to CO compilation. They are updated whenever changed
-	FEditorCompileSettings EditorSettings;
-#endif
 	
 private:
 	UPROPERTY(Transient)
diff --git a/Engine/Plugins/Mutable/Source/CustomizableObjectEditor/Private/MuCOE/CustomizableObjectCompileRunnable.cpp b/Engine/Plugins/Mutable/Source/CustomizableObjectEditor/Private/MuCOE/CustomizableObjectCompileRunnable.cpp
--- a/Engine/Plugins/Mutable/Source/CustomizableObjectEditor/Private/MuCOE/CustomizableObjectCompileRunnable.cpp
+++ b/Engine/Plugins/Mutable/Source/CustomizableObjectEditor/Private/MuCOE/CustomizableObjectCompileRunnable.cpp
@@ -380,7 +380,7 @@ uint32 FCustomizableObjectSaveDDRunnable::Run()
 		bool bStoredSuccessfully = false;
 
 		// TODO UE-222775: Allow using DDC in editor builds, not just for cooking.
-		if (Options.bIsCooking && Options.bStoreCompiledDataInDDC && !DDCKey.Hash.IsZero())
+		if (Options.bStoreCompiledDataInDDC && !DDCKey.Hash.IsZero())
 		{
 			StoreCachedPlatformDataInDDC(bStoredSuccessfully);
 		}
@@ -460,14 +460,15 @@ void FCustomizableObjectSaveDDRunnable::StoreCachedPlatformDataInDDC(bool& bStor
 	FCacheRecordBuilder RecordBuilder(DDCKey);
 
 	// Store streamable resources info as FValues
+	FModelStreamableBulkData ModelStreamablesDDC = *PlatformData->ModelStreamableBulkData.Get();
 	{
 		MUTABLE_CPUPROFILER_SCOPE(SerializeModelStreamables);
 
 		// ModelStreamable  will be modified for the DDC record. Modify a copy
-		FModelStreamableBulkData ModelStreamablesDDC = *PlatformData->ModelStreamableBulkData.Get();
 
 		// Generate list of files and update streamable blocks ids and offsets
-		MutablePrivate::GenerateBulkDataFilesListWithFileLimit(PlatformData->Model, ModelStreamablesDDC, MAX_int16, BulkDataFilesDDC);
+		constexpr int32 MaxDDCFiles = 1 << 13;
+		MutablePrivate::GenerateBulkDataFilesListWithFileLimit(PlatformData->Model, ModelStreamablesDDC, MaxDDCFiles, BulkDataFilesDDC);
 
 		TArray64<uint8> ModelStreamablesBytesDDC;
 		FMemoryWriter64 MemoryWriterDDC(ModelStreamablesBytesDDC);
@@ -481,24 +482,9 @@ void FCustomizableObjectSaveDDRunnable::StoreCachedPlatformDataInDDC(bool& bStor
 	{
 		MUTABLE_CPUPROFILER_SCOPE(SerializeBulkDataForDDC);
 	
-		{
-			MutablePrivate::FFile File;
-			check(sizeof(FValueId::ByteArray) >= sizeof(File.DataType) + sizeof(File.Id) + sizeof(File.ResourceType) + sizeof(File.Flags));
-		}
-
-		FValueId::ByteArray ValueIdBytes = {};
-		const auto WriteBulkDataDDC = [&RecordBuilder, &ValueIdBytes](MutablePrivate::FFile& File, TArray64<uint8>& FileBulkData, uint32 FileIndex)
+		const auto WriteBulkDataDDC = [&RecordBuilder](MutablePrivate::FFile& File, TArray64<uint8>& FileBulkData, uint32 FileIndex)
 			{
-				int8 ValueIdOffset = 0;
-				FMemory::Memcpy(&ValueIdBytes[ValueIdOffset], &File.DataType, sizeof(File.DataType));
-				ValueIdOffset += sizeof(File.DataType);
-				FMemory::Memcpy(&ValueIdBytes[ValueIdOffset], &File.Id, sizeof(File.Id));
-				ValueIdOffset += sizeof(File.Id);
-				FMemory::Memcpy(&ValueIdBytes[ValueIdOffset], &File.ResourceType, sizeof(File.ResourceType));
-				ValueIdOffset += sizeof(File.ResourceType);
-				FMemory::Memcpy(&ValueIdBytes[ValueIdOffset], &File.Flags, sizeof(File.Flags));
-
-				const FValueId ValueId(ValueIdBytes);
+				const FValueId ValueId = GetDerivedDataValueIdForResource(File.DataType, FileIndex, File.ResourceType, File.Flags);
 				const FValue Value = FValue::Compress(FSharedBuffer::MakeView(FileBulkData.GetData(), FileBulkData.Num()));
 
 				RecordBuilder.AddValue(ValueId, Value);
@@ -554,6 +540,18 @@ void FCustomizableObjectSaveDDRunnable::StoreCachedPlatformDataInDDC(bool& bStor
 			});
 
 		RequestOwner.Wait();
+
+		if (bStoredSuccessfully)
+		{
+			if (!Options.bIsCooking)
+			{
+				*PlatformData->ModelStreamableBulkData.Get() = ModelStreamablesDDC;
+			}
+
+			PlatformData->ModelStreamableBulkData->bIsStoredInDDC = true;
+			PlatformData->ModelStreamableBulkData->DDCKey = DDCKey;
+			PlatformData->ModelStreamableBulkData->DDCDefaultPolicy = DefaultDDCPolicy;
+		}
 	}
 }
 
diff --git a/Engine/Plugins/Mutable/Source/CustomizableObjectEditor/Private/MuCOE/CustomizableObjectCompiler.cpp b/Engine/Plugins/Mutable/Source/CustomizableObjectEditor/Private/MuCOE/CustomizableObjectCompiler.cpp
--- a/Engine/Plugins/Mutable/Source/CustomizableObjectEditor/Private/MuCOE/CustomizableObjectCompiler.cpp
+++ b/Engine/Plugins/Mutable/Source/CustomizableObjectEditor/Private/MuCOE/CustomizableObjectCompiler.cpp
@@ -1284,14 +1284,19 @@ bool FCustomizableObjectCompiler::TryLoadCompiledDataFromDDC(UCustomizableObject
 	CurrentRequest->BuildDerivedDataCacheKey();
 
 	FCacheKey CacheKey = CurrentRequest->GetDerivedDataCacheKey();
-	CustomizableObject.GetPrivate()->LoadCompiledDataFromDDC(CurrentOptions, DefaultPolicy, &CacheKey);
-	
-	if (CurrentOptions.bIsCooking)
+
+	if (CustomizableObject.GetPrivate()->LoadCompiledDataFromDDC(CurrentOptions, DefaultPolicy, &CacheKey, *PlatformData.Get()))
 	{
-		return CustomizableObject.GetPrivate()->CachedPlatformsData.Find(CurrentOptions.TargetPlatform->PlatformName()) != nullptr;
+		if (CurrentOptions.bIsCooking)
+		{
+			MutablePrivate::FMutableCachedPlatformData& CachedPlatformData = CurrentObject->GetPrivate()->CachedPlatformsData.Add(CurrentOptions.TargetPlatform->PlatformName(), {});
+			CachedPlatformData = MoveTemp(*PlatformData.Get());
+		}
+
+		return true;
 	}
 
-	return CustomizableObject.IsCompiled();
+	return false;
 }
 
 
@@ -1384,18 +1389,11 @@ void FCustomizableObjectCompiler::FinishCompilationTask()
 		TMap<uint32, FMutableStreamableBlock>& ModelStreamables = PlatformData->ModelStreamableBulkData->ModelStreamables;
 		ModelStreamables.Empty(NumStreamingFiles);
 
-		// TODO: Temp. Remove after unifying generated output files code between editor an package. UE-222777
-		const bool bRequiresCookedData = CurrentOptions.TargetPlatform->RequiresCookedData();
-
 		uint64 Offset = 0;
 		for (int32 FileIndex = 0; FileIndex < NumStreamingFiles; ++FileIndex)
 		{
 			const uint32 ResourceSize = Model->GetRomSize(FileIndex);
-			EMutableFileFlags Flags = EMutableFileFlags::None;
-			if (bRequiresCookedData)
-			{
-				Flags = Model->IsRomHighRes(FileIndex) ? EMutableFileFlags::HighRes : EMutableFileFlags::None;
-			}
+			EMutableFileFlags Flags = Model->IsRomHighRes(FileIndex) ? EMutableFileFlags::HighRes : EMutableFileFlags::None;
 			ModelStreamables.Add(FileIndex, FMutableStreamableBlock{ .FileId = 0, .Flags = uint16(Flags), .Offset = Offset });
 			Offset += ResourceSize;
 		}
diff --git a/Engine/Plugins/Mutable/Source/CustomizableObjectEditor/Private/MuCOE/CustomizableObjectEditor.cpp b/Engine/Plugins/Mutable/Source/CustomizableObjectEditor/Private/MuCOE/CustomizableObjectEditor.cpp
--- a/Engine/Plugins/Mutable/Source/CustomizableObjectEditor/Private/MuCOE/CustomizableObjectEditor.cpp
+++ b/Engine/Plugins/Mutable/Source/CustomizableObjectEditor/Private/MuCOE/CustomizableObjectEditor.cpp
@@ -1760,7 +1760,8 @@ void FCustomizableObjectEditor::OnCustomizableObjectStatusChanged(FCustomizableO
 			TSharedRef<FCompilationRequest> CompileRequest = MakeShared<FCompilationRequest>(*CustomizableObject);
 			CompileRequest->bSkipIfOutOfDate = true;
 			CompileRequest->bSilentCompilation = false;
-			
+			CompileRequest->SetDerivedDataCachePolicy(GetDerivedDataCachePolicyForEditor());
+
 			ICustomizableObjectEditorModulePrivate::GetChecked().EnqueueCompileRequest(CompileRequest);
 		}
 	}
diff --git a/Engine/Plugins/Mutable/Source/CustomizableObjectEditor/Private/MuCOE/CustomizableObjectEditorModule.cpp b/Engine/Plugins/Mutable/Source/CustomizableObjectEditor/Private/MuCOE/CustomizableObjectEditorModule.cpp
--- a/Engine/Plugins/Mutable/Source/CustomizableObjectEditor/Private/MuCOE/CustomizableObjectEditorModule.cpp
+++ b/Engine/Plugins/Mutable/Source/CustomizableObjectEditor/Private/MuCOE/CustomizableObjectEditorModule.cpp
@@ -126,6 +126,12 @@ const FName CustomizableObjectMacroLibraryEditorAppIdentifier = FName(TEXT("Cust
 /** Max timespan in days before a Saved/MutableStreamedDataEditor file is deleted. */
 constexpr int32 MaxAccessTimespan = 30;
 
+TAutoConsoleVariable<bool> CVarMutableDerivedDataCacheUsage(
+	TEXT("mutable.DerivedDataCacheUsage"),
+	true,
+	TEXT("Derived data cache access for compiled data.")
+	TEXT("true - Allows DDC access according to the DDC policies specified in the plugin settings."),
+	ECVF_Default);
 
 IMPLEMENT_MODULE( FCustomizableObjectEditorModule, CustomizableObjectEditor );
 
@@ -206,6 +212,29 @@ TMap<FString, FString> GetCompileOnlySelectedParameters(const UCustomizableObjec
 }
 
 
+UE::DerivedData::ECachePolicy ConvertDerivedDataCachePolicy(ECustomizableObjectDDCPolicy InPolicy)
+{
+	switch (InPolicy)
+	{
+	case ECustomizableObjectDDCPolicy::Default: return UE::DerivedData::ECachePolicy::Default;
+	case ECustomizableObjectDDCPolicy::Local: return UE::DerivedData::ECachePolicy::Local;
+	default: return UE::DerivedData::ECachePolicy::None;
+	}
+}
+
+
+UE::DerivedData::ECachePolicy GetDerivedDataCachePolicyForEditor()
+{
+	if (UCustomizableObjectSystem::IsCreated() && CVarMutableDerivedDataCacheUsage.GetValueOnAnyThread())
+	{
+		UCustomizableObjectSystem* System = UCustomizableObjectSystem::GetInstance();
+		return ConvertDerivedDataCachePolicy(System->GetPrivate()->EditorSettings.EditorDerivedDataCachePolicy);
+	}
+
+	return UE::DerivedData::ECachePolicy::None;
+}
+
+
 void FCustomizableObjectEditorModule::StartupModule()
 {
 	// Delete unused local compiled data
@@ -338,7 +367,9 @@ bool FCustomizableObjectEditorModule::HandleSettingsSaved()
 		CompileSettings.bEnableAutomaticCompilation = CustomizableObjectSettings->bEnableAutomaticCompilation;
 		CompileSettings.bCompileObjectsSynchronously = CustomizableObjectSettings->bCompileObjectsSynchronously;
 		CompileSettings.bCompileRootObjectsOnStartPIE = CustomizableObjectSettings->bCompileRootObjectsOnStartPIE;
-		
+		CompileSettings.EditorDerivedDataCachePolicy = CustomizableObjectSettings->EditorDerivedDataCachePolicy;
+		CompileSettings.CookDerivedDataCachePolicy = CustomizableObjectSettings->CookDerivedDataCachePolicy;
+
 		UCustomizableObjectSystem::GetInstance()->EditorSettingsChanged(CompileSettings);
 	}
 
@@ -373,7 +404,9 @@ void FCustomizableObjectEditorModule::RegisterSettings()
 				CompileSettings.bEnableAutomaticCompilation = CustomizableObjectSettings->bEnableAutomaticCompilation;
 				CompileSettings.bCompileObjectsSynchronously = CustomizableObjectSettings->bCompileObjectsSynchronously;
 				CompileSettings.bCompileRootObjectsOnStartPIE = CustomizableObjectSettings->bCompileRootObjectsOnStartPIE;
-				
+				CompileSettings.EditorDerivedDataCachePolicy = CustomizableObjectSettings->EditorDerivedDataCachePolicy;
+				CompileSettings.CookDerivedDataCachePolicy = CustomizableObjectSettings->CookDerivedDataCachePolicy;
+
 				UCustomizableObjectSystem::GetInstance()->EditorSettingsChanged(CompileSettings);
 			}
 		}
@@ -1094,7 +1127,7 @@ void FCustomizableObjectEditorModule::CompileCustomizableObject(UCustomizableObj
 		CompileRequest->Options.bGatherReferences = Params->bGatherReferneces;
 		CompileRequest->Callback = Params->Callback;
 		CompileRequest->CallbackNative = Params->CallbackNative;
-		
+
 		// Override the optimization level provided in the request with the one set in the CO
 		if (Params->OptimizationLevel == ECustomizableObjectOptimizationLevel::FromCustomizableObject)
 		{
@@ -1111,6 +1144,8 @@ void FCustomizableObjectEditorModule::CompileCustomizableObject(UCustomizableObj
 		}
 	}
 
+	
+	CompileRequest->SetDerivedDataCachePolicy(GetDerivedDataCachePolicyForEditor());
 	CompileRequest->bSilentCompilation = bSilent;
 	
 	EnqueueCompileRequest(CompileRequest, bForce);
@@ -1159,17 +1194,11 @@ void FCustomizableObjectEditorModule::BeginCacheForCookedPlatformData(UCustomiza
 	CompileRequest->Options.bUseBulkData = CVarMutableUseBulkData.GetValueOnAnyThread();
 	CompileRequest->Options.TargetPlatform = TargetPlatform;
 
-	const int32 DDCUsage = CVarMutableDerivedDataCacheUsage.GetValueOnAnyThread();
-	UE::DerivedData::ECachePolicy DefaultCachePolicy = UE::DerivedData::ECachePolicy::None;
-	if (DDCUsage == 1)
-	{
-		DefaultCachePolicy = UE::DerivedData::ECachePolicy::Local;
-	}
-	else if (DDCUsage == 2)
+	if (CVarMutableDerivedDataCacheUsage.GetValueOnAnyThread())
 	{
-		DefaultCachePolicy = UE::DerivedData::ECachePolicy::Default;
+		UCustomizableObjectSystem* System = UCustomizableObjectSystem::GetInstanceChecked();
+		CompileRequest->SetDerivedDataCachePolicy(ConvertDerivedDataCachePolicy(System->GetPrivate()->EditorSettings.CookDerivedDataCachePolicy));
 	}
-	CompileRequest->SetDerivedDataCachePolicy(DefaultCachePolicy);
 
 	ObjectCompileRequests.Add(CompileRequest);
 
@@ -1231,7 +1260,7 @@ void FCustomizableObjectEditorModule::OnPreBeginPIE(const bool bIsSimulatingInEd
 	}
 
 	UCustomizableObjectSystem* System = UCustomizableObjectSystem::GetInstanceChecked();
-	if (!System->EditorSettings.bCompileRootObjectsOnStartPIE)
+	if (!System->GetPrivate()->EditorSettings.bCompileRootObjectsOnStartPIE)
 	{
 		return;
 	}
@@ -1262,6 +1291,7 @@ void FCustomizableObjectEditorModule::OnPreBeginPIE(const bool bIsSimulatingInEd
 
 		// Add uncompiled objects to the objects to cook list
 		TSharedRef<FCompilationRequest> NewRequest = MakeShared<FCompilationRequest>(*Object);
+		NewRequest->SetDerivedDataCachePolicy(GetDerivedDataCachePolicyForEditor());
 		NewRequest->bSilentCompilation = true;
 
 		Requests.Add(NewRequest);
diff --git a/Engine/Plugins/Mutable/Source/CustomizableObjectEditor/Private/MuCOE/CustomizableObjectEditorModule.h b/Engine/Plugins/Mutable/Source/CustomizableObjectEditor/Private/MuCOE/CustomizableObjectEditorModule.h
--- a/Engine/Plugins/Mutable/Source/CustomizableObjectEditor/Private/MuCOE/CustomizableObjectEditorModule.h
+++ b/Engine/Plugins/Mutable/Source/CustomizableObjectEditor/Private/MuCOE/CustomizableObjectEditorModule.h
@@ -19,6 +19,11 @@ class ICustomizableObjectInstanceEditor;
 class FBakeOperationCompletedDelegate;
 class FName;
 
+namespace UE::DerivedData
+{
+	enum class ECachePolicy : uint32;
+}
+
 /** Get a list of packages that are used by the compilation but are not directly referenced.
   * List includes:
   * - Child UCustomizableObjects: Have inverted references.
@@ -30,6 +35,7 @@ int32 ConvertOptimizationLevel(ECustomizableObjectOptimizationLevel Optimization
 
 TMap<FString, FString> GetCompileOnlySelectedParameters(const UCustomizableObjectInstance& Instance);
 
+UE::DerivedData::ECachePolicy GetDerivedDataCachePolicyForEditor();
 
 /**
  * StaticMesh editor module
diff --git a/Engine/Plugins/Mutable/Source/CustomizableObjectEditor/Private/MuCOE/CustomizableObjectEditorSettings.cpp b/Engine/Plugins/Mutable/Source/CustomizableObjectEditor/Private/MuCOE/CustomizableObjectEditorSettings.cpp
--- a/Engine/Plugins/Mutable/Source/CustomizableObjectEditor/Private/MuCOE/CustomizableObjectEditorSettings.cpp
+++ b/Engine/Plugins/Mutable/Source/CustomizableObjectEditor/Private/MuCOE/CustomizableObjectEditorSettings.cpp
@@ -2,8 +2,11 @@
 
 #include "MuCOE/CustomizableObjectEditorSettings.h"
 
+#include "MuCO/CustomizableObjectSystemPrivate.h"
 
 UCustomizableObjectEditorSettings::UCustomizableObjectEditorSettings()
+	: EditorDerivedDataCachePolicy(ECustomizableObjectDDCPolicy::Default)
+	, CookDerivedDataCachePolicy(ECustomizableObjectDDCPolicy::Default)
 {
 
 }
diff --git a/Engine/Plugins/Mutable/Source/CustomizableObjectEditor/Private/MuCOE/CustomizableObjectEditorSettings.h b/Engine/Plugins/Mutable/Source/CustomizableObjectEditor/Private/MuCOE/CustomizableObjectEditorSettings.h
--- a/Engine/Plugins/Mutable/Source/CustomizableObjectEditor/Private/MuCOE/CustomizableObjectEditorSettings.h
+++ b/Engine/Plugins/Mutable/Source/CustomizableObjectEditor/Private/MuCOE/CustomizableObjectEditorSettings.h
@@ -5,6 +5,8 @@
 
 #include "CustomizableObjectEditorSettings.generated.h"
 
+enum class ECustomizableObjectDDCPolicy : uint8;
+
 UCLASS(config = Engine)
 class UCustomizableObjectEditorSettings : public UObject
 {
@@ -30,4 +32,9 @@ public:
 	UPROPERTY(config, EditAnywhere, Category = AutomaticCompilation)
 	bool bCompileRootObjectsOnStartPIE = false;
 
+	UPROPERTY(config, EditAnywhere, Category = DerivedDataCache, DisplayName = "DDC policy for editor compilations")
+	ECustomizableObjectDDCPolicy EditorDerivedDataCachePolicy;
+	
+	UPROPERTY(config, EditAnywhere, Category = DerivedDataCache, DisplayName = "DDC policy for cook compilations")
+	ECustomizableObjectDDCPolicy CookDerivedDataCachePolicy;
 };
diff --git a/Engine/Plugins/Mutable/Source/MutableRuntime/Private/MuR/CodeRunnerTasked.cpp b/Engine/Plugins/Mutable/Source/MutableRuntime/Private/MuR/CodeRunnerTasked.cpp
--- a/Engine/Plugins/Mutable/Source/MutableRuntime/Private/MuR/CodeRunnerTasked.cpp
+++ b/Engine/Plugins/Mutable/Source/MutableRuntime/Private/MuR/CodeRunnerTasked.cpp
@@ -2096,7 +2096,7 @@ namespace mu
 
 		
 			const uint32 RomId = RomIndex;
-			RomLoadOp.m_streamID = Runner->System->StreamInterface->BeginReadBlock(Runner->Model.Get(), RomId, RomLoadOp.m_streamBuffer.GetData(), RomSize, &Callback);
+			RomLoadOp.m_streamID = Runner->System->StreamInterface->BeginReadBlock(Runner->Model.Get(), RomId, RomLoadOp.m_streamBuffer.GetData(), RomSize, EDataType::Mesh, &Callback);
 			if (RomLoadOp.m_streamID < 0)
 			{
 				bOutFailed = true;
@@ -2266,7 +2266,7 @@ namespace mu
 			};
 			
 			const uint32 RomId = RomIndex;
-			RomLoadOp.m_streamID = Runner->System->StreamInterface->BeginReadBlock(Runner->Model.Get(), RomId, RomLoadOp.m_streamBuffer.GetData(), RomSize, &Callback);
+			RomLoadOp.m_streamID = Runner->System->StreamInterface->BeginReadBlock(Runner->Model.Get(), RomId, RomLoadOp.m_streamBuffer.GetData(), RomSize, EDataType::Image, &Callback);
 			if (RomLoadOp.m_streamID < 0)
 			{
 				bOutFailed = true;
diff --git a/Engine/Plugins/Mutable/Source/MutableRuntime/Public/MuR/Serialisation.h b/Engine/Plugins/Mutable/Source/MutableRuntime/Public/MuR/Serialisation.h
--- a/Engine/Plugins/Mutable/Source/MutableRuntime/Public/MuR/Serialisation.h
+++ b/Engine/Plugins/Mutable/Source/MutableRuntime/Public/MuR/Serialisation.h
@@ -22,6 +22,7 @@ namespace mu
 {    
     class FImage;
 	class FModel;
+	enum class EDataType : uint8;
 
 #define MUTABLE_DEFINE_POD_SERIALISABLE(Type)									\
 	void DLLEXPORT operator<<(mu::FOutputArchive& Arch, const Type& T);			\
@@ -105,7 +106,7 @@ namespace mu
 		//! \return a previously unused identifier, now used for this operation, that can be used in
 		//! calls to the other methods of this interface. If the return value is negative it indicates
 		//! an unrecoverable error.
-		virtual FOperationID BeginReadBlock(const FModel*, uint32 BlockKey, void* Buffer, uint64 Size, TFunction<void(bool bSuccess)>* CompletionCallback = nullptr) = 0;
+		virtual FOperationID BeginReadBlock(const FModel*, uint32 BlockKey, void* Buffer, uint64 Size, EDataType DataType, TFunction<void(bool bSuccess)>* CompletionCallback = nullptr) = 0;
 
         //! Check if a data request operation has been completed.
         //! This is a weak check than *may* return true if the given operation has completed, but
