[Mutable] Convert UCustomizableObjectNodeTable DT and struct references to soft references.

#rnx

diff --git a/Engine/Plugins/Mutable/Source/CustomizableObjectEditor/Private/MuCOE/CustomizableObjectEditorModule.cpp b/Engine/Plugins/Mutable/Source/CustomizableObjectEditor/Private/MuCOE/CustomizableObjectEditorModule.cpp
--- a/Engine/Plugins/Mutable/Source/CustomizableObjectEditor/Private/MuCOE/CustomizableObjectEditorModule.cpp
+++ b/Engine/Plugins/Mutable/Source/CustomizableObjectEditor/Private/MuCOE/CustomizableObjectEditorModule.cpp
@@ -1022,7 +1022,7 @@ TMap<FName, FGuid> FCustomizableObjectEditorModule::GetParticipatingObjects(cons
 			}
 			else
 			{
-				DataTable = TableNode->Table;
+				DataTable = MutablePrivate::LoadObject(TableNode->Table);
 			}
 
 			if (DataTable)
diff --git a/Engine/Plugins/Mutable/Source/CustomizableObjectEditor/Private/MuCOE/GenerateMutableSource/GenerateMutableSourceImage.cpp b/Engine/Plugins/Mutable/Source/CustomizableObjectEditor/Private/MuCOE/GenerateMutableSource/GenerateMutableSourceImage.cpp
--- a/Engine/Plugins/Mutable/Source/CustomizableObjectEditor/Private/MuCOE/GenerateMutableSource/GenerateMutableSourceImage.cpp
+++ b/Engine/Plugins/Mutable/Source/CustomizableObjectEditor/Private/MuCOE/GenerateMutableSource/GenerateMutableSourceImage.cpp
@@ -1010,7 +1010,9 @@ mu::Ptr<mu::NodeImage> GenerateMutableSourceImage(const UEdGraphPin* Pin, FMutab
 		}
 		else
 		{
-			const FString TableName = TypedNodeTable->Table ? GetNameSafe(TypedNodeTable->Table).ToLower() : GetNameSafe(TypedNodeTable->Structure).ToLower();
+			TObjectPtr<UDataTable> Table = MutablePrivate::LoadObject(TypedNodeTable->Table);
+			TObjectPtr<UScriptStruct> Structure = MutablePrivate::LoadObject(TypedNodeTable->Structure);
+			const FString TableName = Table ? GetNameSafe(Table).ToLower() : GetNameSafe(Structure).ToLower();
 			const uint32 TableId = CityHash32(reinterpret_cast<const char*>(*TableName), TableName.Len() * sizeof(FString::ElementType));
 
 			// This node will add a checker texture in case of error
@@ -1053,10 +1055,10 @@ mu::Ptr<mu::NodeImage> GenerateMutableSourceImage(const UEdGraphPin* Pin, FMutab
 				if (bSuccess)
 				{
 					// Generating a new data table if not exists
-					mu::Ptr<mu::FTable> Table = nullptr;
-					Table = GenerateMutableSourceTable(DataTable, TypedNodeTable, GenerationContext);
+					mu::Ptr<mu::FTable> GeneratedTable = nullptr;
+					GeneratedTable = GenerateMutableSourceTable(DataTable, TypedNodeTable, GenerationContext);
 
-					if (Table)
+					if (GeneratedTable)
 					{
 						mu::Ptr<mu::NodeImageTable> ImageTableNode = new mu::NodeImageTable();
 
@@ -1068,7 +1070,7 @@ mu::Ptr<mu::NodeImage> GenerateMutableSourceImage(const UEdGraphPin* Pin, FMutab
 						else
 						{
 							// Checking if this pin texture has been used in another table node but with a different texture mode
-							const int32 ImageDataIndex = GenerationContext.GeneratedTableImages.Find({ ColumnName, Pin->PinType.PinCategory, Table });
+							const int32 ImageDataIndex = GenerationContext.GeneratedTableImages.Find({ ColumnName, Pin->PinType.PinCategory, GeneratedTable });
 
 							if (ImageDataIndex != INDEX_NONE )
 							{
@@ -1084,7 +1086,7 @@ mu::Ptr<mu::NodeImage> GenerateMutableSourceImage(const UEdGraphPin* Pin, FMutab
 							}
 							else
 							{
-								GenerationContext.GeneratedTableImages.Add({ColumnName, Pin->PinType.PinCategory, Table, TypedNodeTable});
+								GenerationContext.GeneratedTableImages.Add({ColumnName, Pin->PinType.PinCategory, GeneratedTable, TypedNodeTable});
 							}
 
 							// Encoding the texture mode of the type to allow different texture modes from the same pin
@@ -1095,11 +1097,11 @@ mu::Ptr<mu::NodeImage> GenerateMutableSourceImage(const UEdGraphPin* Pin, FMutab
 						}
 
 						// Generating a new Texture column if not exists
-						if (Table->FindColumn(ColumnName) == INDEX_NONE)
+						if (GeneratedTable->FindColumn(ColumnName) == INDEX_NONE)
 						{
 							int32 Dummy = -1; // TODO MTBL-1512
 							bool Dummy2 = false;
-							bSuccess = GenerateTableColumn(TypedNodeTable, Pin, Table, ColumnName, Property, 
+							bSuccess = GenerateTableColumn(TypedNodeTable, Pin, GeneratedTable, ColumnName, Property, 
 								FMutableSourceMeshData(), Dummy, Dummy, GenerationContext.CurrentLOD, Dummy, Dummy2, GenerationContext);
 
 							if (!bSuccess)
@@ -1113,7 +1115,7 @@ mu::Ptr<mu::NodeImage> GenerateMutableSourceImage(const UEdGraphPin* Pin, FMutab
 						{
 							Result = ImageTableNode;
 
-							ImageTableNode->Table = Table;
+							ImageTableNode->Table = GeneratedTable;
 							ImageTableNode->ColumnName = ColumnName;
 							ImageTableNode->ParameterName = TypedNodeTable->ParameterName;
 							ImageTableNode->bNoneOption = TypedNodeTable->bAddNoneOption;
diff --git a/Engine/Plugins/Mutable/Source/CustomizableObjectEditor/Private/MuCOE/GenerateMutableSource/GenerateMutableSourceMesh.cpp b/Engine/Plugins/Mutable/Source/CustomizableObjectEditor/Private/MuCOE/GenerateMutableSource/GenerateMutableSourceMesh.cpp
--- a/Engine/Plugins/Mutable/Source/CustomizableObjectEditor/Private/MuCOE/GenerateMutableSource/GenerateMutableSourceMesh.cpp
+++ b/Engine/Plugins/Mutable/Source/CustomizableObjectEditor/Private/MuCOE/GenerateMutableSource/GenerateMutableSourceMesh.cpp
@@ -3497,7 +3497,9 @@ mu::NodeMeshPtr GenerateMorphMesh(const UEdGraphPin* Pin,
 	
 	if (const UCustomizableObjectNodeTable* TypedNodeTable = Cast<UCustomizableObjectNodeTable>(Pin->GetOwningNode()))
 	{
-		const FString TableName = TypedNodeTable->Table ? GetNameSafe(TypedNodeTable->Table).ToLower() : GetNameSafe(TypedNodeTable->Structure).ToLower();
+		TObjectPtr<UDataTable> Table = MutablePrivate::LoadObject(TypedNodeTable->Table);
+		TObjectPtr<UScriptStruct> Structure = MutablePrivate::LoadObject(TypedNodeTable->Structure);
+		const FString TableName = Table ? GetNameSafe(Table).ToLower() : GetNameSafe(Structure).ToLower();
 		const uint32 TableId = CityHash32(reinterpret_cast<const char*>(*TableName), TableName.Len() * sizeof(FString::ElementType));
 
 		UDataTable* DataTable = GetDataTable(TypedNodeTable, GenerationContext);
@@ -3510,8 +3512,8 @@ mu::NodeMeshPtr GenerateMorphMesh(const UEdGraphPin* Pin,
 		// Should exist
 		FGeneratedMutableDataTableKey MutableTableKey(DataTable->GetName(), TypedNodeTable->VersionColumn,
 			TypedNodeTable->CompilationFilterColumn, TypedNodeTable->CompilationFilters, TypedNodeTable->FilterOperationType);
-		mu::Ptr<mu::FTable> Table = GenerationContext.GeneratedTables[MutableTableKey].GeneratedTable;
-		check(Table);
+		mu::Ptr<mu::FTable> GeneratedTable = GenerationContext.GeneratedTables[MutableTableKey].GeneratedTable;
+		check(GeneratedTable);
 
 		FString ColumnName = TableColumnName + TypedNodeMorphs[MorphIndex].MorphTargetName;
 		int32 ColumnIndex = INDEX_NONE;
@@ -3521,15 +3523,15 @@ mu::NodeMeshPtr GenerateMorphMesh(const UEdGraphPin* Pin,
 			const FName RowName = RowNames[RowIndex];
 			const uint32 RowId = RowIds[RowIndex];
 
-			ColumnIndex = Table->FindColumn(ColumnName);
+			ColumnIndex = GeneratedTable->FindColumn(ColumnName);
 
 			if (ColumnIndex == INDEX_NONE)
 			{
-				ColumnIndex = Table->AddColumn(ColumnName, mu::ETableColumnType::Mesh);
+				ColumnIndex = GeneratedTable->AddColumn(ColumnName, mu::ETableColumnType::Mesh);
 			}
 
 			TSharedPtr<mu::FMesh> MorphedSourceTableMesh = BuildMorphedMutableMesh(Pin, TypedNodeMorphs[MorphIndex].MorphTargetName, GenerationContext, bOnlyConnectedLOD, RowName);
-			Table->SetCell(ColumnIndex, RowId, MorphedSourceTableMesh);
+			GeneratedTable->SetCell(ColumnIndex, RowId, MorphedSourceTableMesh);
 		}
 
 		if (ColumnIndex > INDEX_NONE)
@@ -3537,7 +3539,7 @@ mu::NodeMeshPtr GenerateMorphMesh(const UEdGraphPin* Pin,
 			bSuccess = true;
 
 			mu::Ptr<mu::NodeMeshTable> MorphedSourceMeshNodeTable = new mu::NodeMeshTable;
-			MorphedSourceMeshNodeTable->Table = Table;
+			MorphedSourceMeshNodeTable->Table = GeneratedTable;
 			MorphedSourceMeshNodeTable->ColumnName = ColumnName;
 			MorphedSourceMeshNodeTable->ParameterName = TypedNodeTable->ParameterName;
 			MorphedSourceMeshNodeTable->SetMessageContext( MorphNode );
@@ -4375,7 +4377,9 @@ mu::Ptr<mu::NodeMesh> GenerateMutableSourceMesh(const UEdGraphPin* Pin,
 
 	else if (const UCustomizableObjectNodeTable* TypedNodeTable = Cast<UCustomizableObjectNodeTable>(Node))
 	{
-		const FString TableName = TypedNodeTable->Table ? GetNameSafe(TypedNodeTable->Table).ToLower() : GetNameSafe(TypedNodeTable->Structure).ToLower();
+		TObjectPtr<UDataTable> Table = MutablePrivate::LoadObject(TypedNodeTable->Table);
+		TObjectPtr<UScriptStruct> Structure = MutablePrivate::LoadObject(TypedNodeTable->Structure);
+		const FString TableName = Table ? GetNameSafe(Table).ToLower() : GetNameSafe(Structure).ToLower();
 		const uint32 TableId = CityHash32(reinterpret_cast<const char*>(*TableName), TableName.Len() * sizeof(FString::ElementType));
 
 		mu::Ptr<mu::NodeMeshConstant> EmptyNode = new mu::NodeMeshConstant();
@@ -4415,8 +4419,8 @@ mu::Ptr<mu::NodeMesh> GenerateMutableSourceMesh(const UEdGraphPin* Pin,
 			if (bSuccess)
 			{
 				// Generating a new data table if not exists
-				mu::Ptr<mu::FTable> Table = nullptr;
-				Table = GenerateMutableSourceTable(DataTable, TypedNodeTable, GenerationContext);
+				mu::Ptr<mu::FTable> GeneratedTable = nullptr;
+				GeneratedTable = GenerateMutableSourceTable(DataTable, TypedNodeTable, GenerationContext);
 
 				if (Table)
 				{
@@ -4447,9 +4451,9 @@ mu::Ptr<mu::NodeMesh> GenerateMutableSourceMesh(const UEdGraphPin* Pin,
 					}
 					
 					// Generating a new FMesh column if not exists
-					if (Table->FindColumn(MutableColumnName) == INDEX_NONE)
+					if (GeneratedTable->FindColumn(MutableColumnName) == INDEX_NONE)
 					{
-						bSuccess = GenerateTableColumn(TypedNodeTable, Pin, Table, ColumnName, Property, BaseMeshData, LODIndexConnected, SectionIndexConnected, LODIndex, SectionIndex, bOnlyConnectedLOD, GenerationContext);
+						bSuccess = GenerateTableColumn(TypedNodeTable, Pin, GeneratedTable, ColumnName, Property, BaseMeshData, LODIndexConnected, SectionIndexConnected, LODIndex, SectionIndex, bOnlyConnectedLOD, GenerationContext);
 
 						if (!bSuccess)
 						{
@@ -4462,7 +4466,7 @@ mu::Ptr<mu::NodeMesh> GenerateMutableSourceMesh(const UEdGraphPin* Pin,
 					{
 						Result = MeshTableNode;
 
-						MeshTableNode->Table = Table;
+						MeshTableNode->Table = GeneratedTable;
 						MeshTableNode->ColumnName = MutableColumnName;
 						MeshTableNode->ParameterName = TypedNodeTable->ParameterName;
 						MeshTableNode->bNoneOption = TypedNodeTable->bAddNoneOption;
diff --git a/Engine/Plugins/Mutable/Source/CustomizableObjectEditor/Private/MuCOE/GenerateMutableSource/GenerateMutableSourceTable.cpp b/Engine/Plugins/Mutable/Source/CustomizableObjectEditor/Private/MuCOE/GenerateMutableSource/GenerateMutableSourceTable.cpp
--- a/Engine/Plugins/Mutable/Source/CustomizableObjectEditor/Private/MuCOE/GenerateMutableSource/GenerateMutableSourceTable.cpp
+++ b/Engine/Plugins/Mutable/Source/CustomizableObjectEditor/Private/MuCOE/GenerateMutableSource/GenerateMutableSourceTable.cpp
@@ -1122,7 +1122,7 @@ UDataTable* GetDataTable(const UCustomizableObjectNodeTable* TableNode, FMutable
 	}
 	else
 	{
-		OutDataTable = TableNode->Table;
+		OutDataTable = MutablePrivate::LoadObject(TableNode->Table);
 	}
 
 	return OutDataTable;
@@ -1131,14 +1131,16 @@ UDataTable* GetDataTable(const UCustomizableObjectNodeTable* TableNode, FMutable
 
 UDataTable* GenerateDataTableFromStruct(const UCustomizableObjectNodeTable* TableNode, FMutableGraphGenerationContext& GenerationContext)
 {
-	if (!TableNode->Structure)
+	TObjectPtr<UScriptStruct> Structure = MutablePrivate::LoadObject(TableNode->Structure);
+
+	if (!Structure)
 	{
 		GenerationContext.Log(LOCTEXT("EmptyStructureError", "Empty structure asset."), TableNode);
 		return nullptr;
 	}
 
 	FMutableGraphGenerationContext::FGeneratedCompositeDataTablesData DataTableData;
-	DataTableData.ParentStruct = TableNode->Structure;
+	DataTableData.ParentStruct = Structure;
 	DataTableData.FilterPaths = TableNode->FilterPaths;
 	
 	//Checking cache of generated data tables
@@ -1164,7 +1166,7 @@ UDataTable* GenerateDataTableFromStruct(const UCustomizableObjectNodeTable* Tabl
 	TArray<FAssetData> DataTableAssets = TableNode->GetParentTables();
 
 	UCompositeDataTable* CompositeDataTable = NewObject<UCompositeDataTable>();
-	CompositeDataTable->RowStruct = TableNode->Structure;
+	CompositeDataTable->RowStruct = Structure;
 
 	TArray<UDataTable*> ParentTables;
 
diff --git a/Engine/Plugins/Mutable/Source/CustomizableObjectEditor/Private/MuCOE/Nodes/CustomizableObjectNodeTable.cpp b/Engine/Plugins/Mutable/Source/CustomizableObjectEditor/Private/MuCOE/Nodes/CustomizableObjectNodeTable.cpp
--- a/Engine/Plugins/Mutable/Source/CustomizableObjectEditor/Private/MuCOE/Nodes/CustomizableObjectNodeTable.cpp
+++ b/Engine/Plugins/Mutable/Source/CustomizableObjectEditor/Private/MuCOE/Nodes/CustomizableObjectNodeTable.cpp
@@ -227,14 +227,16 @@ void UCustomizableObjectNodeTable::BackwardsCompatibleFixup(int32 CustomizableOb
 
 	if (CustomizableObjectCustomVersion == FCustomizableObjectCustomVersion::SetDisplayNamePropertyAsPinNameOfTableNodes)
 	{
-		if (Structure)
+		TObjectPtr<UScriptStruct> LoadedStructure = MutablePrivate::LoadObject(Structure);
+		if (LoadedStructure)
 		{
-			ConditionalPostLoadReference(*Structure);
+			ConditionalPostLoadReference(*LoadedStructure);
 		}
 
-		if (Table)
+		TObjectPtr<UDataTable> LoadedTable = MutablePrivate::LoadObject(Table);
+		if (LoadedTable)
 		{
-			ConditionalPostLoadReference(*Table);
+			ConditionalPostLoadReference(*LoadedTable);
 		}
 		
 		// We asume that until now we only supported UserDefinedStructs (Structure assets generated trhough the Editor)
@@ -288,14 +290,16 @@ void UCustomizableObjectNodeTable::BackwardsCompatibleFixup(int32 CustomizableOb
 	{
 		if (bDisableCheckedRows_DEPRECATED)
 		{
-			if (Structure)
+			TObjectPtr<UScriptStruct> LoadedStructure = MutablePrivate::LoadObject(Structure);
+			if (LoadedStructure)
 			{
-				ConditionalPostLoadReference(*Structure);
+				ConditionalPostLoadReference(*LoadedStructure);
 			}
 
-			if (Table)
+			TObjectPtr<UDataTable> LoadedTable = MutablePrivate::LoadObject(Table);
+			if (LoadedTable)
 			{
-				ConditionalPostLoadReference(*Table);
+				ConditionalPostLoadReference(*LoadedTable);
 			}
 
 			const UScriptStruct* TableStruct = GetTableNodeStruct();
@@ -324,10 +328,12 @@ void UCustomizableObjectNodeTable::PreEditChange(FProperty* PropertyAboutToChang
 {
 	Super::PreEditChange(PropertyAboutToChange);
 
+	TObjectPtr<UDataTable> LoadedTable = MutablePrivate::LoadObject(Table);
+
 	if (PropertyAboutToChange && PropertyAboutToChange->GetFName() == GET_MEMBER_NAME_CHECKED(UCustomizableObjectNodeTable, Table) &&
-		Table)
+		LoadedTable)
 	{
-		Table->OnDataTableChanged().Remove(OnTableChangedDelegateHandle);
+		LoadedTable->OnDataTableChanged().Remove(OnTableChangedDelegateHandle);
 	}
 }
 
@@ -347,14 +353,14 @@ FText UCustomizableObjectNodeTable::GetNodeTitle(ENodeTitleType::Type TitleType)
 	}
 	else
 	{
-		if (Table)
+		if (TObjectPtr<UDataTable> LoadedTable = MutablePrivate::LoadObject(Table))
 		{
-			Args.Add(TEXT("TableName"), FText::FromString(Table->GetName()));
+			Args.Add(TEXT("TableName"), FText::FromString(LoadedTable->GetName()));
 			return FText::Format(LOCTEXT("TableNode_Title_DataTable", "{ParamName}\n{TableName} - Data Table"), Args);
 		}
-		else if (Structure)
+		else if (TObjectPtr<UScriptStruct> LoadedStructure = MutablePrivate::LoadObject(Structure))
 		{
-			Args.Add(TEXT("StructureName"), FText::FromString(Structure->GetName()));
+			Args.Add(TEXT("StructureName"), FText::FromString(LoadedStructure->GetName()));
 			return FText::Format(LOCTEXT("TableNode_Title_ScriptedStruct", "{ParamName}\n{StructureName} - Script Struct"), Args);
 		}
 	}
@@ -401,9 +407,9 @@ void UCustomizableObjectNodeTable::PostEditChangeProperty(FPropertyChangedEvent&
 		}
 		else if (PropertyThatChanged->GetName() == GET_MEMBER_NAME_CHECKED(UCustomizableObjectNodeTable, TableDataGatheringMode))
 		{
-			if (Table)
+			if (TObjectPtr<UDataTable> LoadedTable = MutablePrivate::LoadObject(Table))
 			{
-				Table->OnDataTableChanged().Remove(OnTableChangedDelegateHandle);
+				LoadedTable->OnDataTableChanged().Remove(OnTableChangedDelegateHandle);
 			}
 
 			Table = nullptr;
@@ -1319,16 +1325,16 @@ const UScriptStruct* UCustomizableObjectNodeTable::GetTableNodeStruct() const
 
 	if (TableDataGatheringMode == ETableDataGatheringSource::ETDGM_AssetRegistry)
 	{
-		if (Structure)
+		if (TObjectPtr<UScriptStruct> LoadedStructure = MutablePrivate::LoadObject(Structure))
 		{
-			TableStruct = Structure;
+			TableStruct = LoadedStructure;
 		}
 	}
 	else
 	{
-		if (Table)
+		if (TObjectPtr<UDataTable> LoadedTable = MutablePrivate::LoadObject(Table))
 		{
-			TableStruct = Table->GetRowStruct();
+			TableStruct = LoadedTable->GetRowStruct();
 		}
 	}
 
@@ -1445,10 +1451,10 @@ TArray<FAssetData> UCustomizableObjectNodeTable::GetParentTables() const
 		Filter.ClassPaths.Add(FTopLevelAssetPath(UDataTable::StaticClass()));
 		Filter.bRecursivePaths = true;
 
-		if (Structure)
+		if (TObjectPtr<UScriptStruct> LoadedStructure = MutablePrivate::LoadObject(Structure))
 		{
 			TArray<FName> ReferencedTables;
-			AssetRegistry.Get()->GetReferencers(Structure.GetPackage().GetFName(), ReferencedTables, UE::AssetRegistry::EDependencyCategory::Package, UE::AssetRegistry::EDependencyQuery::NoRequirements);
+			AssetRegistry.Get()->GetReferencers(LoadedStructure.GetPackage().GetFName(), ReferencedTables, UE::AssetRegistry::EDependencyCategory::Package, UE::AssetRegistry::EDependencyQuery::NoRequirements);
 
 			for (const FName& ReferencedTable : ReferencedTables)
 			{
diff --git a/Engine/Plugins/Mutable/Source/CustomizableObjectEditor/Private/MuCOE/Nodes/CustomizableObjectNodeTable.h b/Engine/Plugins/Mutable/Source/CustomizableObjectEditor/Private/MuCOE/Nodes/CustomizableObjectNodeTable.h
--- a/Engine/Plugins/Mutable/Source/CustomizableObjectEditor/Private/MuCOE/Nodes/CustomizableObjectNodeTable.h
+++ b/Engine/Plugins/Mutable/Source/CustomizableObjectEditor/Private/MuCOE/Nodes/CustomizableObjectNodeTable.h
@@ -223,11 +223,11 @@ public:
 
 	// Pointer to the Data Table Asset represented in this node
 	UPROPERTY(EditAnywhere, Category = TableProperties, meta = (DontUpdateWhileEditing, EditCondition = "TableDataGatheringMode == ETableDataGatheringSource::ETDGM_DataTable", EditConditionHides))
-	TObjectPtr<UDataTable> Table = nullptr;
+	TSoftObjectPtr<UDataTable> Table = nullptr;
 
 	// Pointer to the Struct Asset represented in this node
 	UPROPERTY(EditAnywhere, Category = TableProperties, meta = (DontUpdateWhileEditing, EditCondition = "TableDataGatheringMode == ETableDataGatheringSource::ETDGM_AssetRegistry", EditConditionHides))
-	TObjectPtr<UScriptStruct> Structure = nullptr;
+	TSoftObjectPtr<UScriptStruct> Structure = nullptr;
 
 	UPROPERTY(EditAnywhere, Category = TableProperties, meta = (EditCondition = "TableDataGatheringMode == ETableDataGatheringSource::ETDGM_AssetRegistry", EditConditionHides))
 	TArray<FName> FilterPaths;
