[Mutable] Refactor OpImageLayer operations to make program size smaller. Remove some uneeded FORCEINLINE functions in ImageDataStorage.

- Move OpImageLayer functions to its own compilation unit. Template usages are manually instantiated.
- Move common batching boilerplate to non-inlined helper functions.
- Move ParallelFor call to a non-inlined helper function.
- Move ImageDataStorage non inlined functions to the cpp file.

#rnx

diff --git a/Engine/Plugins/Mutable/Source/MutableRuntime/Private/MuR/ImageDataStorage.cpp b/Engine/Plugins/Mutable/Source/MutableRuntime/Private/MuR/ImageDataStorage.cpp
--- a/Engine/Plugins/Mutable/Source/MutableRuntime/Private/MuR/ImageDataStorage.cpp
+++ b/Engine/Plugins/Mutable/Source/MutableRuntime/Private/MuR/ImageDataStorage.cpp
@@ -370,6 +370,232 @@ namespace mu
 		}
 	}
 
+	TArrayView<const uint8> FImageDataStorage::GetLOD(int32 LODIndex) const
+	{
+		check(!IsVoid());
+		if (LODIndex >= NumLODs)
+		{
+			//check(false);
+			return TArrayView<const uint8>();
+		}
+
+		const int32 FirstTailLOD = ComputeFirstCompactedTailLOD();
+		if (LODIndex >= FirstTailLOD)
+		{
+			const int32 LODIndexInTail = FMath::Max(0, LODIndex - FirstTailLOD);
+			
+			const int32 LODInTailBegin = LODIndexInTail == 0 ? 0 : CompactedTailOffsets[LODIndexInTail - 1];
+			const int32 LODInTailEnd = CompactedTailOffsets[LODIndexInTail];
+
+			const FImageArray& TailBuffer = Buffers.Last();
+			return MakeArrayView(TailBuffer.GetData() + LODInTailBegin, LODInTailEnd - LODInTailBegin);
+		}
+
+		return MakeArrayView(Buffers[LODIndex].GetData(), Buffers[LODIndex].Num());
+	}
+
+	TArrayView<uint8> FImageDataStorage::GetLOD(int32 LODIndex)
+	{	
+		const TArrayView<const uint8> ConstLODView = const_cast<const FImageDataStorage*>(this)->GetLOD(LODIndex);
+		return TArrayView<uint8>(const_cast<uint8*>(ConstLODView.GetData()), ConstLODView.Num());
+	}
+
+	int32 FImageDataStorage::GetNumBatchesLODRange(int32 BatchSizeInElems, int32 BatchElemSizeInBytes, int32 LODBegin, int32 LODEnd) const
+	{
+		check(!IsVoid());
+		check(LODBegin < LODEnd);
+		check(LODBegin >= 0 && LODEnd <= NumLODs);
+
+		const int32 BatchNumBytes = BatchSizeInElems*BatchElemSizeInBytes;
+		
+		const int32 FirstCompactedTailLOD = ComputeFirstCompactedTailLOD();
+		const int32 NumLODsInTail = FMath::Max(0, NumLODs - FirstCompactedTailLOD);
+		const int32 NumLODsNotInTail = FMath::Max(0, NumLODs - NumLODsInTail);
+
+		const int32 LODsNotInTailLODRangeEnd = FMath::Min(LODEnd, NumLODsNotInTail);
+
+		int32 NumBatches = 0;
+		for (int32 BufferIndex = LODBegin; BufferIndex < LODsNotInTailLODRangeEnd; ++BufferIndex)
+		{	
+			NumBatches += FMath::DivideAndRoundUp(Buffers[BufferIndex].Num(), BatchNumBytes);
+		}
+
+		if (FirstCompactedTailLOD < LODEnd)
+		{
+			TArrayView<const uint8> FirstLODInRangeView = GetLOD(FMath::Max(FirstCompactedTailLOD, LODBegin));
+			TArrayView<const uint8> LastLODInRangeView = GetLOD(LODEnd - 1);
+
+			const int32 TailInLODRangeNumBytes = LastLODInRangeView.GetData() - FirstLODInRangeView.GetData() + LastLODInRangeView.Num();
+			NumBatches += FMath::DivideAndRoundUp(TailInLODRangeNumBytes, BatchNumBytes);
+		}
+		return NumBatches;
+	}
+
+	int32 FImageDataStorage::GetNumBatches(int32 BatchSizeInElems, int32 BatchElemSizeInBytes) const
+	{	
+		return GetNumBatchesLODRange(BatchSizeInElems, BatchElemSizeInBytes, 0, NumLODs);
+	}
+
+	TArrayView<const uint8> FImageDataStorage::GetBatch(int32 BatchId, int32 BatchSizeInElems, int32 BatchElemSizeInBytes) const
+	{
+		return GetBatchLODRange(BatchId, BatchSizeInElems, BatchElemSizeInBytes, 0, NumLODs);
+	}
+
+	TArrayView<uint8> FImageDataStorage::GetBatch(int32 BatchId, int32 BatchSizeInElems, int32 BatchElemSizeInBytes)
+	{
+		// Use the const_cast idiom to generate the non-const operation. 
+		const TArrayView<const uint8> ConstBatchView = 
+				const_cast<const FImageDataStorage*>(this)->GetBatch(BatchId, BatchSizeInElems, BatchElemSizeInBytes);
+		return TArrayView<uint8>(const_cast<uint8*>(ConstBatchView.GetData()), ConstBatchView.Num());
+	}
+
+	int32 FImageDataStorage::GetNumBatchesFirstLODOffset(int32 BatchSizeInElems, int32 BatchElemSizeInBytes, int32 OffsetInBytes) const
+	{
+		check(!IsVoid());
+		TArrayView<const uint8> FirstLODView = GetLOD(0);
+		
+		check(OffsetInBytes % BatchElemSizeInBytes == 0);
+		check(FirstLODView.Num() % BatchElemSizeInBytes == 0);
+		check(FirstLODView.Num() > OffsetInBytes); 
+
+		return FMath::DivideAndRoundUp(FirstLODView.Num() - OffsetInBytes, BatchSizeInElems*BatchElemSizeInBytes);
+	}
+
+	TArrayView<const uint8> FImageDataStorage::GetBatchFirstLODOffset(int32 BatchId, int32 BatchSizeInElems, int32 BatchElemSizeInBytes, int32 OffsetInBytes) const
+	{
+		check(!IsVoid());
+		TArrayView<const uint8> FirstLODView = GetLOD(0);
+		
+		check(OffsetInBytes < FirstLODView.Num());
+		check(FirstLODView.Num() % BatchElemSizeInBytes == 0);
+		check(OffsetInBytes % BatchElemSizeInBytes == 0);
+	
+		TArrayView<const uint8> OffsetedLODView = MakeArrayView(FirstLODView.GetData() + OffsetInBytes, FirstLODView.Num() - OffsetInBytes);
+
+		const int32 BatchSizeInBytes = BatchSizeInElems * BatchElemSizeInBytes;
+		const int32 BatchOffset = BatchId * BatchSizeInBytes;
+
+		if (BatchOffset >= OffsetedLODView.Num())
+		{
+			return TArrayView<const uint8>();
+		}
+
+		return MakeArrayView(OffsetedLODView.GetData() + BatchOffset, FMath::Min(BatchSizeInBytes, OffsetedLODView.Num() - BatchOffset));
+	}
+
+	TArrayView<uint8> FImageDataStorage::GetBatchFirstLODOffset(int32 BatchId, int32 BatchSizeInElems, int32 BatchElemSizeInBytes, int32 OffsetInBytes)
+	{
+		const TArrayView<const uint8> ConstBatchView = 
+				const_cast<const FImageDataStorage*>(this)->GetBatchFirstLODOffset(BatchId, BatchSizeInElems, BatchElemSizeInBytes, OffsetInBytes);
+
+		return TArrayView<uint8>(const_cast<uint8*>(ConstBatchView.GetData()), ConstBatchView.Num());
+	}
+
+	TArrayView<const uint8> FImageDataStorage::GetBatchLODRange(int32 BatchId, int32 BatchSizeInElems, int32 BatchElemSizeInBytes, int32 LODBegin, int32 LODEnd) const
+	{
+		check(!IsVoid());
+		check(LODBegin < LODEnd);
+		check(LODBegin >= 0 && LODEnd <= NumLODs);
+
+		const int32 BatchNumBytes = BatchSizeInElems*BatchElemSizeInBytes;
+
+		const int32 FirstCompactedTailLOD = ComputeFirstCompactedTailLOD();
+		const int32 NonTailBuffersEnd = FMath::Min(LODEnd, FirstCompactedTailLOD);
+
+		int32 NumBatches = 0;	
+		for (int32 BufferIndex = LODBegin; BufferIndex < NonTailBuffersEnd; ++BufferIndex)
+		{
+			const TArrayView<const uint8> BufferView = MakeArrayView(Buffers[BufferIndex].GetData(), Buffers[BufferIndex].Num());
+			const int32 BufferNumBatches = FMath::DivideAndRoundUp(BufferView.Num(), BatchNumBytes); 
+			
+			check(BatchId >= NumBatches);
+			int32 BufferBatchId = BatchId - NumBatches;
+
+			if (BufferBatchId < BufferNumBatches)
+			{ 
+				const int32 BufferBatchBeginOffset = BufferBatchId * BatchNumBytes;
+
+				return TArrayView<const uint8>(
+						BufferView.GetData() + BufferBatchBeginOffset,
+						FMath::Min(BufferView.Num() - BufferBatchBeginOffset, BatchNumBytes));
+			}
+
+			NumBatches += BufferNumBatches;
+		}
+
+		if (FirstCompactedTailLOD < LODEnd)
+		{
+			TArrayView<const uint8> FirstLODInRangeView = GetLOD(FMath::Max(FirstCompactedTailLOD, LODBegin));
+			TArrayView<const uint8> LastLODInRangeView = GetLOD(LODEnd - 1);
+			
+			const int32 TailInLODRangeNumBytes = LastLODInRangeView.GetData() - FirstLODInRangeView.GetData() + LastLODInRangeView.Num();
+			TArrayView<const uint8> TailInLODRangeView = MakeArrayView(FirstLODInRangeView.GetData(), TailInLODRangeNumBytes);
+
+			check(BatchId >= NumBatches);
+			const int32 TailInLODRangeBatchId = BatchId - NumBatches;
+			const int32 TailInLODRangeNumBatches = FMath::DivideAndRoundUp(TailInLODRangeView.Num(), BatchNumBytes);
+
+			if (TailInLODRangeBatchId < TailInLODRangeNumBatches)
+			{
+				const int32 BufferBatchBeginOffset = TailInLODRangeBatchId * BatchNumBytes;
+				
+				return TArrayView<const uint8>(
+							TailInLODRangeView.GetData() + BufferBatchBeginOffset,
+							FMath::Min(TailInLODRangeView.Num() - BufferBatchBeginOffset, BatchNumBytes));
+			}
+		}
+		// If the BatchId is not found, return an empty view.
+		return TArrayView<const uint8>();
+	}
+
+	TArrayView<uint8> FImageDataStorage::GetBatchLODRange(int32 BatchId, int32 BatchSizeInElems, int32 BatchElemSizeInBytes, int32 LODBegin, int32 LODEnd)
+	{
+		const TArrayView<const uint8> ConstBatchView = 
+			const_cast<const FImageDataStorage*>(this)->GetBatchLODRange(
+					BatchId, BatchSizeInElems, BatchElemSizeInBytes, LODBegin, LODEnd);
+
+		return TArrayView<uint8>(const_cast<uint8*>(ConstBatchView.GetData()), ConstBatchView.Num());
+	}
+
+	int32 FImageDataStorage::GetAllocatedSize() const
+	{ 
+		SIZE_T Result = 0;
+		for (const FImageArray& Buffer : Buffers)
+		{
+			Result += Buffer.GetAllocatedSize();
+		}
+
+		check(Result < TNumericLimits<int32>::Max())
+		return (int32)Result;
+	}
+
+	int32 FImageDataStorage::GetDataSize() const
+	{
+		SIZE_T Result = 0;
+		for (const FImageArray& Buffer : Buffers)
+		{
+			Result += Buffer.Num();
+		}
+
+		check(Result < TNumericLimits<int32>::Max())
+		return (int32)Result;
+	}
+
+	bool FImageDataStorage::IsEmpty() const
+	{
+		check(!IsVoid());
+
+		for (const FImageArray& Buffer : Buffers)
+		{
+			if (!Buffer.IsEmpty())
+			{
+				return false;
+			}
+		}
+
+		return true;
+	}
+
 	void FImageDataStorage::Serialise(FOutputArchive& Arch) const
 	{
 		Arch << ImageSize.X;
diff --git a/Engine/Plugins/Mutable/Source/MutableRuntime/Private/MuR/OpImageApplyComposite.h b/Engine/Plugins/Mutable/Source/MutableRuntime/Private/MuR/OpImageApplyComposite.h
--- a/Engine/Plugins/Mutable/Source/MutableRuntime/Private/MuR/OpImageApplyComposite.h
+++ b/Engine/Plugins/Mutable/Source/MutableRuntime/Private/MuR/OpImageApplyComposite.h
@@ -7,47 +7,10 @@
 
 #include "MuR/ImagePrivate.h"
 #include "MuR/OpImageLayer.h"
-
+#include "MuR/OpImageBlend.h"
 
 namespace mu
 {
-	// Combine layer oprations get packed color.
-	inline unsigned NormalComposite(
-			const uint32 Base, 
-			const uint32 Normal, 
-			const uint8 Channel, 
-			const float Power )
-	{
-		const int32 ChannelShift = Channel << 3;
-
-		float ReferencValue = static_cast<float>(Base >> ChannelShift & 0xFF) / 255.0f;
-
-		// Rougness is stored as Unorm8
-		float Roughness = static_cast<float>( (Base >> ChannelShift) & 0xFF ) / 255.0f;
-
-		FVector3f N = ( FVector3f( static_cast<float>((Normal >> 0 ) & 0xFF) / 255.0f,
-							       static_cast<float>((Normal >> 8 ) & 0xFF) / 255.0f,
-							       static_cast<float>((Normal >> 16) & 0xFF) / 255.0f ) * 2.0f - 1.0f );	
-	
-		// See TexturecCompressorModule.cpp::ApplyCompositeTexture
-		// Toksvig estimation of variance
-		float LengthN = FMath::Min( N.Size(), 1.0f );
-		float Variance = ( 1.0f - LengthN ) / LengthN;
-		Variance = FMath::Max( 0.0f, Variance - 0.00004f );
-
-		Variance *= Power;
-		
-		float a = Roughness * Roughness;
-		float a2 = a * a;
-		float B = 2.0f * Variance * (a2 - 1.0f);
-		a2 = ( B - a2 ) / ( B - 1.0f );
-		Roughness = FMath::Pow( a2, 0.25f );
-
-		const uint32 ResMask = ~(0xFF << ChannelShift);
-		const uint32 Value = FMath::Clamp<uint32>( static_cast<uint32>(Roughness * 255.0f), 0u, 255u );
-
-		return  (Value << ChannelShift) | (Base & ResMask); 
-	}
 
 	inline void ImageNormalComposite(FImage* pResult, const FImage* pBase, const FImage* pNormal, ECompositeImageMode mode, float power)
 	{
@@ -55,28 +18,25 @@ namespace mu
 		{
 			// Copy base to result. Using functor layer combine with identity to do the copy for simplicity.
 			// This could be optimized but will never happen since it is checked at compile time. 	
-			ImageLayerCombineFunctor( pResult, pBase, pNormal, 
-				[](uint32 base, uint32) -> uint32 { return base; });
+			ImageLayerCombineFunctor( pResult, pBase, pNormal, FNormalCompositeIdentityFunctor{}); 
 
 			return;
 		}	
 
-		const int8 channel = [mode]() 
+		const uint8 channel = [mode]() 
 			{
 				switch (mode)
 				{
-					case ECompositeImageMode::CIM_NormalRoughnessToRed   : return 0;
-					case ECompositeImageMode::CIM_NormalRoughnessToGreen : return 1;
-					case ECompositeImageMode::CIM_NormalRoughnessToBlue  : return 2;
-					case ECompositeImageMode::CIM_NormalRoughnessToAlpha : return 3;
+					case ECompositeImageMode::CIM_NormalRoughnessToRed   : return 0u;
+					case ECompositeImageMode::CIM_NormalRoughnessToGreen : return 1u;
+					case ECompositeImageMode::CIM_NormalRoughnessToBlue  : return 2u;
+					case ECompositeImageMode::CIM_NormalRoughnessToAlpha : return 3u;
 				}
 			
 				check(false);
-				return 0;		
+				return 0u;		
 			}();
 
-		ImageLayerCombineFunctor( pResult, pBase, pNormal,
-				[channel, power](uint32 base, uint32 normal) -> uint32 
-					{ return NormalComposite(base, normal, channel, power); });
+		ImageLayerCombineFunctor( pResult, pBase, pNormal, FNormalCompositeFunctor{channel, power});
 	}
 }
diff --git a/Engine/Plugins/Mutable/Source/MutableRuntime/Private/MuR/OpImageBlend.h b/Engine/Plugins/Mutable/Source/MutableRuntime/Private/MuR/OpImageBlend.h
--- a/Engine/Plugins/Mutable/Source/MutableRuntime/Private/MuR/OpImageBlend.h
+++ b/Engine/Plugins/Mutable/Source/MutableRuntime/Private/MuR/OpImageBlend.h
@@ -180,4 +180,117 @@ namespace mu
 		return Base + (Blended * (255 - Base) >> 8);
 	}
 
+	// Combine layer oprations get packed color.
+	inline uint32 CombineNormal(uint32 base, uint32 blended)
+	{
+		// See /Engine/Functions/Engine_MaterialFunctions02/Utility/BlendAngleCorrectedNormals.BlendAngleCorrectedNormals
+		// for the source of the effect.
+
+		const float baseRf = (static_cast<float>((base >> 0)  & 0xFF) / 255.0f) * 2.0f - 1.0f;
+		const float baseGf = (static_cast<float>((base >> 8)  & 0xFF) / 255.0f) * 2.0f - 1.0f;
+		const float baseBf = (static_cast<float>((base >> 16) & 0xFF) / 255.0f) * 2.0f; //One added to the b channel.
+
+		const float blendedRf = (static_cast<float>((blended >> 0)  & 0xFF) / 255.0f) * 2.0f - 1.0f;
+		const float blendedGf = (static_cast<float>((blended >> 8)  & 0xFF) / 255.0f) * 2.0f - 1.0f;
+		const float blendedBf = (static_cast<float>((blended >> 16) & 0xFF) / 255.0f) * 2.0f - 1.0f;
+
+		FVector3f a = FVector3f(baseRf, baseGf, baseBf);
+		FVector3f b = FVector3f(-blendedRf, -blendedGf, blendedBf);
+
+		// The original does not normalize but if not done results don't look good due to signal clipping.
+		FVector3f n = (a* FVector3f::DotProduct(a, b) - b*a.Z);
+		n.Normalize();
+
+		return 
+			(FMath::Clamp(static_cast<uint32>((n.X + 1.0f) * 0.5f * 255.0f), 0u, 255u) << 0) |
+			(FMath::Clamp(static_cast<uint32>((n.Y + 1.0f) * 0.5f * 255.0f), 0u, 255u) << 8) |
+			(FMath::Clamp(static_cast<uint32>((n.Z + 1.0f) * 0.5f * 255.0f), 0u, 255u) << 16);
+	}
+
+	inline uint32 CombineNormalMasked(uint32 base, uint32 blended, uint32 mask)
+	{
+		// See /Engine/Functions/Engine_MaterialFunctions02/Utility/BlendAngleCorrectedNormals.BlendAngleCorrectedNormals
+		// for the source of the effect.
+
+		const float baseRf = (static_cast<float>((base >> 0)  & 0xFF) / 255.0f) * 2.0f - 1.0f;
+		const float baseGf = (static_cast<float>((base >> 8)  & 0xFF) / 255.0f) * 2.0f - 1.0f;
+		const float baseBf = (static_cast<float>((base >> 16) & 0xFF) / 255.0f) * 2.0f - 1.0f;
+
+		const float blendedRf = (static_cast<float>((blended >> 0)  & 0xFF) / 255.0f) * 2.0f - 1.0f;
+		const float blendedGf = (static_cast<float>((blended >> 8)  & 0xFF) / 255.0f) * 2.0f - 1.0f;
+		const float blendedBf = (static_cast<float>((blended >> 16) & 0xFF) / 255.0f) * 2.0f - 1.0f;
+
+		const float maskRf = (static_cast<float>(mask & 0xFF) / 255.0f);
+
+		FVector3f a = FVector3f(baseRf, baseGf, baseBf + 1.0f);
+		FVector3f b = FVector3f(-blendedRf, -blendedGf, blendedBf);
+
+		// The original does not normalize but if not done results don't look good due to signal clipping.
+		FVector3f n = ( FMath::Lerp( FVector3f(baseRf, baseGf, baseBf), a*FVector3f::DotProduct(a, b) - b*a.Z, maskRf ) );
+		n.Normalize();
+
+		return 
+			(FMath::Clamp(static_cast<uint32>((n.X + 1.0f) * 0.5f * 255.0f), 0u, 255u) << 0) |
+			(FMath::Clamp(static_cast<uint32>((n.Y + 1.0f) * 0.5f * 255.0f), 0u, 255u) << 8) |
+			(FMath::Clamp(static_cast<uint32>((n.Z + 1.0f) * 0.5f * 255.0f), 0u, 255u) << 16);
+	}
+
+
+	// Combine layer oprations get packed color.
+	inline uint32 NormalComposite(
+			const uint32 Base, 
+			const uint32 Normal, 
+			const uint8 Channel, 
+			const float Power )
+	{
+		const int32 ChannelShift = Channel << 3;
+
+		float ReferencValue = static_cast<float>(Base >> ChannelShift & 0xFF) / 255.0f;
+
+		// Rougness is stored as Unorm8
+		float Roughness = static_cast<float>( (Base >> ChannelShift) & 0xFF ) / 255.0f;
+
+		FVector3f N = ( FVector3f( static_cast<float>((Normal >> 0 ) & 0xFF) / 255.0f,
+							       static_cast<float>((Normal >> 8 ) & 0xFF) / 255.0f,
+							       static_cast<float>((Normal >> 16) & 0xFF) / 255.0f ) * 2.0f - 1.0f );	
+	
+		// See TexturecCompressorModule.cpp::ApplyCompositeTexture
+		// Toksvig estimation of variance
+		float LengthN = FMath::Min( N.Size(), 1.0f );
+		float Variance = ( 1.0f - LengthN ) / LengthN;
+		Variance = FMath::Max( 0.0f, Variance - 0.00004f );
+
+		Variance *= Power;
+		
+		float a = Roughness * Roughness;
+		float a2 = a * a;
+		float B = 2.0f * Variance * (a2 - 1.0f);
+		a2 = ( B - a2 ) / ( B - 1.0f );
+		Roughness = FMath::Pow( a2, 0.25f );
+
+		const uint32 ResMask = ~(0xFF << ChannelShift);
+		const uint32 Value = FMath::Clamp<uint32>( static_cast<uint32>(Roughness * 255.0f), 0u, 255u );
+
+		return  (Value << ChannelShift) | (Base & ResMask); 
+	}
+
+	struct FNormalCompositeIdentityFunctor
+	{
+		uint32 operator()(uint32 Base, uint32) const
+		{
+			return Base;		
+		}
+	};
+
+
+	struct FNormalCompositeFunctor
+	{
+		const uint8 Channel;
+		const float Power;
+
+		uint32 operator()(uint32 Base, uint32 Normal) const
+		{
+			return NormalComposite(Base, Normal, Channel, Power);
+		}
+	};
 }
diff --git a/Engine/Plugins/Mutable/Source/MutableRuntime/Private/MuR/OpImageInvert.h b/Engine/Plugins/Mutable/Source/MutableRuntime/Private/MuR/OpImageInvert.h
--- a/Engine/Plugins/Mutable/Source/MutableRuntime/Private/MuR/OpImageInvert.h
+++ b/Engine/Plugins/Mutable/Source/MutableRuntime/Private/MuR/OpImageInvert.h
@@ -5,11 +5,33 @@
 #include "MuR/ImagePrivate.h"
 #include "MuR/Platform.h"
 
+#include "Async/ParallelFor.h"
 #include "Math/VectorRegister.h"
 
 namespace mu
 {
 
+namespace OpImageInvertInternal
+{
+	// TODO: This is used in multiple places, make it an utility so it can be shared.  
+	FORCENOINLINE void InvokeParallelFor(int32 NumBatches, TFunctionRef<void(int32)> BatchFunc)
+	{
+		if (NumBatches == 0)
+		{
+			return;
+		}
+
+		if (NumBatches == 1)
+		{
+			BatchFunc(0);
+		}
+		else
+		{
+			ParallelFor(NumBatches, BatchFunc);
+		}
+	}
+} // namespace mu::OpImageInvertInternal
+
 	/** Invert the colour (RGB or L) components of an image. Leave alpha untouched. */
 	template<bool bUseVectorImpl = false>
 	inline void ImageInvert(FImage* AImage)
@@ -27,7 +49,10 @@ namespace mu
 		const int32 BytesPerElem = GetImageFormatData(AImage->GetFormat()).BytesPerBlock;
 		const int32 NumBatches = AImage->DataStorage.GetNumBatches(NumBatchElems, BytesPerElem);
 
-		auto ProcessBatch = [AImage, NumBatchElems, BytesPerElem](int32 BatchId)
+		OpImageInvertInternal::InvokeParallelFor(NumBatchElems,
+		[
+			AImage, NumBatchElems, BytesPerElem
+		](int32 BatchId)
 		{
 			TArrayView<uint8> AView = AImage->DataStorage.GetBatch(BatchId, NumBatchElems, BytesPerElem);
 			const int32 NumBytes = AView.Num();
@@ -107,15 +132,6 @@ namespace mu
 			default:
 				check(false);
 			}
-		};
-
-		if (NumBatches == 1)
-		{
-			ProcessBatch(0);
-		}
-		else if (NumBatches > 1)
-		{
-			ParallelFor(NumBatches, ProcessBatch);
-		}
+		});
 	}
 }
diff --git a/Engine/Plugins/Mutable/Source/MutableRuntime/Private/MuR/OpImageLayer.cpp b/Engine/Plugins/Mutable/Source/MutableRuntime/Private/MuR/OpImageLayer.cpp
new file mode 100644
--- /dev/null
+++ b/Engine/Plugins/Mutable/Source/MutableRuntime/Private/MuR/OpImageLayer.cpp
@@ -0,0 +1,2829 @@
+// Copyright Epic Games, Inc. All Rights Reserved.
+
+#include "MuR/OpImageLayer.h"
+
+#include "MuR/ImagePrivate.h"
+#include "MuR/ImageRLE.h"
+#include "Async/ParallelFor.h"
+#include "Templates/UnrealTemplate.h"
+
+#include "MuR/OpImageBlend.h"
+#include "MuR/MutableRuntimeModule.h"
+
+namespace mu
+{
+
+namespace OpImageLayerInternal
+{
+
+	struct FOpLayerBatchArgs
+	{
+		int32 BatchNumElems  = 0;
+		int32 LODBegin       = 0;
+		int32 LODEnd         = 0;
+		int32 FirstLODOffset = 0;
+
+		FImage* Result       = nullptr;
+		const FImage* Base   = nullptr;
+		const FImage* Blend  = nullptr;
+		const FImage* Mask   = nullptr;
+
+		uint32 ResultBytesPerElem = 0;
+		uint32 BaseBytesPerElem   = 0;
+		uint32 BlendBytesPerElem  = 0;
+		uint32 MaskBytesPerElem   = 0;
+	};
+
+	struct FOpLayerBatchViews
+	{
+		int32 NumElems = 0;
+
+		TArrayView<uint8> Result;
+		TArrayView<const uint8> Base;
+		TArrayView<const uint8> Blend;
+		TArrayView<const uint8> Mask;
+	};
+	
+	FORCENOINLINE int32 GetOpLayerNumBatches(const FOpLayerBatchArgs& Args)
+	{	
+		check(Args.Result);
+		const int32 NumBatches = Args.Result->DataStorage.GetNumBatches(Args.BatchNumElems, Args.ResultBytesPerElem);
+
+#if DO_CHECK
+		if (Args.Base)
+		{
+			check(NumBatches == Args.Base->DataStorage.GetNumBatches(Args.BatchNumElems, Args.BaseBytesPerElem));
+		}
+	
+		if (Args.Blend)
+		{
+			check(NumBatches == Args.Blend->DataStorage.GetNumBatches(Args.BatchNumElems, Args.BlendBytesPerElem));
+		}
+
+		if (Args.Mask)
+		{
+			check(NumBatches == Args.Mask->DataStorage.GetNumBatches(Args.BatchNumElems, Args.MaskBytesPerElem));
+		}
+#endif
+		return NumBatches;
+	}
+
+	FORCENOINLINE int32 GetOpLayerNumBatchesLODRange(const FOpLayerBatchArgs& Args)
+	{
+		check(Args.Result);
+		const int32 NumBatches = Args.Result->DataStorage.GetNumBatchesLODRange(Args.BatchNumElems, Args.ResultBytesPerElem, Args.LODBegin, Args.LODEnd);
+
+#if DO_CHECK
+		if (Args.Base)
+		{
+			check(NumBatches == Args.Base->DataStorage.GetNumBatchesLODRange(Args.BatchNumElems, Args.BaseBytesPerElem, Args.LODBegin, Args.LODEnd));
+		}
+	
+		if (Args.Blend)
+		{
+			check(NumBatches == Args.Blend->DataStorage.GetNumBatchesLODRange(Args.BatchNumElems, Args.BlendBytesPerElem, Args.LODBegin, Args.LODEnd));
+		}
+
+		if (Args.Mask)
+		{
+			check(NumBatches == Args.Mask->DataStorage.GetNumBatchesLODRange(Args.BatchNumElems, Args.MaskBytesPerElem, Args.LODBegin, Args.LODEnd));
+		}
+#endif
+		return NumBatches;
+	}
+	
+	FORCENOINLINE int32 GetOpLayerNumBatchesLODRangeOffsetViews(const FOpLayerBatchArgs& Args)
+	{
+		const bool bOnlyFirstLOD = Args.FirstLODOffset >= 0;
+
+		check(Args.Result);
+		const int32 NumBatches = bOnlyFirstLOD
+				? Args.Result->DataStorage.GetNumBatchesFirstLODOffset(Args.BatchNumElems, Args.ResultBytesPerElem, Args.FirstLODOffset)
+				: Args.Result->DataStorage.GetNumBatchesLODRange(Args.BatchNumElems, Args.ResultBytesPerElem, Args.LODBegin, Args.LODEnd);
+
+#if DO_CHECK
+		if (Args.Base)
+		{
+			const int32 BaseNumBatches = bOnlyFirstLOD
+					? Args.Base->DataStorage.GetNumBatchesFirstLODOffset(Args.BatchNumElems, Args.BaseBytesPerElem, Args.FirstLODOffset)
+					: Args.Base->DataStorage.GetNumBatchesLODRange(Args.BatchNumElems, Args.BaseBytesPerElem, Args.LODBegin, Args.LODEnd);
+
+			check(NumBatches == BaseNumBatches);
+		}
+	
+		if (Args.Blend)
+		{
+
+			const int32 BlendNumBatches = bOnlyFirstLOD
+					? Args.Blend->DataStorage.GetNumBatchesFirstLODOffset(Args.BatchNumElems, Args.BlendBytesPerElem, Args.FirstLODOffset)
+					: Args.Blend->DataStorage.GetNumBatchesLODRange(Args.BatchNumElems, Args.BlendBytesPerElem, Args.LODBegin, Args.LODEnd);
+
+			check(NumBatches == BlendNumBatches);
+		}
+
+		if (Args.Mask)
+		{
+			const int32 MaskNumBatches = bOnlyFirstLOD
+					? Args.Mask->DataStorage.GetNumBatchesFirstLODOffset(Args.BatchNumElems, Args.MaskBytesPerElem, Args.FirstLODOffset)
+					: Args.Mask->DataStorage.GetNumBatchesLODRange(Args.BatchNumElems, Args.MaskBytesPerElem, Args.LODBegin, Args.LODEnd);
+
+			check(NumBatches == MaskNumBatches);
+		}
+#endif
+		return NumBatches;
+	}
+
+
+	FORCENOINLINE FOpLayerBatchViews GetOpLayerBatchViews(int32 BatchId, const FOpLayerBatchArgs& Args)
+	{
+		FOpLayerBatchViews BatchViewsResult;
+	
+		check(Args.Result);
+		BatchViewsResult.Result = Args.Result->DataStorage.GetBatch(BatchId, Args.BatchNumElems, Args.ResultBytesPerElem);
+
+		BatchViewsResult.NumElems = BatchViewsResult.Result.Num() / Args.ResultBytesPerElem;
+
+		if (Args.Base)
+		{
+			BatchViewsResult.Base = Args.Base->DataStorage.GetBatch(BatchId, Args.BatchNumElems, Args.BaseBytesPerElem);
+			check(BatchViewsResult.NumElems == BatchViewsResult.Base.Num() / Args.BaseBytesPerElem);
+		}
+
+		if (Args.Blend)
+		{
+			BatchViewsResult.Blend = Args.Blend->DataStorage.GetBatch(BatchId, Args.BatchNumElems, Args.BlendBytesPerElem);
+			check(BatchViewsResult.NumElems == BatchViewsResult.Blend.Num() / Args.BlendBytesPerElem);
+		}
+
+		if (Args.Mask)
+		{
+			BatchViewsResult.Mask = Args.Mask->DataStorage.GetBatch(BatchId, Args.BatchNumElems, Args.MaskBytesPerElem);
+			check(BatchViewsResult.NumElems == BatchViewsResult.Mask.Num() / Args.MaskBytesPerElem);
+		}
+
+		return BatchViewsResult;
+	}
+
+	FORCENOINLINE FOpLayerBatchViews GetOpLayerBatchLODRangeViews(int32 BatchId, const FOpLayerBatchArgs& Args)
+	{
+		FOpLayerBatchViews BatchViewsResult;
+		
+		check(Args.Result);
+		BatchViewsResult.Result = 
+					Args.Result->DataStorage.GetBatchLODRange(BatchId, Args.BatchNumElems, Args.ResultBytesPerElem, Args.LODBegin, Args.LODEnd);
+
+		BatchViewsResult.NumElems = BatchViewsResult.Result.Num() / Args.ResultBytesPerElem;
+
+		if (Args.Base)
+		{
+			BatchViewsResult.Base = 
+					Args.Base->DataStorage.GetBatchLODRange(BatchId, Args.BatchNumElems, Args.BaseBytesPerElem, Args.LODBegin, Args.LODEnd);
+			check(BatchViewsResult.NumElems == BatchViewsResult.Base.Num() / Args.BaseBytesPerElem);
+		}
+
+		if (Args.Blend)
+		{
+			BatchViewsResult.Blend = 
+					Args.Blend->DataStorage.GetBatchLODRange(BatchId, Args.BatchNumElems, Args.BlendBytesPerElem, Args.LODBegin, Args.LODEnd);
+			check(BatchViewsResult.NumElems == BatchViewsResult.Blend.Num() / Args.BlendBytesPerElem);
+		}
+
+		if (Args.Mask)
+		{
+			BatchViewsResult.Mask = 
+					Args.Mask->DataStorage.GetBatchLODRange(BatchId, Args.BatchNumElems, Args.MaskBytesPerElem, Args.LODBegin, Args.LODEnd);
+			check(BatchViewsResult.NumElems == BatchViewsResult.Mask.Num() / Args.MaskBytesPerElem);
+		}
+
+		return BatchViewsResult;
+	}
+
+	FORCENOINLINE FOpLayerBatchViews GetOpLayerBatchLODRangeOffsetViews(int32 BatchId, const FOpLayerBatchArgs& Args)
+	{
+		FOpLayerBatchViews BatchViewsResult;
+		
+		const bool bOnlyFirstLOD = Args.FirstLODOffset >= 0;
+
+		check(Args.Result);
+		BatchViewsResult.Result = bOnlyFirstLOD
+				? Args.Result->DataStorage.GetBatchFirstLODOffset(BatchId, Args.BatchNumElems, Args.ResultBytesPerElem, Args.FirstLODOffset) 
+				: Args.Result->DataStorage.GetBatchLODRange(BatchId, Args.BatchNumElems, Args.ResultBytesPerElem, Args.LODBegin, Args.LODEnd);
+
+		BatchViewsResult.NumElems = BatchViewsResult.Result.Num() / Args.ResultBytesPerElem;
+		
+		if (Args.Base)
+		{
+			BatchViewsResult.Base = bOnlyFirstLOD
+					? Args.Base->DataStorage.GetBatchFirstLODOffset(BatchId, Args.BatchNumElems, Args.BaseBytesPerElem, Args.FirstLODOffset) 
+					: Args.Base->DataStorage.GetBatchLODRange(BatchId, Args.BatchNumElems, Args.BaseBytesPerElem, Args.LODBegin, Args.LODEnd); 
+
+			check(BatchViewsResult.NumElems == BatchViewsResult.Base.Num() / Args.BaseBytesPerElem);
+		}
+
+		if (Args.Blend)
+		{
+			BatchViewsResult.Blend = bOnlyFirstLOD
+					? Args.Blend->DataStorage.GetBatchFirstLODOffset(BatchId, Args.BatchNumElems, Args.BlendBytesPerElem, Args.FirstLODOffset) 
+					: Args.Blend->DataStorage.GetBatchLODRange(BatchId, Args.BatchNumElems, Args.BlendBytesPerElem, Args.LODBegin, Args.LODEnd); 
+
+			check(BatchViewsResult.NumElems == BatchViewsResult.Blend.Num() / Args.BlendBytesPerElem);
+		}
+
+		if (Args.Mask)
+		{
+			BatchViewsResult.Mask = bOnlyFirstLOD
+					? Args.Mask->DataStorage.GetBatchFirstLODOffset(BatchId, Args.BatchNumElems, Args.MaskBytesPerElem, Args.FirstLODOffset) 
+					: Args.Mask->DataStorage.GetBatchLODRange(BatchId, Args.BatchNumElems, Args.MaskBytesPerElem, Args.LODBegin, Args.LODEnd); 
+
+			check(BatchViewsResult.NumElems == BatchViewsResult.Mask.Num() / Args.MaskBytesPerElem);
+		}
+
+		return BatchViewsResult;
+	}
+
+	FORCENOINLINE void InvokeParallelFor(int32 NumBatches, TFunctionRef<void(int32)> BatchFunc)
+	{
+		if (NumBatches == 0)
+		{
+			return;
+		}
+
+		if (NumBatches == 1)
+		{
+			BatchFunc(0);
+		}
+		else
+		{
+			ParallelFor(NumBatches, BatchFunc);
+		}
+
+	}
+
+	FORCEINLINE bool IsAnyComponentLargerThan1(FVector4f Value)
+	{
+		return (Value[0] > 1) | (Value[1] > 1) | (Value[2] > 1) | (Value[3] > 1);
+	}
+
+	/** 
+	 * Apply a blending function to an image with a colour source.
+	 * It only affects the RGB or L channels, leaving alpha untouched.
+	 */
+	template<uint32 NumChannels, uint32 (*BLEND_FUNC)(uint32, uint32), bool bClamp>
+	FORCENOINLINE void BufferLayerColourGenericChannel(uint8* DestBuf, const uint8* BaseBuf, int32 NumElems, const FIntVector& Color)
+	{
+		static_assert(NumChannels > 0 && NumChannels <= 4);
+
+		for (int32 I = 0; I < NumElems; ++I)
+		{
+			for (uint32 C = 0; C < NumChannels; ++C)
+			{
+				uint32 Base = BaseBuf[NumChannels * I + C];
+				uint32 Result = BLEND_FUNC(Base, Color[C]);
+				if constexpr (bClamp)
+				{
+					DestBuf[NumChannels * I + C] = (uint8)FMath::Min(255u, Result);
+				}
+				else
+				{
+					DestBuf[NumChannels * I + C] = (uint8)Result;
+				}
+			}
+		}
+	}
+
+	template<uint32 NumChannels, uint32 (*BLEND_FUNC)(uint32,uint32), bool Clamp>
+	FORCENOINLINE void BufferLayerColourFromAlphaGenericChannel(uint8* DestBuf, const uint8* BaseBuf, int32 NumElems, const FIntVector& Color)
+	{
+		static_assert(NumChannels > 0 && NumChannels <= 4);
+
+		for (int32 I = 0; I < NumElems; ++I)
+		{
+			const uint32 Alpha = Invoke([&]() -> uint32
+			{
+				if constexpr (NumChannels <= 3)
+				{
+					return 255;
+				}
+				else
+				{
+					return BaseBuf[NumChannels * I + 3];
+				}
+			});
+
+			for (uint32 C = 0; C < NumChannels; ++C)
+			{
+				uint32 Result = BLEND_FUNC(Alpha, Color[C]);
+				if constexpr (Clamp)
+				{
+					DestBuf[NumChannels * I + C] = (uint8)FMath::Min(255u, Result);
+				}
+				else
+				{
+					DestBuf[NumChannels * I + C] = (uint8)Result;
+				}
+			}
+		}
+	}
+} // namespace OpImageLayerInternal
+
+	template< uint32 (*BLEND_FUNC)(uint32,uint32), bool CLAMP >
+    FORCENOINLINE void BufferLayerColourImpl(FImage* ResultImage, const FImage* BaseImage, FVector4f Color)
+	{
+        check(ResultImage->GetFormat() == BaseImage->GetFormat());
+        check(ResultImage->GetSizeX() == BaseImage->GetSizeX());
+        check(ResultImage->GetSizeY() == BaseImage->GetSizeY());
+        check(ResultImage->GetLODCount() == BaseImage->GetLODCount());
+
+		const EImageFormat BaseFormat = BaseImage->GetFormat();
+		// Generic implementation
+		const int32 BytesPerElem = GetImageFormatData(BaseFormat).BytesPerBlock;	
+
+		constexpr int32 BatchNumElems = 4096*4;
+		OpImageLayerInternal::FOpLayerBatchArgs BatchArgs;
+		BatchArgs.BatchNumElems      = BatchNumElems;
+		BatchArgs.Result             = ResultImage;
+		BatchArgs.Base               = BaseImage;
+		BatchArgs.BaseBytesPerElem   = BytesPerElem;
+		BatchArgs.ResultBytesPerElem = BytesPerElem;
+
+		const int32 NumBatches = OpImageLayerInternal::GetOpLayerNumBatches(BatchArgs);
+
+		const FIntVector ColorValue = FIntVector(Color.X * 255.0f, Color.Y * 255.0f, Color.Z * 255.0f);
+		OpImageLayerInternal::InvokeParallelFor(NumBatches,
+		[
+			&BatchArgs, ColorValue, BaseFormat
+		](int32 BatchId)
+		{
+			OpImageLayerInternal::FOpLayerBatchViews BatchViews = 
+					OpImageLayerInternal::GetOpLayerBatchViews(BatchId, BatchArgs);
+	
+			switch (BaseFormat)
+			{
+			case EImageFormat::L_UByte:
+			{
+				check(BatchArgs.ResultBytesPerElem == 1);
+				OpImageLayerInternal::BufferLayerColourGenericChannel<1, BLEND_FUNC, CLAMP>(
+						BatchViews.Result.GetData(), BatchViews.Base.GetData(), BatchViews.NumElems, ColorValue);
+				break;
+			}
+			case EImageFormat::RGB_UByte:
+			{
+				check(BatchArgs.ResultBytesPerElem == 3);
+				OpImageLayerInternal::BufferLayerColourGenericChannel<3, BLEND_FUNC, CLAMP>(
+						BatchViews.Result.GetData(), BatchViews.Base.GetData(), BatchViews.NumElems, ColorValue);
+				break;
+			}
+			case EImageFormat::RGBA_UByte:
+			{
+				check(BatchArgs.ResultBytesPerElem == 4);
+				OpImageLayerInternal::BufferLayerColourGenericChannel<4, BLEND_FUNC, CLAMP>(
+						BatchViews.Result.GetData(), BatchViews.Base.GetData(), BatchViews.NumElems, ColorValue);
+				break;
+			}
+			case EImageFormat::BGRA_UByte:
+			{
+				check(BatchArgs.ResultBytesPerElem == 4);
+				FIntVector BGRAColorValue = FIntVector(ColorValue.Z, ColorValue.Y, ColorValue.X);
+				OpImageLayerInternal::BufferLayerColourGenericChannel<4, BLEND_FUNC, CLAMP>(
+						BatchViews.Result.GetData(), BatchViews.Base.GetData(), BatchViews.NumElems, ColorValue);
+				break;
+			}
+			default:
+			{
+				checkf(false, TEXT("Unsupported format."));
+				break;
+			}
+			}
+		});
+    }
+
+
+	template< uint32 (*BLEND_FUNC)(uint32, uint32) >
+	void BufferLayerColour(FImage* ResultImage, const FImage* BaseImage, FVector4f Color)
+	{
+		bool bIsClampNeeded = OpImageLayerInternal::IsAnyComponentLargerThan1(Color);
+		if (bIsClampNeeded)
+		{
+			BufferLayerColourImpl<BLEND_FUNC, true>(ResultImage, BaseImage, Color);
+		}
+		else
+		{
+			BufferLayerColourImpl<BLEND_FUNC, false>(ResultImage, BaseImage, Color);
+		}
+	}
+
+	template<
+		uint32 (*BLEND_FUNC)(uint32, uint32), 
+		bool CLAMP>
+    FORCENOINLINE void BufferLayerColourFromAlpha(FImage* ResultImage, const FImage* BaseImage, FVector4f Color)
+	{
+        check(ResultImage->GetFormat() == BaseImage->GetFormat());
+        check(ResultImage->GetSizeX() == BaseImage->GetSizeX());
+        check(ResultImage->GetSizeY() == BaseImage->GetSizeY());
+        check(ResultImage->GetLODCount() == BaseImage->GetLODCount());
+
+		const EImageFormat BaseFormat = BaseImage->GetFormat();
+		// Generic implementation
+		constexpr int32 BatchNumElems = 4096*2;
+		const int32 BytesPerElem = GetImageFormatData(BaseFormat).BytesPerBlock;	
+
+		OpImageLayerInternal::FOpLayerBatchArgs BatchArgs;
+		BatchArgs.BatchNumElems      = BatchNumElems;
+		BatchArgs.Result             = ResultImage;
+		BatchArgs.Base               = BaseImage;
+		BatchArgs.BaseBytesPerElem   = BytesPerElem;
+		BatchArgs.ResultBytesPerElem = BytesPerElem;
+
+		const int32 NumBatches = OpImageLayerInternal::GetOpLayerNumBatches(BatchArgs);
+
+		const FIntVector ColorValue = FIntVector(Color.X * 255.0f, Color.Y * 255.0f, Color.Z * 255.0f);
+		OpImageLayerInternal::InvokeParallelFor(NumBatches,
+		[
+			&BatchArgs, ColorValue, BaseFormat
+		](int32 BatchId)
+		{
+			OpImageLayerInternal::FOpLayerBatchViews BatchViews = 
+					OpImageLayerInternal::GetOpLayerBatchViews(BatchId, BatchArgs);
+	
+			switch (BaseFormat)
+			{
+			case EImageFormat::L_UByte:
+			{
+				check(BatchArgs.ResultBytesPerElem == 1);
+				OpImageLayerInternal::BufferLayerColourFromAlphaGenericChannel<1, BLEND_FUNC, CLAMP>(
+						BatchViews.Result.GetData(), BatchViews.Base.GetData(), BatchViews.NumElems, ColorValue);
+				break;
+			}
+			case EImageFormat::RGB_UByte:
+			{
+				check(BatchArgs.ResultBytesPerElem == 3);
+				OpImageLayerInternal::BufferLayerColourFromAlphaGenericChannel<3, BLEND_FUNC, CLAMP>(
+						BatchViews.Result.GetData(), BatchViews.Base.GetData(), BatchViews.NumElems, ColorValue);
+				break;
+			}
+			case EImageFormat::RGBA_UByte:
+			{
+				check(BatchArgs.ResultBytesPerElem == 4);
+				OpImageLayerInternal::BufferLayerColourFromAlphaGenericChannel<4, BLEND_FUNC, CLAMP>(
+						BatchViews.Result.GetData(), BatchViews.Base.GetData(), BatchViews.NumElems, ColorValue);
+				break;
+			}
+			case EImageFormat::BGRA_UByte:
+			{
+				check(BatchArgs.ResultBytesPerElem == 4);
+				FIntVector BGRAColorValues = FIntVector(ColorValue.Z, ColorValue.Y, ColorValue.X);
+				OpImageLayerInternal::BufferLayerColourFromAlphaGenericChannel<4, BLEND_FUNC, CLAMP>(
+						BatchViews.Result.GetData(), BatchViews.Base.GetData(), BatchViews.NumElems, ColorValue);
+				break;
+			}
+			default:
+			{
+				checkf(false, TEXT("Unsupported format."));
+				break;
+			}
+			}
+		});
+    }
+
+
+	template<uint32 (*BLEND_FUNC)(uint32, uint32)>
+	void BufferLayerColourFromAlpha(FImage* ResultImage, const FImage* BaseImage, FVector4f Color)
+	{
+		bool bIsClampNeeded = OpImageLayerInternal::IsAnyComponentLargerThan1(Color);
+		if (bIsClampNeeded)
+		{
+			BufferLayerColourFromAlpha<BLEND_FUNC, true>(ResultImage, BaseImage, Color);
+		}
+		else
+		{
+			BufferLayerColourFromAlpha<BLEND_FUNC, false>(ResultImage, BaseImage, Color);
+		}
+	}
+
+
+	template< 
+		uint32 (*BLEND_FUNC_MASKED)(uint32, uint32, uint32),
+		uint32 (*BLEND_FUNC)(uint32, uint32),
+		bool CLAMP,
+		// Number of total channels to actually process
+		uint32 CHANNELS_TO_BLEND,
+		// Number of total channels in the base image
+		int32 BASE_CHANNEL_STRIDE>
+    void BufferLayerColourFormat(FImage* DestImage, const FImage* BaseImage, const FImage* MaskImage, FVector4f Col, uint32 BaseChannelOffset, uint8 ColorChannelOffset, bool bOnlyFirstLOD)
+	{
+		check(CHANNELS_TO_BLEND + BaseChannelOffset <= BASE_CHANNEL_STRIDE);
+		check(DestImage->GetLODCount() <= BaseImage->GetLODCount());
+
+		FUintVector4 TopColor = FUintVector4(Col.X * 255.0f, Col.Y * 255.0f, Col.Z * 255.0f, Col.W * 255);
+
+		const EImageFormat BaseFormat = BaseImage->GetFormat();
+		const EImageFormat MaskFormat = MaskImage->GetFormat();
+
+		int32 NumLODs = bOnlyFirstLOD ? 1 : BaseImage->GetLODCount();
+
+        const bool bIsMaskUncompressed = (MaskFormat == EImageFormat::L_UByte);
+
+		constexpr uint32 NumColorChannels = FMath::Min(CHANNELS_TO_BLEND, 3u);
+		check(NumColorChannels + ColorChannelOffset < 4);
+
+        if (bIsMaskUncompressed)
+        {
+			const int32 BytesPerElem = GetImageFormatData(BaseFormat).BytesPerBlock;
+			check(GetImageFormatData(MaskFormat).BytesPerBlock == 1);
+			check(GetImageFormatData(DestImage->GetFormat()).BytesPerBlock == BytesPerElem);
+
+			constexpr int32 BatchNumElems = 4096 * 2;
+			OpImageLayerInternal::FOpLayerBatchArgs BatchArgs;
+			BatchArgs.BatchNumElems      = BatchNumElems; 
+			BatchArgs.LODBegin           = 0;
+			BatchArgs.LODEnd             = NumLODs;
+			BatchArgs.Result             = DestImage;
+			BatchArgs.Base               = BaseImage;
+			BatchArgs.Mask               = MaskImage;
+			BatchArgs.ResultBytesPerElem = BytesPerElem;
+			BatchArgs.BaseBytesPerElem   = BytesPerElem;
+			BatchArgs.MaskBytesPerElem   = 1;
+
+			const int32 NumBatches = OpImageLayerInternal::GetOpLayerNumBatchesLODRange(BatchArgs);
+
+			OpImageLayerInternal::InvokeParallelFor(NumBatches,
+			[
+				&BatchArgs, BaseChannelOffset, TopColor, ColorChannelOffset, NumColorChannels
+			](int32 BatchId)
+			{
+				OpImageLayerInternal::FOpLayerBatchViews BatchViews = 
+						OpImageLayerInternal::GetOpLayerBatchLODRangeViews(BatchId, BatchArgs);
+
+				// This could happen in case of missing data files.
+				if (!BatchViews.Base.GetData() || !BatchViews.Mask.GetData() || !BatchViews.Result.GetData())
+				{
+					return;
+				}
+
+				const uint8* BaseBuf = BatchViews.Base.GetData() + BaseChannelOffset;
+				const uint8* MaskBuf = BatchViews.Mask.GetData();
+				uint8* DestBuf       = BatchViews.Result.GetData() + BaseChannelOffset;
+
+				for (int32 I = 0; I < BatchViews.NumElems; ++I)
+				{
+					uint32 MaskData = MaskBuf[I];
+					for (int32 C = 0; C < NumColorChannels; ++C)
+					{
+						const uint32 Base = BaseBuf[BASE_CHANNEL_STRIDE*I + C];
+						const uint32 Result = BLEND_FUNC_MASKED(Base, TopColor[C + ColorChannelOffset], MaskData);
+						if constexpr (CLAMP)
+						{
+							DestBuf[BASE_CHANNEL_STRIDE*I + C] = (uint8)FMath::Min(255u, Result);
+						}
+						else
+						{
+							DestBuf[BASE_CHANNEL_STRIDE*I + C] = (uint8)Result;
+						}
+					}
+
+					constexpr bool bIsNC4 = (BASE_CHANNEL_STRIDE == 4);
+					if constexpr (bIsNC4)
+					{
+						DestBuf[BASE_CHANNEL_STRIDE*I + 3] = BaseBuf[BASE_CHANNEL_STRIDE*I + 3];
+					}
+				}
+			});
+        }
+        else if (MaskFormat == EImageFormat::L_UByteRLE)
+        {
+            int32 Rows = BaseImage->GetSizeY();
+            int32 Width = BaseImage->GetSizeX();
+
+            for (int32 Lod = 0; Lod < NumLODs; ++Lod)
+            {
+                const uint8* BaseBuf = BaseImage->GetLODData(Lod);
+                const uint8* MaskBuf = MaskImage->GetLODData(Lod);
+				uint8* DestBuf = DestImage->GetLODData(Lod);
+
+				// This could happen in case of missing data files.
+				if (!BaseBuf || !MaskBuf || !DestBuf)
+				{
+					continue;
+				}
+
+				// Remove RLE header, mip size and row sizes.
+				MaskBuf += sizeof(uint32);
+				MaskBuf += Rows * sizeof(uint32);
+
+                for (int32 RowIndex = 0; RowIndex < Rows; ++RowIndex)
+                {
+                    const uint8* DestRowEnd = DestBuf + Width * BASE_CHANNEL_STRIDE;
+                    while (DestBuf != DestRowEnd)
+                    {
+                        // Decode header
+						uint16 Equal = 0;
+						FMemory::Memmove(&Equal, MaskBuf, sizeof(uint16));
+                        MaskBuf += 2;
+
+                        uint8 Different = *MaskBuf;
+                        ++MaskBuf;
+
+                        uint8 EqualPixel = *MaskBuf;
+                        ++MaskBuf;
+
+                        // Equal pixels
+						//check(DestBuf + BASE_CHANNEL_STRIDE * Equal <= BaseImage->GetDataSize(Lod));
+                        if (EqualPixel == 255)
+                        {
+                            for (int32 I = 0; I < Equal; ++I)
+                            {
+                                for (int32 C = 0; C < NumColorChannels; ++C)
+                                {
+                                    uint32 Base = BaseBuf[BASE_CHANNEL_STRIDE * I + C];
+                                    uint32 Result = BLEND_FUNC(Base, TopColor[C + ColorChannelOffset]);
+                                    if constexpr (CLAMP)
+                                    {
+                                        DestBuf[BASE_CHANNEL_STRIDE * I + C] = (uint8)FMath::Min(255u, Result);
+                                    }
+                                    else
+                                    {
+                                        DestBuf[BASE_CHANNEL_STRIDE * I + C] = (uint8)Result;
+                                    }
+                                }
+
+                                constexpr bool bIsNC4 = (BASE_CHANNEL_STRIDE == 4);
+                                if (bIsNC4)
+                                {
+                                    DestBuf[BASE_CHANNEL_STRIDE * I + 3] = BaseBuf[BASE_CHANNEL_STRIDE * I + 3];
+                                }
+                            }
+                        }
+                        else if (EqualPixel > 0)
+                        {
+                            for (int32 I = 0; I < Equal; ++I)
+                            {
+                                for (int32 C = 0; C < NumColorChannels; ++C)
+                                {
+                                    uint32 Base = BaseBuf[BASE_CHANNEL_STRIDE * I + C];
+                                    uint32 Result = BLEND_FUNC_MASKED(Base, TopColor[C + ColorChannelOffset], EqualPixel);
+                                    if constexpr (CLAMP)
+                                    {
+                                        DestBuf[BASE_CHANNEL_STRIDE * I + C] = (uint8)FMath::Min(255u, Result);
+                                    }
+                                    else
+                                    {
+                                        DestBuf[BASE_CHANNEL_STRIDE * I + C] = (uint8)Result;
+                                    }
+                                }
+
+                                constexpr bool bIsNC4 = (BASE_CHANNEL_STRIDE == 4);
+                                if (bIsNC4)
+                                {
+                                    DestBuf[BASE_CHANNEL_STRIDE * I + 3] = BaseBuf[BASE_CHANNEL_STRIDE * I + 3];
+                                }
+                            }
+                        }
+                        else
+                        {
+                            // It could happen if xxxxxOnBase
+                            if (DestBuf != BaseBuf)
+                            {
+                                FMemory::Memmove(DestBuf, BaseBuf, BASE_CHANNEL_STRIDE*Equal);
+                            }
+                        }
+
+                        DestBuf += BASE_CHANNEL_STRIDE * Equal;
+                        BaseBuf += BASE_CHANNEL_STRIDE * Equal;
+
+                        // Different pixels
+						//check(DestBuf + BASE_CHANNEL_STRIDE * Different <= StartDestBuf + BaseImage->GetDataSize(Lod));
+                        for (int32 I = 0; I < Different; ++I)
+                        {
+                            for (int32 C = 0; C < NumColorChannels; ++C)
+                            {
+                                uint32 Mask = MaskBuf[I];
+                                uint32 Base = BaseBuf[BASE_CHANNEL_STRIDE * I + C];
+                                uint32 Result = BLEND_FUNC_MASKED(Base, TopColor[C + ColorChannelOffset], Mask);
+                                if constexpr (CLAMP)
+                                {
+                                    DestBuf[BASE_CHANNEL_STRIDE * I + C] = (uint8)FMath::Min(255u, Result);
+                                }
+                                else
+                                {
+                                    DestBuf[BASE_CHANNEL_STRIDE * I + C] = (uint8)Result;
+                                }
+                            }
+
+                            constexpr bool bIsNC4 = (BASE_CHANNEL_STRIDE == 4);
+                            if (bIsNC4)
+                            {
+                                DestBuf[BASE_CHANNEL_STRIDE * I + 3] = BaseBuf[BASE_CHANNEL_STRIDE * I + 3];
+                            }
+                        }
+
+                        DestBuf += BASE_CHANNEL_STRIDE * Different;
+                        BaseBuf += BASE_CHANNEL_STRIDE * Different;
+                        MaskBuf += Different;
+                    }
+                }
+
+                Rows = FMath::DivideAndRoundUp(Rows, 2);
+                Width = FMath::DivideAndRoundUp(Width, 2);
+            }
+        }
+        else
+        {
+            checkf( false, TEXT("Unsupported mask format.") );
+        }
+	}
+
+
+	/**
+	* Apply a blending function to an image with a colour source and a mask
+	*/
+	template<
+		uint32 (*BLEND_FUNC_MASKED)(uint32,uint32,uint32),
+		uint32 (*BLEND_FUNC)(uint32,uint32),
+		bool CLAMP>
+    FORCENOINLINE void BufferLayerColourImpl(FImage* ResultImage, const FImage* BaseImage, const FImage* MaskImage, FVector4f Col)
+	{
+		check(BaseImage->GetSizeX() == MaskImage->GetSizeX());
+		check(BaseImage->GetSizeY() == MaskImage->GetSizeY());
+		check(MaskImage->GetFormat() == EImageFormat::L_UByte ||
+			  MaskImage->GetFormat() == EImageFormat::L_UByteRLE);
+
+		const bool bValid = (BaseImage->GetSizeX() == MaskImage->GetSizeX()) &&
+						    (BaseImage->GetSizeY() == MaskImage->GetSizeY()) &&
+							(MaskImage->GetFormat() == EImageFormat::L_UByte || MaskImage->GetFormat() == EImageFormat::L_UByteRLE);
+		if (!bValid)
+		{
+			return;
+		}
+
+		EImageFormat BaseFormat = BaseImage->GetFormat();
+		if (BaseFormat == EImageFormat::RGB_UByte)
+		{
+			BufferLayerColourFormat<BLEND_FUNC_MASKED, BLEND_FUNC, CLAMP, 3, 3>(ResultImage, BaseImage, MaskImage, Col, 0, 0, false);
+		}
+        else if (BaseFormat == EImageFormat::RGBA_UByte)
+        {
+            BufferLayerColourFormat<BLEND_FUNC_MASKED, BLEND_FUNC, CLAMP, 3, 4>(ResultImage, BaseImage, MaskImage, Col, 0, 0, false);
+        }
+        else if (BaseFormat == EImageFormat::BGRA_UByte)
+        {
+            float Temp = Col[0];
+            Col[0] = Col[2];
+            Col[2] = Temp;
+            BufferLayerColourFormat<BLEND_FUNC_MASKED, BLEND_FUNC, CLAMP, 3, 4>(ResultImage, BaseImage, MaskImage, Col, 0, 0, false);
+        }
+        else if (BaseFormat == EImageFormat::L_UByte)
+		{
+			BufferLayerColourFormat<BLEND_FUNC_MASKED, BLEND_FUNC, CLAMP, 1, 1>(ResultImage, BaseImage, MaskImage, Col, 0, 0, false);
+		}
+		else
+		{
+			checkf(false, TEXT("Unsupported format."));
+		}
+	}
+
+
+	template< 
+		uint32 (*BLEND_FUNC_MASKED)(uint32, uint32, uint32),
+		uint32 (*BLEND_FUNC)(uint32, uint32)>
+	void BufferLayerColour(FImage* DestImage, const FImage* BaseImage, const FImage* MaskImage, FVector4f Col)
+	{
+		bool bIsClampNeeded = OpImageLayerInternal::IsAnyComponentLargerThan1(Col);
+		if (bIsClampNeeded)
+		{
+			BufferLayerColourImpl<BLEND_FUNC_MASKED, BLEND_FUNC, true>(DestImage, BaseImage, MaskImage, Col);
+		}
+		else
+		{
+			BufferLayerColourImpl<BLEND_FUNC_MASKED, BLEND_FUNC, false>(DestImage, BaseImage, MaskImage, Col);
+		}
+	}
+
+
+	/**
+	* Apply a blending function to an image with another image as blending layer
+	*/
+	template< 
+		uint32 (*BLEND_FUNC)(uint32, uint32), 
+		bool CLAMP,
+		// Number of total channels to actually process
+		uint32 CHANNELS_TO_BLEND,
+		// Number of total channels in the base image
+		int32 BASE_CHANNEL_STRIDE>
+	void BufferLayerColourFormatInPlace(FImage* BaseImage, FVector4f Color, uint32 BaseChannelOffset, uint8 ColorChannelOffset, bool bOnlyFirstLOD)
+	{
+		FUintVector4 TopColor = FUintVector4(Color.X * 255.0f, Color.Y * 255.0f, Color.Z * 255.0f, Color.W * 255.0f);
+
+		int32 NumLODs = bOnlyFirstLOD ? 1 : BaseImage->GetLODCount();
+
+		const int32 BytesPerElem = GetImageFormatData(BaseImage->GetFormat()).BytesPerBlock;
+		
+		constexpr int32 BatchNumElems = 4096*2;
+		OpImageLayerInternal::FOpLayerBatchArgs BatchArgs;
+		BatchArgs.BatchNumElems      = BatchNumElems;
+		BatchArgs.LODBegin           = 0;
+		BatchArgs.LODEnd             = NumLODs;
+		BatchArgs.Result             = BaseImage;
+		BatchArgs.ResultBytesPerElem = BytesPerElem;
+
+		const int32 NumBatches = OpImageLayerInternal::GetOpLayerNumBatchesLODRange(BatchArgs);
+
+		OpImageLayerInternal::InvokeParallelFor(NumBatches,
+		[
+			&BatchArgs, TopColor, BaseChannelOffset, ColorChannelOffset
+		] (uint32 BatchId)
+		{
+			OpImageLayerInternal::FOpLayerBatchViews BatchViews = 
+					OpImageLayerInternal::GetOpLayerBatchLODRangeViews(BatchId, BatchArgs);
+			
+			uint8* BaseBuf = BatchViews.Result.GetData() + BaseChannelOffset;
+			
+			for (int32 I = 0; I < BatchViews.NumElems; ++I)
+			{
+				for (int32 C = 0; C < CHANNELS_TO_BLEND; ++C)
+				{
+					uint32 Base = BaseBuf[BASE_CHANNEL_STRIDE*I + C];
+					uint32 Blended = TopColor[C + ColorChannelOffset];
+					uint32 Result = BLEND_FUNC(Base, Blended);
+					if constexpr (CLAMP)
+					{
+						BaseBuf[BASE_CHANNEL_STRIDE*I + C] = (uint8)FMath::Min(255u, Result);
+					}
+					else
+					{
+						BaseBuf[BASE_CHANNEL_STRIDE*I + C] = (uint8)Result;
+					}
+				}
+			}
+		});
+	}
+
+
+	template<
+		uint32 (*BLEND_FUNC)(uint32, uint32), 
+		bool CLAMP, 
+		uint32 CHANNEL_COUNT>
+	FORCENOINLINE void BufferLayerColourInPlaceImpl(FImage* BaseImage, FVector4f Col, bool bOnlyOneMip, uint32 BaseOffset, uint8 ColOffset)
+	{
+		EImageFormat BaseFormat = BaseImage->GetFormat();
+
+		if (BaseFormat == EImageFormat::RGB_UByte)
+		{
+			check(BaseOffset + CHANNEL_COUNT <= 3);
+			BufferLayerColourFormatInPlace<BLEND_FUNC, CLAMP, CHANNEL_COUNT, 3>
+					(BaseImage, Col, BaseOffset, ColOffset, bOnlyOneMip);
+		}
+		else if (BaseFormat == EImageFormat::RGBA_UByte)
+		{
+			check(BaseOffset + CHANNEL_COUNT <= 4);
+			BufferLayerColourFormatInPlace<BLEND_FUNC, CLAMP, CHANNEL_COUNT, 4>
+					(BaseImage, Col, BaseOffset, ColOffset, bOnlyOneMip);
+		}
+		else if (BaseFormat == EImageFormat::BGRA_UByte)
+		{
+			float Temp = Col[0];
+			Col[0] = Col[2];
+			Col[2] = Temp;
+			BufferLayerColourFormatInPlace<BLEND_FUNC, CLAMP, CHANNEL_COUNT, 4>
+					(BaseImage, Col, BaseOffset, ColOffset, bOnlyOneMip);
+		}
+
+		else if (BaseFormat == EImageFormat::L_UByte)
+		{
+			check(BaseOffset + CHANNEL_COUNT <= 1);
+			BufferLayerColourFormatInPlace< BLEND_FUNC, CLAMP, CHANNEL_COUNT, 1>
+					(BaseImage, Col, BaseOffset, ColOffset, bOnlyOneMip);
+		}
+		else
+		{
+			checkf(false, TEXT("Unsupported format."));
+		}
+	}
+
+
+	template< 
+		uint32 (*BLEND_FUNC)(uint32, uint32), 
+		uint32 CHANNEL_COUNT>
+	void BufferLayerColourInPlace(FImage* BaseImage, FVector4f Color, bool bOnlyOneMip, uint32 BaseOffset, uint8 ColorOffset)
+	{
+		bool bIsClampNeeded = OpImageLayerInternal::IsAnyComponentLargerThan1(Color);
+		if (bIsClampNeeded)
+		{
+			BufferLayerColourInPlaceImpl<BLEND_FUNC, true, CHANNEL_COUNT>(BaseImage, Color, bOnlyOneMip, BaseOffset, ColorOffset);
+		}
+		else
+		{
+			BufferLayerColourInPlaceImpl<BLEND_FUNC, false, CHANNEL_COUNT>(BaseImage, Color, bOnlyOneMip, BaseOffset, ColorOffset);
+		}
+	}
+
+
+	template< 
+		uint32 (*BLEND_FUNC_MASKED)(uint32, uint32, uint32),
+		uint32 (*BLEND_FUNC)(uint32, uint32), 
+		bool CLAMP, 
+		int32 CHANNEL_COUNT>
+	FORCENOINLINE void BufferLayerColourInPlaceImpl(FImage* BaseImage, const FImage* MaskImage, FVector4f Color, bool bOnlyOneMip, uint32 BaseOffset, uint8 ColorOffset)
+	{
+		check(BaseImage->GetSizeX() == MaskImage->GetSizeX());
+		check(BaseImage->GetSizeY() == MaskImage->GetSizeY());
+
+		EImageFormat BaseFormat = BaseImage->GetFormat();
+
+		if (BaseFormat == EImageFormat::RGB_UByte)
+		{
+			check(BaseOffset + CHANNEL_COUNT <= 3);
+			BufferLayerColourFormat<BLEND_FUNC_MASKED, BLEND_FUNC, CLAMP, CHANNEL_COUNT, 3>
+					(BaseImage, BaseImage, MaskImage, Color, BaseOffset, ColorOffset, bOnlyOneMip);
+		}
+		else if (BaseFormat == EImageFormat::RGBA_UByte)
+		{
+			check(BaseOffset + CHANNEL_COUNT <= 4);
+			BufferLayerColourFormat<BLEND_FUNC_MASKED, BLEND_FUNC, CLAMP, CHANNEL_COUNT, 4>
+					(BaseImage, BaseImage, MaskImage, Color, BaseOffset, ColorOffset, bOnlyOneMip);
+		}
+		else if (BaseFormat == EImageFormat::BGRA_UByte)
+		{
+			check(BaseOffset + CHANNEL_COUNT <= 4);
+			float Temp = Color[0];
+			Color[0] = Color[2];
+			Color[2] = Temp;
+
+			BufferLayerColourFormat<BLEND_FUNC_MASKED, BLEND_FUNC, CLAMP, CHANNEL_COUNT, 4>
+					(BaseImage, BaseImage, MaskImage, Color, BaseOffset, ColorOffset, bOnlyOneMip);
+		}
+		else if (BaseFormat == EImageFormat::L_UByte)
+		{
+			check(BaseOffset + CHANNEL_COUNT <= 1);
+			BufferLayerColourFormat<BLEND_FUNC_MASKED, BLEND_FUNC, CLAMP, CHANNEL_COUNT, 1>
+					(BaseImage, BaseImage, MaskImage, Color, BaseOffset, ColorOffset, bOnlyOneMip);
+		}
+		else
+		{
+			checkf(false, TEXT("Unsupported format."));
+		}
+	}
+
+	template< 
+		uint32 (*BLEND_FUNC_MASKED)(uint32, uint32, uint32), 
+		uint32 (*BLEND_FUNC)(uint32, uint32), uint32 CHANNEL_COUNT>
+	void BufferLayerColourInPlace(FImage* BaseImage, const FImage* MaskImage, FVector4f Color, bool bOnlyOneMip, uint32 BaseOffset, uint8 ColorOffset)
+	{
+		bool bIsClampNeeded = OpImageLayerInternal::IsAnyComponentLargerThan1(Color);
+		if (bIsClampNeeded)
+		{
+			BufferLayerColourInPlaceImpl<BLEND_FUNC_MASKED, BLEND_FUNC, true, CHANNEL_COUNT>
+					(BaseImage, MaskImage, Color, bOnlyOneMip, BaseOffset, ColorOffset);
+		}
+		else
+		{
+			BufferLayerColourInPlaceImpl<BLEND_FUNC_MASKED, BLEND_FUNC, false, CHANNEL_COUNT>
+					(BaseImage, MaskImage, Color, bOnlyOneMip, BaseOffset, ColorOffset);
+		}
+	}
+
+
+	/**	
+	* Apply a blending function to an image with another image as blending layer
+	*/
+	template< uint32 (*BLEND_FUNC)(uint32, uint32), bool CLAMP,
+		// Number of total channels to actually process
+		int32 CHANNELS_TO_BLEND,
+		// Number of total channels in the base image
+		int32 BASE_CHANNEL_STRIDE,
+		// Number of total channels in the blend image
+		int32 BLENDED_CHANNEL_STRIDE>
+    void BufferLayerFormatInPlace(
+		FImage* BaseImage, 
+		const FImage* BlendedImage,
+		uint32 BaseChannelOffset,
+		uint32 BlendedChannelOffset,
+		bool bOnlyFirstLOD)
+	{
+		check(BaseImage->GetSizeX() == BlendedImage->GetSizeX());
+		check(BaseImage->GetSizeY() == BlendedImage->GetSizeY());
+
+		// No longer required.
+		//check(BaseImage->GetFormat() == BlendedImage->GetFormat());
+		check(BaseChannelOffset + CHANNELS_TO_BLEND <= GetImageFormatData(BaseImage->GetFormat()).Channels);
+		check(BlendedChannelOffset + CHANNELS_TO_BLEND <= GetImageFormatData(BlendedImage->GetFormat()).Channels);
+
+		check(bOnlyFirstLOD || BaseImage->GetLODCount() <= BlendedImage->GetLODCount());
+
+		const int32 BaseBytesPerElem = GetImageFormatData(BaseImage->GetFormat()).BytesPerBlock;
+		const int32 BlendBytesPerElem = GetImageFormatData(BlendedImage->GetFormat()).BytesPerBlock;
+	
+		const int32 NumLODs = bOnlyFirstLOD ? 1 : BaseImage->GetLODCount();
+
+		constexpr int32 BatchNumElems = 4096*2;
+
+		OpImageLayerInternal::FOpLayerBatchArgs BatchArgs;
+		BatchArgs.BatchNumElems      = BatchNumElems;
+		BatchArgs.LODBegin           = 0;
+		BatchArgs.LODEnd             = NumLODs;
+		BatchArgs.Result             = BaseImage;
+		BatchArgs.Blend              = BlendedImage;
+		BatchArgs.ResultBytesPerElem = BaseBytesPerElem;
+		BatchArgs.BlendBytesPerElem  = BlendBytesPerElem;
+
+		const int32 NumBatches = OpImageLayerInternal::GetOpLayerNumBatchesLODRange(BatchArgs);
+
+		OpImageLayerInternal::InvokeParallelFor(NumBatches,
+		[
+			&BatchArgs, BaseChannelOffset, BlendedChannelOffset
+		] (int32 BatchId)
+		{
+			OpImageLayerInternal::FOpLayerBatchViews BatchViews = 
+					OpImageLayerInternal::GetOpLayerBatchLODRangeViews(BatchId, BatchArgs);
+
+        	uint8* BaseBuf = BatchViews.Result.GetData() + BaseChannelOffset;
+        	const uint8* BlendedBuf = BatchViews.Blend.GetData() + BlendedChannelOffset;
+
+			for (int32 I = 0; I < BatchViews.NumElems; ++I)
+			{
+				for (int32 C = 0; C < CHANNELS_TO_BLEND; ++C)
+				{
+					uint32 Base = BaseBuf[BASE_CHANNEL_STRIDE * I + C];
+					uint32 Blended = BlendedBuf[BLENDED_CHANNEL_STRIDE * I + C];
+					uint32 Result = BLEND_FUNC(Base, Blended);
+					
+					if constexpr (CLAMP)
+					{
+						BaseBuf[BASE_CHANNEL_STRIDE * I + C] = (uint8)FMath::Min(255u, Result);
+					}
+					else
+					{
+						BaseBuf[BASE_CHANNEL_STRIDE * I + C] = (uint8)Result;
+					}
+				}
+			}
+		});
+	}
+	
+	/**
+	* Apply a blending function to an image with another image as blending layer
+	*/
+	template<
+		uint32 (*BLEND_FUNC)(uint32, uint32), 
+		bool CLAMP,
+		// Number of total channels to actually process
+		int32 CHANNELS_TO_BLEND,
+		// Number of total channels in the base image
+		int32 BASE_CHANNEL_STRIDE,
+		// Number of total channels in the blend image
+		int32 BLENDED_CHANNEL_STRIDE,
+		int32 BLENDED_CHANNEL_OFFSET>
+	void BufferLayerFormat(FImage* DestImage, const FImage* BaseImage, const FImage* BlendedImage, bool bOnlyFirstLOD)
+	{
+		check(BaseImage->GetSizeX() == BlendedImage->GetSizeX());
+		check(BaseImage->GetSizeY() == BlendedImage->GetSizeY());
+		// Not true anymore, since the BLENDED_CHANNEL_OFFSET has been added.
+		// check(BaseImage->GetFormat() == BlendedImage->GetFormat());
+		check(bOnlyFirstLOD || BaseImage->GetLODCount() <= BlendedImage->GetLODCount());
+
+		constexpr int32 UnblendedChannels = BASE_CHANNEL_STRIDE - CHANNELS_TO_BLEND;
+		constexpr int32 BatchNumElems = 4096*2;
+		
+		const int32 BaseBytesPerElem = GetImageFormatData(BaseImage->GetFormat()).BytesPerBlock;
+		const int32 BlendBytesPerElem = GetImageFormatData(BlendedImage->GetFormat()).BytesPerBlock;
+		const int32 DestBytesPerElem = GetImageFormatData(DestImage->GetFormat()).BytesPerBlock;
+
+		const int32 NumLODs = bOnlyFirstLOD ? 1 : BaseImage->GetLODCount();
+
+		OpImageLayerInternal::FOpLayerBatchArgs BatchArgs;
+		BatchArgs.BatchNumElems      = BatchNumElems;
+		BatchArgs.LODBegin           = 0;
+		BatchArgs.LODEnd             = NumLODs;
+		BatchArgs.Result             = DestImage;
+		BatchArgs.Base               = BaseImage;
+		BatchArgs.Blend              = BlendedImage;
+		BatchArgs.ResultBytesPerElem = DestBytesPerElem;
+		BatchArgs.BaseBytesPerElem   = BaseBytesPerElem;
+		BatchArgs.BlendBytesPerElem  = BlendBytesPerElem;
+
+		const int32 NumBatches = OpImageLayerInternal::GetOpLayerNumBatchesLODRange(BatchArgs);
+
+		OpImageLayerInternal::InvokeParallelFor(NumBatches,
+		[
+			&BatchArgs, UnblendedChannels
+		] (int32 BatchId)
+		{
+			OpImageLayerInternal::FOpLayerBatchViews BatchViews = 
+					OpImageLayerInternal::GetOpLayerBatchLODRangeViews(BatchId, BatchArgs);
+	
+			const uint8* BaseBuf = BatchViews.Base.GetData();
+			const uint8* BlendedBuf = BatchViews.Blend.GetData() + BLENDED_CHANNEL_OFFSET;
+			uint8* DestBuf = BatchViews.Result.GetData();
+
+			for (int32 I = 0; I < BatchViews.NumElems; ++I)
+			{
+				for (int32 C = 0; C < CHANNELS_TO_BLEND; ++C)
+				{
+					uint32 Base = BaseBuf[BASE_CHANNEL_STRIDE * I + C];
+					uint32 Blended = BlendedBuf[BLENDED_CHANNEL_STRIDE * I + C];
+					uint32 Result = BLEND_FUNC(Base, Blended);
+
+					if constexpr (CLAMP)
+					{
+						DestBuf[BASE_CHANNEL_STRIDE * I + C] = (uint8)FMath::Min(255u, Result);
+					}
+					else
+					{
+						DestBuf[BASE_CHANNEL_STRIDE * I + C] = (uint8)Result;
+					}
+				}
+
+				// Copy the unblended channels
+				// \TODO: unnecessary when doing it in-place?
+				if constexpr (UnblendedChannels > 0)
+				{
+					for (int32 C = 0; C < UnblendedChannels; ++C)
+					{
+						DestBuf[BASE_CHANNEL_STRIDE * I + CHANNELS_TO_BLEND + C] 
+							= BaseBuf[BASE_CHANNEL_STRIDE * I + CHANNELS_TO_BLEND + C];
+					}
+				}
+			}
+		});
+	}
+
+	/**
+	* Apply a blending function to an image with another image as blending layer
+	*/
+	template< uint32 (*BLEND_FUNC)(uint32, uint32), bool CLAMP >
+	void BufferLayer(FImage* ResultImage, const FImage* BaseImage, const FImage* BlendedImage, bool bApplyToAlpha, bool bOnlyOneMip, bool bUseBlendSourceFromBlendAlpha)
+	{
+		check(ResultImage->GetFormat() == BaseImage->GetFormat());
+		check(ResultImage->GetSizeX() == BaseImage->GetSizeX());
+		check(ResultImage->GetSizeY() == BaseImage->GetSizeY());
+		check(bOnlyOneMip || ResultImage->GetLODCount() == BaseImage->GetLODCount());
+		check(BaseImage->GetSizeX() == BlendedImage->GetSizeX());
+		check(BaseImage->GetSizeY() == BlendedImage->GetSizeY());
+		check(bOnlyOneMip || ResultImage->GetLODCount() <= BlendedImage->GetLODCount());
+
+		EImageFormat BaseFormat = BaseImage->GetFormat();
+		EImageFormat BlendedFormat = BlendedImage->GetFormat();
+
+		if (bUseBlendSourceFromBlendAlpha)
+		{
+			if (BlendedFormat == EImageFormat::RGBA_UByte || BlendedFormat == EImageFormat::BGRA_UByte)
+			{
+				if (BaseFormat == EImageFormat::L_UByte)
+				{
+					BufferLayerFormat<BLEND_FUNC, CLAMP, 1, 1, 4, 3>(ResultImage, BaseImage, BlendedImage, bOnlyOneMip);
+				}
+				else
+				{
+					checkf(false, TEXT("Unsupported format."));
+				}
+			}
+			else if (BlendedFormat == EImageFormat::L_UByte)
+			{
+				BufferLayerFormat<BLEND_FUNC, CLAMP, 1, 1, 1, 0>(ResultImage, BaseImage, BlendedImage, bOnlyOneMip);
+			}
+		}
+		else
+		{
+			check(BaseFormat == BlendedFormat);
+			if (BaseFormat == EImageFormat::RGB_UByte)
+			{
+				check(!bUseBlendSourceFromBlendAlpha);
+				BufferLayerFormat<BLEND_FUNC, CLAMP, 3, 3, 3, 0>(ResultImage, BaseImage, BlendedImage, bOnlyOneMip);
+			}
+			else if (BaseFormat == EImageFormat::RGBA_UByte || BaseFormat == EImageFormat::BGRA_UByte)
+			{
+				check(!bUseBlendSourceFromBlendAlpha);
+				if (bApplyToAlpha)
+				{
+					BufferLayerFormat<BLEND_FUNC, CLAMP, 4, 4, 4, 0>(ResultImage, BaseImage, BlendedImage, bOnlyOneMip);
+				}
+				else
+				{
+					BufferLayerFormat<BLEND_FUNC, CLAMP, 3, 4, 4, 0>(ResultImage, BaseImage, BlendedImage, bOnlyOneMip);
+				}
+			}
+			else if (BaseFormat == EImageFormat::L_UByte)
+			{
+				BufferLayerFormat<BLEND_FUNC, CLAMP, 1, 1, 1, 0>(ResultImage, BaseImage, BlendedImage, bOnlyOneMip);
+			}
+			else
+			{
+				checkf(false, TEXT("Unsupported format."));
+			}
+		}
+	}
+
+
+	template< uint32 (*BLEND_FUNC)(uint32, uint32), bool CLAMP, int32 CHANNEL_COUNT >
+	void BufferLayerInPlace(FImage* BaseImage, const FImage* BlendedImage, bool bOnlyOneMip, uint32 BaseOffset, uint32 BlendedOffset)
+	{
+		check(BaseImage->GetSizeX() == BlendedImage->GetSizeX());
+		check(BaseImage->GetSizeY() == BlendedImage->GetSizeY());
+		// Not required since we have the CHANNEL_COUNT and offsets. 
+		// check(BaseImage->GetFormat() == BlendedImage->GetFormat());
+
+		EImageFormat BaseFormat = BaseImage->GetFormat();
+		EImageFormat BlendFormat = BlendedImage->GetFormat();
+
+		if (BaseFormat == EImageFormat::RGB_UByte && BlendFormat==EImageFormat::RGB_UByte)
+		{
+			check(BaseOffset + CHANNEL_COUNT <= 3);
+			check(BlendedOffset + CHANNEL_COUNT <= 3);
+			BufferLayerFormatInPlace<BLEND_FUNC, CLAMP, CHANNEL_COUNT, 3, 3>(BaseImage, BlendedImage, BaseOffset, BlendedOffset, bOnlyOneMip);
+		}
+		else if ((BaseFormat == EImageFormat::RGBA_UByte || BaseFormat == EImageFormat::BGRA_UByte) &&
+				 (BlendFormat == EImageFormat::RGBA_UByte || BlendFormat == EImageFormat::BGRA_UByte) )
+		{
+			check(BaseOffset + CHANNEL_COUNT <= 4);
+			check(BlendedOffset + CHANNEL_COUNT <= 4);
+			BufferLayerFormatInPlace< BLEND_FUNC, CLAMP, CHANNEL_COUNT, 4, 4>
+					(BaseImage, BlendedImage, BaseOffset, BlendedOffset, bOnlyOneMip);
+		}
+		else if ((BaseFormat == EImageFormat::RGBA_UByte || BaseFormat == EImageFormat::BGRA_UByte) &&
+				 BlendFormat == EImageFormat::L_UByte )
+		{
+			check(BaseOffset + CHANNEL_COUNT <= 4);
+			check(BlendedOffset + CHANNEL_COUNT <= 1);
+			BufferLayerFormatInPlace< BLEND_FUNC, CLAMP, CHANNEL_COUNT, 4, 1>
+					(BaseImage, BlendedImage, BaseOffset, BlendedOffset, bOnlyOneMip);
+		}
+		else if (BaseFormat == EImageFormat::L_UByte && BlendFormat==EImageFormat::L_UByte)
+		{
+			check(BaseOffset + CHANNEL_COUNT <= 1);
+			BufferLayerFormatInPlace< BLEND_FUNC, CLAMP, CHANNEL_COUNT, 1, 1>
+					(BaseImage, BlendedImage, BaseOffset, BlendedOffset, bOnlyOneMip);
+		}
+		else
+		{
+			checkf(false, TEXT("Unsupported format."));
+		}
+	}
+
+
+	template< uint32 (*BLEND_FUNC_MASKED)(uint32, uint32, uint32),
+		uint32 (*BLEND_FUNC)(uint32, uint32),
+		bool CLAMP,
+		// Number of total channels to actually process
+		int32 CHANNELS_TO_BLEND,
+		// Number of total channels in the base image
+		int32 BASE_CHANNEL_STRIDE,
+		// Number of total channels in the blend image
+		int32 BLENDED_CHANNEL_STRIDE>
+	void BufferLayerFormat(
+		FImage* DestImage,
+		const FImage* BaseImage,
+		const FImage* MaskImage,
+		const FImage* BlendImage,
+		uint32 DestOffset,
+		uint32 BaseChannelOffset,
+		uint32 BlendedChannelOffset,
+		bool bOnlyFirstLOD)
+	{
+		check(BaseImage->GetSizeX() == MaskImage->GetSizeX() && BaseImage->GetSizeY() == MaskImage->GetSizeY());
+		check(BaseImage->GetSizeX() == BlendImage->GetSizeX() && BaseImage->GetSizeY() == BlendImage->GetSizeY());
+		check(bOnlyFirstLOD || BaseImage->GetLODCount() <= MaskImage->GetLODCount());
+		check(bOnlyFirstLOD || BaseImage->GetLODCount() <= BlendImage->GetLODCount());
+		check(BaseImage->GetFormat() == BlendImage->GetFormat());
+		check(MaskImage->GetFormat() == EImageFormat::L_UByte ||
+			  MaskImage->GetFormat() == EImageFormat::L_UByteRLE);
+
+		const EImageFormat MaskFormat = MaskImage->GetFormat();
+
+		const int32 BaseBytesPerElem = GetImageFormatData(BaseImage->GetFormat()).BytesPerBlock; 
+		const int32 BlendBytesPerElem = GetImageFormatData(BlendImage->GetFormat()).BytesPerBlock; 
+
+		const int32 NumLODs = bOnlyFirstLOD ? 1 : BaseImage->GetLODCount();
+
+		const bool bIsMaskUncompressed = MaskFormat == EImageFormat::L_UByte;
+		constexpr int32 UnblendedChannels = BASE_CHANNEL_STRIDE - CHANNELS_TO_BLEND;
+
+		if (bIsMaskUncompressed)
+		{	
+			constexpr int32 BatchNumElems = 4096*2;
+
+			OpImageLayerInternal::FOpLayerBatchArgs BatchArgs;
+			BatchArgs.BatchNumElems      = BatchNumElems;
+			BatchArgs.LODBegin           = 0;
+			BatchArgs.LODEnd             = NumLODs;
+			BatchArgs.Result             = DestImage;
+			BatchArgs.Base               = BaseImage;
+			BatchArgs.Blend              = BlendImage;
+			BatchArgs.Mask               = MaskImage;
+			BatchArgs.ResultBytesPerElem = BaseBytesPerElem;
+			BatchArgs.BaseBytesPerElem   = BaseBytesPerElem;
+			BatchArgs.BlendBytesPerElem  = BlendBytesPerElem;
+			BatchArgs.MaskBytesPerElem   = 1;
+
+			const int32 NumBatches = OpImageLayerInternal::GetOpLayerNumBatchesLODRange(BatchArgs);
+
+			OpImageLayerInternal::InvokeParallelFor(NumBatches,
+			[
+				&BatchArgs, UnblendedChannels
+			] (uint32 BatchId)
+			{
+				OpImageLayerInternal::FOpLayerBatchViews BatchViews = 
+						OpImageLayerInternal::GetOpLayerBatchLODRangeViews(BatchId, BatchArgs);
+
+				const uint8* BaseBuf = BatchViews.Base.GetData();
+				const uint8* BlendedBuf = BatchViews.Blend.GetData();
+				const uint8* MaskBuf = BatchViews.Mask.GetData();
+				uint8* DestBuf = BatchViews.Result.GetData();
+
+				// This could happen in case of missing data files.
+				if (!BaseBuf || !BlendedBuf || !MaskBuf || !DestBuf)
+				{
+					return;
+				}
+
+				for (int32 I = 0; I < BatchViews.NumElems; ++I)
+				{
+					uint32 Mask = MaskBuf[I];
+					for (int32 C = 0; C < CHANNELS_TO_BLEND; ++C)
+					{
+						const uint32 Base = BaseBuf[BASE_CHANNEL_STRIDE * I + C];
+						const uint32 Blended = BlendedBuf[BLENDED_CHANNEL_STRIDE * I + C];
+						const uint32 Result = BLEND_FUNC_MASKED(Base, Blended, Mask);
+						if constexpr (CLAMP)
+						{
+							DestBuf[BASE_CHANNEL_STRIDE * I + C] = (uint8)FMath::Min(255u, Result);
+						}
+						else
+						{
+							DestBuf[BASE_CHANNEL_STRIDE * I + C] = (uint8)Result;
+						}
+					}
+					// Copy the unblended channels
+					// \TODO: unnecessary when doing it in-place?
+					if constexpr (UnblendedChannels > 0)
+					{
+						for (int32 C = 0; C < UnblendedChannels; ++C)
+						{
+							DestBuf[BASE_CHANNEL_STRIDE * I + CHANNELS_TO_BLEND + C] = BaseBuf[BASE_CHANNEL_STRIDE * I + CHANNELS_TO_BLEND + C];
+						}
+					}
+				}
+			});
+		}
+		else if (MaskFormat == EImageFormat::L_UByteRLE)
+		{
+			int32 Rows = BaseImage->GetSizeY();
+			int32 Width = BaseImage->GetSizeX();
+
+			for (int32 LOD = 0; LOD < NumLODs; ++LOD)
+			{
+				const uint8* BaseBuf = BaseImage->GetLODData(LOD);
+				uint8* DestBuf = DestImage->GetLODData(LOD);
+				const uint8* BlendedBuf = BlendImage->GetLODData(LOD);
+				const uint8* MaskBuf = MaskImage->GetLODData(LOD);
+
+				// This could happen in case of missing data files.
+				if (!BaseBuf || !BlendedBuf || !MaskBuf || !DestBuf)
+				{
+					continue;
+				}
+
+				// Remove RLE header, mip size and row sizes.
+				MaskBuf += sizeof(uint32) + Rows * sizeof(uint32);
+
+				for (int32 RowIndex = 0; RowIndex < Rows; ++RowIndex)
+				{
+					const uint8* DestRowEnd = DestBuf + Width * BASE_CHANNEL_STRIDE;
+					while (DestBuf != DestRowEnd)
+					{
+						// Decode header
+						uint16 Equal;
+						FMemory::Memcpy(&Equal, MaskBuf, 2);
+						MaskBuf += 2;
+
+						uint8 Different = *MaskBuf;
+						++MaskBuf;
+
+						uint8 EqualPixel = *MaskBuf;
+						++MaskBuf;
+
+						// Equal pixels
+						//check(DestBuf + BASE_CHANNEL_STRIDE * Equal <= BaseImage->GetLODDataSize(0));
+						if (EqualPixel == 255)
+						{
+							for (int32 I = 0; I < Equal; ++I)
+							{
+								for (int32 C = 0; C < CHANNELS_TO_BLEND; ++C)
+								{
+									uint32 Base = BaseBuf[BASE_CHANNEL_STRIDE * I + C];
+									uint32 Blended = BlendedBuf[BLENDED_CHANNEL_STRIDE * I + C];
+									uint32 Result = BLEND_FUNC(Base, Blended);
+									if constexpr (CLAMP)
+									{
+										DestBuf[BASE_CHANNEL_STRIDE * I + C] = (uint8)FMath::Min(255u, Result);
+									}
+									else
+									{
+										DestBuf[BASE_CHANNEL_STRIDE * I + C] = (uint8)Result;
+									}
+								}
+
+								// Copy the unblended channels
+								// \TODO: unnecessary when doing it in-place?
+								if constexpr (UnblendedChannels > 0)
+								{
+									for (int32 C = 0; C < UnblendedChannels; ++C)
+									{
+										DestBuf[BASE_CHANNEL_STRIDE * I + CHANNELS_TO_BLEND + C] = BaseBuf[BASE_CHANNEL_STRIDE * I + CHANNELS_TO_BLEND + C];
+									}
+								}
+							}
+						}
+						else if (EqualPixel > 0)
+						{
+							for (int32 I = 0; I < Equal; ++I)
+							{
+								for (int32 C = 0; C < CHANNELS_TO_BLEND; ++C)
+								{
+									uint32 Base = BaseBuf[BASE_CHANNEL_STRIDE * I + C];
+									uint32 Blended = BlendedBuf[BLENDED_CHANNEL_STRIDE * I + C];
+									uint32 Result = BLEND_FUNC_MASKED(Base, Blended, EqualPixel);
+
+									if constexpr (CLAMP)
+									{
+										DestBuf[BASE_CHANNEL_STRIDE * I + C] = (uint8)FMath::Min(255u, Result);
+									}
+									else
+									{
+										DestBuf[BASE_CHANNEL_STRIDE * I + C] = (uint8)Result;
+									}
+								}
+
+								// Copy the unblended channels
+								// \TODO: unnecessary when doing it in-place?
+								if constexpr (UnblendedChannels > 0)
+								{
+									for (int32 C = 0; C < UnblendedChannels; ++C)
+									{
+										DestBuf[BASE_CHANNEL_STRIDE * I + CHANNELS_TO_BLEND + C] = BaseBuf[BASE_CHANNEL_STRIDE * I + CHANNELS_TO_BLEND + C];
+									}
+								}
+							}
+						}
+						else
+						{
+							// It could happen if xxxxxOnBase
+							if (DestBuf != BaseBuf)
+							{
+								FMemory::Memmove(DestBuf, BaseBuf, BASE_CHANNEL_STRIDE * Equal);
+							}
+						}
+
+						DestBuf += BASE_CHANNEL_STRIDE * Equal;
+						BaseBuf += BASE_CHANNEL_STRIDE * Equal;
+						BlendedBuf += BLENDED_CHANNEL_STRIDE * Equal;
+
+						// Different pixels
+						//check(DestBuf + BASE_CHANNEL_STRIDE * Different <= BaseImage->GetDataSize(0));
+						for (int32 I = 0; I < Different; ++I)
+						{
+							for (int32 C = 0; C < CHANNELS_TO_BLEND; ++C)
+							{
+								uint32 Mask = MaskBuf[I];
+								uint32 Base = BaseBuf[BASE_CHANNEL_STRIDE * I + C];
+								uint32 Blended = BlendedBuf[BLENDED_CHANNEL_STRIDE * I + C];
+								uint32 Result = BLEND_FUNC_MASKED(Base, Blended, Mask);
+								if (CLAMP)
+								{
+									DestBuf[BASE_CHANNEL_STRIDE * I + C] = (uint8)FMath::Min(255u, Result);
+								}
+								else
+								{
+									DestBuf[BASE_CHANNEL_STRIDE * I + C] = (uint8)Result;
+								}
+							}
+
+							// Copy the unblended channels
+							// \TODO: unnecessary when doing it in-place?
+							if constexpr (UnblendedChannels > 0)
+							{
+								for (int32 C = 0; C < UnblendedChannels; ++C)
+								{
+									DestBuf[BASE_CHANNEL_STRIDE * I + CHANNELS_TO_BLEND + C] = BaseBuf[BASE_CHANNEL_STRIDE * I + CHANNELS_TO_BLEND + C];
+								}
+							}
+						}
+
+						DestBuf += BASE_CHANNEL_STRIDE * Different;
+						BaseBuf += BASE_CHANNEL_STRIDE * Different;
+						BlendedBuf += BLENDED_CHANNEL_STRIDE * Different;
+						MaskBuf += Different;
+					}
+				}
+
+				Rows = FMath::DivideAndRoundUp(Rows, 2);
+				Width = FMath::DivideAndRoundUp(Width, 2);
+			}
+		}
+		else
+		{
+			checkf(false, TEXT("Unsupported mask format."));
+		}
+	}
+
+	template< uint32 (*BLEND_FUNC_MASKED)(uint32, uint32, uint32),
+		uint32 (*BLEND_FUNC)(uint32, uint32),
+		bool CLAMP,
+		// Number of total channels to actually process
+		int32 CHANNELS_TO_BLEND,
+		// Number of total channels in the base image
+		int32 BASE_CHANNEL_STRIDE,
+		// Number of total channels in the blend image
+		int32 BLENDED_CHANNEL_STRIDE>
+	void BufferLayerFormatEmbeddedMask(
+		FImage* DestImage,
+		const FImage* BaseImage,
+		const FImage* BlendImage,
+		uint32 DestOffset,
+		uint32 BaseChannelOffset,
+		bool bOnlyFirstLOD)
+	{
+		check(BaseImage->GetSizeX() == BlendImage->GetSizeX() && BaseImage->GetSizeY() == BlendImage->GetSizeY());
+		check(bOnlyFirstLOD || BaseImage->GetLODCount() <= BlendImage->GetLODCount());
+
+		constexpr int32 UnblendedChannels = BASE_CHANNEL_STRIDE - CHANNELS_TO_BLEND;
+
+		const int32 BytesPerElem = GetImageFormatData(BaseImage->GetFormat()).BytesPerBlock; 
+	
+		const int32 NumLODs = bOnlyFirstLOD ? 1 : BaseImage->GetLODCount();
+
+		constexpr int32 BatchNumElems = 4096*2;
+		OpImageLayerInternal::FOpLayerBatchArgs BatchArgs;
+		BatchArgs.BatchNumElems      = BatchNumElems;
+		BatchArgs.LODBegin           = 0;
+		BatchArgs.LODEnd             = NumLODs;
+		BatchArgs.Result             = DestImage;
+		BatchArgs.Base               = BaseImage;
+		BatchArgs.Blend              = BlendImage;
+		BatchArgs.ResultBytesPerElem = BytesPerElem;
+		BatchArgs.BaseBytesPerElem   = BytesPerElem;
+		BatchArgs.BlendBytesPerElem  = BytesPerElem;
+
+		const int32 NumBatches = OpImageLayerInternal::GetOpLayerNumBatchesLODRange(BatchArgs);
+
+		OpImageLayerInternal::InvokeParallelFor(NumBatches, 
+		[
+			&BatchArgs, UnblendedChannels, BlendImageFormat = BlendImage->GetFormat()
+		] (int32 BatchId)
+		{
+			OpImageLayerInternal::FOpLayerBatchViews BatchViews = 
+					OpImageLayerInternal::GetOpLayerBatchLODRangeViews(BatchId, BatchArgs);
+
+			uint8* const DestBuf = BatchViews.Result.GetData();
+			uint8 const * const BlendedBuf = BatchViews.Blend.GetData();
+			uint8 const * const BaseBuf = BatchViews.Base.GetData(); 
+
+			// This could happen in case of missing data files.
+			if (!BaseBuf || !BlendedBuf)
+			{
+				return;
+			}
+
+			if (BlendImageFormat == EImageFormat::RGBA_UByte)
+			{
+				const uint8* MaskBuf = BlendedBuf + 3;
+
+				for (int32 I = 0; I < BatchViews.NumElems; ++I)
+				{
+					uint32 Mask = MaskBuf[BLENDED_CHANNEL_STRIDE * I];
+					for (int32 C = 0; C < CHANNELS_TO_BLEND; ++C)
+					{
+						uint32 Base = BaseBuf[BASE_CHANNEL_STRIDE * I + C];
+						uint32 Blended = BlendedBuf[BLENDED_CHANNEL_STRIDE * I + C];
+						uint32 Result = BLEND_FUNC_MASKED(Base, Blended, Mask);
+						if constexpr (CLAMP)
+						{
+							DestBuf[BASE_CHANNEL_STRIDE * I + C] = (uint8)FMath::Min(255u, Result);
+						}
+						else
+						{
+							DestBuf[BASE_CHANNEL_STRIDE * I + C] = (uint8)Result;
+						}
+					}
+					// Copy the unblended channels
+					// \TODO: unnecessary when doing it in-place?
+					if constexpr (UnblendedChannels > 0)
+					{
+						for (int32 C = 0; C < UnblendedChannels; ++C)
+						{
+							DestBuf[BASE_CHANNEL_STRIDE * I + CHANNELS_TO_BLEND + C] = BaseBuf[BASE_CHANNEL_STRIDE * I + CHANNELS_TO_BLEND + C];
+						}
+					}
+				}
+			}
+			else
+			{
+				for (int32 I = 0; I < BatchViews.NumElems; ++I)
+				{
+					for (int32 C = 0; C < CHANNELS_TO_BLEND; ++C)
+					{
+						uint32 Base = BaseBuf[BASE_CHANNEL_STRIDE * I + C];
+						uint32 Blended = BlendedBuf[BLENDED_CHANNEL_STRIDE * I + C];
+						uint32 Result = BLEND_FUNC(Base, Blended);
+						if constexpr (CLAMP)
+						{
+							DestBuf[BASE_CHANNEL_STRIDE * I + C] = (uint8)FMath::Min(255u, Result);
+						}
+						else
+						{
+							DestBuf[BASE_CHANNEL_STRIDE * I + C] = (uint8)Result;
+						}
+					}
+					// Copy the unblended channels
+					// \TODO: unnecessary when doing it in-place?
+					if constexpr (UnblendedChannels > 0)
+					{
+						for (int32 C = 0; C < UnblendedChannels; ++C)
+						{
+							DestBuf[BASE_CHANNEL_STRIDE * I + CHANNELS_TO_BLEND + C] = BaseBuf[BASE_CHANNEL_STRIDE * I + CHANNELS_TO_BLEND + C];
+						}
+					}
+				}
+			}
+		});
+
+	}
+
+	template< 
+		uint32 (*BLEND_FUNC)(uint32, uint32),
+		bool CLAMP,
+		uint32 NC>
+    void BufferLayerFormatStrideNoAlpha(
+		FImage* DestImage,
+		int32 DestOffset,
+		int32 Stride,
+		const FImage* MaskImage,
+		const FImage* BlendImage/*, int32 LODCount*/)
+	{
+        const uint8* MaskBuf = MaskImage->GetLODData(0);
+        const uint8* BlendedBuf = BlendImage->GetLODData(0);
+		uint8* DestBuf = DestImage->GetLODData(0) + DestOffset;
+
+		// This could happen in case of missing data files.
+		if (!BlendedBuf || !MaskBuf || !DestImage->GetLODData(0))
+		{
+			return;
+		}
+
+		EImageFormat MaskFormat = MaskImage->GetFormat();
+        bool bIsUncompressed = (MaskFormat == EImageFormat::L_UByte);
+
+        if (bIsUncompressed)
+		{
+			int32 RowCount = BlendImage->GetSizeY();
+			int32 PixelCount = BlendImage->GetSizeX();
+			for (int32 RowIndex = 0; RowIndex < RowCount; ++RowIndex)
+			{
+				for (int32 PixelIndex = 0; PixelIndex < PixelCount; ++PixelIndex)
+				{
+					uint32 Mask = *MaskBuf;
+					if (Mask)
+					{
+						for (int32 C = 0; C < NC; ++C)
+						{
+							uint32 Base = *DestBuf;
+							uint32 Blended = *BlendedBuf;
+							uint32 Result = BLEND_FUNC(Base, Blended);
+							if constexpr (CLAMP)
+							{
+								*DestBuf = (uint8)FMath::Min(255u, Result);
+							}
+							else
+							{
+								*DestBuf = (uint8)Result;
+							}
+							++DestBuf;
+							++BlendedBuf;
+						}
+					}
+					else
+					{
+						DestBuf += NC;
+						BlendedBuf += NC;
+					}
+					++MaskBuf;
+				}
+
+				DestBuf += Stride;
+			}
+		}
+        else if (MaskFormat == EImageFormat::L_UBitRLE)
+		{
+			int32 Rows = MaskImage->GetSizeY();
+			int32 Width = MaskImage->GetSizeX();
+
+            //for (int32 lod = 0; lod < LODCount; ++lod)
+            //{
+				// Remove RLE header.
+                MaskBuf += 4 + Rows*sizeof(uint32);
+
+                for (int32 RowIndex = 0; RowIndex < Rows; ++RowIndex)
+                {
+                    const uint8* DestRowEnd = DestBuf + Width*NC;
+                    while (DestBuf != DestRowEnd)
+                    {
+                        // Decode header
+                        uint16 Zeros = *(const uint16*)MaskBuf;
+                        MaskBuf += 2;
+
+                        uint16 Ones = *(const uint16*)MaskBuf;
+                        MaskBuf += 2;
+
+                        // Skip
+                        DestBuf += Zeros*NC;
+                        BlendedBuf += Zeros*NC;
+
+                        // Copy
+                        FMemory::Memmove(DestBuf, BlendedBuf, Ones*NC);
+
+                        DestBuf += NC*Ones;
+                        BlendedBuf += NC*Ones;
+                    }
+
+                    DestBuf += Stride;
+                }
+
+                //Rows = FMath::DivideAndRoundUp(Rows, 2);
+                //Width = FMath::DivideAndRoundUp(Width, 2);
+            //}
+		}
+		else
+		{
+			checkf( false, TEXT("Unsupported mask format.") );
+		}
+	}
+
+	template<uint32(*BLEND_FUNC_MASKED)(uint32, uint32, uint32),
+			 uint32(*BLEND_FUNC)(uint32, uint32),
+			 bool CLAMP>
+    void BufferLayer(FImage* DestImage,
+							const FImage* BaseImage,
+							const FImage* MaskImage,
+                            const FImage* BlendImage,
+                            bool bApplyToAlpha,
+							bool bOnlyFirstLOD)
+	{
+		if (BaseImage->GetFormat() == EImageFormat::RGB_UByte)
+		{
+            BufferLayerFormat<BLEND_FUNC_MASKED, BLEND_FUNC, CLAMP, 3, 3, 3>
+                    (DestImage, BaseImage, MaskImage, BlendImage, 0, 0, 0, bOnlyFirstLOD);
+		}
+        else if (BaseImage->GetFormat() == EImageFormat::RGBA_UByte || 
+				 BaseImage->GetFormat() == EImageFormat::BGRA_UByte)
+		{
+			if (bApplyToAlpha)
+			{
+				BufferLayerFormat<BLEND_FUNC_MASKED, BLEND_FUNC, CLAMP, 4, 4, 4>
+					(DestImage, BaseImage, MaskImage, BlendImage, 0, 0, 0, bOnlyFirstLOD);
+			}
+			else
+			{
+				BufferLayerFormat<BLEND_FUNC_MASKED, BLEND_FUNC, CLAMP, 3, 4, 4>
+					(DestImage, BaseImage, MaskImage, BlendImage, 0, 0, 0, bOnlyFirstLOD);
+			}
+		}
+		else if (BaseImage->GetFormat() == EImageFormat::L_UByte)
+		{
+            BufferLayerFormat<BLEND_FUNC_MASKED, BLEND_FUNC, CLAMP, 1, 1, 1>
+                    (DestImage, BaseImage, MaskImage, BlendImage, 0, 0, 0, bOnlyFirstLOD);
+		}
+        else
+		{
+			checkf(false, TEXT("Unsupported format."));
+		}
+	}
+
+
+	template< uint32 (*BLEND_FUNC_MASKED)(uint32, uint32, uint32),
+		uint32 (*BLEND_FUNC)(uint32, uint32),
+		bool CLAMP >
+	void BufferLayerEmbeddedMask(
+		FImage* DestImage,
+		const FImage* BaseImage,
+		const FImage* BlendImage,
+		bool bApplyToAlpha,
+		bool bOnlyFirstLOD)
+	{
+		if (BaseImage->GetFormat() == EImageFormat::RGB_UByte)
+		{
+			BufferLayerFormatEmbeddedMask<BLEND_FUNC_MASKED, BLEND_FUNC, CLAMP, 3, 3, 3>
+				(DestImage, BaseImage, BlendImage, 0, 0, bOnlyFirstLOD);
+		}
+		else if (BaseImage->GetFormat() == EImageFormat::RGBA_UByte || 
+				 BaseImage->GetFormat() == EImageFormat::BGRA_UByte)
+		{
+			if (bApplyToAlpha)
+			{
+				BufferLayerFormatEmbeddedMask<BLEND_FUNC_MASKED, BLEND_FUNC, CLAMP, 4, 4, 4>
+					(DestImage, BaseImage, BlendImage, 0, 0, bOnlyFirstLOD);
+			}
+			else
+			{
+				BufferLayerFormatEmbeddedMask<BLEND_FUNC_MASKED, BLEND_FUNC, CLAMP, 3, 4, 4>
+					(DestImage, BaseImage, BlendImage, 0, 0, bOnlyFirstLOD);
+			}
+		}
+		else
+		{
+			checkf(false, TEXT("Unsupported format."));
+		}
+	}
+
+
+	template< 
+		uint32 (*RGB_FUNC_MASKED)(uint32, uint32, uint32),
+		uint32 (*A_FUNC)(uint32, uint32),
+		bool CLAMP >
+	void BufferLayerComposite(
+		FImage* BaseImage,
+		const FImage* BlendImage,
+		bool bOnlyFirstLOD,
+		uint8 BlendAlphaSourceChannel)
+	{
+		check(BaseImage->GetFormat() == EImageFormat::RGBA_UByte);
+		check(BlendImage->GetFormat() == EImageFormat::RGBA_UByte);
+		check(BaseImage->GetSizeX() == BlendImage->GetSizeX() && BaseImage->GetSizeY() == BlendImage->GetSizeY());
+		check(bOnlyFirstLOD || BaseImage->GetLODCount() <= BlendImage->GetLODCount());
+
+		bOnlyFirstLOD = bOnlyFirstLOD || BaseImage->GetLODCount() == 1;
+		
+		int32 FirstLODDataOffset = 0;	
+		int32 NumRelevantElems = -1;
+
+		if (BlendImage->Flags & FImage::IF_HAS_RELEVANCY_MAP && bOnlyFirstLOD)
+		{
+			check(BlendImage->RelevancyMaxY < BaseImage->GetSizeY());
+			check(BlendImage->RelevancyMaxY >= BlendImage->RelevancyMinY);
+
+			uint16 SizeX = BaseImage->GetSizeX();
+			NumRelevantElems = (BlendImage->RelevancyMaxY - BlendImage->RelevancyMinY + 1) * SizeX * 4;
+			
+			FirstLODDataOffset = BlendImage->RelevancyMinY * SizeX * 4;
+		}
+
+		const int32 BytesPerElem = GetImageFormatData(BaseImage->GetFormat()).BytesPerBlock;
+		check(BytesPerElem == GetImageFormatData(BlendImage->GetFormat()).BytesPerBlock);
+
+		const int32 LODBegin = 0;
+		const int32 LODEnd = BaseImage->GetLODCount();
+
+		constexpr int32 BatchNumElems = 4096*2;
+
+		OpImageLayerInternal::FOpLayerBatchArgs BatchArgs;
+		BatchArgs.BatchNumElems      = BatchNumElems;
+		BatchArgs.LODBegin           = LODBegin;
+		BatchArgs.LODEnd             = LODEnd;
+		BatchArgs.FirstLODOffset     = bOnlyFirstLOD ? FirstLODDataOffset : -1;
+		BatchArgs.Result             = BaseImage;
+		BatchArgs.Blend              = BlendImage;
+		BatchArgs.ResultBytesPerElem = BytesPerElem;
+		BatchArgs.BlendBytesPerElem  = BytesPerElem;
+
+		const int32 NumBatches = OpImageLayerInternal::GetOpLayerNumBatchesLODRangeOffsetViews(BatchArgs);
+
+		// This will always be an upper-bound for bOnlyFirtsLODs, check if it performs as expected or it needs more fine tune.
+		const int32 NumRelevantBatches = NumRelevantElems != -1
+				? FMath::DivideAndRoundUp(NumRelevantElems, BatchNumElems)
+				: NumBatches;
+
+		OpImageLayerInternal::InvokeParallelFor(NumBatches, 
+		[
+			&BatchArgs, BlendAlphaSourceChannel
+		] (int32 BatchId)
+		{
+			OpImageLayerInternal::FOpLayerBatchViews BatchViews = 
+					OpImageLayerInternal::GetOpLayerBatchLODRangeOffsetViews(BatchId, BatchArgs);
+
+			uint8* BaseBuf = BatchViews.Result.GetData();
+			const uint8* BlendBuf = BatchViews.Blend.GetData();
+
+			for (int32 I = 0; I < BatchViews.NumElems; ++I)
+			{
+				// TODO: Optimize this (SIMD?)
+				uint32 Mask = BlendBuf[4 * I + 3];
+
+				// RGB
+				for (int32 C = 0; C < 3; ++C)
+				{
+					uint32 Base = BaseBuf[4 * I + C];
+					uint32 Blended = BlendBuf[4 * I + C];
+					uint32 Result = RGB_FUNC_MASKED(Base, Blended, Mask);
+					if constexpr (CLAMP)
+					{
+						BaseBuf[4 * I + C] = (uint8)FMath::Min(255u, Result);
+					}
+					else
+					{
+						BaseBuf[4 * I + C] = (uint8)Result;
+					}
+				}
+
+				// A
+				{
+					uint32 Base = BaseBuf[4 * I + 3];
+					uint32 Blended = BlendBuf[4 * I + BlendAlphaSourceChannel];
+					uint32 Result = A_FUNC(Base, Blended);
+					if constexpr (CLAMP)
+					{
+						BaseBuf[4 * I + 3] = (uint8)FMath::Min(255u, Result);
+					}
+					else
+					{
+						BaseBuf[4 * I + 3] = (uint8)Result;
+					}
+				}
+			}
+		});
+	}
+
+	template<>
+	void BufferLayerComposite<BlendChannelMasked, LightenChannel, false>
+	(
+		FImage* BaseImage,
+		const FImage* BlendImage,
+		bool bOnlyFirstLOD,
+		uint8 BlendAlphaSourceChannel)
+	{
+		check(BaseImage->GetFormat() == EImageFormat::RGBA_UByte);
+		check(BlendImage->GetFormat() == EImageFormat::RGBA_UByte);
+		check(BaseImage->GetSizeX() == BlendImage->GetSizeX() && BaseImage->GetSizeY() == BlendImage->GetSizeY());
+		check(bOnlyFirstLOD || BaseImage->GetLODCount() <= BlendImage->GetLODCount());
+		check(BlendAlphaSourceChannel == 3);
+
+		bOnlyFirstLOD = bOnlyFirstLOD || BaseImage->GetLODCount() == 1;
+
+		int32 FirstLODDataOffset = 0;
+		int32 NumRelevantElems = -1;
+		if (BlendImage->Flags & FImage::IF_HAS_RELEVANCY_MAP && bOnlyFirstLOD)
+		{
+			check(BlendImage->RelevancyMaxY < BaseImage->GetSizeY());
+			check(BlendImage->RelevancyMaxY >= BlendImage->RelevancyMinY);
+
+			uint16 SizeX = BaseImage->GetSizeX();
+			NumRelevantElems = (BlendImage->RelevancyMaxY - BlendImage->RelevancyMinY + 1) * SizeX;
+			
+			FirstLODDataOffset = BlendImage->RelevancyMinY * SizeX * 4;
+		}
+
+		const int32 NumLODs = BaseImage->GetLODCount();
+	
+		constexpr int32 BatchNumElems = 4096 * 2;
+		constexpr int32 BytesPerElem = 4;
+		OpImageLayerInternal::FOpLayerBatchArgs BatchArgs;
+		BatchArgs.BatchNumElems      = BatchNumElems;
+		BatchArgs.LODBegin           = 0;
+		BatchArgs.LODEnd             = NumLODs;
+		BatchArgs.FirstLODOffset     = bOnlyFirstLOD ? FirstLODDataOffset : -1;
+		BatchArgs.Result             = BaseImage;
+		BatchArgs.Blend              = BlendImage;
+		BatchArgs.ResultBytesPerElem = BytesPerElem;
+		BatchArgs.BlendBytesPerElem  = BytesPerElem;
+
+		const int32 NumBatches = OpImageLayerInternal::GetOpLayerNumBatchesLODRangeOffsetViews(BatchArgs);
+
+		// This will always be an upper-bound for bOnlyFirtsLODs, check if it performs as expected or it needs more fine tune.
+		const int32 NumRelevantBatches = NumRelevantElems != -1
+				? FMath::DivideAndRoundUp(NumRelevantElems, BatchNumElems)
+				: NumBatches;
+
+		OpImageLayerInternal::InvokeParallelFor(NumRelevantBatches,
+		[
+			&BatchArgs
+		] (uint32 BatchId)
+		{
+			OpImageLayerInternal::FOpLayerBatchViews BatchViews = 
+					OpImageLayerInternal::GetOpLayerBatchLODRangeOffsetViews(BatchId, BatchArgs);
+	
+			uint8* BaseBuf = BatchViews.Result.GetData();
+			const uint8* BlendBuf = BatchViews.Blend.GetData();
+			
+			for (int32 I = 0; I < BatchViews.NumElems; ++I)
+			{
+				// TODO: Optimize this (SIMD?)
+				uint32 FullBase;
+				FMemory::Memcpy(&FullBase, BaseBuf + I*sizeof(uint32), sizeof(uint32)); 
+
+				uint32 FullBlended;
+				FMemory::Memcpy(&FullBlended, BlendBuf + I*sizeof(uint32), sizeof(uint32)); 
+				uint32 Mask = (FullBlended & 0xff000000) >> 24;
+
+				uint32 FullResult = 0;
+				FullResult |= BlendChannelMasked((FullBase >>  0) & 0xff, (FullBlended >>  0) & 0xff, Mask) << 0;
+				FullResult |= BlendChannelMasked((FullBase >>  8) & 0xff, (FullBlended >>  8) & 0xff, Mask) << 8;
+				FullResult |= BlendChannelMasked((FullBase >> 16) & 0xff, (FullBlended >> 16) & 0xff, Mask) << 16;
+				FullResult |= LightenChannel	((FullBase >> 24) & 0xff, (FullBlended >> 24) & 0xff) << 24;
+
+				FMemory::Memcpy(BaseBuf + I*sizeof(uint32), &FullResult, sizeof(uint32));
+			}
+		});
+	}
+
+	template< 
+		VectorRegister4Int (*RGB_FUNC_MASKED)(const VectorRegister4Int&, const VectorRegister4Int&, const VectorRegister4Int&),
+		int32 (*A_FUNC)(int32, int32),
+		bool CLAMP >
+	void BufferLayerCompositeVector(
+		FImage* BaseImage,
+		const FImage* BlendImage,
+		bool bOnlyFirstLOD,
+		uint8 BlendAlphaSourceChannel)
+	{
+		check(BaseImage->GetFormat() == EImageFormat::RGBA_UByte);
+		check(BlendImage->GetFormat() == EImageFormat::RGBA_UByte);
+		check(BaseImage->GetSizeX() == BlendImage->GetSizeX() && BaseImage->GetSizeY() == BlendImage->GetSizeY());
+		check(bOnlyFirstLOD || BaseImage->GetLODCount() <= BlendImage->GetLODCount());
+
+		const int32 NumLODs = bOnlyFirstLOD ? 1 : BaseImage->GetLODCount();
+
+		constexpr int32 BatchNumElems = 4096*2;
+		constexpr int32 BytesPerElem = 4;
+		OpImageLayerInternal::FOpLayerBatchArgs BatchArgs;
+		BatchArgs.BatchNumElems      = BatchNumElems;
+		BatchArgs.LODBegin           = 0;
+		BatchArgs.LODEnd             = NumLODs;
+		BatchArgs.Result             = BaseImage;
+		BatchArgs.Blend              = BlendImage;
+		BatchArgs.ResultBytesPerElem = BytesPerElem;
+		BatchArgs.BlendBytesPerElem  = BytesPerElem;
+
+		const int32 NumBatches = OpImageLayerInternal::GetOpLayerNumBatchesLODRange(BatchArgs);
+
+		OpImageLayerInternal::InvokeParallelFor(NumBatches, 
+		[
+			&BatchArgs, BlendAlphaSourceChannel
+		] (uint32 BatchId)
+		{
+			OpImageLayerInternal::FOpLayerBatchViews BatchViews = 
+					OpImageLayerInternal::GetOpLayerBatchLODRangeViews(BatchId, BatchArgs);
+
+			uint8* BaseBuf = BatchViews.Result.GetData();
+			const uint8* BlendBuf = BatchViews.Blend.GetData();
+
+			for (int32 I = 0; I < BatchViews.NumElems; ++I)
+			{
+				// TODO: Optimize this (SIMD?)
+				const int32 BaseAlpha = BaseBuf[4 * I + BlendAlphaSourceChannel];
+				const int32 BlendedAlpha = BlendBuf[4 * I + BlendAlphaSourceChannel];
+
+				const VectorRegister4Int Mask = VectorIntSet1(BlendedAlpha);
+
+				uint32 BlendPixel;
+				FMemory::Memcpy(&BlendPixel, BlendBuf, sizeof(uint32));
+
+				uint32 BasePixel;
+				FMemory::Memcpy(&BasePixel, BaseBuf, sizeof(uint32));
+
+				const VectorRegister4Int Blended = MakeVectorRegisterInt(
+						(BlendPixel >> 0) & 0xFF, 
+						(BlendPixel >> 8) & 0xFF, 
+						(BlendPixel >> 16) & 0xFF, 
+						(BlendPixel >> 24) & 0xFF);
+
+				const VectorRegister4Int Base = MakeVectorRegisterInt(
+						(BasePixel >> 0) & 0xFF, 
+						(BasePixel >> 8) & 0xFF, 
+						(BasePixel >> 16) & 0xFF, 
+						(BasePixel >> 24) & 0xFF);
+
+				VectorRegister4Int Result = RGB_FUNC_MASKED(Base, Blended, Mask);
+				if constexpr (CLAMP)
+				{
+					Result = VectorIntMin(MakeVectorRegisterIntConstant(255, 255, 255, 255), Result);
+				}
+
+				int32 AlphaResult = A_FUNC(BaseAlpha, BlendedAlpha);
+				if constexpr (CLAMP)
+				{
+					AlphaResult = FMath::Min(255, AlphaResult);
+				}
+
+				alignas(VectorRegister4Int) int32 IndexableRegister[4];
+				VectorIntStoreAligned(Result, &IndexableRegister);
+
+				BaseBuf[4 * I + 0] = static_cast<uint8>(IndexableRegister[0]);
+				BaseBuf[4 * I + 1] = static_cast<uint8>(IndexableRegister[1]);
+				BaseBuf[4 * I + 2] = static_cast<uint8>(IndexableRegister[2]);
+				BaseBuf[4 * I + 3] = static_cast<uint8>(IndexableRegister[3]);
+
+				BaseBuf[4 * I + BlendAlphaSourceChannel] = static_cast<uint8>(AlphaResult);
+			}
+		});
+	}
+
+	template<uint32 (*BLEND_FUNC)(uint32,uint32),
+			 bool CLAMP>
+    void BufferLayerStrideNoAlpha(FImage* DestImage, int32 DestOffset, int32 Stride, const FImage* MaskImage, const FImage* BlendImage/*, int32 LODCount*/)
+	{
+		if (BlendImage->GetFormat() == EImageFormat::RGB_UByte)
+		{
+			BufferLayerFormatStrideNoAlpha<BLEND_FUNC, CLAMP, 3>
+					(DestImage, DestOffset, Stride, MaskImage, BlendImage/*, LODCount*/);
+		}
+        else if (BlendImage->GetFormat() == EImageFormat::RGBA_UByte || 
+				 BlendImage->GetFormat() == EImageFormat::BGRA_UByte)
+		{
+			BufferLayerFormatStrideNoAlpha<BLEND_FUNC, CLAMP, 4>
+					(DestImage, DestOffset, Stride, MaskImage, BlendImage/*, LODCount*/);
+		}
+		else if (BlendImage->GetFormat() == EImageFormat::L_UByte)
+		{
+			BufferLayerFormatStrideNoAlpha<BLEND_FUNC, CLAMP, 1>
+					(DestImage, DestOffset, Stride, MaskImage, BlendImage/*, LODCount*/);
+		}
+		else
+		{
+			checkf( false, TEXT("Unsupported format.") );
+		}
+	}
+
+
+	/**
+	* Apply a blending function to an image with another image as blending layer, on a subrect of
+	* the base image.
+	* \warning this method applies the blending function to the alpha channel too
+	* \warning this method uses the mask as a binary mask (>0)
+	*/
+	template<uint32 (*BLEND_FUNC)(uint32, uint32), bool CLAMP>
+	void ImageLayerOnBaseNoAlpha(
+			FImage* BaseImage,
+			const FImage* MaskImage,
+			const FImage* BlendedImage,
+			const box<FIntVector2>& Rect)
+	{
+		check(BaseImage->GetSizeX() >= Rect.min[0] + Rect.size[0]);
+		check(BaseImage->GetSizeY() >= Rect.min[1] + Rect.size[1]);
+		check(MaskImage->GetSizeX() == BlendedImage->GetSizeX());
+		check(MaskImage->GetSizeY() == BlendedImage->GetSizeY());
+		check(BaseImage->GetFormat() == BlendedImage->GetFormat());
+		check(MaskImage->GetFormat() == EImageFormat::L_UByte ||
+			  //UBYTE_RLE does not look to be supported.
+			  //MaskImage->GetFormat() == EImageFormat::L_UByteRLE || 
+			  MaskImage->GetFormat() == EImageFormat::L_UBitRLE);
+        check(BaseImage->GetLODCount() <= MaskImage->GetLODCount());
+        check(BaseImage->GetLODCount() <= BlendedImage->GetLODCount());
+
+		int32 PixelSize = GetImageFormatData(BaseImage->GetFormat()).BytesPerBlock;
+
+		int32 Start = (BaseImage->GetSizeX() * Rect.min[1] + Rect.min[0]) * PixelSize;
+		int32 Stride = (BaseImage->GetSizeX() - Rect.size[0]) * PixelSize;
+
+		// Stride is only valid for LOD 0, BufferLayerStride variants cannot operate on multiple lods.
+		// TODO: review if this needs to be supported, and implement using a rect lod reducction at this level.
+		BufferLayerStrideNoAlpha<BLEND_FUNC, CLAMP>(BaseImage, Start, Stride, MaskImage, BlendedImage/*, BaseImage->GetLODCount()*/);
+	}
+
+	template<uint32 NC>
+	FORCEINLINE uint32 PackPixel(const uint8* PixelPtr)
+	{
+		static_assert(NC > 0 && NC <= 4);
+
+		uint32 PixelPack = 0;
+
+		// The compiler should be able to optimize this given that NC is a constant expression.
+		FMemory::Memcpy(&PixelPack, PixelPtr, NC);
+
+		return PixelPack;
+	}
+
+	template<uint32 NC>
+	FORCEINLINE void UnpackPixel(uint8* PixelPtr, uint32 PixelData)
+	{
+		static_assert(NC > 0 && NC <= 4);
+
+		// The compiler should be able to optimize this given that NC is a constant expression
+		FMemory::Memcpy(PixelPtr, &PixelData, NC);
+	}
+
+	template< 
+		uint32 (*BLEND_FUNC)(uint32, uint32), 
+		uint32 NC>
+	void BufferLayerCombineColour(FImage* ResultImage, const FImage* BaseImage, FVector4f Color, bool bOnlyFirstLOD = false)
+	{
+		static_assert(NC > 0 && NC <= 4);
+
+		check(BaseImage->GetSizeX() == ResultImage->GetSizeX());
+		check(BaseImage->GetSizeY() == ResultImage->GetSizeY());
+
+		const uint32 TopColor = 
+			static_cast<uint32>(255.0f * Color[0]) << 0 |  
+			static_cast<uint32>(255.0f * Color[1]) << 8 |  
+			static_cast<uint32>(255.0f * Color[2]) << 16;  
+
+		const int32 NumLODs = bOnlyFirstLOD ? 1 : ResultImage->GetLODCount();
+		
+		constexpr int32 BatchNumElems = 4098*2;
+		OpImageLayerInternal::FOpLayerBatchArgs BatchArgs;
+		BatchArgs.BatchNumElems      = BatchNumElems;
+		BatchArgs.LODBegin           = 0;
+		BatchArgs.LODEnd             = NumLODs;
+		BatchArgs.Result             = ResultImage;
+		BatchArgs.Base               = BaseImage;
+		BatchArgs.ResultBytesPerElem = NC;
+		BatchArgs.BaseBytesPerElem   = NC;
+
+		const int32 NumBatches = OpImageLayerInternal::GetOpLayerNumBatchesLODRange(BatchArgs);
+
+		OpImageLayerInternal::InvokeParallelFor(NumBatches,
+		[
+			&BatchArgs, TopColor
+		](int32 BatchId)
+		{
+
+			OpImageLayerInternal::FOpLayerBatchViews BatchViews = 
+					OpImageLayerInternal::GetOpLayerBatchLODRangeViews(BatchId, BatchArgs);
+	
+			const uint8* BaseBuf = BatchViews.Base.GetData();
+			uint8* ResultBuf = BatchViews.Result.GetData();
+			
+			for (int32 I = 0; I < BatchViews.NumElems; ++I)
+			{
+				const uint32 Base = PackPixel<NC>(&BaseBuf[NC * I]);
+
+				const uint32 Result = BLEND_FUNC(Base, TopColor);
+				UnpackPixel<NC>(&ResultBuf[NC * I], Result);
+			}
+		});
+	}
+
+	template< 
+		uint32 (*BLEND_FUNC)(uint32, uint32), 
+		uint32 NC>
+	void BufferLayerCombine(FImage* ResultImage, const FImage* BaseImage, const FImage* BlendImage, bool bOnlyFirstLOD)
+	{
+		static_assert(NC > 0 && NC <= 4);
+
+		check(BaseImage->GetSizeX() == BlendImage->GetSizeX());
+		check(BaseImage->GetSizeY() == BlendImage->GetSizeY());
+		check(BaseImage->GetSizeX() == ResultImage->GetSizeX());
+		check(BaseImage->GetSizeY() == ResultImage->GetSizeY());
+		check(BaseImage->GetFormat() == BlendImage->GetFormat());
+		check(BaseImage->GetFormat() == ResultImage->GetFormat());
+
+		const int32 NumLODs = bOnlyFirstLOD ? 1 : ResultImage->GetLODCount();
+
+		constexpr int32 BatchNumElems = 4098*2;
+		OpImageLayerInternal::FOpLayerBatchArgs BatchArgs;
+		BatchArgs.BatchNumElems      = BatchNumElems;
+		BatchArgs.LODBegin           = 0;
+		BatchArgs.LODEnd             = NumLODs;
+		BatchArgs.Result             = ResultImage;
+		BatchArgs.Base               = BaseImage;
+		BatchArgs.ResultBytesPerElem = NC;
+		BatchArgs.BaseBytesPerElem   = NC;
+
+		const int32 NumBatches = OpImageLayerInternal::GetOpLayerNumBatchesLODRange(BatchArgs);
+
+		OpImageLayerInternal::InvokeParallelFor(NumBatches,
+		[
+			&BatchArgs
+		](int32 BatchId)
+		{
+			OpImageLayerInternal::FOpLayerBatchViews BatchViews = 
+					OpImageLayerInternal::GetOpLayerBatchLODRangeViews(BatchId, BatchArgs);
+			
+			const uint8* BaseBuf = BatchViews.Base.GetData();
+			const uint8* BlendBuf = BatchViews.Blend.GetData();
+			uint8* ResultBuf = BatchViews.Result.GetData();
+			
+			for (int32 I = 0; I < BatchViews.NumElems; ++I)
+			{
+				const uint32 Base = PackPixel<NC>(&BaseBuf[NC * I]);
+				const uint32 Blend = PackPixel<NC>(&BlendBuf[NC * I]);
+
+				const uint32 Result = BLEND_FUNC(Base, Blend);
+				UnpackPixel<NC>(&ResultBuf[NC * I], Result);
+			}
+		});
+	}
+
+	template< 
+		uint32 (*BLEND_FUNC_MASKED)(uint32, uint32, uint32), 
+		uint32 NC>
+	void BufferLayerCombine(FImage* ResultImage, const FImage* BaseImage, const FImage* MaskImage, const FImage* BlendImage, bool bOnlyFirstLOD)
+	{
+		static_assert(NC > 0 && NC <= 4);
+
+		check(BaseImage->GetSizeX() == BlendImage->GetSizeX());
+		check(BaseImage->GetSizeY() == BlendImage->GetSizeY());
+		check(BaseImage->GetSizeX() == MaskImage->GetSizeX());
+		check(BaseImage->GetSizeY() == MaskImage->GetSizeY());
+		check(BaseImage->GetFormat() == BlendImage->GetFormat());
+
+		const int32 NumLODs = bOnlyFirstLOD ? 1 : ResultImage->GetLODCount();
+
+		constexpr int32 BatchNumElems = 4098*2;
+		OpImageLayerInternal::FOpLayerBatchArgs BatchArgs;
+		BatchArgs.BatchNumElems      = BatchNumElems;
+		BatchArgs.LODBegin           = 0;
+		BatchArgs.LODEnd             = NumLODs;
+		BatchArgs.Result             = ResultImage;
+		BatchArgs.Base               = BaseImage;
+		BatchArgs.Blend              = BlendImage;
+		BatchArgs.Mask               = MaskImage;
+		BatchArgs.ResultBytesPerElem = NC;
+		BatchArgs.BaseBytesPerElem   = NC;
+		BatchArgs.BlendBytesPerElem  = NC;
+		BatchArgs.MaskBytesPerElem   = 1;
+
+		const int32 NumBatches = OpImageLayerInternal::GetOpLayerNumBatchesLODRange(BatchArgs);
+		
+		OpImageLayerInternal::InvokeParallelFor(NumBatches,
+		[
+			&BatchArgs
+		](int32 BatchId)
+		{
+			OpImageLayerInternal::FOpLayerBatchViews BatchViews = 
+					OpImageLayerInternal::GetOpLayerBatchLODRangeViews(BatchId, BatchArgs);
+
+			const uint8* BaseBuf = BatchViews.Base.GetData();
+			const uint8* BlendBuf = BatchViews.Blend.GetData();
+			const uint8* MaskBuf = BatchViews.Mask.GetData();
+			uint8* ResultBuf = BatchViews.Result.GetData();
+			
+			// This could happen in case of missing data files.
+			if (!BaseBuf || !BlendBuf || !MaskBuf)
+			{
+				return;
+			}
+
+			for (int32 I = 0; I < BatchViews.NumElems; ++I)
+			{
+				const uint32 Base = PackPixel<NC>(&BaseBuf[NC * I]);
+				const uint32 Blend = PackPixel<NC>(&BlendBuf[NC * I]);
+				const uint32 Mask = PackPixel<1>(&MaskBuf[1 * I]);
+
+				const uint32 Result = BLEND_FUNC_MASKED(Base, Blend, Mask);
+				UnpackPixel<NC>(&ResultBuf[NC * I], Result);
+			}
+		});
+	}
+
+	template< 
+		uint32 (*BLEND_FUNC_MASKED)(uint32, uint32, uint32), 
+		uint32 NC>
+	void BufferLayerCombineColour(FImage* ResultImage, const FImage* BaseImage, const FImage* MaskImage, FVector4f Color)
+	{
+		static_assert(NC > 0 && NC <= 4);
+
+		check(BaseImage->GetSizeX() == MaskImage->GetSizeX());
+		check(BaseImage->GetSizeY() == MaskImage->GetSizeY());
+
+		const uint32 TopColor = 
+			static_cast<uint32>(255.0f * Color[0]) << 0 |  
+			static_cast<uint32>(255.0f * Color[1]) << 8 |  
+			static_cast<uint32>(255.0f * Color[2]) << 16;  
+
+		constexpr int32 BatchNumElems = 4098*2;
+
+		OpImageLayerInternal::FOpLayerBatchArgs BatchArgs;
+		BatchArgs.BatchNumElems      = BatchNumElems;
+		BatchArgs.Result             = ResultImage;
+		BatchArgs.Base               = BaseImage;
+		BatchArgs.Mask               = MaskImage;
+		BatchArgs.ResultBytesPerElem = NC;
+		BatchArgs.BaseBytesPerElem   = NC;
+		BatchArgs.MaskBytesPerElem   = 1;
+		
+		const int32 NumBatches = OpImageLayerInternal::GetOpLayerNumBatches(BatchArgs);
+
+		OpImageLayerInternal::InvokeParallelFor(NumBatches,
+		[
+			&BatchArgs, TopColor
+		](int32 BatchId)
+		{
+			OpImageLayerInternal::FOpLayerBatchViews BatchViews = 
+					OpImageLayerInternal::GetOpLayerBatchViews(BatchId, BatchArgs);
+	
+			const uint8* BaseBuf = BatchViews.Base.GetData();
+			const uint8* MaskBuf = BatchViews.Mask.GetData();
+			uint8* ResultBuf = BatchViews.Result.GetData();
+
+			// This could happen in case of missing data files.
+			if (!BaseBuf || !MaskBuf)
+			{
+				return;
+			}
+
+			for (int32 I = 0; I < BatchViews.NumElems; ++I)
+			{
+				const uint32 Base = PackPixel<NC>(&BaseBuf[NC * I]);
+				const uint32 Mask = PackPixel<1>(&MaskBuf[1 * I]);
+
+				const uint32 Result = BLEND_FUNC_MASKED(Base, TopColor, Mask);
+				UnpackPixel<NC>(&ResultBuf[NC * I], Result);
+			}
+		});
+	}
+
+	template<uint32 (*BLEND_FUNC)(uint32, uint32)>
+	void ImageLayerCombine(FImage* ResultImage, const FImage* BaseImage, const FImage* BlendedImage, bool bOnlyFirstLOD)
+	{
+		check(ResultImage->GetFormat() == BaseImage->GetFormat());
+		check(ResultImage->GetSizeX() == BaseImage->GetSizeX());
+		check(ResultImage->GetSizeY() == BaseImage->GetSizeY());
+		check(bOnlyFirstLOD || ResultImage->GetLODCount() == BaseImage->GetLODCount());
+		check(BaseImage->GetSizeX() == BlendedImage->GetSizeX());
+		check(BaseImage->GetSizeY() == BlendedImage->GetSizeY());
+		check(BaseImage->GetFormat() == BlendedImage->GetFormat());
+		check(bOnlyFirstLOD || ResultImage->GetLODCount() <= BlendedImage->GetLODCount());
+
+		const EImageFormat BaseFormat = BaseImage->GetFormat();
+
+		if (BaseFormat == EImageFormat::L_UByte)
+		{
+			BufferLayerCombine<BLEND_FUNC, 1>(ResultImage, BaseImage, BlendedImage, bOnlyFirstLOD);
+		}
+		else if (BaseFormat == EImageFormat::RGB_UByte)
+		{
+			BufferLayerCombine<BLEND_FUNC, 3>(ResultImage, BaseImage, BlendedImage, bOnlyFirstLOD);
+		}
+		else if (BaseFormat == EImageFormat::RGBA_UByte || BaseFormat == EImageFormat::BGRA_UByte)
+		{
+			// \todo: pass swizzle template argument if BGRA_UBYTE, not yet supported.
+			BufferLayerCombine<BLEND_FUNC, 4>(ResultImage, BaseImage, BlendedImage, bOnlyFirstLOD);
+		}
+		else
+		{
+			checkf(false, TEXT("Unsupported format."));
+		}
+	}
+
+	template< 
+		uint32 (*BLEND_FUNC)(uint32, uint32), 
+		uint32 (*BLEND_FUNC_MASKED)(uint32, uint32, uint32)>
+	void ImageLayerCombine(FImage* ResultImage, const FImage* BaseImage, const FImage* MaskImage, const FImage* BlendedImage, bool bOnlyFirstLOD)
+	{
+		check(ResultImage->GetFormat() == BaseImage->GetFormat());
+		check(ResultImage->GetSizeX() == BaseImage->GetSizeX());
+		check(ResultImage->GetSizeY() == BaseImage->GetSizeY());
+		check(bOnlyFirstLOD || ResultImage->GetLODCount() == BaseImage->GetLODCount());
+		check(BaseImage->GetSizeX() == BlendedImage->GetSizeX());
+		check(BaseImage->GetSizeY() == BlendedImage->GetSizeY());
+		check(BaseImage->GetFormat() == BlendedImage->GetFormat());
+		check(bOnlyFirstLOD || ResultImage->GetLODCount() <= BlendedImage->GetLODCount());
+
+		const EImageFormat BaseFormat = BaseImage->GetFormat();
+
+		TSharedPtr<FImage> TempMaskImage;
+		if (MaskImage->GetFormat() != EImageFormat::L_UByte)
+		{
+			UE_LOG(LogMutableCore, Log, TEXT("Image layer format not supported. A generic one will be used. "));
+
+			FImageOperator ImOp = FImageOperator::GetDefault(nullptr);
+			constexpr int32 Quality = 4;
+			TempMaskImage = ImOp.ImagePixelFormat( Quality, MaskImage, EImageFormat::L_UByte );
+			MaskImage = TempMaskImage.Get();
+		}
+
+		if (BaseFormat == EImageFormat::L_UByte)
+		{
+			BufferLayerCombine<BLEND_FUNC_MASKED, 1>(ResultImage, BaseImage, MaskImage, BlendedImage, bOnlyFirstLOD);
+		}
+		else if (BaseFormat == EImageFormat::RGB_UByte)
+		{
+			BufferLayerCombine<BLEND_FUNC_MASKED, 3>(ResultImage, BaseImage, MaskImage, BlendedImage, bOnlyFirstLOD);
+		}
+		else if (BaseFormat == EImageFormat::RGBA_UByte || BaseFormat == EImageFormat::BGRA_UByte)
+		{
+			// \todo: pass swizzle template argument if BGRA_UBYTE, not yet supported.
+			BufferLayerCombine<BLEND_FUNC_MASKED, 4>(ResultImage, BaseImage, MaskImage, BlendedImage, bOnlyFirstLOD);
+		}
+		else
+		{
+			UE_LOG(LogMutableCore, Log, TEXT("Image layer format not supported. A generic one will be used. "));
+
+			FImageOperator ImOp = FImageOperator::GetDefault(nullptr);
+			constexpr int32 Quality = 4;
+			TSharedPtr<FImage> TempBaseImage = ImOp.ImagePixelFormat(Quality, BaseImage, EImageFormat::RGBA_UByte);
+			TSharedPtr<FImage> TempBlededImage = ImOp.ImagePixelFormat(Quality, BlendedImage, EImageFormat::RGBA_UByte);
+			BufferLayerCombine<BLEND_FUNC_MASKED, 4>(ResultImage, TempBaseImage.Get(), MaskImage, TempBlededImage.Get(), bOnlyFirstLOD);
+		}
+	}
+
+	template<uint32 (*BLEND_FUNC)(uint32, uint32)>
+	void ImageLayerCombineColour(FImage* ResultImage, const FImage* BaseImage, FVector4f Color)
+	{
+		check(ResultImage->GetFormat() == BaseImage->GetFormat());
+		check(ResultImage->GetSizeX() == BaseImage->GetSizeX());
+		check(ResultImage->GetSizeY() == BaseImage->GetSizeY());
+		check(ResultImage->GetLODCount() == BaseImage->GetLODCount());
+
+		const EImageFormat BaseFormat = BaseImage->GetFormat();
+
+		if (BaseFormat == EImageFormat::L_UByte)
+		{
+			BufferLayerCombineColour<BLEND_FUNC, 1>(ResultImage, BaseImage, Color);
+		}
+		else if (BaseFormat == EImageFormat::RGB_UByte)
+		{
+			BufferLayerCombineColour<BLEND_FUNC, 3>(ResultImage, BaseImage, Color);
+		}
+		else if (BaseFormat == EImageFormat::RGBA_UByte || BaseFormat == EImageFormat::BGRA_UByte)
+		{
+			// \todo: pass swizzle template argument if BGRA_UBYTE, not yet supported.
+ 			BufferLayerCombineColour<BLEND_FUNC, 4>(ResultImage, BaseImage, Color);
+		}
+		else
+		{
+			checkf(false, TEXT("Unsupported format."));
+		}
+	}
+
+	template< 
+		uint32 (*BLEND_FUNC)(uint32, uint32),
+		uint32 (*BLEND_FUNC_MASKED)(uint32, uint32, uint32)>
+	void ImageLayerCombineColour(FImage* ResultImage, const FImage* BaseImage, const FImage* MaskImage, FVector4f Color)
+	{
+		check(ResultImage->GetFormat() == BaseImage->GetFormat());
+		check(ResultImage->GetSizeX() == BaseImage->GetSizeX());
+		check(ResultImage->GetSizeY() == BaseImage->GetSizeY());
+		check(ResultImage->GetLODCount() == BaseImage->GetLODCount());
+
+		const EImageFormat BaseFormat = BaseImage->GetFormat();
+
+		if (MaskImage->GetFormat() != EImageFormat::L_UByte)
+		{
+			checkf(false, TEXT("Unsupported mask format."));
+
+			BufferLayerCombineColour<BLEND_FUNC, 1>(ResultImage, BaseImage, Color);
+		}
+
+		if (BaseFormat == EImageFormat::L_UByte)
+		{
+			BufferLayerCombineColour<BLEND_FUNC_MASKED, 1>(ResultImage, BaseImage, MaskImage, Color);
+		}
+		else if (BaseFormat == EImageFormat::RGB_UByte)
+		{
+			BufferLayerCombineColour<BLEND_FUNC_MASKED, 3>(ResultImage, BaseImage, MaskImage, Color);
+		}
+		else if (BaseFormat == EImageFormat::RGBA_UByte || BaseFormat == EImageFormat::BGRA_UByte)
+		{
+			// \todo: pass swizzle template argument if BGRA_UBYTE, not yet supported.
+ 			BufferLayerCombineColour<BLEND_FUNC_MASKED, 4>(ResultImage, BaseImage, MaskImage, Color);
+		}
+		else
+		{
+			checkf(false, TEXT("Unsupported format."));
+		}
+	}
+
+	template<uint32 NCBase, uint32 NCBlend, class ImageCombineFn>
+	void BufferLayerCombineFunctor(FImage* DestImage, const FImage* BaseImage, const FImage* BlendImage, ImageCombineFn&& ImageCombine)
+	{
+		static_assert(NCBase > 0 && NCBase <= 4);
+		static_assert(NCBlend > 0 && NCBlend <= 4);
+
+		check(BaseImage->GetFormat() == DestImage->GetFormat());
+		check(BaseImage->GetSizeX() == BlendImage->GetSizeX());
+		check(BaseImage->GetSizeY() == BlendImage->GetSizeY());
+		check(BaseImage->GetSizeX() == DestImage->GetSizeX());
+		check(BaseImage->GetSizeY() == DestImage->GetSizeY());
+		check(BaseImage->GetLODCount() <= BlendImage->GetLODCount());
+		check(BaseImage->GetLODCount() <= DestImage->GetLODCount());
+	
+		constexpr int32 BatchNumElems = 4098*2;
+
+		OpImageLayerInternal::FOpLayerBatchArgs BatchArgs;
+		BatchArgs.BatchNumElems      = BatchNumElems;
+		BatchArgs.Result             = DestImage;
+		BatchArgs.Base               = BaseImage;
+		BatchArgs.Blend              = BlendImage;
+		BatchArgs.ResultBytesPerElem = NCBase;
+		BatchArgs.BaseBytesPerElem   = NCBase;
+		BatchArgs.BlendBytesPerElem  = NCBlend;
+		
+		const int32 NumBatches = OpImageLayerInternal::GetOpLayerNumBatches(BatchArgs);
+
+		OpImageLayerInternal::InvokeParallelFor(NumBatches, 
+		[
+			&BatchArgs, ImageCombine
+		](int32 BatchId)
+		{
+			OpImageLayerInternal::FOpLayerBatchViews BatchViews = 
+					OpImageLayerInternal::GetOpLayerBatchViews(BatchId, BatchArgs);
+
+			const uint8* BaseBuf = BatchViews.Base.GetData();
+			const uint8* BlendBuf = BatchViews.Blend.GetData();
+			uint8* DestBuf = BatchViews.Result.GetData();
+	
+			for (int32 I = 0; I < BatchViews.NumElems; ++I)
+			{
+				const uint32 Base = PackPixel<NCBase>(&BaseBuf[NCBase * I]);
+				const uint32 Blend = PackPixel<NCBlend>(&BlendBuf[NCBlend * I]);
+
+				const uint32 Result = ImageCombine(Base, Blend);
+				UnpackPixel<NCBase>(&DestBuf[NCBase* I], Result);
+			}
+		});
+	}
+
+	// Same functionality as above, in this case we use a functor which allows to pass user data. 
+	template<class ImageCombineFn>
+	void ImageLayerCombineFunctor(FImage* ResultImage, const FImage* BaseImage, const FImage* BlendedImage, ImageCombineFn&& ImageCombine)
+	{
+		check(ResultImage->GetFormat() == BaseImage->GetFormat());
+		check(ResultImage->GetSizeX() == BaseImage->GetSizeX());
+		check(ResultImage->GetSizeY() == BaseImage->GetSizeY());
+		check(ResultImage->GetLODCount() == BaseImage->GetLODCount());
+		check(BaseImage->GetSizeX() == BlendedImage->GetSizeX());
+		check(BaseImage->GetSizeY() == BlendedImage->GetSizeY());
+		check(ResultImage->GetLODCount() <= BlendedImage->GetLODCount());
+
+		const EImageFormat BaseFormat = BaseImage->GetFormat();
+		const EImageFormat BlendFormat = BlendedImage->GetFormat();
+
+		if (BaseFormat == EImageFormat::L_UByte )
+		{
+			if (BlendFormat == EImageFormat::L_UByte )
+			{
+				BufferLayerCombineFunctor<1, 1>(ResultImage, BaseImage, BlendedImage, Forward<ImageCombineFn>(ImageCombine));
+			}
+			else if (BlendFormat == EImageFormat::RGB_UByte )
+			{
+				BufferLayerCombineFunctor<1, 3>(ResultImage, BaseImage, BlendedImage, Forward<ImageCombineFn>(ImageCombine));
+			}
+			else if (BlendFormat == EImageFormat::RGBA_UByte || BlendFormat == EImageFormat::BGRA_UByte)
+			{
+				BufferLayerCombineFunctor<1, 4>(ResultImage, BaseImage, BlendedImage, Forward<ImageCombineFn>(ImageCombine));
+			}
+			else
+			{
+				checkf(false, TEXT("Unsupported format."));
+			}
+		}
+		else if (BaseFormat == EImageFormat::RGB_UByte)
+		{
+			if (BlendFormat == EImageFormat::L_UByte )
+			{
+				BufferLayerCombineFunctor<3, 1>(ResultImage, BaseImage, BlendedImage, Forward<ImageCombineFn>(ImageCombine));
+			}
+			else if (BlendFormat == EImageFormat::RGB_UByte )
+			{
+				BufferLayerCombineFunctor<3, 3>(ResultImage, BaseImage, BlendedImage, Forward<ImageCombineFn>(ImageCombine));
+			}
+			else if (BlendFormat == EImageFormat::RGBA_UByte || BlendFormat == EImageFormat::BGRA_UByte )
+			{
+				BufferLayerCombineFunctor<3, 4>(ResultImage, BaseImage, BlendedImage, Forward<ImageCombineFn>(ImageCombine));
+			}
+			else
+			{
+				checkf(false, TEXT("Unsupported format."));
+			}
+		}
+		else if (BaseFormat == EImageFormat::RGBA_UByte || BaseFormat == EImageFormat::BGRA_UByte)
+		{
+			if (BlendFormat == EImageFormat::L_UByte )
+			{
+				BufferLayerCombineFunctor<4, 1>(ResultImage, BaseImage, BlendedImage, Forward<ImageCombineFn>(ImageCombine));
+			}
+			else if (BlendFormat == EImageFormat::RGB_UByte )
+			{
+				BufferLayerCombineFunctor<4, 3>(ResultImage, BaseImage, BlendedImage, Forward<ImageCombineFn>(ImageCombine));
+			}
+			else if (BlendFormat == EImageFormat::RGBA_UByte || BlendFormat == EImageFormat::BGRA_UByte )
+			{
+				BufferLayerCombineFunctor<4, 4>(ResultImage, BaseImage, BlendedImage, Forward<ImageCombineFn>(ImageCombine));
+			}
+			else
+			{
+				checkf(false, TEXT("Unsupported format."));
+			}
+		}
+		else
+		{
+			checkf(false, TEXT("Unsupported format."));
+		}
+	}
+
+	//! Blend a subimage on the base using a mask.
+	void ImageBlendOnBaseNoAlpha(FImage* BaseImage, const FImage* MaskImage, const FImage* BlendedImage, const box<FIntVector2>& Rect)
+	{
+		ImageLayerOnBaseNoAlpha<BlendChannel, false>(BaseImage, MaskImage, BlendedImage, Rect);
+	}
+
+	template void BufferLayer<SoftLightChannelMasked, SoftLightChannel, false>(FImage*, const FImage*, const FImage*, const FImage*, bool, bool);
+    template void BufferLayer<HardLightChannelMasked, HardLightChannel, false>(FImage*, const FImage*, const FImage*, const FImage*, bool, bool);
+    template void BufferLayer<BurnChannelMasked     , BurnChannel     , false>(FImage*, const FImage*, const FImage*, const FImage*, bool, bool);
+    template void BufferLayer<DodgeChannelMasked    , DodgeChannel    , false>(FImage*, const FImage*, const FImage*, const FImage*, bool, bool);
+    template void BufferLayer<ScreenChannelMasked   , ScreenChannel   , false>(FImage*, const FImage*, const FImage*, const FImage*, bool, bool);
+    template void BufferLayer<OverlayChannelMasked  , OverlayChannel  , false>(FImage*, const FImage*, const FImage*, const FImage*, bool, bool);
+    template void BufferLayer<LightenChannelMasked  , LightenChannel  , false>(FImage*, const FImage*, const FImage*, const FImage*, bool, bool);
+    template void BufferLayer<MultiplyChannelMasked , MultiplyChannel , false>(FImage*, const FImage*, const FImage*, const FImage*, bool, bool);
+    template void BufferLayer<BlendChannelMasked    , BlendChannel    , false>(FImage*, const FImage*, const FImage*, const FImage*, bool, bool);
+
+	template void BufferLayer<SoftLightChannelMasked, SoftLightChannel, true>(FImage*, const FImage*, const FImage*, const FImage*, bool, bool);
+    template void BufferLayer<HardLightChannelMasked, HardLightChannel, true>(FImage*, const FImage*, const FImage*, const FImage*, bool, bool);
+    template void BufferLayer<BurnChannelMasked     , BurnChannel     , true>(FImage*, const FImage*, const FImage*, const FImage*, bool, bool);
+    template void BufferLayer<DodgeChannelMasked    , DodgeChannel    , true>(FImage*, const FImage*, const FImage*, const FImage*, bool, bool);
+    template void BufferLayer<ScreenChannelMasked   , ScreenChannel   , true>(FImage*, const FImage*, const FImage*, const FImage*, bool, bool);
+    template void BufferLayer<OverlayChannelMasked  , OverlayChannel  , true>(FImage*, const FImage*, const FImage*, const FImage*, bool, bool);
+    template void BufferLayer<LightenChannelMasked  , LightenChannel  , true>(FImage*, const FImage*, const FImage*, const FImage*, bool, bool);
+    template void BufferLayer<MultiplyChannelMasked , MultiplyChannel , true>(FImage*, const FImage*, const FImage*, const FImage*, bool, bool);
+    template void BufferLayer<BlendChannelMasked    , BlendChannel    , true>(FImage*, const FImage*, const FImage*, const FImage*, bool, bool);
+
+	template void BufferLayerEmbeddedMask<SoftLightChannelMasked, SoftLightChannel, false>(FImage*, const FImage*, const FImage*, bool, bool); 
+	template void BufferLayerEmbeddedMask<HardLightChannelMasked, HardLightChannel, false>(FImage*, const FImage*, const FImage*, bool, bool); 
+	template void BufferLayerEmbeddedMask<BurnChannelMasked     , BurnChannel     , false>(FImage*, const FImage*, const FImage*, bool, bool); 
+	template void BufferLayerEmbeddedMask<DodgeChannelMasked    , DodgeChannel    , false>(FImage*, const FImage*, const FImage*, bool, bool); 
+	template void BufferLayerEmbeddedMask<ScreenChannelMasked   , ScreenChannel   , false>(FImage*, const FImage*, const FImage*, bool, bool); 
+	template void BufferLayerEmbeddedMask<OverlayChannelMasked  , OverlayChannel  , false>(FImage*, const FImage*, const FImage*, bool, bool); 
+	template void BufferLayerEmbeddedMask<LightenChannelMasked  , LightenChannel  , false>(FImage*, const FImage*, const FImage*, bool, bool); 
+	template void BufferLayerEmbeddedMask<MultiplyChannelMasked , MultiplyChannel , false>(FImage*, const FImage*, const FImage*, bool, bool); 
+	template void BufferLayerEmbeddedMask<BlendChannelMasked    , BlendChannel    , false>(FImage*, const FImage*, const FImage*, bool, bool); 
+   
+	template void BufferLayer<SoftLightChannel, false>(FImage*, const FImage*, const FImage*, bool, bool, bool); 
+    template void BufferLayer<HardLightChannel, false>(FImage*, const FImage*, const FImage*, bool, bool, bool); 
+    template void BufferLayer<BurnChannel     , false>(FImage*, const FImage*, const FImage*, bool, bool, bool); 
+    template void BufferLayer<DodgeChannel    , false>(FImage*, const FImage*, const FImage*, bool, bool, bool); 
+    template void BufferLayer<ScreenChannel   , false>(FImage*, const FImage*, const FImage*, bool, bool, bool); 
+    template void BufferLayer<OverlayChannel  , false>(FImage*, const FImage*, const FImage*, bool, bool, bool); 
+    template void BufferLayer<LightenChannel  , false>(FImage*, const FImage*, const FImage*, bool, bool, bool); 
+    template void BufferLayer<MultiplyChannel , false>(FImage*, const FImage*, const FImage*, bool, bool, bool); 
+    template void BufferLayer<BlendChannel    , false>(FImage*, const FImage*, const FImage*, bool, bool, bool); 
+
+	template void BufferLayer<SoftLightChannel, true>(FImage*, const FImage*, const FImage*, bool, bool, bool); 
+    template void BufferLayer<HardLightChannel, true>(FImage*, const FImage*, const FImage*, bool, bool, bool); 
+    template void BufferLayer<BurnChannel     , true>(FImage*, const FImage*, const FImage*, bool, bool, bool); 
+    template void BufferLayer<DodgeChannel    , true>(FImage*, const FImage*, const FImage*, bool, bool, bool); 
+    template void BufferLayer<ScreenChannel   , true>(FImage*, const FImage*, const FImage*, bool, bool, bool); 
+    template void BufferLayer<OverlayChannel  , true>(FImage*, const FImage*, const FImage*, bool, bool, bool); 
+    template void BufferLayer<LightenChannel  , true>(FImage*, const FImage*, const FImage*, bool, bool, bool); 
+    template void BufferLayer<MultiplyChannel , true>(FImage*, const FImage*, const FImage*, bool, bool, bool);
+    template void BufferLayer<BlendChannel    , true>(FImage*, const FImage*, const FImage*, bool, bool, bool);
+
+	template void BufferLayerInPlace<SoftLightChannel, false, 1>(FImage*, const FImage*, bool, uint32, uint32); 
+	template void BufferLayerInPlace<HardLightChannel, false, 1>(FImage*, const FImage*, bool, uint32, uint32); 
+	template void BufferLayerInPlace<BurnChannel     , false, 1>(FImage*, const FImage*, bool, uint32, uint32); 
+	template void BufferLayerInPlace<DodgeChannel    , false, 1>(FImage*, const FImage*, bool, uint32, uint32); 
+	template void BufferLayerInPlace<ScreenChannel   , false, 1>(FImage*, const FImage*, bool, uint32, uint32); 
+	template void BufferLayerInPlace<OverlayChannel  , false, 1>(FImage*, const FImage*, bool, uint32, uint32); 
+	template void BufferLayerInPlace<LightenChannel  , false, 1>(FImage*, const FImage*, bool, uint32, uint32); 
+	template void BufferLayerInPlace<MultiplyChannel , false, 1>(FImage*, const FImage*, bool, uint32, uint32); 
+	template void BufferLayerInPlace<BlendChannel    , false, 1>(FImage*, const FImage*, bool, uint32, uint32);
+
+	template void BufferLayerColour<SoftLightChannelMasked, SoftLightChannel>(FImage*, const FImage*, const FImage*, FVector4f); 
+	template void BufferLayerColour<HardLightChannelMasked, HardLightChannel>(FImage*, const FImage*, const FImage*, FVector4f); 
+	template void BufferLayerColour<BurnChannelMasked     , BurnChannel     >(FImage*, const FImage*, const FImage*, FVector4f); 
+	template void BufferLayerColour<DodgeChannelMasked    , DodgeChannel    >(FImage*, const FImage*, const FImage*, FVector4f); 
+	template void BufferLayerColour<ScreenChannelMasked   , ScreenChannel   >(FImage*, const FImage*, const FImage*, FVector4f); 
+	template void BufferLayerColour<OverlayChannelMasked  , OverlayChannel  >(FImage*, const FImage*, const FImage*, FVector4f); 
+	template void BufferLayerColour<LightenChannelMasked  , LightenChannel  >(FImage*, const FImage*, const FImage*, FVector4f); 
+	template void BufferLayerColour<MultiplyChannelMasked , MultiplyChannel >(FImage*, const FImage*, const FImage*, FVector4f); 
+	template void BufferLayerColour<BlendChannelMasked    , BlendChannel    >(FImage*, const FImage*, const FImage*, FVector4f); 
+
+	template void BufferLayerColourFromAlpha<SoftLightChannel>(FImage*, const FImage*, FVector4f); 
+	template void BufferLayerColourFromAlpha<HardLightChannel>(FImage*, const FImage*, FVector4f); 
+	template void BufferLayerColourFromAlpha<BurnChannel     >(FImage*, const FImage*, FVector4f); 
+	template void BufferLayerColourFromAlpha<DodgeChannel    >(FImage*, const FImage*, FVector4f); 
+	template void BufferLayerColourFromAlpha<ScreenChannel   >(FImage*, const FImage*, FVector4f); 
+	template void BufferLayerColourFromAlpha<OverlayChannel  >(FImage*, const FImage*, FVector4f); 
+	template void BufferLayerColourFromAlpha<LightenChannel  >(FImage*, const FImage*, FVector4f); 
+	template void BufferLayerColourFromAlpha<MultiplyChannel >(FImage*, const FImage*, FVector4f); 
+
+	template void BufferLayerColour<SoftLightChannel>(FImage*, const FImage*, FVector4f); 
+	template void BufferLayerColour<HardLightChannel>(FImage*, const FImage*, FVector4f); 
+	template void BufferLayerColour<BurnChannel     >(FImage*, const FImage*, FVector4f); 
+	template void BufferLayerColour<DodgeChannel    >(FImage*, const FImage*, FVector4f); 
+	template void BufferLayerColour<ScreenChannel   >(FImage*, const FImage*, FVector4f); 
+	template void BufferLayerColour<OverlayChannel  >(FImage*, const FImage*, FVector4f); 
+	template void BufferLayerColour<LightenChannel  >(FImage*, const FImage*, FVector4f); 
+	template void BufferLayerColour<MultiplyChannel >(FImage*, const FImage*, FVector4f); 
+
+	template void BufferLayerColourInPlace<SoftLightChannelMasked, SoftLightChannel, 1>(FImage*, const FImage*, FVector4f, bool, uint32, uint8); 
+	template void BufferLayerColourInPlace<HardLightChannelMasked, HardLightChannel, 1>(FImage*, const FImage*, FVector4f, bool, uint32, uint8); 
+	template void BufferLayerColourInPlace<BurnChannelMasked     , BurnChannel     , 1>(FImage*, const FImage*, FVector4f, bool, uint32, uint8); 
+	template void BufferLayerColourInPlace<DodgeChannelMasked    , DodgeChannel    , 1>(FImage*, const FImage*, FVector4f, bool, uint32, uint8); 
+	template void BufferLayerColourInPlace<ScreenChannelMasked   , ScreenChannel   , 1>(FImage*, const FImage*, FVector4f, bool, uint32, uint8); 
+	template void BufferLayerColourInPlace<OverlayChannelMasked  , OverlayChannel  , 1>(FImage*, const FImage*, FVector4f, bool, uint32, uint8); 
+	template void BufferLayerColourInPlace<LightenChannelMasked  , LightenChannel  , 1>(FImage*, const FImage*, FVector4f, bool, uint32, uint8); 
+	template void BufferLayerColourInPlace<MultiplyChannelMasked , MultiplyChannel , 1>(FImage*, const FImage*, FVector4f, bool, uint32, uint8); 
+	template void BufferLayerColourInPlace<BlendChannelMasked    , BlendChannel    , 1>(FImage*, const FImage*, FVector4f, bool, uint32, uint8); 
+	
+	template void BufferLayerColourInPlace<SoftLightChannel, 1>(FImage*, FVector4f, bool, uint32, uint8); 
+	template void BufferLayerColourInPlace<HardLightChannel, 1>(FImage*, FVector4f, bool, uint32, uint8); 
+	template void BufferLayerColourInPlace<BurnChannel     , 1>(FImage*, FVector4f, bool, uint32, uint8); 
+	template void BufferLayerColourInPlace<DodgeChannel    , 1>(FImage*, FVector4f, bool, uint32, uint8); 
+	template void BufferLayerColourInPlace<ScreenChannel   , 1>(FImage*, FVector4f, bool, uint32, uint8); 
+	template void BufferLayerColourInPlace<OverlayChannel  , 1>(FImage*, FVector4f, bool, uint32, uint8); 
+	template void BufferLayerColourInPlace<LightenChannel  , 1>(FImage*, FVector4f, bool, uint32, uint8); 
+	template void BufferLayerColourInPlace<MultiplyChannel , 1>(FImage*, FVector4f, bool, uint32, uint8); 
+	template void BufferLayerColourInPlace<BlendChannel    , 1>(FImage*, FVector4f, bool, uint32, uint8); 
+
+	// BufferLayerComposite* uses are specializations, no need to explicitly instanciate.
+	//template void BufferLayerComposite<BlendChannelMasked, LightenChannel, false>(FImage*, const FImage*, bool, uint8);
+
+	//template void BufferLayerCompositeVector<VectorBlendChannelMasked, VectorLightenChannel, false>(FImage*, const FImage*, bool, uint8);
+
+	template void ImageLayerCombine<CombineNormal>(FImage*, const FImage*, const FImage*, bool);
+	template void ImageLayerCombine<CombineNormal, CombineNormalMasked>(FImage*, const FImage*, const FImage*, const FImage*, bool);
+
+	template void ImageLayerCombineColour<CombineNormal>(FImage*, const FImage*, FVector4f);
+	template void ImageLayerCombineColour<CombineNormal, CombineNormalMasked>(FImage*, const FImage*, const FImage*, FVector4f);
+
+	template void ImageLayerCombineFunctor<FNormalCompositeFunctor>(FImage*, const FImage*, const FImage*, FNormalCompositeFunctor&&);
+	template void ImageLayerCombineFunctor<FNormalCompositeIdentityFunctor>(FImage*, const FImage*, const FImage*, FNormalCompositeIdentityFunctor&&);
+
+}
diff --git a/Engine/Plugins/Mutable/Source/MutableRuntime/Private/MuR/OpImageLayer.h b/Engine/Plugins/Mutable/Source/MutableRuntime/Private/MuR/OpImageLayer.h
--- a/Engine/Plugins/Mutable/Source/MutableRuntime/Private/MuR/OpImageLayer.h
+++ b/Engine/Plugins/Mutable/Source/MutableRuntime/Private/MuR/OpImageLayer.h
@@ -3,2607 +3,67 @@
 #pragma once
 
 #include "MuR/ImagePrivate.h"
-#include "MuR/ImageRLE.h"
-#include "Async/ParallelFor.h"
-#include "Templates/UnrealTemplate.h"
-
 #include "MuR/OpImageBlend.h"
-#include "MuR/MutableRuntimeModule.h"
 
 namespace mu
-{
-	inline bool IsAnyComponentLargerThan1(FVector4f Value)
-	{
-		return (Value[0] > 1) | (Value[1] > 1) | (Value[2] > 1) | (Value[3] > 1);
-	}
-	
-	/** Apply a blending function to an image with a colour source.
-	* It only affects the RGB or L channels, leaving alpha untouched.
-	*/
-	namespace Private
-	{
-		template<uint32 NumChannels, uint32 (*BLEND_FUNC)(uint32, uint32), bool bClamp>
-		FORCENOINLINE void BufferLayerColourGenericChannel(uint8* DestBuf, const uint8* BaseBuf, int32 NumElems, const FIntVector& Color)
-		{
-			static_assert(NumChannels > 0 && NumChannels <= 4);
-
-			for (int32 I = 0; I < NumElems; ++I)
-			{
-				for (uint32 C = 0; C < NumChannels; ++C)
-				{
-					uint32 Base = BaseBuf[NumChannels * I + C];
-					uint32 Result = BLEND_FUNC(Base, Color[C]);
-					if constexpr (bClamp)
-					{
-						DestBuf[NumChannels * I + C] = (uint8)FMath::Min(255u, Result);
-					}
-					else
-					{
-						DestBuf[NumChannels * I + C] = (uint8)Result;
-					}
-				}
-			}
-		}
-	}
-
-	template< unsigned (*BLEND_FUNC)(unsigned,unsigned), bool CLAMP >
-    inline void BufferLayerColour(FImage* ResultImage, const FImage* BaseImage, FVector4f Color)
-	{
-        check(ResultImage->GetFormat() == BaseImage->GetFormat());
-        check(ResultImage->GetSizeX() == BaseImage->GetSizeX());
-        check(ResultImage->GetSizeY() == BaseImage->GetSizeY());
-        check(ResultImage->GetLODCount() == BaseImage->GetLODCount());
-
-		const EImageFormat BaseFormat = BaseImage->GetFormat();
-		// Generic implementation
-		constexpr int32 NumBatchElems = 4096*4;
-		const int32 BytesPerElem = GetImageFormatData(BaseFormat).BytesPerBlock;	
-
-		const int32 NumBatches = BaseImage->DataStorage.GetNumBatches(NumBatchElems, BytesPerElem);
-
-		const FIntVector ColorValue = FIntVector(Color.X * 255.0f, Color.Y * 255.0f, Color.Z * 255.0f);
-		auto ProcessBatch = [NumBatchElems, BytesPerElem, ColorValue, BaseImage, ResultImage, BaseFormat](int32 BatchId)
-		{
-			TArrayView<const uint8> BaseView = BaseImage->DataStorage.GetBatch(BatchId, NumBatchElems, BytesPerElem);
-			TArrayView<uint8> ResultView = ResultImage->DataStorage.GetBatch(BatchId, NumBatchElems, BytesPerElem);
-			
-			check(BaseView.Num() == ResultView.Num());
-
-			const int32 NumElems = BaseView.Num() / BytesPerElem;
-
-			switch (BaseFormat)
-			{
-			case EImageFormat::L_UByte:
-			{
-				check(BytesPerElem == 1);
-				Private::BufferLayerColourGenericChannel<1, BLEND_FUNC, CLAMP>(ResultView.GetData(), BaseView.GetData(), NumElems, ColorValue);
-				break;
-			}
-			case EImageFormat::RGB_UByte:
-			{
-				check(BytesPerElem == 3);
-				Private::BufferLayerColourGenericChannel<3, BLEND_FUNC, CLAMP>(ResultView.GetData(), BaseView.GetData(), NumElems, ColorValue);
-				break;
-			}
-			case EImageFormat::RGBA_UByte:
-			{
-				check(BytesPerElem == 4);
-				Private::BufferLayerColourGenericChannel<4, BLEND_FUNC, CLAMP>(ResultView.GetData(), BaseView.GetData(), NumElems, ColorValue);
-				break;
-			}
-			case EImageFormat::BGRA_UByte:
-			{
-				check(BytesPerElem == 4);
-				FIntVector BGRAColorValue = FIntVector(ColorValue.Z, ColorValue.Y, ColorValue.X);
-				Private::BufferLayerColourGenericChannel<4, BLEND_FUNC, CLAMP>(ResultView.GetData(), BaseView.GetData(), NumElems, BGRAColorValue);
-				break;
-			}
-			default:
-			{
-				checkf(false, TEXT("Unsupported format."));
-				break;
-			}
-			}
-		};
-
-		if (NumBatches == 1)
-		{
-			ProcessBatch(0);
-		}
-		else if (NumBatches > 1)
-		{
-			ParallelFor(NumBatches, ProcessBatch);
-		}
-    }
-
-
-	template< uint32 (*BLEND_FUNC)(uint32, uint32) >
-	inline void BufferLayerColour(FImage* ResultImage, const FImage* BaseImage, FVector4f Color)
-	{
-		bool bIsClampNeeded = IsAnyComponentLargerThan1(Color);
-		if (bIsClampNeeded)
-		{
-			BufferLayerColour<BLEND_FUNC, true>(ResultImage, BaseImage, Color);
-		}
-		else
-		{
-			BufferLayerColour<BLEND_FUNC, false>(ResultImage, BaseImage, Color);
-		}
-	}
-
-	namespace Private
-	{
-		template<uint32 NumChannels, uint32 (*BLEND_FUNC)(uint32,uint32), bool Clamp>
-		FORCENOINLINE void BufferLayerColourFromAlphaGenericChannel(uint8* DestBuf, const uint8* BaseBuf, int32 NumElems, const FIntVector& Color)
-		{
-			static_assert(NumChannels > 0 && NumChannels <= 4);
-
-			for (int32 I = 0; I < NumElems; ++I)
-			{
-				const uint32 Alpha = Invoke([&]() -> uint32
-				{
-					if constexpr (NumChannels <= 3)
-					{
-						return 255;
-					}
-					else
-					{
-						return BaseBuf[NumChannels * I + 3];
-					}
-				});
-
-				for (uint32 C = 0; C < NumChannels; ++C)
-				{
-					uint32 Result = BLEND_FUNC(Alpha, Color[C]);
-					if constexpr (Clamp)
-					{
-						DestBuf[NumChannels * I + C] = (uint8)FMath::Min(255u, Result);
-					}
-					else
-					{
-						DestBuf[NumChannels * I + C] = (uint8)Result;
-					}
-				}
-			}
-		}
-	}
-
-	template<
-		uint32 (*BLEND_FUNC)(uint32, uint32), 
-		bool CLAMP>
-    inline void BufferLayerColourFromAlpha(FImage* ResultImage, const FImage* BaseImage, FVector4f Color)
-	{
-        check(ResultImage->GetFormat() == BaseImage->GetFormat());
-        check(ResultImage->GetSizeX() == BaseImage->GetSizeX());
-        check(ResultImage->GetSizeY() == BaseImage->GetSizeY());
-        check(ResultImage->GetLODCount() == BaseImage->GetLODCount());
-
-		const EImageFormat BaseFormat = BaseImage->GetFormat();
-		// Generic implementation
-		constexpr int32 NumBatchElems = 4096*2;
-		const int32 BytesPerElem = GetImageFormatData(BaseFormat).BytesPerBlock;	
-
-		const int32 NumBatches = BaseImage->DataStorage.GetNumBatches(NumBatchElems, BytesPerElem);
-		check(NumBatches == ResultImage->DataStorage.GetNumBatches(NumBatchElems, BytesPerElem));
-
-		const FIntVector ColorValue = FIntVector(Color.X * 255.0f, Color.Y * 255.0f, Color.Z * 255.0f);
-		auto ProcessBatch = [NumBatchElems, BytesPerElem, ColorValue, BaseImage, ResultImage, BaseFormat](int32 BatchId)
-		{
-			TArrayView<const uint8> BaseView = BaseImage->DataStorage.GetBatch(BatchId, NumBatchElems, BytesPerElem);
-			TArrayView<uint8> ResultView = ResultImage->DataStorage.GetBatch(BatchId, NumBatchElems, BytesPerElem);
-			
-			check(BaseView.Num() == ResultView.Num());
-
-			const int32 NumElems = ResultView.Num() / BytesPerElem;
-			switch (BaseFormat)
-			{
-			case EImageFormat::L_UByte:
-			{
-				check(BytesPerElem == 1);
-				Private::BufferLayerColourFromAlphaGenericChannel<1, BLEND_FUNC, CLAMP>(ResultView.GetData(), BaseView.GetData(), NumElems, ColorValue);
-				break;
-			}
-			case EImageFormat::RGB_UByte:
-			{
-				check(BytesPerElem == 3);
-				Private::BufferLayerColourFromAlphaGenericChannel<3, BLEND_FUNC, CLAMP>(ResultView.GetData(), BaseView.GetData(), NumElems, ColorValue);
-				break;
-			}
-			case EImageFormat::RGBA_UByte:
-			{
-				check(BytesPerElem == 4);
-				Private::BufferLayerColourFromAlphaGenericChannel<4, BLEND_FUNC, CLAMP>(ResultView.GetData(), BaseView.GetData(), NumElems, ColorValue);
-				break;
-			}
-			case EImageFormat::BGRA_UByte:
-			{
-				check(BytesPerElem == 4);
-				FIntVector BGRAColorValues = FIntVector(ColorValue.Z, ColorValue.Y, ColorValue.X);
-				Private::BufferLayerColourFromAlphaGenericChannel<4, BLEND_FUNC, CLAMP>(ResultView.GetData(), BaseView.GetData(), NumElems , ColorValue);
-				break;
-			}
-			default:
-			{
-				checkf(false, TEXT("Unsupported format."));
-				break;
-			}
-			}
-		};
-
-		if (NumBatches == 1)
-		{
-			ProcessBatch(0);
-		}
-		else if (NumBatches > 1)
-		{
-			ParallelFor(NumBatches, ProcessBatch);
-		}
-    }
-
-
-	template<uint32 (*BLEND_FUNC)(uint32, uint32)>
-	inline void BufferLayerColourFromAlpha(FImage* ResultImage, const FImage* BaseImage, FVector4f Color)
-	{
-		bool bIsClampNeeded = IsAnyComponentLargerThan1(Color);
-		if (bIsClampNeeded)
-		{
-			BufferLayerColourFromAlpha<BLEND_FUNC, true>(ResultImage, BaseImage, Color);
-		}
-		else
-		{
-			BufferLayerColourFromAlpha<BLEND_FUNC, false>(ResultImage, BaseImage, Color);
-		}
-	}
-
-
-	template< 
-		uint32 (*BLEND_FUNC_MASKED)(uint32, uint32, uint32),
-		uint32 (*BLEND_FUNC)(uint32, uint32),
-		bool CLAMP,
-		// Number of total channels to actually process
-		uint32 CHANNELS_TO_BLEND,
-		// Number of total channels in the base image
-		int32 BASE_CHANNEL_STRIDE>
-    inline void BufferLayerColourFormat(FImage* DestImage, const FImage* BaseImage, const FImage* MaskImage, FVector4f Col, uint32 BaseChannelOffset, uint8 ColorChannelOffset, bool bOnlyFirstLOD)
-	{
-		check(CHANNELS_TO_BLEND + BaseChannelOffset <= BASE_CHANNEL_STRIDE);
-		check(DestImage->GetLODCount() <= BaseImage->GetLODCount());
-
-		FUintVector4 TopColor = FUintVector4(Col.X * 255.0f, Col.Y * 255.0f, Col.Z * 255.0f, Col.W * 255);
-
-		const EImageFormat BaseFormat = BaseImage->GetFormat();
-		const EImageFormat MaskFormat = MaskImage->GetFormat();
-
-		int32 NumLODs = bOnlyFirstLOD ? 1 : BaseImage->GetLODCount();
-
-        const bool bIsMaskUncompressed = (MaskFormat == EImageFormat::L_UByte);
-
-		constexpr uint32 NumColorChannels = FMath::Min(CHANNELS_TO_BLEND, 3u);
-		check(NumColorChannels + ColorChannelOffset < 4);
-
-        if (bIsMaskUncompressed)
-        {
-			constexpr int32 NumBatchElems = 4096 * 2;
-			const int32 BytesPerElem = GetImageFormatData(BaseFormat).BytesPerBlock;
-			check(GetImageFormatData(MaskFormat).BytesPerBlock == 1);
-			check(GetImageFormatData(DestImage->GetFormat()).BytesPerBlock == BytesPerElem);
-
-			const int32 NumBatches = 
-					BaseImage->DataStorage.GetNumBatchesLODRange(NumBatchElems, BytesPerElem, 0, NumLODs);
-
-			check(NumBatches == DestImage->DataStorage.GetNumBatchesLODRange(NumBatchElems, BytesPerElem, 0, NumLODs));
-			check(NumBatches == MaskImage->DataStorage.GetNumBatchesLODRange(NumBatchElems, 1, 0, NumLODs));
-
-			auto ProcessBatch = [
-				BaseImage, MaskImage, DestImage, 
-				BytesPerElem, NumBatchElems, 
-				BaseChannelOffset, TopColor, ColorChannelOffset, NumColorChannels, 
-				NumLODs
-			](int32 BatchId)
-			{
-				TArrayView<const uint8> BaseBatchView = 
-						BaseImage->DataStorage.GetBatchLODRange(BatchId, NumBatchElems, BytesPerElem, 0, NumLODs);
-
-				TArrayView<const uint8> MaskBatchView =
-						MaskImage->DataStorage.GetBatchLODRange(BatchId, NumBatchElems, 1, 0, NumLODs);
-
-				TArrayView<uint8> DestBatchView =
-						DestImage->DataStorage.GetBatchLODRange(BatchId, NumBatchElems, BytesPerElem, 0, NumLODs);
-
-				const uint8* BaseBuf = BaseBatchView.GetData() + BaseChannelOffset;
-				const uint8* MaskBuf = MaskBatchView.GetData();
-				uint8* DestBuf = DestBatchView.GetData() + BaseChannelOffset;
-
-				// This could happen in case of missing data files.
-				if (!BaseBatchView.GetData() || !MaskBuf || !DestBatchView.GetData())
-				{
-					return;
-				}
-
-				const int32 NumElems = BaseBatchView.Num() / BytesPerElem;
-				check(NumElems == MaskBatchView.Num());
-				check(NumElems == DestBatchView.Num() / BytesPerElem);
-
-				for (int32 I = 0; I < NumElems; ++I)
-				{
-					uint32 MaskData = MaskBuf[I];
-					for (int32 C = 0; C < NumColorChannels; ++C)
-					{
-						const uint32 Base = BaseBuf[BASE_CHANNEL_STRIDE*I + C];
-						const uint32 Result = BLEND_FUNC_MASKED(Base, TopColor[C + ColorChannelOffset], MaskData);
-						if constexpr (CLAMP)
-						{
-							DestBuf[BASE_CHANNEL_STRIDE*I + C] = (uint8)FMath::Min(255u, Result);
-						}
-						else
-						{
-							DestBuf[BASE_CHANNEL_STRIDE*I + C] = (uint8)Result;
-						}
-					}
-
-					constexpr bool bIsNC4 = (BASE_CHANNEL_STRIDE == 4);
-					if constexpr (bIsNC4)
-					{
-						DestBuf[BASE_CHANNEL_STRIDE*I + 3] = BaseBuf[BASE_CHANNEL_STRIDE*I + 3];
-					}
-				}
-			};
-
-			if (NumBatches == 1)
-			{
-				ProcessBatch(0);
-			}
-			else if (NumBatches > 1)
-			{
-				ParallelFor(NumBatches, ProcessBatch);
-			}
-        }
-        else if (MaskFormat == EImageFormat::L_UByteRLE)
-        {
-            int32 Rows = BaseImage->GetSizeY();
-            int32 Width = BaseImage->GetSizeX();
-
-            for (int32 Lod = 0; Lod < NumLODs; ++Lod)
-            {
-                const uint8* BaseBuf = BaseImage->GetLODData(Lod);
-                const uint8* MaskBuf = MaskImage->GetLODData(Lod);
-				uint8* DestBuf = DestImage->GetLODData(Lod);
-
-				// This could happen in case of missing data files.
-				if (!BaseBuf || !MaskBuf || !DestBuf)
-				{
-					continue;
-				}
-
-				// Remove RLE header, mip size and row sizes.
-				MaskBuf += sizeof(uint32);
-				MaskBuf += Rows * sizeof(uint32);
-
-                for (int32 RowIndex = 0; RowIndex < Rows; ++RowIndex)
-                {
-                    const uint8* DestRowEnd = DestBuf + Width * BASE_CHANNEL_STRIDE;
-                    while (DestBuf != DestRowEnd)
-                    {
-                        // Decode header
-						uint16 Equal = 0;
-						FMemory::Memmove(&Equal, MaskBuf, sizeof(uint16));
-                        MaskBuf += 2;
-
-                        uint8 Different = *MaskBuf;
-                        ++MaskBuf;
-
-                        uint8 EqualPixel = *MaskBuf;
-                        ++MaskBuf;
-
-                        // Equal pixels
-						//check(DestBuf + BASE_CHANNEL_STRIDE * Equal <= BaseImage->GetDataSize(Lod));
-                        if (EqualPixel == 255)
-                        {
-                            for (int32 I = 0; I < Equal; ++I)
-                            {
-                                for (int32 C = 0; C < NumColorChannels; ++C)
-                                {
-                                    uint32 Base = BaseBuf[BASE_CHANNEL_STRIDE * I + C];
-                                    uint32 Result = BLEND_FUNC(Base, TopColor[C + ColorChannelOffset]);
-                                    if constexpr (CLAMP)
-                                    {
-                                        DestBuf[BASE_CHANNEL_STRIDE * I + C] = (uint8)FMath::Min(255u, Result);
-                                    }
-                                    else
-                                    {
-                                        DestBuf[BASE_CHANNEL_STRIDE * I + C] = (uint8)Result;
-                                    }
-                                }
-
-                                constexpr bool bIsNC4 = (BASE_CHANNEL_STRIDE == 4);
-                                if (bIsNC4)
-                                {
-                                    DestBuf[BASE_CHANNEL_STRIDE * I + 3] = BaseBuf[BASE_CHANNEL_STRIDE * I + 3];
-                                }
-                            }
-                        }
-                        else if (EqualPixel > 0)
-                        {
-                            for (int32 I = 0; I < Equal; ++I)
-                            {
-                                for (int32 C = 0; C < NumColorChannels; ++C)
-                                {
-                                    uint32 Base = BaseBuf[BASE_CHANNEL_STRIDE * I + C];
-                                    uint32 Result = BLEND_FUNC_MASKED(Base, TopColor[C + ColorChannelOffset], EqualPixel);
-                                    if constexpr (CLAMP)
-                                    {
-                                        DestBuf[BASE_CHANNEL_STRIDE * I + C] = (uint8)FMath::Min(255u, Result);
-                                    }
-                                    else
-                                    {
-                                        DestBuf[BASE_CHANNEL_STRIDE * I + C] = (uint8)Result;
-                                    }
-                                }
-
-                                constexpr bool bIsNC4 = (BASE_CHANNEL_STRIDE == 4);
-                                if (bIsNC4)
-                                {
-                                    DestBuf[BASE_CHANNEL_STRIDE * I + 3] = BaseBuf[BASE_CHANNEL_STRIDE * I + 3];
-                                }
-                            }
-                        }
-                        else
-                        {
-                            // It could happen if xxxxxOnBase
-                            if (DestBuf != BaseBuf)
-                            {
-                                FMemory::Memmove(DestBuf, BaseBuf, BASE_CHANNEL_STRIDE*Equal);
-                            }
-                        }
-
-                        DestBuf += BASE_CHANNEL_STRIDE * Equal;
-                        BaseBuf += BASE_CHANNEL_STRIDE * Equal;
-
-                        // Different pixels
-						//check(DestBuf + BASE_CHANNEL_STRIDE * Different <= StartDestBuf + BaseImage->GetDataSize(Lod));
-                        for (int32 I = 0; I < Different; ++I)
-                        {
-                            for (int32 C = 0; C < NumColorChannels; ++C)
-                            {
-                                uint32 Mask = MaskBuf[I];
-                                uint32 Base = BaseBuf[BASE_CHANNEL_STRIDE * I + C];
-                                uint32 Result = BLEND_FUNC_MASKED(Base, TopColor[C + ColorChannelOffset], Mask);
-                                if constexpr (CLAMP)
-                                {
-                                    DestBuf[BASE_CHANNEL_STRIDE * I + C] = (uint8)FMath::Min(255u, Result);
-                                }
-                                else
-                                {
-                                    DestBuf[BASE_CHANNEL_STRIDE * I + C] = (uint8)Result;
-                                }
-                            }
-
-                            constexpr bool bIsNC4 = (BASE_CHANNEL_STRIDE == 4);
-                            if (bIsNC4)
-                            {
-                                DestBuf[BASE_CHANNEL_STRIDE * I + 3] = BaseBuf[BASE_CHANNEL_STRIDE * I + 3];
-                            }
-                        }
-
-                        DestBuf += BASE_CHANNEL_STRIDE * Different;
-                        BaseBuf += BASE_CHANNEL_STRIDE * Different;
-                        MaskBuf += Different;
-                    }
-                }
-
-                Rows = FMath::DivideAndRoundUp(Rows, 2);
-                Width = FMath::DivideAndRoundUp(Width, 2);
-            }
-        }
-        else
-        {
-            checkf( false, TEXT("Unsupported mask format.") );
-        }
-	}
-
-
-	/**
-	* Apply a blending function to an image with a colour source and a mask
-	*/
-	template<
-		uint32 (*BLEND_FUNC_MASKED)(uint32,uint32,uint32),
-		uint32 (*BLEND_FUNC)(uint32,uint32),
-		bool CLAMP>
-    inline void BufferLayerColour(FImage* ResultImage, const FImage* BaseImage, const FImage* MaskImage, FVector4f Col)
-	{
-		check(BaseImage->GetSizeX() == MaskImage->GetSizeX());
-		check(BaseImage->GetSizeY() == MaskImage->GetSizeY());
-		check(MaskImage->GetFormat() == EImageFormat::L_UByte ||
-			  MaskImage->GetFormat() == EImageFormat::L_UByteRLE);
-
-		const bool bValid = (BaseImage->GetSizeX() == MaskImage->GetSizeX()) &&
-						    (BaseImage->GetSizeY() == MaskImage->GetSizeY()) &&
-							(MaskImage->GetFormat() == EImageFormat::L_UByte || MaskImage->GetFormat() == EImageFormat::L_UByteRLE);
-		if (!bValid)
-		{
-			return;
-		}
-
-		EImageFormat BaseFormat = BaseImage->GetFormat();
-		if (BaseFormat == EImageFormat::RGB_UByte)
-		{
-			BufferLayerColourFormat<BLEND_FUNC_MASKED, BLEND_FUNC, CLAMP, 3, 3>(ResultImage, BaseImage, MaskImage, Col, 0, 0, false);
-		}
-        else if (BaseFormat == EImageFormat::RGBA_UByte)
-        {
-            BufferLayerColourFormat<BLEND_FUNC_MASKED, BLEND_FUNC, CLAMP, 3, 4>(ResultImage, BaseImage, MaskImage, Col, 0, 0, false);
-        }
-        else if (BaseFormat == EImageFormat::BGRA_UByte)
-        {
-            float Temp = Col[0];
-            Col[0] = Col[2];
-            Col[2] = Temp;
-            BufferLayerColourFormat<BLEND_FUNC_MASKED, BLEND_FUNC, CLAMP, 3, 4>(ResultImage, BaseImage, MaskImage, Col, 0, 0, false);
-        }
-        else if (BaseFormat == EImageFormat::L_UByte)
-		{
-			BufferLayerColourFormat<BLEND_FUNC_MASKED, BLEND_FUNC, CLAMP, 1, 1>(ResultImage, BaseImage, MaskImage, Col, 0, 0, false);
-		}
-		else
-		{
-			checkf(false, TEXT("Unsupported format."));
-		}
-	}
-
-
-	template< 
-		uint32 (*BLEND_FUNC_MASKED)(uint32, uint32, uint32),
-		uint32 (*BLEND_FUNC)(uint32, uint32)>
-	inline void BufferLayerColour(FImage* DestImage, const FImage* BaseImage, const FImage* MaskImage, FVector4f Col)
-	{
-		bool bIsClampNeeded = IsAnyComponentLargerThan1(Col);
-		if (bIsClampNeeded)
-		{
-			BufferLayerColour<BLEND_FUNC_MASKED, BLEND_FUNC, true>(DestImage, BaseImage, MaskImage, Col);
-		}
-		else
-		{
-			BufferLayerColour<BLEND_FUNC_MASKED, BLEND_FUNC, false>(DestImage, BaseImage, MaskImage, Col);
-		}
-	}
-
-
-	/**
-	* Apply a blending function to an image with another image as blending layer
-	*/
-	template< 
-		uint32 (*BLEND_FUNC)(uint32, uint32), 
-		bool CLAMP,
-		// Number of total channels to actually process
-		uint32 CHANNELS_TO_BLEND,
-		// Number of total channels in the base image
-		int32 BASE_CHANNEL_STRIDE>
-	inline void BufferLayerColourFormatInPlace(FImage* BaseImage, FVector4f Color, uint32 BaseChannelOffset, uint8 ColorChannelOffset, bool bOnlyFirstLOD)
-	{
-		FUintVector4 TopColor = FUintVector4(Color.X * 255.0f, Color.Y * 255.0f, Color.Z * 255.0f, Color.W * 255.0f);
-
-		int32 NumLODs = bOnlyFirstLOD ? 1 : BaseImage->GetLODCount();
-
-		constexpr int32 NumBatchElems = 4096*2;
-
-		const int32 BytesPerElem = GetImageFormatData(BaseImage->GetFormat()).BytesPerBlock;
-		const int32 NumBatches = 
-				BaseImage->DataStorage.GetNumBatchesLODRange(NumBatchElems, BytesPerElem, 0, NumLODs);
-
-		auto ProcessBatch =
-		[
-			BaseImage, TopColor, BaseChannelOffset, ColorChannelOffset, BytesPerElem, NumBatchElems, NumLODs
-		] (uint32 BatchId)
-		{	
-			const TArrayView<uint8> BaseBatchView = 
-					BaseImage->DataStorage.GetBatchLODRange(BatchId, NumBatchElems, BytesPerElem, 0, NumLODs);
-			
-			uint8* BaseBuf = BaseBatchView.GetData() + BaseChannelOffset;
-			const int32 NumElems = BaseBatchView.Num() / BytesPerElem;
-			
-			for (int32 I = 0; I < NumElems; ++I)
-			{
-				for (int32 C = 0; C < CHANNELS_TO_BLEND; ++C)
-				{
-					uint32 Base = BaseBuf[BASE_CHANNEL_STRIDE*I + C];
-					uint32 Blended = TopColor[C + ColorChannelOffset];
-					uint32 Result = BLEND_FUNC(Base, Blended);
-					if constexpr (CLAMP)
-					{
-						BaseBuf[BASE_CHANNEL_STRIDE*I + C] = (uint8)FMath::Min(255u, Result);
-					}
-					else
-					{
-						BaseBuf[BASE_CHANNEL_STRIDE*I + C] = (uint8)Result;
-					}
-				}
-			}
-		};
-
-		if (NumBatches == 1)
-		{
-			ProcessBatch(0);
-		}
-		else if (NumBatches > 1)
-		{
-			ParallelFor(NumBatches, ProcessBatch);
-		}
-	}
-
-
-	template<
-		uint32 (*BLEND_FUNC)(uint32, uint32), 
-		bool CLAMP, 
-		uint32 CHANNEL_COUNT>
-	inline void BufferLayerColourInPlace(FImage* BaseImage, FVector4f Col, bool bOnlyOneMip, uint32 BaseOffset, uint8 ColOffset)
-	{
-		EImageFormat BaseFormat = BaseImage->GetFormat();
-
-		if (BaseFormat == EImageFormat::RGB_UByte)
-		{
-			check(BaseOffset + CHANNEL_COUNT <= 3);
-			BufferLayerColourFormatInPlace<BLEND_FUNC, CLAMP, CHANNEL_COUNT, 3>
-					(BaseImage, Col, BaseOffset, ColOffset, bOnlyOneMip);
-		}
-		else if (BaseFormat == EImageFormat::RGBA_UByte)
-		{
-			check(BaseOffset + CHANNEL_COUNT <= 4);
-			BufferLayerColourFormatInPlace<BLEND_FUNC, CLAMP, CHANNEL_COUNT, 4>
-					(BaseImage, Col, BaseOffset, ColOffset, bOnlyOneMip);
-		}
-		else if (BaseFormat == EImageFormat::BGRA_UByte)
-		{
-			float Temp = Col[0];
-			Col[0] = Col[2];
-			Col[2] = Temp;
-			BufferLayerColourFormatInPlace<BLEND_FUNC, CLAMP, CHANNEL_COUNT, 4>
-					(BaseImage, Col, BaseOffset, ColOffset, bOnlyOneMip);
-		}
-
-		else if (BaseFormat == EImageFormat::L_UByte)
-		{
-			check(BaseOffset + CHANNEL_COUNT <= 1);
-			BufferLayerColourFormatInPlace< BLEND_FUNC, CLAMP, CHANNEL_COUNT, 1>
-					(BaseImage, Col, BaseOffset, ColOffset, bOnlyOneMip);
-		}
-		else
-		{
-			checkf(false, TEXT("Unsupported format."));
-		}
-	}
-
-
-	template< 
-		uint32 (*BLEND_FUNC)(uint32, uint32), 
-		uint32 CHANNEL_COUNT>
-	inline void BufferLayerColourInPlace(FImage* BaseImage, FVector4f Color, bool bOnlyOneMip, uint32 BaseOffset, uint8 ColorOffset)
-	{
-		bool bIsClampNeeded = IsAnyComponentLargerThan1(Color);
-		if (bIsClampNeeded)
-		{
-			BufferLayerColourInPlace<BLEND_FUNC, true, CHANNEL_COUNT>(BaseImage, Color, bOnlyOneMip, BaseOffset, ColorOffset);
-		}
-		else
-		{
-			BufferLayerColourInPlace<BLEND_FUNC, false, CHANNEL_COUNT>(BaseImage, Color, bOnlyOneMip, BaseOffset, ColorOffset);
-		}
-	}
-
-
-	template< 
-		uint32 (*BLEND_FUNC_MASKED)(uint32, uint32, uint32),
-		uint32 (*BLEND_FUNC)(uint32, uint32), 
-		bool CLAMP, 
-		int32 CHANNEL_COUNT>
-	inline void BufferLayerColourInPlace(FImage* BaseImage, const FImage* MaskImage, FVector4f Color, bool bOnlyOneMip, uint32 BaseOffset, uint8 ColorOffset)
-	{
-		check(BaseImage->GetSizeX() == MaskImage->GetSizeX());
-		check(BaseImage->GetSizeY() == MaskImage->GetSizeY());
-
-		EImageFormat BaseFormat = BaseImage->GetFormat();
-
-		if (BaseFormat == EImageFormat::RGB_UByte)
-		{
-			check(BaseOffset + CHANNEL_COUNT <= 3);
-			BufferLayerColourFormat<BLEND_FUNC_MASKED, BLEND_FUNC, CLAMP, CHANNEL_COUNT, 3>
-					(BaseImage, BaseImage, MaskImage, Color, BaseOffset, ColorOffset, bOnlyOneMip);
-		}
-		else if (BaseFormat == EImageFormat::RGBA_UByte)
-		{
-			check(BaseOffset + CHANNEL_COUNT <= 4);
-			BufferLayerColourFormat<BLEND_FUNC_MASKED, BLEND_FUNC, CLAMP, CHANNEL_COUNT, 4>
-					(BaseImage, BaseImage, MaskImage, Color, BaseOffset, ColorOffset, bOnlyOneMip);
-		}
-		else if (BaseFormat == EImageFormat::BGRA_UByte)
-		{
-			check(BaseOffset + CHANNEL_COUNT <= 4);
-			float Temp = Color[0];
-			Color[0] = Color[2];
-			Color[2] = Temp;
-
-			BufferLayerColourFormat<BLEND_FUNC_MASKED, BLEND_FUNC, CLAMP, CHANNEL_COUNT, 4>
-					(BaseImage, BaseImage, MaskImage, Color, BaseOffset, ColorOffset, bOnlyOneMip);
-		}
-		else if (BaseFormat == EImageFormat::L_UByte)
-		{
-			check(BaseOffset + CHANNEL_COUNT <= 1);
-			BufferLayerColourFormat<BLEND_FUNC_MASKED, BLEND_FUNC, CLAMP, CHANNEL_COUNT, 1>
-					(BaseImage, BaseImage, MaskImage, Color, BaseOffset, ColorOffset, bOnlyOneMip);
-		}
-		else
-		{
-			checkf(false, TEXT("Unsupported format."));
-		}
-	}
-
-	template< 
-		uint32 (*BLEND_FUNC_MASKED)(uint32, uint32, uint32), 
-		uint32 (*BLEND_FUNC)(uint32, uint32), uint32 CHANNEL_COUNT>
-	inline void BufferLayerColourInPlace(FImage* BaseImage, const FImage* MaskImage, FVector4f Color, bool bOnlyOneMip, uint32 BaseOffset, uint8 ColorOffset)
-	{
-		bool bIsClampNeeded = IsAnyComponentLargerThan1(Color);
-		if (bIsClampNeeded)
-		{
-			BufferLayerColourInPlace<BLEND_FUNC_MASKED, BLEND_FUNC, true, CHANNEL_COUNT>
-					(BaseImage, MaskImage, Color, bOnlyOneMip, BaseOffset, ColorOffset);
-		}
-		else
-		{
-			BufferLayerColourInPlace<BLEND_FUNC_MASKED, BLEND_FUNC, false, CHANNEL_COUNT>
-					(BaseImage, MaskImage, Color, bOnlyOneMip, BaseOffset, ColorOffset);
-		}
-	}
-
-
-	/**	
-	* Apply a blending function to an image with another image as blending layer
-	*/
-	template< uint32 (*BLEND_FUNC)(uint32, uint32), bool CLAMP,
-		// Number of total channels to actually process
-		int32 CHANNELS_TO_BLEND,
-		// Number of total channels in the base image
-		int32 BASE_CHANNEL_STRIDE,
-		// Number of total channels in the blend image
-		int32 BLENDED_CHANNEL_STRIDE>
-    inline void BufferLayerFormatInPlace(
-		FImage* BaseImage, 
-		const FImage* BlendedImage,
-		uint32 BaseChannelOffset,
-		uint32 BlendedChannelOffset,
-		bool bOnlyFirstLOD)
-	{
-		check(BaseImage->GetSizeX() == BlendedImage->GetSizeX());
-		check(BaseImage->GetSizeY() == BlendedImage->GetSizeY());
-
-		// No longer required.
-		//check(BaseImage->GetFormat() == BlendedImage->GetFormat());
-		check(BaseChannelOffset + CHANNELS_TO_BLEND <= GetImageFormatData(BaseImage->GetFormat()).Channels);
-		check(BlendedChannelOffset + CHANNELS_TO_BLEND <= GetImageFormatData(BlendedImage->GetFormat()).Channels);
-
-		check(bOnlyFirstLOD || BaseImage->GetLODCount() <= BlendedImage->GetLODCount());
-
-		const int32 BaseBytesPerElem = GetImageFormatData(BaseImage->GetFormat()).BytesPerBlock;
-		const int32 BlendBytesPerElem = GetImageFormatData(BlendedImage->GetFormat()).BytesPerBlock;
-	
-		const int32 NumLODs = bOnlyFirstLOD ? 1 : BaseImage->GetLODCount();
-
-		constexpr int32 NumBatchElems = 4096*2;
-		const int32 NumBatches =
-				BaseImage->DataStorage.GetNumBatchesLODRange(NumBatchElems, BaseBytesPerElem, 0, NumLODs);
-		
-		check(NumBatches <= BlendedImage->DataStorage.GetNumBatches(NumBatchElems, BlendBytesPerElem));
-
-		auto ProcessBatch =
-		[
-			BaseImage, BlendedImage, BaseBytesPerElem, BlendBytesPerElem, 
-			NumBatchElems, BaseChannelOffset, BlendedChannelOffset, NumLODs
-		] (int32 BatchId)
-		{
-			TArrayView<uint8> BaseBatchView =
-					BaseImage->DataStorage.GetBatchLODRange(BatchId, NumBatchElems, BaseBytesPerElem, 0, NumLODs);
-
-			TArrayView<const uint8> BlendedBatchView =
-					BlendedImage->DataStorage.GetBatchLODRange(BatchId, NumBatchElems, BlendBytesPerElem, 0, NumLODs);
-
-        	uint8* BaseBuf = BaseBatchView.GetData() + BaseChannelOffset;
-        	const uint8* BlendedBuf = BlendedBatchView.GetData() + BlendedChannelOffset;
-
-			const int32 NumElems = BaseBatchView.Num() / BaseBytesPerElem;
-			check(BlendedBatchView.Num() / BlendBytesPerElem == NumElems);
-
-			for (int32 I = 0; I < NumElems; ++I)
-			{
-				for (int32 C = 0; C < CHANNELS_TO_BLEND; ++C)
-				{
-					uint32 Base = BaseBuf[BASE_CHANNEL_STRIDE * I + C];
-					uint32 Blended = BlendedBuf[BLENDED_CHANNEL_STRIDE * I + C];
-					uint32 Result = BLEND_FUNC(Base, Blended);
-					
-					if constexpr (CLAMP)
-					{
-						BaseBuf[BASE_CHANNEL_STRIDE * I + C] = (uint8)FMath::Min(255u, Result);
-					}
-					else
-					{
-						BaseBuf[BASE_CHANNEL_STRIDE * I + C] = (uint8)Result;
-					}
-				}
-			}
-		};
-
-		if (NumBatches == 1)
-		{
-			ProcessBatch(0);
-		}
-		else if (NumBatches > 1)
-		{
-			ParallelFor(NumBatches, ProcessBatch);
-		}
-	}
-	
-	/**
-	* Apply a blending function to an image with another image as blending layer
-	*/
-	template<
-		uint32 (*BLEND_FUNC)(uint32, uint32), 
-		bool CLAMP,
-		// Number of total channels to actually process
-		int32 CHANNELS_TO_BLEND,
-		// Number of total channels in the base image
-		int32 BASE_CHANNEL_STRIDE,
-		// Number of total channels in the blend image
-		int32 BLENDED_CHANNEL_STRIDE,
-		int32 BLENDED_CHANNEL_OFFSET>
-	inline void BufferLayerFormat(FImage* DestImage, const FImage* BaseImage, const FImage* BlendedImage, bool bOnlyFirstLOD)
-	{
-		check(BaseImage->GetSizeX() == BlendedImage->GetSizeX());
-		check(BaseImage->GetSizeY() == BlendedImage->GetSizeY());
-		// Not true anymore, since the BLENDED_CHANNEL_OFFSET has been added.
-		// check(BaseImage->GetFormat() == BlendedImage->GetFormat());
-		check(bOnlyFirstLOD || BaseImage->GetLODCount() <= BlendedImage->GetLODCount());
-
-		constexpr int32 UnblendedChannels = BASE_CHANNEL_STRIDE - CHANNELS_TO_BLEND;
-		constexpr int32 NumBatchElems = 4096*2;
-		
-		const int32 BaseBytesPerElem = GetImageFormatData(BaseImage->GetFormat()).BytesPerBlock;
-		const int32 BlendBytesPerElem = GetImageFormatData(BlendedImage->GetFormat()).BytesPerBlock;
-		const int32 DestBytesPerElem = GetImageFormatData(DestImage->GetFormat()).BytesPerBlock;
-
-		const int32 NumLODs = bOnlyFirstLOD ? 1 : BaseImage->GetLODCount();
-
-		const int32 NumBatches = 
-				BaseImage->DataStorage.GetNumBatchesLODRange(NumBatchElems, BaseBytesPerElem, 0, NumLODs);
-
-		check(NumBatches == BlendedImage->DataStorage.GetNumBatchesLODRange(NumBatchElems, BlendBytesPerElem, 0, NumLODs));
-		check(NumBatches == DestImage->DataStorage.GetNumBatchesLODRange(NumBatchElems, DestBytesPerElem, 0, NumLODs));
-
-		auto ProcessBatch =
-		[
-			BaseImage, BlendedImage, DestImage, 
-			BaseBytesPerElem, BlendBytesPerElem, DestBytesPerElem, 
-			UnblendedChannels, NumBatchElems, NumLODs
-		] (int32 BatchId)
-		{
-			TArrayView<const uint8> BaseBatchView =
-					BaseImage->DataStorage.GetBatchLODRange(BatchId, NumBatchElems, BaseBytesPerElem, 0, NumLODs);
-
-			TArrayView<const uint8> BlendedBatchView =
-					BlendedImage->DataStorage.GetBatchLODRange(BatchId, NumBatchElems, BlendBytesPerElem, 0, NumLODs);
-
-			TArrayView<uint8> DestBatchView =
-					DestImage->DataStorage.GetBatchLODRange(BatchId, NumBatchElems, DestBytesPerElem, 0, NumLODs);
-	
-			const int32 NumElems = BaseBatchView.Num() / BaseBytesPerElem;
-			check(NumElems == BlendedBatchView.Num() / BlendBytesPerElem);
-			check(NumElems == DestBatchView.Num() / DestBytesPerElem);
-
-			const uint8* BaseBuf = BaseBatchView.GetData();
-			const uint8* BlendedBuf = BlendedBatchView.GetData() + BLENDED_CHANNEL_OFFSET;
-			uint8* DestBuf = DestBatchView.GetData();
-
-			for (int32 I = 0; I < NumElems; ++I)
-			{
-				for (int32 C = 0; C < CHANNELS_TO_BLEND; ++C)
-				{
-					uint32 Base = BaseBuf[BASE_CHANNEL_STRIDE * I + C];
-					uint32 Blended = BlendedBuf[BLENDED_CHANNEL_STRIDE * I + C];
-					uint32 Result = BLEND_FUNC(Base, Blended);
-
-					if constexpr (CLAMP)
-					{
-						DestBuf[BASE_CHANNEL_STRIDE * I + C] = (uint8)FMath::Min(255u, Result);
-					}
-					else
-					{
-						DestBuf[BASE_CHANNEL_STRIDE * I + C] = (uint8)Result;
-					}
-				}
-
-				// Copy the unblended channels
-				// \TODO: unnecessary when doing it in-place?
-				if constexpr (UnblendedChannels > 0)
-				{
-					for (int32 C = 0; C < UnblendedChannels; ++C)
-					{
-						DestBuf[BASE_CHANNEL_STRIDE * I + CHANNELS_TO_BLEND + C] 
-							= BaseBuf[BASE_CHANNEL_STRIDE * I + CHANNELS_TO_BLEND + C];
-					}
-				}
-			}
-		};
-
-		if (NumBatches == 1)
-		{
-			ProcessBatch(0);
-		}
-		else if (NumBatches > 1)
-		{
-			ParallelFor(NumBatches, ProcessBatch);
-		}
-	}
-
-	/**
-	* Apply a blending function to an image with another image as blending layer
-	*/
-	template< uint32 (*BLEND_FUNC)(uint32, uint32), bool CLAMP >
-	inline void BufferLayer(FImage* ResultImage, const FImage* BaseImage, const FImage* BlendedImage, bool bApplyToAlpha, bool bOnlyOneMip, bool bUseBlendSourceFromBlendAlpha)
-	{
-		check(ResultImage->GetFormat() == BaseImage->GetFormat());
-		check(ResultImage->GetSizeX() == BaseImage->GetSizeX());
-		check(ResultImage->GetSizeY() == BaseImage->GetSizeY());
-		check(bOnlyOneMip || ResultImage->GetLODCount() == BaseImage->GetLODCount());
-		check(BaseImage->GetSizeX() == BlendedImage->GetSizeX());
-		check(BaseImage->GetSizeY() == BlendedImage->GetSizeY());
-		check(bOnlyOneMip || ResultImage->GetLODCount() <= BlendedImage->GetLODCount());
-
-		EImageFormat BaseFormat = BaseImage->GetFormat();
-		EImageFormat BlendedFormat = BlendedImage->GetFormat();
-
-		if (bUseBlendSourceFromBlendAlpha)
-		{
-			if (BlendedFormat == EImageFormat::RGBA_UByte || BlendedFormat == EImageFormat::BGRA_UByte)
-			{
-				if (BaseFormat == EImageFormat::L_UByte)
-				{
-					BufferLayerFormat<BLEND_FUNC, CLAMP, 1, 1, 4, 3>(ResultImage, BaseImage, BlendedImage, bOnlyOneMip);
-				}
-				else
-				{
-					checkf(false, TEXT("Unsupported format."));
-				}
-			}
-			else if (BlendedFormat == EImageFormat::L_UByte)
-			{
-				BufferLayerFormat<BLEND_FUNC, CLAMP, 1, 1, 1, 0>(ResultImage, BaseImage, BlendedImage, bOnlyOneMip);
-			}
-		}
-		else
-		{
-			check(BaseFormat == BlendedFormat);
-			if (BaseFormat == EImageFormat::RGB_UByte)
-			{
-				check(!bUseBlendSourceFromBlendAlpha);
-				BufferLayerFormat<BLEND_FUNC, CLAMP, 3, 3, 3, 0>(ResultImage, BaseImage, BlendedImage, bOnlyOneMip);
-			}
-			else if (BaseFormat == EImageFormat::RGBA_UByte || BaseFormat == EImageFormat::BGRA_UByte)
-			{
-				check(!bUseBlendSourceFromBlendAlpha);
-				if (bApplyToAlpha)
-				{
-					BufferLayerFormat<BLEND_FUNC, CLAMP, 4, 4, 4, 0>(ResultImage, BaseImage, BlendedImage, bOnlyOneMip);
-				}
-				else
-				{
-					BufferLayerFormat<BLEND_FUNC, CLAMP, 3, 4, 4, 0>(ResultImage, BaseImage, BlendedImage, bOnlyOneMip);
-				}
-			}
-			else if (BaseFormat == EImageFormat::L_UByte)
-			{
-				BufferLayerFormat<BLEND_FUNC, CLAMP, 1, 1, 1, 0>(ResultImage, BaseImage, BlendedImage, bOnlyOneMip);
-			}
-			else
-			{
-				checkf(false, TEXT("Unsupported format."));
-			}
-		}
-	}
-
-
-	template< uint32 (*BLEND_FUNC)(uint32, uint32), bool CLAMP, int32 CHANNEL_COUNT >
-	inline void BufferLayerInPlace(FImage* BaseImage, const FImage* BlendedImage, bool bOnlyOneMip, uint32 BaseOffset, uint32 BlendedOffset)
-	{
-		check(BaseImage->GetSizeX() == BlendedImage->GetSizeX());
-		check(BaseImage->GetSizeY() == BlendedImage->GetSizeY());
-		// Not required since we have the CHANNEL_COUNT and offsets. 
-		// check(BaseImage->GetFormat() == BlendedImage->GetFormat());
-
-		EImageFormat BaseFormat = BaseImage->GetFormat();
-		EImageFormat BlendFormat = BlendedImage->GetFormat();
-
-		if (BaseFormat == EImageFormat::RGB_UByte && BlendFormat==EImageFormat::RGB_UByte)
-		{
-			check(BaseOffset + CHANNEL_COUNT <= 3);
-			check(BlendedOffset + CHANNEL_COUNT <= 3);
-			BufferLayerFormatInPlace<BLEND_FUNC, CLAMP, CHANNEL_COUNT, 3, 3>(BaseImage, BlendedImage, BaseOffset, BlendedOffset, bOnlyOneMip);
-		}
-		else if ((BaseFormat == EImageFormat::RGBA_UByte || BaseFormat == EImageFormat::BGRA_UByte) &&
-				 (BlendFormat == EImageFormat::RGBA_UByte || BlendFormat == EImageFormat::BGRA_UByte) )
-		{
-			check(BaseOffset + CHANNEL_COUNT <= 4);
-			check(BlendedOffset + CHANNEL_COUNT <= 4);
-			BufferLayerFormatInPlace< BLEND_FUNC, CLAMP, CHANNEL_COUNT, 4, 4>
-					(BaseImage, BlendedImage, BaseOffset, BlendedOffset, bOnlyOneMip);
-		}
-		else if ((BaseFormat == EImageFormat::RGBA_UByte || BaseFormat == EImageFormat::BGRA_UByte) &&
-				 BlendFormat == EImageFormat::L_UByte )
-		{
-			check(BaseOffset + CHANNEL_COUNT <= 4);
-			check(BlendedOffset + CHANNEL_COUNT <= 1);
-			BufferLayerFormatInPlace< BLEND_FUNC, CLAMP, CHANNEL_COUNT, 4, 1>
-					(BaseImage, BlendedImage, BaseOffset, BlendedOffset, bOnlyOneMip);
-		}
-		else if (BaseFormat == EImageFormat::L_UByte && BlendFormat==EImageFormat::L_UByte)
-		{
-			check(BaseOffset + CHANNEL_COUNT <= 1);
-			BufferLayerFormatInPlace< BLEND_FUNC, CLAMP, CHANNEL_COUNT, 1, 1>
-					(BaseImage, BlendedImage, BaseOffset, BlendedOffset, bOnlyOneMip);
-		}
-		else
-		{
-			checkf(false, TEXT("Unsupported format."));
-		}
-	}
-
-
-	template< uint32 (*BLEND_FUNC_MASKED)(uint32, uint32, uint32),
-		uint32 (*BLEND_FUNC)(uint32, uint32),
-		bool CLAMP,
-		// Number of total channels to actually process
-		int32 CHANNELS_TO_BLEND,
-		// Number of total channels in the base image
-		int32 BASE_CHANNEL_STRIDE,
-		// Number of total channels in the blend image
-		int32 BLENDED_CHANNEL_STRIDE>
-	inline void BufferLayerFormat(
-		FImage* DestImage,
-		const FImage* BaseImage,
-		const FImage* MaskImage,
-		const FImage* BlendImage,
-		uint32 DestOffset,
-		uint32 BaseChannelOffset,
-		uint32 BlendedChannelOffset,
-		bool bOnlyFirstLOD)
-	{
-		check(BaseImage->GetSizeX() == MaskImage->GetSizeX() && BaseImage->GetSizeY() == MaskImage->GetSizeY());
-		check(BaseImage->GetSizeX() == BlendImage->GetSizeX() && BaseImage->GetSizeY() == BlendImage->GetSizeY());
-		check(bOnlyFirstLOD || BaseImage->GetLODCount() <= MaskImage->GetLODCount());
-		check(bOnlyFirstLOD || BaseImage->GetLODCount() <= BlendImage->GetLODCount());
-		check(BaseImage->GetFormat() == BlendImage->GetFormat());
-		check(MaskImage->GetFormat() == EImageFormat::L_UByte ||
-			  MaskImage->GetFormat() == EImageFormat::L_UByteRLE);
-
-		const EImageFormat MaskFormat = MaskImage->GetFormat();
-
-		const int32 BaseBytesPerElem = GetImageFormatData(BaseImage->GetFormat()).BytesPerBlock; 
-		const int32 BlendBytesPerElem = GetImageFormatData(BlendImage->GetFormat()).BytesPerBlock; 
-
-		const int32 NumLODs = bOnlyFirstLOD ? 1 : BaseImage->GetLODCount();
-
-		const bool bIsMaskUncompressed = MaskFormat == EImageFormat::L_UByte;
-		constexpr int32 UnblendedChannels = BASE_CHANNEL_STRIDE - CHANNELS_TO_BLEND;
-
-		if (bIsMaskUncompressed)
-		{	
-			constexpr int32 NumBatchElems = 4096*2;
-
-			const int32 NumBatches = 
-				BaseImage->DataStorage.GetNumBatchesLODRange(NumBatchElems, BaseBytesPerElem, 0, NumLODs);
-
-			check(NumBatches == BlendImage->DataStorage.GetNumBatchesLODRange(NumBatchElems, BlendBytesPerElem, 0, NumLODs));
-			check(NumBatches == MaskImage->DataStorage.GetNumBatchesLODRange(NumBatchElems, 1, 0, NumLODs));
-			check(NumBatches == DestImage->DataStorage.GetNumBatchesLODRange(NumBatchElems, BaseBytesPerElem, 0, NumLODs));
-
-			auto ProcessBatch = 
-			[
-				BaseImage, BlendImage, MaskImage, DestImage, UnblendedChannels, BaseBytesPerElem, BlendBytesPerElem, NumLODs
-			] (uint32 BatchId)
-			{
-				TArrayView<const uint8> BaseBatchView =
-						BaseImage->DataStorage.GetBatchLODRange(BatchId, NumBatchElems, BaseBytesPerElem, 0, NumLODs);
-
-				TArrayView<const uint8> BlendBatchView =
-						BlendImage->DataStorage.GetBatchLODRange(BatchId, NumBatchElems, BlendBytesPerElem, 0, NumLODs);
-
-				TArrayView<const uint8> MaskBatchView =
-						MaskImage->DataStorage.GetBatchLODRange(BatchId, NumBatchElems, 1, 0, NumLODs);
-
-				TArrayView<uint8> DestBatchView = 
-						DestImage->DataStorage.GetBatchLODRange(BatchId, NumBatchElems, BaseBytesPerElem, 0, NumLODs);
-
-				const uint8* BaseBuf = BaseBatchView.GetData();
-				const uint8* BlendedBuf = BlendBatchView.GetData();
-				const uint8* MaskBuf = MaskBatchView.GetData();
-				uint8* DestBuf = DestBatchView.GetData();
-
-				// This could happen in case of missing data files.
-				if (!BaseBuf || !BlendedBuf || !MaskBuf || !DestBuf)
-				{
-					return;
-				}
-
-				const int32 NumElems = BaseBatchView.Num() / BaseBytesPerElem;
-				check(NumElems == BlendBatchView.Num() / BlendBytesPerElem);
-				check(NumElems == MaskBatchView.Num() / 1);
-				check(NumElems == DestBatchView.Num() / BaseBytesPerElem);
-
-				for (int32 I = 0; I < NumElems; ++I)
-				{
-					uint32 Mask = MaskBuf[I];
-					for (int32 C = 0; C < CHANNELS_TO_BLEND; ++C)
-					{
-						const uint32 Base = BaseBuf[BASE_CHANNEL_STRIDE * I + C];
-						const uint32 Blended = BlendedBuf[BLENDED_CHANNEL_STRIDE * I + C];
-						const uint32 Result = BLEND_FUNC_MASKED(Base, Blended, Mask);
-						if constexpr (CLAMP)
-						{
-							DestBuf[BASE_CHANNEL_STRIDE * I + C] = (uint8)FMath::Min(255u, Result);
-						}
-						else
-						{
-							DestBuf[BASE_CHANNEL_STRIDE * I + C] = (uint8)Result;
-						}
-					}
-					// Copy the unblended channels
-					// \TODO: unnecessary when doing it in-place?
-					if constexpr (UnblendedChannels > 0)
-					{
-						for (int32 C = 0; C < UnblendedChannels; ++C)
-						{
-							DestBuf[BASE_CHANNEL_STRIDE * I + CHANNELS_TO_BLEND + C] = BaseBuf[BASE_CHANNEL_STRIDE * I + CHANNELS_TO_BLEND + C];
-						}
-					}
-				}
-			};
-
-			if (NumBatches == 1)
-			{
-				ProcessBatch(0);
-			}
-			else if (NumBatches > 1)
-			{
-				ParallelFor(NumBatches, ProcessBatch);
-			}
-		}
-		else if (MaskFormat == EImageFormat::L_UByteRLE)
-		{
-			int32 Rows = BaseImage->GetSizeY();
-			int32 Width = BaseImage->GetSizeX();
-
-			for (int32 LOD = 0; LOD < NumLODs; ++LOD)
-			{
-				const uint8* BaseBuf = BaseImage->GetLODData(LOD);
-				uint8* DestBuf = DestImage->GetLODData(LOD);
-				const uint8* BlendedBuf = BlendImage->GetLODData(LOD);
-				const uint8* MaskBuf = MaskImage->GetLODData(LOD);
-
-				// This could happen in case of missing data files.
-				if (!BaseBuf || !BlendedBuf || !MaskBuf || !DestBuf)
-				{
-					continue;
-				}
-
-				// Remove RLE header, mip size and row sizes.
-				MaskBuf += sizeof(uint32) + Rows * sizeof(uint32);
-
-				for (int32 RowIndex = 0; RowIndex < Rows; ++RowIndex)
-				{
-					const uint8* DestRowEnd = DestBuf + Width * BASE_CHANNEL_STRIDE;
-					while (DestBuf != DestRowEnd)
-					{
-						// Decode header
-						uint16 Equal;
-						FMemory::Memcpy(&Equal, MaskBuf, 2);
-						MaskBuf += 2;
-
-						uint8 Different = *MaskBuf;
-						++MaskBuf;
-
-						uint8 EqualPixel = *MaskBuf;
-						++MaskBuf;
-
-						// Equal pixels
-						//check(DestBuf + BASE_CHANNEL_STRIDE * Equal <= BaseImage->GetLODDataSize(0));
-						if (EqualPixel == 255)
-						{
-							for (int32 I = 0; I < Equal; ++I)
-							{
-								for (int32 C = 0; C < CHANNELS_TO_BLEND; ++C)
-								{
-									uint32 Base = BaseBuf[BASE_CHANNEL_STRIDE * I + C];
-									uint32 Blended = BlendedBuf[BLENDED_CHANNEL_STRIDE * I + C];
-									uint32 Result = BLEND_FUNC(Base, Blended);
-									if constexpr (CLAMP)
-									{
-										DestBuf[BASE_CHANNEL_STRIDE * I + C] = (uint8)FMath::Min(255u, Result);
-									}
-									else
-									{
-										DestBuf[BASE_CHANNEL_STRIDE * I + C] = (uint8)Result;
-									}
-								}
-
-								// Copy the unblended channels
-								// \TODO: unnecessary when doing it in-place?
-								if constexpr (UnblendedChannels > 0)
-								{
-									for (int32 C = 0; C < UnblendedChannels; ++C)
-									{
-										DestBuf[BASE_CHANNEL_STRIDE * I + CHANNELS_TO_BLEND + C] = BaseBuf[BASE_CHANNEL_STRIDE * I + CHANNELS_TO_BLEND + C];
-									}
-								}
-							}
-						}
-						else if (EqualPixel > 0)
-						{
-							for (int32 I = 0; I < Equal; ++I)
-							{
-								for (int32 C = 0; C < CHANNELS_TO_BLEND; ++C)
-								{
-									uint32 Base = BaseBuf[BASE_CHANNEL_STRIDE * I + C];
-									uint32 Blended = BlendedBuf[BLENDED_CHANNEL_STRIDE * I + C];
-									uint32 Result = BLEND_FUNC_MASKED(Base, Blended, EqualPixel);
-
-									if constexpr (CLAMP)
-									{
-										DestBuf[BASE_CHANNEL_STRIDE * I + C] = (uint8)FMath::Min(255u, Result);
-									}
-									else
-									{
-										DestBuf[BASE_CHANNEL_STRIDE * I + C] = (uint8)Result;
-									}
-								}
-
-								// Copy the unblended channels
-								// \TODO: unnecessary when doing it in-place?
-								if constexpr (UnblendedChannels > 0)
-								{
-									for (int32 C = 0; C < UnblendedChannels; ++C)
-									{
-										DestBuf[BASE_CHANNEL_STRIDE * I + CHANNELS_TO_BLEND + C] = BaseBuf[BASE_CHANNEL_STRIDE * I + CHANNELS_TO_BLEND + C];
-									}
-								}
-							}
-						}
-						else
-						{
-							// It could happen if xxxxxOnBase
-							if (DestBuf != BaseBuf)
-							{
-								FMemory::Memmove(DestBuf, BaseBuf, BASE_CHANNEL_STRIDE * Equal);
-							}
-						}
-
-						DestBuf += BASE_CHANNEL_STRIDE * Equal;
-						BaseBuf += BASE_CHANNEL_STRIDE * Equal;
-						BlendedBuf += BLENDED_CHANNEL_STRIDE * Equal;
-
-						// Different pixels
-						//check(DestBuf + BASE_CHANNEL_STRIDE * Different <= BaseImage->GetDataSize(0));
-						for (int32 I = 0; I < Different; ++I)
-						{
-							for (int32 C = 0; C < CHANNELS_TO_BLEND; ++C)
-							{
-								uint32 Mask = MaskBuf[I];
-								uint32 Base = BaseBuf[BASE_CHANNEL_STRIDE * I + C];
-								uint32 Blended = BlendedBuf[BLENDED_CHANNEL_STRIDE * I + C];
-								uint32 Result = BLEND_FUNC_MASKED(Base, Blended, Mask);
-								if (CLAMP)
-								{
-									DestBuf[BASE_CHANNEL_STRIDE * I + C] = (uint8)FMath::Min(255u, Result);
-								}
-								else
-								{
-									DestBuf[BASE_CHANNEL_STRIDE * I + C] = (uint8)Result;
-								}
-							}
-
-							// Copy the unblended channels
-							// \TODO: unnecessary when doing it in-place?
-							if constexpr (UnblendedChannels > 0)
-							{
-								for (int32 C = 0; C < UnblendedChannels; ++C)
-								{
-									DestBuf[BASE_CHANNEL_STRIDE * I + CHANNELS_TO_BLEND + C] = BaseBuf[BASE_CHANNEL_STRIDE * I + CHANNELS_TO_BLEND + C];
-								}
-							}
-						}
-
-						DestBuf += BASE_CHANNEL_STRIDE * Different;
-						BaseBuf += BASE_CHANNEL_STRIDE * Different;
-						BlendedBuf += BLENDED_CHANNEL_STRIDE * Different;
-						MaskBuf += Different;
-					}
-				}
-
-				Rows = FMath::DivideAndRoundUp(Rows, 2);
-				Width = FMath::DivideAndRoundUp(Width, 2);
-			}
-		}
-		else
-		{
-			checkf(false, TEXT("Unsupported mask format."));
-		}
-	}
-
-	template< uint32 (*BLEND_FUNC_MASKED)(uint32, uint32, uint32),
-		uint32 (*BLEND_FUNC)(uint32, uint32),
-		bool CLAMP,
-		// Number of total channels to actually process
-		int32 CHANNELS_TO_BLEND,
-		// Number of total channels in the base image
-		int32 BASE_CHANNEL_STRIDE,
-		// Number of total channels in the blend image
-		int32 BLENDED_CHANNEL_STRIDE>
-	inline void BufferLayerFormatEmbeddedMask(
-		FImage* DestImage,
-		const FImage* BaseImage,
-		const FImage* BlendImage,
-		uint32 DestOffset,
-		uint32 BaseChannelOffset,
-		bool bOnlyFirstLOD)
-	{
-		check(BaseImage->GetSizeX() == BlendImage->GetSizeX() && BaseImage->GetSizeY() == BlendImage->GetSizeY());
-		check(bOnlyFirstLOD || BaseImage->GetLODCount() <= BlendImage->GetLODCount());
-
-		constexpr int32 UnblendedChannels = BASE_CHANNEL_STRIDE - CHANNELS_TO_BLEND;
-
-		const int32 BytesPerElem = GetImageFormatData(BaseImage->GetFormat()).BytesPerBlock; 
-	
-		const int32 NumLODs = bOnlyFirstLOD ? 1 : BaseImage->GetLODCount();
-
-		constexpr int32 NumBatchElems = 4096*2;
-		const int32 NumBatches = BaseImage->DataStorage.GetNumBatchesLODRange(NumBatchElems, BytesPerElem, 0, NumLODs);
+{	
+	/** Used function templates need to be manually instanciated in OpImageLayer.cpp */ 
 
-		check(BytesPerElem == GetImageFormatData(DestImage->GetFormat()).BytesPerBlock);
-		check(BytesPerElem == GetImageFormatData(BlendImage->GetFormat()).BytesPerBlock);
-		
-		check(NumBatches == BlendImage->DataStorage.GetNumBatchesLODRange(NumBatchElems, BytesPerElem, 0, NumLODs));
-		check(NumBatches == DestImage->DataStorage.GetNumBatchesLODRange(NumBatchElems, BytesPerElem, 0, NumLODs));
+	template<uint32(*BlendFuncMasked)(uint32, uint32, uint32), uint32(*BlendFunc)(uint32, uint32), bool bClamp>
+    void BufferLayer(FImage* DestImage, const FImage* BaseImage, const FImage* MaskImage, const FImage* BlendImage, bool bApplyToAlpha, bool bOnlyFirstLOD);
 
-		auto ProcessBatch = 
-		[
-			BaseImage, BlendImage, DestImage, UnblendedChannels, NumBatchElems, BytesPerElem, NumLODs
-		] (int32 BatchId)
-		{
-			TArrayView<const uint8> BlendedView = 
-					BlendImage->DataStorage.GetBatchLODRange(BatchId, NumBatchElems, BytesPerElem, 0, NumLODs);
+	template<uint32(*BlendFunc)(uint32, uint32), bool bClamp>
+	void BufferLayer(FImage* ResultImage, const FImage* BaseImage, const FImage* BlendedImage, bool bApplyToAlpha, bool bOnlyOneMip, bool bUseBlendSourceFromBlendAlpha);
 
-			TArrayView<const uint8> BaseView =
-					BaseImage->DataStorage.GetBatchLODRange(BatchId, NumBatchElems, BytesPerElem, 0, NumLODs);
-
-			TArrayView<uint8> DestView =
-					DestImage->DataStorage.GetBatchLODRange(BatchId, NumBatchElems, BytesPerElem, 0, NumLODs);
-
-			uint8* const DestBuf = DestView.GetData();
-			uint8 const * const BlendedBuf = BlendedView.GetData();
-			uint8 const * const BaseBuf = BaseView.GetData(); 
-
-			// This could happen in case of missing data files.
-			if (!BaseBuf || !BlendedBuf)
-			{
-				return;
-			}
-
-			const int32 NumElems = BaseView.Num() / BytesPerElem;
-			check(NumElems == DestView.Num() / BytesPerElem);
-			check(NumElems == BlendedView.Num() / BytesPerElem);
-
-			if (BlendImage->GetFormat() == EImageFormat::RGBA_UByte)
-			{
-				const uint8* MaskBuf = BlendedBuf + 3;
-
-				for (int32 I = 0; I < NumElems; ++I)
-				{
-					uint32 Mask = MaskBuf[BLENDED_CHANNEL_STRIDE * I];
-					for (int32 C = 0; C < CHANNELS_TO_BLEND; ++C)
-					{
-						uint32 Base = BaseBuf[BASE_CHANNEL_STRIDE * I + C];
-						uint32 Blended = BlendedBuf[BLENDED_CHANNEL_STRIDE * I + C];
-						uint32 Result = BLEND_FUNC_MASKED(Base, Blended, Mask);
-						if constexpr (CLAMP)
-						{
-							DestBuf[BASE_CHANNEL_STRIDE * I + C] = (uint8)FMath::Min(255u, Result);
-						}
-						else
-						{
-							DestBuf[BASE_CHANNEL_STRIDE * I + C] = (uint8)Result;
-						}
-					}
-					// Copy the unblended channels
-					// \TODO: unnecessary when doing it in-place?
-					if constexpr (UnblendedChannels > 0)
-					{
-						for (int32 C = 0; C < UnblendedChannels; ++C)
-						{
-							DestBuf[BASE_CHANNEL_STRIDE * I + CHANNELS_TO_BLEND + C] = BaseBuf[BASE_CHANNEL_STRIDE * I + CHANNELS_TO_BLEND + C];
-						}
-					}
-				}
-			}
-			else
-			{
-				for (int32 I = 0; I < NumElems; ++I)
-				{
-					for (int32 C = 0; C < CHANNELS_TO_BLEND; ++C)
-					{
-						uint32 Base = BaseBuf[BASE_CHANNEL_STRIDE * I + C];
-						uint32 Blended = BlendedBuf[BLENDED_CHANNEL_STRIDE * I + C];
-						uint32 Result = BLEND_FUNC(Base, Blended);
-						if constexpr (CLAMP)
-						{
-							DestBuf[BASE_CHANNEL_STRIDE * I + C] = (uint8)FMath::Min(255u, Result);
-						}
-						else
-						{
-							DestBuf[BASE_CHANNEL_STRIDE * I + C] = (uint8)Result;
-						}
-					}
-					// Copy the unblended channels
-					// \TODO: unnecessary when doing it in-place?
-					if constexpr (UnblendedChannels > 0)
-					{
-						for (int32 C = 0; C < UnblendedChannels; ++C)
-						{
-							DestBuf[BASE_CHANNEL_STRIDE * I + CHANNELS_TO_BLEND + C] = BaseBuf[BASE_CHANNEL_STRIDE * I + CHANNELS_TO_BLEND + C];
-						}
-					}
-				}
-			}
-		};
-
-		if (NumBatches == 1)
-		{
-			ProcessBatch(0);
-		}
-		else if (NumBatches > 1)
-		{
-			ParallelFor(NumBatches, ProcessBatch);
-		}
-	}
-
-	template< 
-		uint32 (*BLEND_FUNC)(uint32, uint32),
-		bool CLAMP,
-		uint32 NC>
-    inline void BufferLayerFormatStrideNoAlpha(
-		FImage* DestImage,
-		int32 DestOffset,
-		int32 Stride,
-		const FImage* MaskImage,
-		const FImage* BlendImage/*, int32 LODCount*/)
-	{
-        const uint8* MaskBuf = MaskImage->GetLODData(0);
-        const uint8* BlendedBuf = BlendImage->GetLODData(0);
-		uint8* DestBuf = DestImage->GetLODData(0) + DestOffset;
-
-		// This could happen in case of missing data files.
-		if (!BlendedBuf || !MaskBuf || !DestImage->GetLODData(0))
-		{
-			return;
-		}
-
-		EImageFormat MaskFormat = MaskImage->GetFormat();
-        bool bIsUncompressed = (MaskFormat == EImageFormat::L_UByte);
-
-        if (bIsUncompressed)
-		{
-			int32 RowCount = BlendImage->GetSizeY();
-			int32 PixelCount = BlendImage->GetSizeX();
-			for (int32 RowIndex = 0; RowIndex < RowCount; ++RowIndex)
-			{
-				for (int32 PixelIndex = 0; PixelIndex < PixelCount; ++PixelIndex)
-				{
-					uint32 Mask = *MaskBuf;
-					if (Mask)
-					{
-						for (int32 C = 0; C < NC; ++C)
-						{
-							uint32 Base = *DestBuf;
-							uint32 Blended = *BlendedBuf;
-							uint32 Result = BLEND_FUNC(Base, Blended);
-							if constexpr (CLAMP)
-							{
-								*DestBuf = (uint8)FMath::Min(255u, Result);
-							}
-							else
-							{
-								*DestBuf = (uint8)Result;
-							}
-							++DestBuf;
-							++BlendedBuf;
-						}
-					}
-					else
-					{
-						DestBuf += NC;
-						BlendedBuf += NC;
-					}
-					++MaskBuf;
-				}
-
-				DestBuf += Stride;
-			}
-		}
-        else if (MaskFormat == EImageFormat::L_UBitRLE)
-		{
-			int32 Rows = MaskImage->GetSizeY();
-			int32 Width = MaskImage->GetSizeX();
-
-            //for (int32 lod = 0; lod < LODCount; ++lod)
-            //{
-				// Remove RLE header.
-                MaskBuf += 4 + Rows*sizeof(uint32);
-
-                for (int32 RowIndex = 0; RowIndex < Rows; ++RowIndex)
-                {
-                    const uint8* DestRowEnd = DestBuf + Width*NC;
-                    while (DestBuf != DestRowEnd)
-                    {
-                        // Decode header
-                        uint16 Zeros = *(const uint16*)MaskBuf;
-                        MaskBuf += 2;
-
-                        uint16 Ones = *(const uint16*)MaskBuf;
-                        MaskBuf += 2;
-
-                        // Skip
-                        DestBuf += Zeros*NC;
-                        BlendedBuf += Zeros*NC;
-
-                        // Copy
-                        FMemory::Memmove(DestBuf, BlendedBuf, Ones*NC);
-
-                        DestBuf += NC*Ones;
-                        BlendedBuf += NC*Ones;
-                    }
-
-                    DestBuf += Stride;
-                }
-
-                //Rows = FMath::DivideAndRoundUp(Rows, 2);
-                //Width = FMath::DivideAndRoundUp(Width, 2);
-            //}
-		}
-		else
-		{
-			checkf( false, TEXT("Unsupported mask format.") );
-		}
-	}
-
-	template<uint32(*BLEND_FUNC_MASKED)(uint32, uint32, uint32),
-			 uint32(*BLEND_FUNC)(uint32, uint32),
-			 bool CLAMP>
-    inline void BufferLayer(FImage* DestImage,
-							const FImage* BaseImage,
-							const FImage* MaskImage,
-                            const FImage* BlendImage,
-                            bool bApplyToAlpha,
-							bool bOnlyFirstLOD)
-	{
-		if (BaseImage->GetFormat() == EImageFormat::RGB_UByte)
-		{
-            BufferLayerFormat<BLEND_FUNC_MASKED, BLEND_FUNC, CLAMP, 3, 3, 3>
-                    (DestImage, BaseImage, MaskImage, BlendImage, 0, 0, 0, bOnlyFirstLOD);
-		}
-        else if (BaseImage->GetFormat() == EImageFormat::RGBA_UByte || 
-				 BaseImage->GetFormat() == EImageFormat::BGRA_UByte)
-		{
-			if (bApplyToAlpha)
-			{
-				BufferLayerFormat<BLEND_FUNC_MASKED, BLEND_FUNC, CLAMP, 4, 4, 4>
-					(DestImage, BaseImage, MaskImage, BlendImage, 0, 0, 0, bOnlyFirstLOD);
-			}
-			else
-			{
-				BufferLayerFormat<BLEND_FUNC_MASKED, BLEND_FUNC, CLAMP, 3, 4, 4>
-					(DestImage, BaseImage, MaskImage, BlendImage, 0, 0, 0, bOnlyFirstLOD);
-			}
-		}
-		else if (BaseImage->GetFormat() == EImageFormat::L_UByte)
-		{
-            BufferLayerFormat<BLEND_FUNC_MASKED, BLEND_FUNC, CLAMP, 1, 1, 1>
-                    (DestImage, BaseImage, MaskImage, BlendImage, 0, 0, 0, bOnlyFirstLOD);
-		}
-        else
-		{
-			checkf(false, TEXT("Unsupported format."));
-		}
-	}
-
-
-	template< uint32 (*BLEND_FUNC_MASKED)(uint32, uint32, uint32),
-		uint32 (*BLEND_FUNC)(uint32, uint32),
-		bool CLAMP >
-	inline void BufferLayerEmbeddedMask(
-		FImage* DestImage,
-		const FImage* BaseImage,
-		const FImage* BlendImage,
-		bool bApplyToAlpha,
-		bool bOnlyFirstLOD)
-	{
-		if (BaseImage->GetFormat() == EImageFormat::RGB_UByte)
-		{
-			BufferLayerFormatEmbeddedMask<BLEND_FUNC_MASKED, BLEND_FUNC, CLAMP, 3, 3, 3>
-				(DestImage, BaseImage, BlendImage, 0, 0, bOnlyFirstLOD);
-		}
-		else if (BaseImage->GetFormat() == EImageFormat::RGBA_UByte || 
-				 BaseImage->GetFormat() == EImageFormat::BGRA_UByte)
-		{
-			if (bApplyToAlpha)
-			{
-				BufferLayerFormatEmbeddedMask<BLEND_FUNC_MASKED, BLEND_FUNC, CLAMP, 4, 4, 4>
-					(DestImage, BaseImage, BlendImage, 0, 0, bOnlyFirstLOD);
-			}
-			else
-			{
-				BufferLayerFormatEmbeddedMask<BLEND_FUNC_MASKED, BLEND_FUNC, CLAMP, 3, 4, 4>
-					(DestImage, BaseImage, BlendImage, 0, 0, bOnlyFirstLOD);
-			}
-		}
-		else
-		{
-			checkf(false, TEXT("Unsupported format."));
-		}
-	}
-
-
-	template< 
-		uint32 (*RGB_FUNC_MASKED)(uint32, uint32, uint32),
-		uint32 (*A_FUNC)(uint32, uint32),
-		bool CLAMP >
-	inline void BufferLayerComposite(
-		FImage* BaseImage,
-		const FImage* BlendImage,
-		bool bOnlyFirstLOD,
-		uint8 BlendAlphaSourceChannel)
-	{
-		check(BaseImage->GetFormat() == EImageFormat::RGBA_UByte);
-		check(BlendImage->GetFormat() == EImageFormat::RGBA_UByte);
-		check(BaseImage->GetSizeX() == BlendImage->GetSizeX() && BaseImage->GetSizeY() == BlendImage->GetSizeY());
-		check(bOnlyFirstLOD || BaseImage->GetLODCount() <= BlendImage->GetLODCount());
-
-		bOnlyFirstLOD = bOnlyFirstLOD || BaseImage->GetLODCount() == 1;
-		
-		int32 FirstLODDataOffset = 0;	
-		int32 NumRelevantElems = -1;
-
-		if (BlendImage->Flags & FImage::IF_HAS_RELEVANCY_MAP && bOnlyFirstLOD)
-		{
-			check(BlendImage->RelevancyMaxY < BaseImage->GetSizeY());
-			check(BlendImage->RelevancyMaxY >= BlendImage->RelevancyMinY);
-
-			uint16 SizeX = BaseImage->GetSizeX();
-			NumRelevantElems = (BlendImage->RelevancyMaxY - BlendImage->RelevancyMinY + 1) * SizeX * 4;
-			
-			FirstLODDataOffset = BlendImage->RelevancyMinY * SizeX * 4;
-		}
-
-		const int32 BytesPerElem = GetImageFormatData(BaseImage->GetFormat()).BytesPerBlock;
-		check(BytesPerElem == GetImageFormatData(BlendImage->GetFormat()).BytesPerBlock);
-
-		const int32 LODBegin = 0;
-		const int32 LODEnd = BaseImage->GetLODCount();
-
-		constexpr int32 NumBatchElems = 4096*2;
-
-		const int32 NumBatches = bOnlyFirstLOD
-				? BaseImage->DataStorage.GetNumBatchesFirstLODOffset(NumBatchElems, BytesPerElem, FirstLODDataOffset)
-				: BaseImage->DataStorage.GetNumBatchesLODRange(NumBatchElems, BytesPerElem, LODBegin, LODEnd);
-
-		// This will always be an upper-bound for bOnlyFirtsLODs, check if it performs as expected or it needs more fine tune.
-		const int32 NumRelevantBatches = NumRelevantElems != -1
-				? FMath::DivideAndRoundUp(NumRelevantElems, NumBatchElems)
-				: NumBatches;
-
-		// This check only applies if bOnlyFirstLOD is false. The other case could give false negatives
-		// but that is already checked above.
-		check(NumBatches <= BlendImage->DataStorage.GetNumBatches(NumBatchElems, BytesPerElem));
-
-		auto ProcessBatch = 
-		[
-			BaseImage, BlendImage, BlendAlphaSourceChannel, BytesPerElem, NumBatchElems, bOnlyFirstLOD, LODBegin, LODEnd, FirstLODDataOffset
-		] (int32 BatchId)
-		{
-			TArrayView<uint8> BaseBatchView = bOnlyFirstLOD 
-					? BaseImage->DataStorage.GetBatchFirstLODOffset(BatchId, NumBatchElems, BytesPerElem, FirstLODDataOffset)
-					: BaseImage->DataStorage.GetBatchLODRange(BatchId, NumBatchElems, BytesPerElem, LODBegin, LODEnd);
-
-			TArrayView<const uint8> BlendBatchView = bOnlyFirstLOD 
-					? BlendImage->DataStorage.GetBatchFirstLODOffset(BatchId, NumBatchElems, BytesPerElem, FirstLODDataOffset)
-					: BlendImage->DataStorage.GetBatchLODRange(BatchId, NumBatchElems, BytesPerElem, LODBegin, LODEnd);
-
-			const int32 NumElems = BaseBatchView.Num() / BytesPerElem; 
-			check(BlendBatchView.Num() / BytesPerElem == NumElems);
-
-			uint8* BaseBuf = BaseBatchView.GetData();
-			const uint8* BlendBuf = BlendBatchView.GetData();
-
-			for (int32 I = 0; I < NumElems; ++I)
-			{
-				// TODO: Optimize this (SIMD?)
-				uint32 Mask = BlendBuf[4 * I + 3];
-
-				// RGB
-				for (int32 C = 0; C < 3; ++C)
-				{
-					uint32 Base = BaseBuf[4 * I + C];
-					uint32 Blended = BlendBuf[4 * I + C];
-					uint32 Result = RGB_FUNC_MASKED(Base, Blended, Mask);
-					if constexpr (CLAMP)
-					{
-						BaseBuf[4 * I + C] = (uint8)FMath::Min(255u, Result);
-					}
-					else
-					{
-						BaseBuf[4 * I + C] = (uint8)Result;
-					}
-				}
-
-				// A
-				{
-					uint32 Base = BaseBuf[4 * I + 3];
-					uint32 Blended = BlendBuf[4 * I + BlendAlphaSourceChannel];
-					uint32 Result = A_FUNC(Base, Blended);
-					if constexpr (CLAMP)
-					{
-						BaseBuf[4 * I + 3] = (uint8)FMath::Min(255u, Result);
-					}
-					else
-					{
-						BaseBuf[4 * I + 3] = (uint8)Result;
-					}
-				}
-			}
-		};
-
-		if (NumRelevantBatches == 1)
-		{
-			ProcessBatch(0);
-		}
-		else if (NumRelevantBatches > 1)
-		{
-			ParallelFor(NumRelevantBatches, ProcessBatch);
-		}
-	}
-
-	template<>
-	inline void BufferLayerComposite<BlendChannelMasked, LightenChannel, false>
-	(
-		FImage* BaseImage,
-		const FImage* BlendImage,
-		bool bOnlyFirstLOD,
-		uint8 BlendAlphaSourceChannel)
-	{
-		check(BaseImage->GetFormat() == EImageFormat::RGBA_UByte);
-		check(BlendImage->GetFormat() == EImageFormat::RGBA_UByte);
-		check(BaseImage->GetSizeX() == BlendImage->GetSizeX() && BaseImage->GetSizeY() == BlendImage->GetSizeY());
-		check(bOnlyFirstLOD || BaseImage->GetLODCount() <= BlendImage->GetLODCount());
-		check(BlendAlphaSourceChannel == 3);
-
-		bOnlyFirstLOD = bOnlyFirstLOD || BaseImage->GetLODCount() == 1;
-
-		int32 FirstLODDataOffset = 0;
-		int32 NumRelevantElems = -1;
-		if (BlendImage->Flags & FImage::IF_HAS_RELEVANCY_MAP && bOnlyFirstLOD)
-		{
-			check(BlendImage->RelevancyMaxY < BaseImage->GetSizeY());
-			check(BlendImage->RelevancyMaxY >= BlendImage->RelevancyMinY);
-
-			uint16 SizeX = BaseImage->GetSizeX();
-			NumRelevantElems = (BlendImage->RelevancyMaxY - BlendImage->RelevancyMinY + 1) * SizeX;
-			
-			FirstLODDataOffset = BlendImage->RelevancyMinY * SizeX * 4;
-		}
-
-		// \TODO: Ensure batches operate on 32Kb aligned buffers?
-		constexpr int32 NumBatchElems = 4096 * 2;
-
-		constexpr int32 BytesPerElem = 4;
-
-		const int32 NumLODs = BaseImage->GetLODCount();
-		const int32 NumBatches = bOnlyFirstLOD
-				? BaseImage->DataStorage.GetNumBatchesFirstLODOffset(NumBatchElems, BytesPerElem, FirstLODDataOffset)
-				: BaseImage->DataStorage.GetNumBatchesLODRange(NumBatchElems, BytesPerElem, 0, NumLODs);
-
-		// This will always be an upper-bound for bOnlyFirtsLODs, check if it performs as expected or it needs more fine tune.
-		const int32 NumRelevantBatches = NumRelevantElems != -1
-				? FMath::DivideAndRoundUp(NumRelevantElems, NumBatchElems)
-				: NumBatches;
-
-		// This check only applies if bOnlyFirstLOD is false. The other case could give false negatives
-		// but that is already checked above.
-		check(NumBatches <= BlendImage->DataStorage.GetNumBatchesLODRange(NumBatchElems, BytesPerElem, 0, NumLODs));
-		
-		auto ProcessBatch =
-			[
-				BaseImage, BlendImage, NumBatchElems, bOnlyFirstLOD, FirstLODDataOffset, NumLODs
-			] (uint32 BatchId)
-		{
-			TArrayView<uint8> BaseBatchView = bOnlyFirstLOD 
-					? BaseImage->DataStorage.GetBatchFirstLODOffset(BatchId, NumBatchElems, BytesPerElem, FirstLODDataOffset)
-					: BaseImage->DataStorage.GetBatchLODRange(BatchId, NumBatchElems, BytesPerElem, 0, NumLODs);
-
-			TArrayView<const uint8> BlendBatchView = bOnlyFirstLOD 
-					? BlendImage->DataStorage.GetBatchFirstLODOffset(BatchId, NumBatchElems, BytesPerElem, FirstLODDataOffset)
-					: BlendImage->DataStorage.GetBatchLODRange(BatchId, NumBatchElems, BytesPerElem, 0, NumLODs);
-
-			const int32 NumElems = BaseBatchView.Num() / BytesPerElem; 
-			check(NumElems == BlendBatchView.Num() / BytesPerElem);
-			
-			uint8* BaseBuf = BaseBatchView.GetData();
-			const uint8* BlendBuf = BlendBatchView.GetData();
-			
-			for (int32 I = 0; I < NumElems; ++I)
-			{
-				// TODO: Optimize this (SIMD?)
-				uint32 FullBase;
-				FMemory::Memcpy(&FullBase, BaseBuf + I*sizeof(uint32), sizeof(uint32)); 
-
-				uint32 FullBlended;
-				FMemory::Memcpy(&FullBlended, BlendBuf + I*sizeof(uint32), sizeof(uint32)); 
-				uint32 Mask = (FullBlended & 0xff000000) >> 24;
-
-				uint32 FullResult = 0;
-				FullResult |= BlendChannelMasked((FullBase >>  0) & 0xff, (FullBlended >>  0) & 0xff, Mask) << 0;
-				FullResult |= BlendChannelMasked((FullBase >>  8) & 0xff, (FullBlended >>  8) & 0xff, Mask) << 8;
-				FullResult |= BlendChannelMasked((FullBase >> 16) & 0xff, (FullBlended >> 16) & 0xff, Mask) << 16;
-				FullResult |= LightenChannel	((FullBase >> 24) & 0xff, (FullBlended >> 24) & 0xff) << 24;
-
-				FMemory::Memcpy(BaseBuf + I*sizeof(uint32), &FullResult, sizeof(uint32));
-			}
-		};
-
-		if (NumRelevantBatches == 1)
-		{
-			ProcessBatch(0);
-		}
-		else if (NumRelevantBatches > 1)
-		{
-			ParallelFor(NumBatches, ProcessBatch);
-		}
-	}
-
-	template< 
-		VectorRegister4Int (*RGB_FUNC_MASKED)(const VectorRegister4Int&, const VectorRegister4Int&, const VectorRegister4Int&),
-		int32 (*A_FUNC)(int32, int32),
-		bool CLAMP >
-	inline void BufferLayerCompositeVector(
-		FImage* BaseImage,
-		const FImage* BlendImage,
-		bool bOnlyFirstLOD,
-		uint8 BlendAlphaSourceChannel)
-	{
-		check(BaseImage->GetFormat() == EImageFormat::RGBA_UByte);
-		check(BlendImage->GetFormat() == EImageFormat::RGBA_UByte);
-		check(BaseImage->GetSizeX() == BlendImage->GetSizeX() && BaseImage->GetSizeY() == BlendImage->GetSizeY());
-		check(bOnlyFirstLOD || BaseImage->GetLODCount() <= BlendImage->GetLODCount());
-
-		constexpr int32 NumBatchElems = 4096*2;
-		constexpr int32 BytesPerElem = 4;
-
-		const int32 NumLODs = bOnlyFirstLOD ? 1 : BaseImage->GetLODCount();
-
-		const int32 NumBatches = BaseImage->DataStorage.GetNumBatchesLODRange(NumBatchElems, BytesPerElem, 0, NumLODs);
+	template<uint32(*BlendFuncMasked)(uint32, uint32, uint32), uint32 (*BlendFunc)(uint32, uint32), bool bClamp>
+	void BufferLayerEmbeddedMask(FImage* DestImage, const FImage* BaseImage, const FImage* BlendImage, bool bApplyToAlpha, bool bOnlyFirstLOD);
 	
-		check(NumBatches == BlendImage->DataStorage.GetNumBatchesLODRange(NumBatchElems, BytesPerElem, 0, NumLODs));
-
-		auto ProcessBatch = 
-		[
-			BaseImage, BlendImage, BytesPerElem, NumBatchElems, BlendAlphaSourceChannel, NumLODs
-		] (uint32 BatchId)
-		{
-			const TArrayView<uint8> BaseBatchView =
-					BaseImage->DataStorage.GetBatchLODRange(BatchId, NumBatchElems, BytesPerElem, 0, NumLODs);
-
-			const TArrayView<const uint8> BlendBatchView =
-					BlendImage->DataStorage.GetBatchLODRange(BatchId, NumBatchElems, BytesPerElem, 0, NumLODs);
-
-			uint8* BaseBuf = BaseBatchView.GetData();
-			const uint8* BlendBuf = BlendBatchView.GetData();
-
-			const int32 NumElems = BaseBatchView.Num() / BytesPerElem;
-			check(NumElems == BlendBatchView.Num()/BytesPerElem);
-
-			for (int32 I = 0; I < NumElems; ++I)
-			{
-				// TODO: Optimize this (SIMD?)
-				const int32 BaseAlpha = BaseBuf[4 * I + BlendAlphaSourceChannel];
-				const int32 BlendedAlpha = BlendBuf[4 * I + BlendAlphaSourceChannel];
-
-				const VectorRegister4Int Mask = VectorIntSet1(BlendedAlpha);
-
-				uint32 BlendPixel;
-				FMemory::Memcpy(&BlendPixel, BlendBuf, sizeof(uint32));
-
-				uint32 BasePixel;
-				FMemory::Memcpy(&BasePixel, BaseBuf, sizeof(uint32));
-
-				const VectorRegister4Int Blended = MakeVectorRegisterInt(
-						(BlendPixel >> 0) & 0xFF, 
-						(BlendPixel >> 8) & 0xFF, 
-						(BlendPixel >> 16) & 0xFF, 
-						(BlendPixel >> 24) & 0xFF);
-
-				const VectorRegister4Int Base = MakeVectorRegisterInt(
-						(BasePixel >> 0) & 0xFF, 
-						(BasePixel >> 8) & 0xFF, 
-						(BasePixel >> 16) & 0xFF, 
-						(BasePixel >> 24) & 0xFF);
-
-				VectorRegister4Int Result = RGB_FUNC_MASKED(Base, Blended, Mask);
-				if constexpr (CLAMP)
-				{
-					Result = VectorIntMin(MakeVectorRegisterIntConstant(255, 255, 255, 255), Result);
-				}
-
-				int32 AlphaResult = A_FUNC(BaseAlpha, BlendedAlpha);
-				if constexpr (CLAMP)
-				{
-					AlphaResult = FMath::Min(255, AlphaResult);
-				}
-
-				alignas(VectorRegister4Int) int32 IndexableRegister[4];
-				VectorIntStoreAligned(Result, &IndexableRegister);
-
-				BaseBuf[4 * I + 0] = static_cast<uint8>(IndexableRegister[0]);
-				BaseBuf[4 * I + 1] = static_cast<uint8>(IndexableRegister[1]);
-				BaseBuf[4 * I + 2] = static_cast<uint8>(IndexableRegister[2]);
-				BaseBuf[4 * I + 3] = static_cast<uint8>(IndexableRegister[3]);
-
-				BaseBuf[4 * I + BlendAlphaSourceChannel] = static_cast<uint8>(AlphaResult);
-			}
-		};
-
-		if (NumBatches == 1)
-		{
-			ProcessBatch(0);
-		}
-		else if (NumBatches > 1)
-		{
-			ParallelFor(NumBatches, ProcessBatch);
-		}
-	}
+	template<uint32(*BlendFunc)(uint32, uint32), bool bClamp, int32 ChannelCount>
+	void BufferLayerInPlace(FImage* BaseImage, const FImage* BlendedImage, bool bOnlyOneMip, uint32 BaseOffset, uint32 BlendedOffset);
 
-	template<uint32 (*BLEND_FUNC)(uint32,uint32),
-			 bool CLAMP>
-    inline void BufferLayerStrideNoAlpha(FImage* DestImage, int32 DestOffset, int32 Stride, const FImage* MaskImage, const FImage* BlendImage/*, int32 LODCount*/)
-	{
-		if (BlendImage->GetFormat() == EImageFormat::RGB_UByte)
-		{
-			BufferLayerFormatStrideNoAlpha<BLEND_FUNC, CLAMP, 3>
-					(DestImage, DestOffset, Stride, MaskImage, BlendImage/*, LODCount*/);
-		}
-        else if (BlendImage->GetFormat() == EImageFormat::RGBA_UByte || 
-				 BlendImage->GetFormat() == EImageFormat::BGRA_UByte)
-		{
-			BufferLayerFormatStrideNoAlpha<BLEND_FUNC, CLAMP, 4>
-					(DestImage, DestOffset, Stride, MaskImage, BlendImage/*, LODCount*/);
-		}
-		else if (BlendImage->GetFormat() == EImageFormat::L_UByte)
-		{
-			BufferLayerFormatStrideNoAlpha<BLEND_FUNC, CLAMP, 1>
-					(DestImage, DestOffset, Stride, MaskImage, BlendImage/*, LODCount*/);
-		}
-		else
-		{
-			checkf( false, TEXT("Unsupported format.") );
-		}
-	}
+	template<uint32(*RGBFuncMasked)(uint32, uint32, uint32), uint32(*AlphaFunc)(uint32, uint32), bool bClamp>
+	void BufferLayerComposite(FImage* BaseImage, const FImage* BlendImage, bool bOnlyFirstLOD, uint8 BlendAlphaSourceChannel);
 
+	template<> 
+	void BufferLayerComposite<BlendChannelMasked, LightenChannel, false>(FImage*, const FImage*, bool, uint8);
 
-	/**
-	* Apply a blending function to an image with another image as blending layer, on a subrect of
-	* the base image.
-	* \warning this method applies the blending function to the alpha channel too
-	* \warning this method uses the mask as a binary mask (>0)
-	*/
-	template<uint32 (*BLEND_FUNC)(uint32, uint32), bool CLAMP>
-	inline void ImageLayerOnBaseNoAlpha(
-			FImage* BaseImage,
-			const FImage* MaskImage,
-			const FImage* BlendedImage,
-			const box<FIntVector2>& Rect)
-	{
-		check(BaseImage->GetSizeX() >= Rect.min[0] + Rect.size[0]);
-		check(BaseImage->GetSizeY() >= Rect.min[1] + Rect.size[1]);
-		check(MaskImage->GetSizeX() == BlendedImage->GetSizeX());
-		check(MaskImage->GetSizeY() == BlendedImage->GetSizeY());
-		check(BaseImage->GetFormat() == BlendedImage->GetFormat());
-		check(MaskImage->GetFormat() == EImageFormat::L_UByte ||
-			  //UBYTE_RLE does not look to be supported.
-			  //MaskImage->GetFormat() == EImageFormat::L_UByteRLE || 
-			  MaskImage->GetFormat() == EImageFormat::L_UBitRLE);
-        check(BaseImage->GetLODCount() <= MaskImage->GetLODCount());
-        check(BaseImage->GetLODCount() <= BlendedImage->GetLODCount());
+	template<VectorRegister4Int(*RGBFuncMasked)(const VectorRegister4Int&, const VectorRegister4Int&, const VectorRegister4Int&), int32(*AlphaFunc)(int32, int32), bool bClamp>
+	void BufferLayerCompositeVector(FImage* BaseImage, const FImage* BlendImage, bool bOnlyFirstLOD, uint8 BlendAlphaSourceChannel);
 
-		int32 PixelSize = GetImageFormatData(BaseImage->GetFormat()).BytesPerBlock;
+	template<> 
+	void BufferLayerCompositeVector<VectorBlendChannelMasked, VectorLightenChannel, false>(FImage*, const FImage*, bool, uint8);
 
-		int32 Start = (BaseImage->GetSizeX() * Rect.min[1] + Rect.min[0]) * PixelSize;
-		int32 Stride = (BaseImage->GetSizeX() - Rect.size[0]) * PixelSize;
+	/** Blend a subimage on the base using a mask */
+	void ImageBlendOnBaseNoAlpha(FImage* BaseImage, const FImage* MaskImage, const FImage* BlendedImage, const box<FIntVector2>& Rect);
 
-		// Stride is only valid for LOD 0, BufferLayerStride variants cannot operate on multiple lods.
-		// TODO: review if this needs to be supported, and implement using a rect lod reducction at this level.
-		BufferLayerStrideNoAlpha<BLEND_FUNC, CLAMP>(BaseImage, Start, Stride, MaskImage, BlendedImage/*, BaseImage->GetLODCount()*/);
-	}
+	template<uint32 (*BlendFuncMasked)(uint32, uint32, uint32), uint32 (*BlendFunc)(uint32, uint32), uint32 ChannelCount>
+	void BufferLayerColourInPlace(FImage* BaseImage, const FImage* MaskImage, FVector4f Color, bool bOnlyOneMip, uint32 BaseOffset, uint8 ColorOffset);
 
-	template<uint32 NC>
-	FORCEINLINE uint32 PackPixel(const uint8* PixelPtr)
-	{
-		static_assert(NC > 0 && NC <= 4);
+	template<uint32 (*BlendFunc)(uint32, uint32), uint32 ChannelCount>
+	void BufferLayerColourInPlace(FImage* BaseImage, FVector4f Color, bool bOnlyOneMip, uint32 BaseOffset, uint8 ColorOffset);
 
-		uint32 PixelPack = 0;
+	template<uint32 (*BlendFuncMasked)(uint32, uint32, uint32), uint32 (*BlendFunc)(uint32, uint32)>
+	void BufferLayerColour(FImage* DestImage, const FImage* BaseImage, const FImage* MaskImage, FVector4f Color);
 
-		// The compiler should be able to optimize this given that NC is a constant expression.
-		FMemory::Memcpy(&PixelPack, PixelPtr, NC);
+	template<uint32 (*BlendFunc)(uint32,uint32)>
+    void BufferLayerColour(FImage* ResultImage, const FImage* BaseImage, FVector4f Color);
 
-		return PixelPack;
-	}
-
-	template<uint32 NC>
-	FORCEINLINE void UnpackPixel(uint8* PixelPtr, uint32 PixelData)
-	{
-		static_assert(NC > 0 && NC <= 4);
-
-		// The compiler should be able to optimize this given that NC is a constant expression
-		FMemory::Memcpy(PixelPtr, &PixelData, NC);
-	}
-
-	template< 
-		uint32 (*BLEND_FUNC)(uint32, uint32), 
-		uint32 NC>
-	inline void BufferLayerCombineColour(FImage* ResultImage, const FImage* BaseImage, FVector4f Color, bool bOnlyFirstLOD = false)
-	{
-		static_assert(NC > 0 && NC <= 4);
-
-		check(BaseImage->GetSizeX() == ResultImage->GetSizeX());
-		check(BaseImage->GetSizeY() == ResultImage->GetSizeY());
-
-		const uint32 TopColor = 
-			static_cast<uint32>(255.0f * Color[0]) << 0 |  
-			static_cast<uint32>(255.0f * Color[1]) << 8 |  
-			static_cast<uint32>(255.0f * Color[2]) << 16;  
-
-		const int32 NumLODs = bOnlyFirstLOD ? 1 : ResultImage->GetLODCount();
-		
-		constexpr int32 BatchNumElems = 4098*2;
-		const int32 NumBatches = BaseImage->DataStorage.GetNumBatchesLODRange(BatchNumElems, NC, 0, NumLODs);
-
-		check(NumBatches == ResultImage->DataStorage.GetNumBatchesLODRange(BatchNumElems, NC, 0, NumLODs));
-
-		auto ProcessBatch = [ResultImage, BaseImage, BatchNumElems, TopColor, NumLODs](int32 BatchId)
-		{
-			const TArrayView<const uint8> BaseBatchView = 
-					BaseImage->DataStorage.GetBatchLODRange(BatchId, BatchNumElems, NC, 0, NumLODs);
-			
-			const TArrayView<uint8> ResultBatchView =
-					ResultImage->DataStorage.GetBatchLODRange(BatchId, BatchNumElems, NC, 0, NumLODs);
-
-			const int32 NumElems = BaseBatchView.Num() / NC;
-
-			check(NumElems == ResultBatchView.Num() / NC);
-			
-			const uint8* BaseBuf = BaseBatchView.GetData();
-			uint8* ResultBuf = ResultBatchView.GetData();
-			
-			for (int32 I = 0; I < NumElems; ++I)
-			{
-				const uint32 Base = PackPixel<NC>(&BaseBuf[NC * I]);
-
-				const uint32 Result = BLEND_FUNC(Base, TopColor);
-				UnpackPixel<NC>(&ResultBuf[NC * I], Result);
-			}
-		};
-
-		if (NumBatches == 1)
-		{
-			ProcessBatch(0);
-		}
-		else if (NumBatches > 1)
-		{
-			ParallelFor(NumBatches, ProcessBatch);
-		}
-	}
-
-	template< 
-		uint32 (*BLEND_FUNC)(uint32, uint32), 
-		uint32 NC>
-	inline void BufferLayerCombine(FImage* ResultImage, const FImage* BaseImage, const FImage* BlendImage, bool bOnlyFirstLOD)
-	{
-		static_assert(NC > 0 && NC <= 4);
-
-		check(BaseImage->GetSizeX() == BlendImage->GetSizeX());
-		check(BaseImage->GetSizeY() == BlendImage->GetSizeY());
-		check(BaseImage->GetSizeX() == ResultImage->GetSizeX());
-		check(BaseImage->GetSizeY() == ResultImage->GetSizeY());
-		check(BaseImage->GetFormat() == BlendImage->GetFormat());
-		check(BaseImage->GetFormat() == ResultImage->GetFormat());
-
-		const int32 NumLODs = bOnlyFirstLOD ? 1 : ResultImage->GetLODCount();
-
-		constexpr int32 BatchNumElems = 4098*2;
-		const int32 NumBatches = BaseImage->DataStorage.GetNumBatchesLODRange(BatchNumElems, NC, 0, NumLODs);
-
-		check(NumBatches == BlendImage->DataStorage.GetNumBatchesLODRange(BatchNumElems, NC, 0, NumLODs));
-		check(NumBatches == ResultImage->DataStorage.GetNumBatchesLODRange(BatchNumElems, NC, 0, NumLODs));
-
-		auto ProcessBatch = [ResultImage, BaseImage, BlendImage, BatchNumElems, NumLODs](int32 BatchId)
-		{
-			const TArrayView<const uint8> BaseBatchView = 
-					BaseImage->DataStorage.GetBatchLODRange(BatchId, BatchNumElems, NC, 0, NumLODs);
-
-			const TArrayView<const uint8> BlendBatchView = 
-					BlendImage->DataStorage.GetBatchLODRange(BatchId, BatchNumElems, NC, 0, NumLODs);
-
-			const TArrayView<uint8> ResultBatchView =
-					ResultImage->DataStorage.GetBatchLODRange(BatchId, BatchNumElems, NC, 0, NumLODs);
-
-			const int32 NumElems = BaseBatchView.Num() / NC;
-			check(NumElems == BlendBatchView.Num() / NC);
-			check(NumElems == ResultBatchView.Num() / NC);
-			
-			const uint8* BaseBuf = BaseBatchView.GetData();
-			const uint8* BlendBuf = BlendBatchView.GetData();
-			uint8* ResultBuf = ResultBatchView.GetData();
-			
-			for (int32 I = 0; I < NumElems; ++I)
-			{
-				const uint32 Base = PackPixel<NC>(&BaseBuf[NC * I]);
-				const uint32 Blend = PackPixel<NC>(&BlendBuf[NC * I]);
-
-				const uint32 Result = BLEND_FUNC(Base, Blend);
-				UnpackPixel<NC>(&ResultBuf[NC * I], Result);
-			}
-		};
-
-		if (NumBatches == 1)
-		{
-			ProcessBatch(0);
-		}
-		else if (NumBatches > 1)
-		{
-			ParallelFor(NumBatches, ProcessBatch);
-		}
-	}
-
-	template< 
-		uint32 (*BLEND_FUNC_MASKED)(uint32, uint32, uint32), 
-		uint32 NC>
-	inline void BufferLayerCombine(FImage* ResultImage, const FImage* BaseImage, const FImage* MaskImage, const FImage* BlendImage, bool bOnlyFirstLOD)
-	{
-		static_assert(NC > 0 && NC <= 4);
-
-		check(BaseImage->GetSizeX() == BlendImage->GetSizeX());
-		check(BaseImage->GetSizeY() == BlendImage->GetSizeY());
-		check(BaseImage->GetSizeX() == MaskImage->GetSizeX());
-		check(BaseImage->GetSizeY() == MaskImage->GetSizeY());
-		check(BaseImage->GetFormat() == BlendImage->GetFormat());
-
-		const int32 NumLODs = bOnlyFirstLOD ? 1 : ResultImage->GetLODCount();
-
-		constexpr int32 BatchNumElems = 4098*2;
-		const int32 NumBatches = BaseImage->DataStorage.GetNumBatchesLODRange(BatchNumElems, NC, 0, NumLODs);
-
-		check(NumBatches <= MaskImage->DataStorage.GetNumBatchesLODRange(BatchNumElems, 1, 0, NumLODs));
-		check(NumBatches <= BlendImage->DataStorage.GetNumBatchesLODRange(BatchNumElems, NC, 0, NumLODs));
-		check(NumBatches <= ResultImage->DataStorage.GetNumBatchesLODRange(BatchNumElems, NC, 0, NumLODs));
-
-		auto ProcessBatch = [ResultImage, BaseImage, MaskImage, BlendImage, BatchNumElems, NumLODs](int32 BatchId)
-		{
-			const TArrayView<const uint8> BaseBatchView = 
-					BaseImage->DataStorage.GetBatchLODRange(BatchId, BatchNumElems, NC, 0, NumLODs);
-
-			const TArrayView<const uint8> BlendBatchView = 
-					BlendImage->DataStorage.GetBatchLODRange(BatchId, BatchNumElems, NC, 0, NumLODs);
-
-			const TArrayView<const uint8> MaskBatchView =
-					MaskImage->DataStorage.GetBatchLODRange(BatchId, BatchNumElems, 1, 0, NumLODs);
-
-			const TArrayView<uint8> ResultBatchView =
-					ResultImage->DataStorage.GetBatchLODRange(BatchId, BatchNumElems, NC, 0, NumLODs);
-
-			const int32 NumElems = BaseBatchView.Num() / NC;
-
-			check(NumElems == BlendBatchView.Num() / NC);
-			check(NumElems == MaskBatchView.Num() / 1);
-			check(NumElems == ResultBatchView.Num() / NC);
-			
-			const uint8* BaseBuf = BaseBatchView.GetData();
-			const uint8* BlendBuf = BlendBatchView.GetData();
-			const uint8* MaskBuf = MaskBatchView.GetData();
-			uint8* ResultBuf = ResultBatchView.GetData();
-			
-			// This could happen in case of missing data files.
-			if (!BaseBuf || !BlendBuf || !MaskBuf)
-			{
-				return;
-			}
-
-			for (int32 I = 0; I < NumElems; ++I)
-			{
-				const uint32 Base = PackPixel<NC>(&BaseBuf[NC * I]);
-				const uint32 Blend = PackPixel<NC>(&BlendBuf[NC * I]);
-				const uint32 Mask = PackPixel<1>(&MaskBuf[1 * I]);
-
-				const uint32 Result = BLEND_FUNC_MASKED(Base, Blend, Mask);
-				UnpackPixel<NC>(&ResultBuf[NC * I], Result);
-			}
-		};
-
-		if (NumBatches == 1)
-		{
-			ProcessBatch(0);
-		}
-		else if (NumBatches > 1)
-		{
-			ParallelFor(NumBatches, ProcessBatch);
-		}
-	}
-
-	template< 
-		uint32 (*BLEND_FUNC_MASKED)(uint32, uint32, uint32), 
-		uint32 NC>
-	inline void BufferLayerCombineColour(FImage* ResultImage, const FImage* BaseImage, const FImage* MaskImage, FVector4f Color)
-	{
-		static_assert(NC > 0 && NC <= 4);
-
-		check(BaseImage->GetSizeX() == MaskImage->GetSizeX());
-		check(BaseImage->GetSizeY() == MaskImage->GetSizeY());
-
-		const uint32 TopColor = 
-			static_cast<uint32>(255.0f * Color[0]) << 0 |  
-			static_cast<uint32>(255.0f * Color[1]) << 8 |  
-			static_cast<uint32>(255.0f * Color[2]) << 16;  
-
-		constexpr int32 BatchNumElems = 4098*2;
-		int32 NumBatches = BaseImage->DataStorage.GetNumBatches(BatchNumElems, NC);
-
-		check(NumBatches == MaskImage->DataStorage.GetNumBatches(BatchNumElems, 1));
-		check(NumBatches == ResultImage->DataStorage.GetNumBatches(BatchNumElems, NC));
-
-		auto ProcessBatch = [ResultImage, BaseImage, MaskImage, BatchNumElems, TopColor](int32 BatchId)
-		{
-			const TArrayView<const uint8> BaseBatchView = BaseImage->DataStorage.GetBatch(BatchId, BatchNumElems, NC);
-			const TArrayView<const uint8> MaskBatchView = MaskImage->DataStorage.GetBatch(BatchId, BatchNumElems, 1);
-			const TArrayView<uint8> ResultBatchView = ResultImage->DataStorage.GetBatch(BatchId, BatchNumElems, NC);
-
-			const int32 NumElems = BaseBatchView.Num() / NC;
-
-			check(NumElems == MaskBatchView.Num() / 1);
-			check(NumElems == ResultBatchView.Num() / NC);
-			
-			const uint8* BaseBuf = BaseBatchView.GetData();
-			const uint8* MaskBuf = MaskBatchView.GetData();
-			uint8* ResultBuf = ResultBatchView.GetData();
-
-			// This could happen in case of missing data files.
-			if (!BaseBuf || !MaskBuf)
-			{
-				return;
-			}
-
-			for (int32 I = 0; I < NumElems; ++I)
-			{
-				const uint32 Base = PackPixel<NC>(&BaseBuf[NC * I]);
-				const uint32 Mask = PackPixel<1>(&MaskBuf[1 * I]);
-
-				const uint32 Result = BLEND_FUNC_MASKED(Base, TopColor, Mask);
-				UnpackPixel<NC>(&ResultBuf[NC * I], Result);
-			}
-		};
-
-		if (NumBatches == 1)
-		{
-			ProcessBatch(0);
-		}
-		else if (NumBatches > 1)
-		{
-			ParallelFor(NumBatches, ProcessBatch);
-		}
-	}
-
-	template<uint32 (*BLEND_FUNC)(uint32, uint32)>
-	inline void ImageLayerCombine(FImage* ResultImage, const FImage* BaseImage, const FImage* BlendedImage, bool bOnlyFirstLOD)
-	{
-		check(ResultImage->GetFormat() == BaseImage->GetFormat());
-		check(ResultImage->GetSizeX() == BaseImage->GetSizeX());
-		check(ResultImage->GetSizeY() == BaseImage->GetSizeY());
-		check(bOnlyFirstLOD || ResultImage->GetLODCount() == BaseImage->GetLODCount());
-		check(BaseImage->GetSizeX() == BlendedImage->GetSizeX());
-		check(BaseImage->GetSizeY() == BlendedImage->GetSizeY());
-		check(BaseImage->GetFormat() == BlendedImage->GetFormat());
-		check(bOnlyFirstLOD || ResultImage->GetLODCount() <= BlendedImage->GetLODCount());
-
-		const EImageFormat BaseFormat = BaseImage->GetFormat();
-
-		if (BaseFormat == EImageFormat::L_UByte)
-		{
-			BufferLayerCombine<BLEND_FUNC, 1>(ResultImage, BaseImage, BlendedImage, bOnlyFirstLOD);
-		}
-		else if (BaseFormat == EImageFormat::RGB_UByte)
-		{
-			BufferLayerCombine<BLEND_FUNC, 3>(ResultImage, BaseImage, BlendedImage, bOnlyFirstLOD);
-		}
-		else if (BaseFormat == EImageFormat::RGBA_UByte || BaseFormat == EImageFormat::BGRA_UByte)
-		{
-			// \todo: pass swizzle template argument if BGRA_UBYTE, not yet supported.
-			BufferLayerCombine<BLEND_FUNC, 4>(ResultImage, BaseImage, BlendedImage, bOnlyFirstLOD);
-		}
-		else
-		{
-			checkf(false, TEXT("Unsupported format."));
-		}
-	}
-
-	template< 
-		uint32 (*BLEND_FUNC)(uint32, uint32), 
-		uint32 (*BLEND_FUNC_MASKED)(uint32, uint32, uint32)>
-	inline void ImageLayerCombine(FImage* ResultImage, const FImage* BaseImage, const FImage* MaskImage, const FImage* BlendedImage, bool bOnlyFirstLOD)
-	{
-		check(ResultImage->GetFormat() == BaseImage->GetFormat());
-		check(ResultImage->GetSizeX() == BaseImage->GetSizeX());
-		check(ResultImage->GetSizeY() == BaseImage->GetSizeY());
-		check(bOnlyFirstLOD || ResultImage->GetLODCount() == BaseImage->GetLODCount());
-		check(BaseImage->GetSizeX() == BlendedImage->GetSizeX());
-		check(BaseImage->GetSizeY() == BlendedImage->GetSizeY());
-		check(BaseImage->GetFormat() == BlendedImage->GetFormat());
-		check(bOnlyFirstLOD || ResultImage->GetLODCount() <= BlendedImage->GetLODCount());
-
-		const EImageFormat BaseFormat = BaseImage->GetFormat();
-
-		TSharedPtr<FImage> TempMaskImage;
-		if (MaskImage->GetFormat() != EImageFormat::L_UByte)
-		{
-			UE_LOG(LogMutableCore, Log, TEXT("Image layer format not supported. A generic one will be used. "));
-
-			FImageOperator ImOp = FImageOperator::GetDefault(nullptr);
-			constexpr int32 Quality = 4;
-			TempMaskImage = ImOp.ImagePixelFormat( Quality, MaskImage, EImageFormat::L_UByte );
-			MaskImage = TempMaskImage.Get();
-		}
-
-		if (BaseFormat == EImageFormat::L_UByte)
-		{
-			BufferLayerCombine<BLEND_FUNC_MASKED, 1>(ResultImage, BaseImage, MaskImage, BlendedImage, bOnlyFirstLOD);
-		}
-		else if (BaseFormat == EImageFormat::RGB_UByte)
-		{
-			BufferLayerCombine<BLEND_FUNC_MASKED, 3>(ResultImage, BaseImage, MaskImage, BlendedImage, bOnlyFirstLOD);
-		}
-		else if (BaseFormat == EImageFormat::RGBA_UByte || BaseFormat == EImageFormat::BGRA_UByte)
-		{
-			// \todo: pass swizzle template argument if BGRA_UBYTE, not yet supported.
-			BufferLayerCombine<BLEND_FUNC_MASKED, 4>(ResultImage, BaseImage, MaskImage, BlendedImage, bOnlyFirstLOD);
-		}
-		else
-		{
-			UE_LOG(LogMutableCore, Log, TEXT("Image layer format not supported. A generic one will be used. "));
-
-			FImageOperator ImOp = FImageOperator::GetDefault(nullptr);
-			constexpr int32 Quality = 4;
-			TSharedPtr<FImage> TempBaseImage = ImOp.ImagePixelFormat(Quality, BaseImage, EImageFormat::RGBA_UByte);
-			TSharedPtr<FImage> TempBlededImage = ImOp.ImagePixelFormat(Quality, BlendedImage, EImageFormat::RGBA_UByte);
-			BufferLayerCombine<BLEND_FUNC_MASKED, 4>(ResultImage, TempBaseImage.Get(), MaskImage, TempBlededImage.Get(), bOnlyFirstLOD);
-		}
-	}
-
-	template<uint32 (*BLEND_FUNC)(uint32, uint32)>
-	inline void ImageLayerCombineColour(FImage* ResultImage, const FImage* BaseImage, FVector4f Color)
-	{
-		check(ResultImage->GetFormat() == BaseImage->GetFormat());
-		check(ResultImage->GetSizeX() == BaseImage->GetSizeX());
-		check(ResultImage->GetSizeY() == BaseImage->GetSizeY());
-		check(ResultImage->GetLODCount() == BaseImage->GetLODCount());
-
-		const EImageFormat BaseFormat = BaseImage->GetFormat();
-
-		if (BaseFormat == EImageFormat::L_UByte)
-		{
-			BufferLayerCombineColour<BLEND_FUNC, 1>(ResultImage, BaseImage, Color);
-		}
-		else if (BaseFormat == EImageFormat::RGB_UByte)
-		{
-			BufferLayerCombineColour<BLEND_FUNC, 3>(ResultImage, BaseImage, Color);
-		}
-		else if (BaseFormat == EImageFormat::RGBA_UByte || BaseFormat == EImageFormat::BGRA_UByte)
-		{
-			// \todo: pass swizzle template argument if BGRA_UBYTE, not yet supported.
- 			BufferLayerCombineColour<BLEND_FUNC, 4>(ResultImage, BaseImage, Color);
-		}
-		else
-		{
-			checkf(false, TEXT("Unsupported format."));
-		}
-	}
-
-	template< 
-		uint32 (*BLEND_FUNC)(uint32, uint32),
-		uint32 (*BLEND_FUNC_MASKED)(uint32, uint32, uint32)>
-	inline void ImageLayerCombineColour(FImage* ResultImage, const FImage* BaseImage, const FImage* MaskImage, FVector4f Color)
-	{
-		check(ResultImage->GetFormat() == BaseImage->GetFormat());
-		check(ResultImage->GetSizeX() == BaseImage->GetSizeX());
-		check(ResultImage->GetSizeY() == BaseImage->GetSizeY());
-		check(ResultImage->GetLODCount() == BaseImage->GetLODCount());
-
-		const EImageFormat BaseFormat = BaseImage->GetFormat();
-
-		if (MaskImage->GetFormat() != EImageFormat::L_UByte)
-		{
-			checkf(false, TEXT("Unsupported mask format."));
-
-			BufferLayerCombineColour<BLEND_FUNC, 1>(ResultImage, BaseImage, Color);
-		}
-
-		if (BaseFormat == EImageFormat::L_UByte)
-		{
-			BufferLayerCombineColour<BLEND_FUNC_MASKED, 1>(ResultImage, BaseImage, MaskImage, Color);
-		}
-		else if (BaseFormat == EImageFormat::RGB_UByte)
-		{
-			BufferLayerCombineColour<BLEND_FUNC_MASKED, 3>(ResultImage, BaseImage, MaskImage, Color);
-		}
-		else if (BaseFormat == EImageFormat::RGBA_UByte || BaseFormat == EImageFormat::BGRA_UByte)
-		{
-			// \todo: pass swizzle template argument if BGRA_UBYTE, not yet supported.
- 			BufferLayerCombineColour<BLEND_FUNC_MASKED, 4>(ResultImage, BaseImage, MaskImage, Color);
-		}
-		else
-		{
-			checkf(false, TEXT("Unsupported format."));
-		}
-	}
-
-	template<uint32 NCBase, uint32 NCBlend, class ImageCombineFn>
-	inline void BufferLayerCombineFunctor(FImage* DestImage, const FImage* BaseImage, const FImage* BlendImage, ImageCombineFn&& ImageCombine)
-	{
-		static_assert(NCBase > 0 && NCBase <= 4);
-		static_assert(NCBlend > 0 && NCBlend <= 4);
-
-		check(BaseImage->GetFormat() == DestImage->GetFormat());
-		check(BaseImage->GetSizeX() == BlendImage->GetSizeX());
-		check(BaseImage->GetSizeY() == BlendImage->GetSizeY());
-		check(BaseImage->GetSizeX() == DestImage->GetSizeX());
-		check(BaseImage->GetSizeY() == DestImage->GetSizeY());
-		check(BaseImage->GetLODCount() <= BlendImage->GetLODCount());
-		check(BaseImage->GetLODCount() <= DestImage->GetLODCount());
+	template<uint32 (*BlendFunc)(uint32, uint32)>
+	void BufferLayerColourFromAlpha(FImage* ResultImage, const FImage* BaseImage, FVector4f Color);
 	
-		constexpr int32 BatchNumElems = 4098*2;
-		int32 NumBatches = BaseImage->DataStorage.GetNumBatches(BatchNumElems, NCBase);
-
-		check(NumBatches == BlendImage->DataStorage.GetNumBatches(BatchNumElems, NCBlend));
-		check(NumBatches == DestImage->DataStorage.GetNumBatches(BatchNumElems, NCBase));
-
-		auto ProcessBatch = [DestImage, BaseImage, BlendImage, BatchNumElems, ImageCombine](int32 BatchId)
-		{
-			const TArrayView<const uint8> BaseBatchView = BaseImage->DataStorage.GetBatch(BatchId, BatchNumElems, NCBase);
-			const TArrayView<const uint8> BlendBatchView = BlendImage->DataStorage.GetBatch(BatchId, BatchNumElems, NCBlend);
-			const TArrayView<uint8> DestBatchView = DestImage->DataStorage.GetBatch(BatchId, BatchNumElems, NCBase);
-
-			const int32 NumElems = BaseBatchView.Num() / NCBase;
-
-			const uint8* BaseBuf = BaseBatchView.GetData();
-			const uint8* BlendBuf = BlendBatchView.GetData();
-			uint8* DestBuf = DestBatchView.GetData();
+	template<uint32 (*BlendFunc)(uint32, uint32)>
+	void ImageLayerCombine(FImage* ResultImage, const FImage* BaseImage, const FImage* BlendedImage, bool bOnlyFirstLOD);
 
-			check(NumElems == BlendBatchView.Num() / NCBlend);
-			check(NumElems == DestBatchView.Num() / NCBase);
-			
-			for (int32 I = 0; I < NumElems; ++I)
-			{
-				const uint32 Base = PackPixel<NCBase>(&BaseBuf[NCBase * I]);
-				const uint32 Blend = PackPixel<NCBlend>(&BlendBuf[NCBlend * I]);
+	template<uint32 (*BlendFunc)(uint32, uint32), uint32 (*BlendFuncMasked)(uint32, uint32, uint32)>
+	void ImageLayerCombine(FImage* ResultImage, const FImage* BaseImage, const FImage* MaskImage, const FImage* BlendedImage, bool bOnlyFirstLOD);
 
-				const uint32 Result = ImageCombine(Base, Blend);
-				UnpackPixel<NCBase>(&DestBuf[NCBase* I], Result);
-			}
-		};
+	template<uint32 (*BlendFunc)(uint32, uint32)>
+	void ImageLayerCombineColour(FImage* ResultImage, const FImage* BaseImage, FVector4f Color);
 
-		if (NumBatches == 1)
-		{
-			ProcessBatch(0);
-		}
-		else if (NumBatches > 1)
-		{
-			ParallelFor(NumBatches, ProcessBatch);
-		}
-	}
+	template<uint32 (*BlendFunc)(uint32, uint32), uint32 (*BlendFuncMasked)(uint32, uint32, uint32)>
+	void ImageLayerCombineColour(FImage* ResultImage, const FImage* BaseImage, const FImage* MaskImage, FVector4f Color);
 
-	// Same functionality as above, in this case we use a functor which allows to pass user data. 
 	template<class ImageCombineFn>
-	inline void ImageLayerCombineFunctor(FImage* ResultImage, const FImage* BaseImage, const FImage* BlendedImage, ImageCombineFn&& ImageCombine)
-	{
-		check(ResultImage->GetFormat() == BaseImage->GetFormat());
-		check(ResultImage->GetSizeX() == BaseImage->GetSizeX());
-		check(ResultImage->GetSizeY() == BaseImage->GetSizeY());
-		check(ResultImage->GetLODCount() == BaseImage->GetLODCount());
-		check(BaseImage->GetSizeX() == BlendedImage->GetSizeX());
-		check(BaseImage->GetSizeY() == BlendedImage->GetSizeY());
-		check(ResultImage->GetLODCount() <= BlendedImage->GetLODCount());
-
-		const EImageFormat BaseFormat = BaseImage->GetFormat();
-		const EImageFormat BlendFormat = BlendedImage->GetFormat();
-
-		if (BaseFormat == EImageFormat::L_UByte )
-		{
-			if (BlendFormat == EImageFormat::L_UByte )
-			{
-				BufferLayerCombineFunctor<1, 1>(ResultImage, BaseImage, BlendedImage, Forward<ImageCombineFn>(ImageCombine));
-			}
-			else if (BlendFormat == EImageFormat::RGB_UByte )
-			{
-				BufferLayerCombineFunctor<1, 3>(ResultImage, BaseImage, BlendedImage, Forward<ImageCombineFn>(ImageCombine));
-			}
-			else if (BlendFormat == EImageFormat::RGBA_UByte || BlendFormat == EImageFormat::BGRA_UByte)
-			{
-				BufferLayerCombineFunctor<1, 4>(ResultImage, BaseImage, BlendedImage, Forward<ImageCombineFn>(ImageCombine));
-			}
-			else
-			{
-				checkf(false, TEXT("Unsupported format."));
-			}
-		}
-		else if (BaseFormat == EImageFormat::RGB_UByte)
-		{
-			if (BlendFormat == EImageFormat::L_UByte )
-			{
-				BufferLayerCombineFunctor<3, 1>(ResultImage, BaseImage, BlendedImage, Forward<ImageCombineFn>(ImageCombine));
-			}
-			else if (BlendFormat == EImageFormat::RGB_UByte )
-			{
-				BufferLayerCombineFunctor<3, 3>(ResultImage, BaseImage, BlendedImage, Forward<ImageCombineFn>(ImageCombine));
-			}
-			else if (BlendFormat == EImageFormat::RGBA_UByte || BlendFormat == EImageFormat::BGRA_UByte )
-			{
-				BufferLayerCombineFunctor<3, 4>(ResultImage, BaseImage, BlendedImage, Forward<ImageCombineFn>(ImageCombine));
-			}
-			else
-			{
-				checkf(false, TEXT("Unsupported format."));
-			}
-		}
-		else if (BaseFormat == EImageFormat::RGBA_UByte || BaseFormat == EImageFormat::BGRA_UByte)
-		{
-			if (BlendFormat == EImageFormat::L_UByte )
-			{
-				BufferLayerCombineFunctor<4, 1>(ResultImage, BaseImage, BlendedImage, Forward<ImageCombineFn>(ImageCombine));
-			}
-			else if (BlendFormat == EImageFormat::RGB_UByte )
-			{
-				BufferLayerCombineFunctor<4, 3>(ResultImage, BaseImage, BlendedImage, Forward<ImageCombineFn>(ImageCombine));
-			}
-			else if (BlendFormat == EImageFormat::RGBA_UByte || BlendFormat == EImageFormat::BGRA_UByte )
-			{
-				BufferLayerCombineFunctor<4, 4>(ResultImage, BaseImage, BlendedImage, Forward<ImageCombineFn>(ImageCombine));
-			}
-			else
-			{
-				checkf(false, TEXT("Unsupported format."));
-			}
-		}
-		else
-		{
-			checkf(false, TEXT("Unsupported format."));
-		}
-	}
-
-
-	//! Blend a subimage on the base using a mask.
-	inline void ImageBlendOnBaseNoAlpha(FImage* BaseImage,
-		const FImage* MaskImage,
-		const FImage* BlendedImage,
-		const box<FIntVector2>& Rect)
-	{
-		ImageLayerOnBaseNoAlpha<BlendChannel, false>(BaseImage, MaskImage, BlendedImage, Rect);
-	}
+	void ImageLayerCombineFunctor(FImage* ResultImage, const FImage* BaseImage, const FImage* BlendedImage, ImageCombineFn&& ImageCombine);
 
 }
diff --git a/Engine/Plugins/Mutable/Source/MutableRuntime/Private/MuR/OpImageNormalCombine.h b/Engine/Plugins/Mutable/Source/MutableRuntime/Private/MuR/OpImageNormalCombine.h
--- a/Engine/Plugins/Mutable/Source/MutableRuntime/Private/MuR/OpImageNormalCombine.h
+++ b/Engine/Plugins/Mutable/Source/MutableRuntime/Private/MuR/OpImageNormalCombine.h
@@ -4,63 +4,10 @@
 
 #include "MuR/ImagePrivate.h"
 #include "MuR/OpImageLayer.h"
+#include "MuR/OpImageBlend.h"
 
 namespace mu
 {
-	// Combine layer oprations get packed color.
-	inline unsigned CombineNormal(unsigned base, unsigned blended)
-	{
-		// See /Engine/Functions/Engine_MaterialFunctions02/Utility/BlendAngleCorrectedNormals.BlendAngleCorrectedNormals
-		// for the source of the effect.
-
-		const float baseRf = (static_cast<float>((base >> 0)  & 0xFF) / 255.0f) * 2.0f - 1.0f;
-		const float baseGf = (static_cast<float>((base >> 8)  & 0xFF) / 255.0f) * 2.0f - 1.0f;
-		const float baseBf = (static_cast<float>((base >> 16) & 0xFF) / 255.0f) * 2.0f; //One added to the b channel.
-
-		const float blendedRf = (static_cast<float>((blended >> 0)  & 0xFF) / 255.0f) * 2.0f - 1.0f;
-		const float blendedGf = (static_cast<float>((blended >> 8)  & 0xFF) / 255.0f) * 2.0f - 1.0f;
-		const float blendedBf = (static_cast<float>((blended >> 16) & 0xFF) / 255.0f) * 2.0f - 1.0f;
-
-		FVector3f a = FVector3f(baseRf, baseGf, baseBf);
-		FVector3f b = FVector3f(-blendedRf, -blendedGf, blendedBf);
-
-		// The original does not normalize but if not done results don't look good due to signal clipping.
-		FVector3f n = (a* FVector3f::DotProduct(a, b) - b*a.Z);
-		n.Normalize();
-
-		return 
-			(FMath::Clamp(static_cast<unsigned>((n.X + 1.0f) * 0.5f * 255.0f), 0u, 255u) << 0) |
-			(FMath::Clamp(static_cast<unsigned>((n.Y + 1.0f) * 0.5f * 255.0f), 0u, 255u) << 8) |
-			(FMath::Clamp(static_cast<unsigned>((n.Z + 1.0f) * 0.5f * 255.0f), 0u, 255u) << 16);
-	}
-
-	inline unsigned CombineNormalMasked(unsigned base, unsigned blended, unsigned mask)
-	{
-		// See /Engine/Functions/Engine_MaterialFunctions02/Utility/BlendAngleCorrectedNormals.BlendAngleCorrectedNormals
-		// for the source of the effect.
-
-		const float baseRf = (static_cast<float>((base >> 0)  & 0xFF) / 255.0f) * 2.0f - 1.0f;
-		const float baseGf = (static_cast<float>((base >> 8)  & 0xFF) / 255.0f) * 2.0f - 1.0f;
-		const float baseBf = (static_cast<float>((base >> 16) & 0xFF) / 255.0f) * 2.0f - 1.0f;
-
-		const float blendedRf = (static_cast<float>((blended >> 0)  & 0xFF) / 255.0f) * 2.0f - 1.0f;
-		const float blendedGf = (static_cast<float>((blended >> 8)  & 0xFF) / 255.0f) * 2.0f - 1.0f;
-		const float blendedBf = (static_cast<float>((blended >> 16) & 0xFF) / 255.0f) * 2.0f - 1.0f;
-
-		const float maskRf = (static_cast<float>(mask & 0xFF) / 255.0f);
-
-		FVector3f a = FVector3f(baseRf, baseGf, baseBf + 1.0f);
-		FVector3f b = FVector3f(-blendedRf, -blendedGf, blendedBf);
-
-		// The original does not normalize but if not done results don't look good due to signal clipping.
-		FVector3f n = ( FMath::Lerp( FVector3f(baseRf, baseGf, baseBf), a*FVector3f::DotProduct(a, b) - b*a.Z, maskRf ) );
-		n.Normalize();
-
-		return 
-			(FMath::Clamp(static_cast<unsigned>((n.X + 1.0f) * 0.5f * 255.0f), 0u, 255u) << 0) |
-			(FMath::Clamp(static_cast<unsigned>((n.Y + 1.0f) * 0.5f * 255.0f), 0u, 255u) << 8) |
-			(FMath::Clamp(static_cast<unsigned>((n.Z + 1.0f) * 0.5f * 255.0f), 0u, 255u) << 16);
-	}
 
 	inline void ImageNormalCombine(FImage* pResult, const FImage* pBase, FVector4f col)
 	{
diff --git a/Engine/Plugins/Mutable/Source/MutableRuntime/Private/MuR/OpImageSaturate.h b/Engine/Plugins/Mutable/Source/MutableRuntime/Private/MuR/OpImageSaturate.h
--- a/Engine/Plugins/Mutable/Source/MutableRuntime/Private/MuR/OpImageSaturate.h
+++ b/Engine/Plugins/Mutable/Source/MutableRuntime/Private/MuR/OpImageSaturate.h
@@ -4,10 +4,34 @@
 
 #include "MuR/ImagePrivate.h"
 #include "MuR/Platform.h"
+
+#include "Async/ParallelFor.h" 
 #include "Math/VectorRegister.h"
 
 namespace mu
 {
+
+namespace OpImageSaturateInternal
+{
+	// TODO: This is used in multiple places, make it an utility so it can be shared.  
+	FORCENOINLINE void InvokeParallelFor(int32 NumBatches, TFunctionRef<void(int32)> BatchFunc)
+	{
+		if (NumBatches == 0)
+		{
+			return;
+		}
+
+		if (NumBatches == 1)
+		{
+			BatchFunc(0);
+		}
+		else
+		{
+			ParallelFor(NumBatches, BatchFunc);
+		}
+	}
+} // namespace mu::OpImageStaurateInternal
+
 	template<bool bUseVectorImpl = false>
 	inline void ImageSaturate(FImage* AImage, float Factor)
 	{
@@ -30,7 +54,10 @@ namespace mu
 		constexpr int32 NumBatchElems = 1 << 14;
 		const int32 NumBatches = AImage->DataStorage.GetNumBatches(NumBatchElems, BytesPerElem); 
 
-		auto ProcessBatch = [AImage, NumBatchElems, BytesPerElem, FactorUNorm](int32 BatchId)
+		OpImageSaturateInternal::InvokeParallelFor(NumBatches,
+		[
+			AImage, NumBatchElems, BytesPerElem, FactorUNorm
+		](int32 BatchId)
 		{
 			TArrayView<uint8> AView = AImage->DataStorage.GetBatch(BatchId, NumBatchElems, BytesPerElem);
 			uint8* ABuf = AView.GetData(); 
@@ -210,15 +237,6 @@ namespace mu
 			default:
 				check(false);
 			}
-		};
-
-		if (NumBatches == 1)
-		{
-			ProcessBatch(0);
-		}
-		else
-		{
-			ParallelFor(NumBatches, ProcessBatch);
-		}
+		});
 	}
 }
diff --git a/Engine/Plugins/Mutable/Source/MutableRuntime/Public/MuR/ImageDataStorage.h b/Engine/Plugins/Mutable/Source/MutableRuntime/Public/MuR/ImageDataStorage.h
--- a/Engine/Plugins/Mutable/Source/MutableRuntime/Public/MuR/ImageDataStorage.h
+++ b/Engine/Plugins/Mutable/Source/MutableRuntime/Public/MuR/ImageDataStorage.h
@@ -11,8 +11,6 @@
 
 #include <atomic>
 
-#define UE_API MUTABLERUNTIME_API
-
 namespace mu::MemoryCounters
 {
 	struct  FImageMemoryCounter
@@ -58,16 +56,16 @@ public:
 
 public:
 
-	UE_API FImageDataStorage();
-	UE_API FImageDataStorage(const FImageDesc& Desc);
+	MUTABLERUNTIME_API FImageDataStorage();
+	MUTABLERUNTIME_API FImageDataStorage(const FImageDesc& Desc);
 
-	UE_API FImageDataStorage(const FImageDataStorage& Other);
-	UE_API FImageDataStorage& operator=(const FImageDataStorage& Other);
+	MUTABLERUNTIME_API FImageDataStorage(const FImageDataStorage& Other);
+	MUTABLERUNTIME_API FImageDataStorage& operator=(const FImageDataStorage& Other);
 	
 	FImageDataStorage(FImageDataStorage&& Other) = default;
 	FImageDataStorage& operator=(FImageDataStorage&& Other) = default;
 
-	UE_API bool operator==(const FImageDataStorage& Other) const;
+	MUTABLERUNTIME_API bool operator==(const FImageDataStorage& Other) const;
 
 	FORCEINLINE FImageDesc MakeImageDesc() const
 	{
@@ -89,10 +87,10 @@ public:
 		return Buffers.IsEmpty() && CompactedTailOffsets[0] == TNumericLimits<uint16>::Max();
 	}
 
-	UE_API void InitInternalArray(int32 Index, int32 Size, EInitializationType InitType);
-	UE_API void Init(const FImageDesc& ImageDesc, EInitializationType InitType);
+	MUTABLERUNTIME_API void InitInternalArray(int32 Index, int32 Size, EInitializationType InitType);
+	MUTABLERUNTIME_API void Init(const FImageDesc& ImageDesc, EInitializationType InitType);
 
-	UE_API void CopyFrom(const FImageDataStorage& Other);
+	MUTABLERUNTIME_API void CopyFrom(const FImageDataStorage& Other);
 
 	FORCEINLINE int32 GetNumLODs() const
 	{
@@ -102,54 +100,28 @@ public:
 	/**
 	* Get a const view to the data containing the LODIndex.
 	*/
-	FORCEINLINE TArrayView<const uint8> GetLOD(int32 LODIndex) const
-	{
-		check(!IsVoid());
-		if (LODIndex >= NumLODs)
-		{
-			//check(false);
-			return TArrayView<const uint8>();
-		}
-
-		const int32 FirstTailLOD = ComputeFirstCompactedTailLOD();
-		if (LODIndex >= FirstTailLOD)
-		{
-			const int32 LODIndexInTail = FMath::Max(0, LODIndex - FirstTailLOD);
-			
-			const int32 LODInTailBegin = LODIndexInTail == 0 ? 0 : CompactedTailOffsets[LODIndexInTail - 1];
-			const int32 LODInTailEnd = CompactedTailOffsets[LODIndexInTail];
-
-			const FImageArray& TailBuffer = Buffers.Last();
-			return MakeArrayView(TailBuffer.GetData() + LODInTailBegin, LODInTailEnd - LODInTailBegin);
-		}
-
-		return MakeArrayView(Buffers[LODIndex].GetData(), Buffers[LODIndex].Num());
-	}
+	MUTABLERUNTIME_API TArrayView<const uint8> GetLOD(int32 LODIndex) const;
 
 	/**
 	* Get a view to the data containing the LODIndex.
 	*/
-	FORCEINLINE TArrayView<uint8> GetLOD(int32 LODIndex)
-	{	
-		const TArrayView<const uint8> ConstLODView = const_cast<const FImageDataStorage*>(this)->GetLOD(LODIndex);
-		return TArrayView<uint8>(const_cast<uint8*>(ConstLODView.GetData()), ConstLODView.Num());
-	}
+	MUTABLERUNTIME_API TArrayView<uint8> GetLOD(int32 LODIndex);
 
 	/**
 	 * Change the number of lods in the image. If InNumLODs is smaller than the current lod count, 
 	 * LODs are dropped from the tail up.
 	 */
-	UE_API void SetNumLODs(int32 InNumLODs, EInitializationType InitType = EInitializationType::NotInitialized);
+	MUTABLERUNTIME_API void SetNumLODs(int32 InNumLODs, EInitializationType InitType = EInitializationType::NotInitialized);
 
 	/**
 	 * Remove NumLODsToDrop starting form LOD 0. The resulting image will be smaller. 
 	 */
-	UE_API void DropLODs(int32 NumLODsToDrop);
+	MUTABLERUNTIME_API void DropLODs(int32 NumLODsToDrop);
 
 	/**
 	 * Resizes LODIndex Data to NewSizeBytes
 	 */
-	UE_API void ResizeLOD(int32 LODIndex, int32 NewSizeBytes);
+	MUTABLERUNTIME_API void ResizeLOD(int32 LODIndex, int32 NewSizeBytes);
 
 	FORCEINLINE const FImageArray& GetInternalArray(int32 LODIndex) const
 	{	
@@ -174,23 +146,9 @@ public:
 	/**
 	 * Stateless iteration in batches.
 	 */
-	FORCEINLINE int32 GetNumBatches(int32 BatchSizeInElems, int32 BatchElemSizeInBytes) const
-	{	
-		return GetNumBatchesLODRange(BatchSizeInElems, BatchElemSizeInBytes, 0, NumLODs);
-	}
-
-	FORCEINLINE TArrayView<const uint8> GetBatch(int32 BatchId, int32 BatchSizeInElems, int32 BatchElemSizeInBytes) const
-	{
-		return GetBatchLODRange(BatchId, BatchSizeInElems, BatchElemSizeInBytes, 0, NumLODs);
-	}
-
-	FORCEINLINE TArrayView<uint8> GetBatch(int32 BatchId, int32 BatchSizeInElems, int32 BatchElemSizeInBytes)
-	{
-		// Use the const_cast idiom to generate the non-const operation. 
-		const TArrayView<const uint8> ConstBatchView = 
-				const_cast<const FImageDataStorage*>(this)->GetBatch(BatchId, BatchSizeInElems, BatchElemSizeInBytes);
-		return TArrayView<uint8>(const_cast<uint8*>(ConstBatchView.GetData()), ConstBatchView.Num());
-	}
+	MUTABLERUNTIME_API int32 GetNumBatches(int32 BatchSizeInElems, int32 BatchElemSizeInBytes) const;
+	MUTABLERUNTIME_API TArrayView<const uint8> GetBatch(int32 BatchId, int32 BatchSizeInElems, int32 BatchElemSizeInBytes) const;
+	MUTABLERUNTIME_API TArrayView<uint8> GetBatch(int32 BatchId, int32 BatchSizeInElems, int32 BatchElemSizeInBytes);
 
 	/**
 	 * Returns the number of batches GetBatchFirtsLODOffet will admit and that will cover
@@ -200,17 +158,7 @@ public:
 	 * BatchSizeInElems*BatchElemSizeInBytes is not multiple of the buffer size.
 	 *
 	 **/
-	FORCEINLINE int32 GetNumBatchesFirstLODOffset(int32 BatchSizeInElems, int32 BatchElemSizeInBytes, int32 OffsetInBytes) const
-	{
-		check(!IsVoid());
-		TArrayView<const uint8> FirstLODView = GetLOD(0);
-		
-		check(OffsetInBytes % BatchElemSizeInBytes == 0);
-		check(FirstLODView.Num() % BatchElemSizeInBytes == 0);
-		check(FirstLODView.Num() > OffsetInBytes); 
-
-		return FMath::DivideAndRoundUp(FirstLODView.Num() - OffsetInBytes, BatchSizeInElems*BatchElemSizeInBytes);
-	}
+	MUTABLERUNTIME_API int32 GetNumBatchesFirstLODOffset(int32 BatchSizeInElems, int32 BatchElemSizeInBytes, int32 OffsetInBytes) const;
 
 	/**
 	 * Returns a non modifiable view to the portion of the first LOD buffer with OffsetInBytes for the BatchId 
@@ -219,74 +167,19 @@ public:
 	 * BatchSizeInElems*BatchElemSizeInBytes is not multiple of the buffer size.
 	 *
 	 */
-	FORCEINLINE TArrayView<const uint8> GetBatchFirstLODOffset(int32 BatchId, int32 BatchSizeInElems, int32 BatchElemSizeInBytes, int32 OffsetInBytes) const
-	{
-		check(!IsVoid());
-		TArrayView<const uint8> FirstLODView = GetLOD(0);
-		
-		check(OffsetInBytes < FirstLODView.Num());
-		check(FirstLODView.Num() % BatchElemSizeInBytes == 0);
-		check(OffsetInBytes % BatchElemSizeInBytes == 0);
-	
-		TArrayView<const uint8> OffsetedLODView = MakeArrayView(FirstLODView.GetData() + OffsetInBytes, FirstLODView.Num() - OffsetInBytes);
-
-		const int32 BatchSizeInBytes = BatchSizeInElems * BatchElemSizeInBytes;
-		const int32 BatchOffset = BatchId * BatchSizeInBytes;
-
-		if (BatchOffset >= OffsetedLODView.Num())
-		{
-			return TArrayView<const uint8>();
-		}
-
-		return MakeArrayView(OffsetedLODView.GetData() + BatchOffset, FMath::Min(BatchSizeInBytes, OffsetedLODView.Num() - BatchOffset));
-	}
+	MUTABLERUNTIME_API TArrayView<const uint8> GetBatchFirstLODOffset(int32 BatchId, int32 BatchSizeInElems, int32 BatchElemSizeInBytes, int32 OffsetInBytes) const;
 
 	/**
 	 * Non const version of GetBatchFirstLODOffset(). 
 	 */
-	FORCEINLINE TArrayView<uint8> GetBatchFirstLODOffset(int32 BatchId, int32 BatchSizeInElems, int32 BatchElemSizeInBytes, int32 OffsetInBytes = 0)
-	{
-		const TArrayView<const uint8> ConstBatchView = 
-				const_cast<const FImageDataStorage*>(this)->GetBatchFirstLODOffset(BatchId, BatchSizeInElems, BatchElemSizeInBytes, OffsetInBytes);
-
-		return TArrayView<uint8>(const_cast<uint8*>(ConstBatchView.GetData()), ConstBatchView.Num());
-	}
+	MUTABLERUNTIME_API TArrayView<uint8> GetBatchFirstLODOffset(int32 BatchId, int32 BatchSizeInElems, int32 BatchElemSizeInBytes, int32 OffsetInBytes = 0);
 
 	/**
 	 * Returns the number of batches GetBatchLODRange will admit and that will cover
 	 * the [LODBegin, LODEnd) range.
 	 *
 	 **/
-	FORCEINLINE int32 GetNumBatchesLODRange(int32 BatchSizeInElems, int32 BatchElemSizeInBytes, int32 LODBegin, int32 LODEnd) const
-	{
-		check(!IsVoid());
-		check(LODBegin < LODEnd);
-		check(LODBegin >= 0 && LODEnd <= NumLODs);
-
-		const int32 BatchNumBytes = BatchSizeInElems*BatchElemSizeInBytes;
-		
-		const int32 FirstCompactedTailLOD = ComputeFirstCompactedTailLOD();
-		const int32 NumLODsInTail = FMath::Max(0, NumLODs - FirstCompactedTailLOD);
-		const int32 NumLODsNotInTail = FMath::Max(0, NumLODs - NumLODsInTail);
-
-		const int32 LODsNotInTailLODRangeEnd = FMath::Min(LODEnd, NumLODsNotInTail);
-
-		int32 NumBatches = 0;
-		for (int32 BufferIndex = LODBegin; BufferIndex < LODsNotInTailLODRangeEnd; ++BufferIndex)
-		{	
-			NumBatches += FMath::DivideAndRoundUp(Buffers[BufferIndex].Num(), BatchNumBytes);
-		}
-
-		if (FirstCompactedTailLOD < LODEnd)
-		{
-			TArrayView<const uint8> FirstLODInRangeView = GetLOD(FMath::Max(FirstCompactedTailLOD, LODBegin));
-			TArrayView<const uint8> LastLODInRangeView = GetLOD(LODEnd - 1);
-
-			const int32 TailInLODRangeNumBytes = LastLODInRangeView.GetData() - FirstLODInRangeView.GetData() + LastLODInRangeView.Num();
-			NumBatches += FMath::DivideAndRoundUp(TailInLODRangeNumBytes, BatchNumBytes);
-		}
-		return NumBatches;
-	}
+	MUTABLERUNTIME_API int32 GetNumBatchesLODRange(int32 BatchSizeInElems, int32 BatchElemSizeInBytes, int32 LODBegin, int32 LODEnd) const;
 
 	/**
 	 * Returns a non modifiable view to the portion of the [LODBegin, LODEnd) for the BatchId. 
@@ -295,122 +188,25 @@ public:
 	 * BatchSizeInElems*BatchElemSizeInBytes is not multiple of the buffer size.
 	 *
 	 */
-	FORCEINLINE TArrayView<const uint8> GetBatchLODRange(int32 BatchId, int32 BatchSizeInElems, int32 BatchElemSizeInBytes, int32 LODBegin, int32 LODEnd) const
-	{
-		check(!IsVoid());
-		check(LODBegin < LODEnd);
-		check(LODBegin >= 0 && LODEnd <= NumLODs);
-
-		const int32 BatchNumBytes = BatchSizeInElems*BatchElemSizeInBytes;
-
-		const int32 FirstCompactedTailLOD = ComputeFirstCompactedTailLOD();
-		const int32 NonTailBuffersEnd = FMath::Min(LODEnd, FirstCompactedTailLOD);
-
-		int32 NumBatches = 0;	
-		for (int32 BufferIndex = LODBegin; BufferIndex < NonTailBuffersEnd; ++BufferIndex)
-		{
-			const TArrayView<const uint8> BufferView = MakeArrayView(Buffers[BufferIndex].GetData(), Buffers[BufferIndex].Num());
-			const int32 BufferNumBatches = FMath::DivideAndRoundUp(BufferView.Num(), BatchNumBytes); 
-			
-			check(BatchId >= NumBatches);
-			int32 BufferBatchId = BatchId - NumBatches;
-
-			if (BufferBatchId < BufferNumBatches)
-			{ 
-				const int32 BufferBatchBeginOffset = BufferBatchId * BatchNumBytes;
-
-				return TArrayView<const uint8>(
-						BufferView.GetData() + BufferBatchBeginOffset,
-						FMath::Min(BufferView.Num() - BufferBatchBeginOffset, BatchNumBytes));
-			}
-
-			NumBatches += BufferNumBatches;
-		}
-
-		if (FirstCompactedTailLOD < LODEnd)
-		{
-			TArrayView<const uint8> FirstLODInRangeView = GetLOD(FMath::Max(FirstCompactedTailLOD, LODBegin));
-			TArrayView<const uint8> LastLODInRangeView = GetLOD(LODEnd - 1);
-			
-			const int32 TailInLODRangeNumBytes = LastLODInRangeView.GetData() - FirstLODInRangeView.GetData() + LastLODInRangeView.Num();
-			TArrayView<const uint8> TailInLODRangeView = MakeArrayView(FirstLODInRangeView.GetData(), TailInLODRangeNumBytes);
-
-			check(BatchId >= NumBatches);
-			const int32 TailInLODRangeBatchId = BatchId - NumBatches;
-			const int32 TailInLODRangeNumBatches = FMath::DivideAndRoundUp(TailInLODRangeView.Num(), BatchNumBytes);
-
-			if (TailInLODRangeBatchId < TailInLODRangeNumBatches)
-			{
-				const int32 BufferBatchBeginOffset = TailInLODRangeBatchId * BatchNumBytes;
-				
-				return TArrayView<const uint8>(
-							TailInLODRangeView.GetData() + BufferBatchBeginOffset,
-							FMath::Min(TailInLODRangeView.Num() - BufferBatchBeginOffset, BatchNumBytes));
-			}
-		}
-		// If the BatchId is not found, return an empty view.
-		return TArrayView<const uint8>();
-	}
+	MUTABLERUNTIME_API TArrayView<const uint8> GetBatchLODRange(int32 BatchId, int32 BatchSizeInElems, int32 BatchElemSizeInBytes, int32 LODBegin, int32 LODEnd) const;
 	
 	/**
 	 * Non const version of GetBatchLODRange().
 	 */
-	FORCEINLINE TArrayView<uint8> GetBatchLODRange(int32 BatchId, int32 BatchSizeInElems, int32 BatchElemSizeInBytes, int32 LODBegin, int32 LODEnd)
-	{
-		const TArrayView<const uint8> ConstBatchView = 
-			const_cast<const FImageDataStorage*>(this)->GetBatchLODRange(
-					BatchId, BatchSizeInElems, BatchElemSizeInBytes, LODBegin, LODEnd);
+	MUTABLERUNTIME_API TArrayView<uint8> GetBatchLODRange(int32 BatchId, int32 BatchSizeInElems, int32 BatchElemSizeInBytes, int32 LODBegin, int32 LODEnd);
 
-		return TArrayView<uint8>(const_cast<uint8*>(ConstBatchView.GetData()), ConstBatchView.Num());
-	}
-
-	FORCEINLINE int32 GetAllocatedSize() const
-	{ 
-		SIZE_T Result = 0;
-		for (const FImageArray& Buffer : Buffers)
-		{
-			Result += Buffer.GetAllocatedSize();
-		}
-
-		check(Result < TNumericLimits<int32>::Max())
-		return (int32)Result;
-	}
-
-	FORCEINLINE int32 GetDataSize() const
-	{
-		SIZE_T Result = 0;
-		for (const FImageArray& Buffer : Buffers)
-		{
-			Result += Buffer.Num();
-		}
-
-		check(Result < TNumericLimits<int32>::Max())
-		return (int32)Result;
-	}
+	MUTABLERUNTIME_API int32 GetAllocatedSize() const;
+	MUTABLERUNTIME_API int32 GetDataSize() const;
 
 	/**
 	 * Returns true if all buffers are empty. This can be true after initialization for image formats that report 
 	 * BytesPerBlock 0, e.g., RLE formats. If not initialized will also return true.
 	 */
-	FORCEINLINE bool IsEmpty() const
-	{
-		check(!IsVoid());
-
-		for (const FImageArray& Buffer : Buffers)
-		{
-			if (!Buffer.IsEmpty())
-			{
-				return false;
-			}
-		}
-
-		return true;
-	}
+	MUTABLERUNTIME_API bool IsEmpty() const;
 
-	UE_API void Serialise(FOutputArchive& Arch) const;
-	UE_API void Unserialise(FInputArchive& Arch);
+	MUTABLERUNTIME_API void Serialise(FOutputArchive& Arch) const;
+	MUTABLERUNTIME_API void Unserialise(FInputArchive& Arch);
 };
 
 }
 
-#undef UE_API
