[Mutable] Cache compiled data from DDC asynchronously. Enable async compilation of COs during cook.

#rnx

diff --git a/Engine/Plugins/Mutable/Source/CustomizableObject/Internal/MuCO/CustomizableObjectPrivate.h b/Engine/Plugins/Mutable/Source/CustomizableObject/Internal/MuCO/CustomizableObjectPrivate.h
--- a/Engine/Plugins/Mutable/Source/CustomizableObject/Internal/MuCO/CustomizableObjectPrivate.h
+++ b/Engine/Plugins/Mutable/Source/CustomizableObject/Internal/MuCO/CustomizableObjectPrivate.h
@@ -1694,12 +1694,6 @@ public:
 
 #if WITH_EDITOR
 
-	static CUSTOMIZABLEOBJECT_API uint32 GetECustomizableObjectVersionEnumHash();
-
-	// Compose folder name where the data is stored
-	static CUSTOMIZABLEOBJECT_API FString GetCompiledDataFolderPath();
-	static CUSTOMIZABLEOBJECT_API FString GetDataTypeExtension(MutablePrivate::EStreamableDataType DataType);
-
 	/** Compose file name. */
 	CUSTOMIZABLEOBJECT_API FString GetCompiledDataFileName(const ITargetPlatform* InTargetPlatform = nullptr, bool bIsDiskStreamer = false) const;
 
@@ -1707,9 +1701,6 @@ public:
 	CUSTOMIZABLEOBJECT_API TArray<uint8> BuildDerivedDataKey(FCompilationOptions Options);
 	CUSTOMIZABLEOBJECT_API UE::DerivedData::FCacheKey GetDerivedDataCacheKeyForOptions(FCompilationOptions InOptions);
 
-	/** Attempts to load the compiled data from DDC. Builds key if not supplied. */
-	CUSTOMIZABLEOBJECT_API bool LoadCompiledDataFromDDC(FCompilationOptions Options, UE::DerivedData::ECachePolicy DefaultPolicy, UE::DerivedData::FCacheKey* DDCKey, MutablePrivate::FMutableCachedPlatformData& InOutPlatformData);
-
 	/** Log data for debugging purposes */
 	CUSTOMIZABLEOBJECT_API void LogMemory();
 #endif
@@ -1737,10 +1728,6 @@ public:
 	CUSTOMIZABLEOBJECT_API void LoadModelStreamableBulk(FArchive& Ar, bool bIsCooking);
 	CUSTOMIZABLEOBJECT_API void LoadModel(FArchive& Ar);
 
-	CUSTOMIZABLEOBJECT_API TObjectPtr<UModelResources> LoadModelResources_Internal(FArchive& Ar, const ITargetPlatform* InTargetPlatform, bool bSkipEditorOnlyData = false) const;
-	CUSTOMIZABLEOBJECT_API const TSharedPtr<FModelStreamableBulkData> LoadModelStreamableBulk_Internal(FArchive& Ar) const;
-	CUSTOMIZABLEOBJECT_API const TSharedPtr<mu::FModel, ESPMode::ThreadSafe> LoadModel_Internal(FArchive& Ar) const;
-
 	/** Load compiled data for the running platform from disk, this is used to load Editor Compilations. */
 	CUSTOMIZABLEOBJECT_API void LoadCompiledDataFromDisk();
 	
@@ -1901,4 +1888,14 @@ public:
 	*/
 CUSTOMIZABLEOBJECT_API UE::DerivedData::FValueId GetDerivedDataValueIdForResource(MutablePrivate::EStreamableDataType StreamableDataType, uint32 FileId, uint16 ResourceType, uint16 Flags);
 
+// Compose folder name where the data is stored
+CUSTOMIZABLEOBJECT_API FString GetCompiledDataFolderPath();
+CUSTOMIZABLEOBJECT_API FString GetDataTypeExtension(MutablePrivate::EStreamableDataType DataType);
+
+CUSTOMIZABLEOBJECT_API uint32 GetECustomizableObjectVersionEnumHash();
+
+CUSTOMIZABLEOBJECT_API TObjectPtr<UModelResources> LoadModelResources_Internal(FArchive& Ar, const UCustomizableObject* Outer, const ITargetPlatform* InTargetPlatform, bool bSkipEditorOnlyData = false);
+CUSTOMIZABLEOBJECT_API const TSharedPtr<FModelStreamableBulkData> LoadModelStreamableBulk_Internal(FArchive& Ar);
+CUSTOMIZABLEOBJECT_API const TSharedPtr<mu::FModel, ESPMode::ThreadSafe> LoadModel_Internal(FArchive& Ar);
+
 #endif
\ No newline at end of file
diff --git a/Engine/Plugins/Mutable/Source/CustomizableObject/Private/MuCO/CustomizableObject.cpp b/Engine/Plugins/Mutable/Source/CustomizableObject/Private/MuCO/CustomizableObject.cpp
--- a/Engine/Plugins/Mutable/Source/CustomizableObject/Private/MuCO/CustomizableObject.cpp
+++ b/Engine/Plugins/Mutable/Source/CustomizableObject/Private/MuCO/CustomizableObject.cpp
@@ -78,7 +78,7 @@ TAutoConsoleVariable<bool> CVarMutableUseBulkData(
 
 TAutoConsoleVariable<bool> CVarMutableAsyncCook(
 	TEXT("Mutable.CookAsync"),
-	false,
+	true,
 	TEXT("True - Customizable Objects will be compiled asynchronously during cook.\n")
 	TEXT("False - Sync compilation.\n"));
 
@@ -497,7 +497,7 @@ FString GetModelResourcesNameForPlatform(const UCustomizableObject& Customizable
 
 bool UCustomizableObjectPrivate::LoadModelResources(FArchive& MemoryReader, const ITargetPlatform* InTargetPlatform, bool bIsCooking)
 {
-	TObjectPtr<UModelResources> LocalModelResources = LoadModelResources_Internal(MemoryReader, InTargetPlatform, bIsCooking);
+	TObjectPtr<UModelResources> LocalModelResources = LoadModelResources_Internal(MemoryReader, GetPublic(), InTargetPlatform, bIsCooking);
 	SetModelResources(LocalModelResources, bIsCooking);
 
 	return LocalModelResources != nullptr;
@@ -516,12 +516,12 @@ void UCustomizableObjectPrivate::LoadModel(FArchive& MemoryReader)
 }
 
 
-TObjectPtr<UModelResources> UCustomizableObjectPrivate::LoadModelResources_Internal(FArchive& MemoryReader, const ITargetPlatform* InTargetPlatform, bool bIsCooking) const
+TObjectPtr<UModelResources> LoadModelResources_Internal(FArchive& MemoryReader, const UCustomizableObject* Outer, const ITargetPlatform* InTargetPlatform, bool bIsCooking)
 {
 	// Make sure mutable has been initialised.
 	UCustomizableObjectSystem::GetInstance();
 
-	const FString ModelResourcesName = GetModelResourcesNameForPlatform(*GetPublic(), *InTargetPlatform);
+	const FString ModelResourcesName = GetModelResourcesNameForPlatform(*Outer, *InTargetPlatform);
 	TObjectPtr<UModelResources> LocalModelResources = NewObject<UModelResources>((UObject*)GetTransientPackage(), FName(*ModelResourcesName), RF_Public);
 
 	FObjectAndNameAsStringProxyArchive ObjectReader(MemoryReader, true);
@@ -532,7 +532,7 @@ TObjectPtr<UModelResources> UCustomizableObjectPrivate::LoadModelResources_Inter
 }
 
 
-const TSharedPtr<FModelStreamableBulkData> UCustomizableObjectPrivate::LoadModelStreamableBulk_Internal(FArchive& MemoryReader) const
+const TSharedPtr<FModelStreamableBulkData> LoadModelStreamableBulk_Internal(FArchive& MemoryReader)
 {
 	TSharedPtr<FModelStreamableBulkData> LocalModelStreamablesPtr = MakeShared<FModelStreamableBulkData>();
 	FModelStreamableBulkData& LocalModelStreamables = *LocalModelStreamablesPtr.Get();
@@ -542,7 +542,7 @@ const TSharedPtr<FModelStreamableBulkData> UCustomizableObjectPrivate::LoadModel
 }
 
 
-const TSharedPtr<mu::FModel, ESPMode::ThreadSafe> UCustomizableObjectPrivate::LoadModel_Internal(FArchive& MemoryReader) const
+const TSharedPtr<mu::FModel, ESPMode::ThreadSafe> LoadModel_Internal(FArchive& MemoryReader)
 {
 	FUnrealMutableInputStream Stream(MemoryReader);
 	mu::FInputArchive Arch(&Stream);
@@ -647,7 +647,7 @@ void UCustomizableObjectPrivate::LoadCompiledDataFromDisk()
 	bool bHasCompiledData = true;
 	for (int32 DataType = 0; DataType < static_cast<int32>(MutablePrivate::EStreamableDataType::DataTypeCount); ++DataType)
 	{
-		const FString FilePath = FullFileName + UCustomizableObjectPrivate::GetDataTypeExtension(static_cast<MutablePrivate::EStreamableDataType>(DataType));
+		const FString FilePath = FullFileName + GetDataTypeExtension(static_cast<MutablePrivate::EStreamableDataType>(DataType));
 		if (IFileHandle* FileHandle = FPlatformFileManager::Get().GetPlatformFile().OpenRead(*FilePath))
 		{
 			FileHandles.Emplace(FileHandle);
@@ -813,7 +813,7 @@ FReply UCustomizableObjectPrivate::AddNewParameterProfile(FString Name, UCustomi
 }
 
 
-uint32 UCustomizableObjectPrivate::GetECustomizableObjectVersionEnumHash()
+uint32 GetECustomizableObjectVersionEnumHash()
 {
 	static const uint32 VersionsHash = []() -> uint32
 		{
@@ -838,7 +838,7 @@ uint32 UCustomizableObjectPrivate::GetECustomizableObjectVersionEnumHash()
 }
 
 
-FString UCustomizableObjectPrivate::GetCompiledDataFolderPath()
+FString GetCompiledDataFolderPath()
 {	
 	return FPaths::ConvertRelativePathToFull(FPaths::ProjectSavedDir() + TEXT("MutableStreamedDataEditor/"));
 }
@@ -854,7 +854,7 @@ FString UCustomizableObjectPrivate::GetCompiledDataFileName(const ITargetPlatfor
 }
 
 
-FString UCustomizableObjectPrivate::GetDataTypeExtension(MutablePrivate::EStreamableDataType DataType)
+FString GetDataTypeExtension(MutablePrivate::EStreamableDataType DataType)
 {
 	switch (DataType)
 	{
@@ -3121,214 +3121,6 @@ UE::DerivedData::FCacheKey UCustomizableObjectPrivate::GetDerivedDataCacheKeyFor
 }
 
 
-bool UCustomizableObjectPrivate::LoadCompiledDataFromDDC(FCompilationOptions Options, UE::DerivedData::ECachePolicy DefaultPolicy, UE::DerivedData::FCacheKey* DDCKey, MutablePrivate::FMutableCachedPlatformData& PlatformData)
-{
-	MUTABLE_CPUPROFILER_SCOPE(UCustomizableObjectPrivate::LoadCompiledDataFromDDC);
-
-	using namespace UE::DerivedData;
-
-	/* Overview.
-	*	1. Create an initial pull request to look for the compiled data in the DDC. Skip streamable binary blobs.
-	*	2. Try to load the compiled data.
-	*	3. (Cooking) Create a second request to pull all streamable blobs and cache the compiled data.
-	*/
-
-	FCacheKey CacheKey = DDCKey ? *DDCKey : GetDerivedDataCacheKeyForOptions(Options);
-	check(CacheKey.Hash.IsZero() == false);
-
-	// Buffers with the compiled data
-	FSharedBuffer ModelBytesDDC;
-	FSharedBuffer ModelResourcesBytesDDC;
-	FSharedBuffer ModelStreamablesBytesDDC;
-	FSharedBuffer BulkDataFilesBytesDDC;
-
-	{	// Create a (sync) request to get the serialized Model, ModelResources, and ModelStreamable files to validate versioning and resources 
-		MUTABLE_CPUPROFILER_SCOPE(CheckDDC);
-
-		// Set the request policy to Default + SkipData to avoid pulling the streamable files until we know the compiled data can be used.
-		FCacheRecordPolicyBuilder PolicyBuilder(DefaultPolicy | ECachePolicy::SkipData);
-
-		// Overwrite the request policy for the resources we want to pull
-		PolicyBuilder.AddValuePolicy(MutablePrivate::GetDerivedDataModelResourcesId(), DefaultPolicy);
-		PolicyBuilder.AddValuePolicy(MutablePrivate::GetDerivedDataModelId(), DefaultPolicy);
-		PolicyBuilder.AddValuePolicy(MutablePrivate::GetDerivedDataModelStreamableBulkDataId(), DefaultPolicy);
-		PolicyBuilder.AddValuePolicy(MutablePrivate::GetDerivedDataBulkDataFilesId(), DefaultPolicy);
-
-		FCacheGetRequest Request;
-		Request.Name = GetPathNameSafe(GetPublic());
-		Request.Key = CacheKey;
-		Request.Policy = PolicyBuilder.Build();
-
-		// Sync request to retrieve the compiled data for validation. Streamable resources are excluded.
-		FRequestOwner RequestOwner(EPriority::Blocking);
-		GetCache().Get(MakeArrayView(&Request, 1), RequestOwner,
-			[&ModelBytesDDC, &ModelResourcesBytesDDC, &ModelStreamablesBytesDDC, &BulkDataFilesBytesDDC](FCacheGetResponse&& Response)
-			{
-				if (Response.Status == EStatus::Ok)
-				{
-					const FCompressedBuffer& ModelCompressedBuffer = Response.Record.GetValue(MutablePrivate::GetDerivedDataModelId()).GetData();
-					ModelBytesDDC = ModelCompressedBuffer.Decompress();
-
-					const FCompressedBuffer& ModelResourcesCompressedBuffer = Response.Record.GetValue(MutablePrivate::GetDerivedDataModelResourcesId()).GetData();
-					ModelResourcesBytesDDC = ModelResourcesCompressedBuffer.Decompress();
-					
-					const FCompressedBuffer& ModelStreamablesCompressedBuffer = Response.Record.GetValue(MutablePrivate::GetDerivedDataModelStreamableBulkDataId()).GetData();
-					ModelStreamablesBytesDDC = ModelStreamablesCompressedBuffer.Decompress();
-
-					const FCompressedBuffer& BulkDataFilesCompressedBuffer = Response.Record.GetValue(MutablePrivate::GetDerivedDataBulkDataFilesId()).GetData();
-					BulkDataFilesBytesDDC = BulkDataFilesCompressedBuffer.Decompress();
-				}
-			});
-		RequestOwner.Wait();
-	}
-
-	// Check if all the requested buffers were found.
-	if (!ModelBytesDDC.IsNull() && !ModelResourcesBytesDDC.IsNull() && !BulkDataFilesBytesDDC.IsNull()  && !ModelStreamablesBytesDDC.IsNull())
-	{
-		// Load the compiled data to validate it.
-		FMemoryReaderView ModelResourcesReader(ModelResourcesBytesDDC.GetView());
-
-		if (TObjectPtr<UModelResources> LocalModelResources = LoadModelResources_Internal(ModelResourcesReader, Options.TargetPlatform, Options.bIsCooking))
-		{
-			PlatformData.ModelResources = TStrongObjectPtr<UModelResources>(LocalModelResources);
-
-			FMemoryReaderView ModelStreamablesReader(ModelStreamablesBytesDDC.GetView());
-			PlatformData.ModelStreamableBulkData = LoadModelStreamableBulk_Internal(ModelStreamablesReader);
-			PlatformData.ModelStreamableBulkData->bIsStoredInDDC = true;
-			PlatformData.ModelStreamableBulkData->DDCKey = CacheKey;
-			PlatformData.ModelStreamableBulkData->DDCDefaultPolicy = ECachePolicy::Default;
-
-			FMemoryReaderView ModelReader(ModelBytesDDC.GetView());
-			PlatformData.Model = LoadModel_Internal(ModelReader);
-		}
-
-		if (!PlatformData.Model || !PlatformData.ModelResources || !PlatformData.ModelStreamableBulkData)
-		{
-			return false;
-		}
-
-		// Cache CookedPlatfomData. 
-		if (Options.bIsCooking)
-		{
-			// Sync cache cooked platform data
-			// TODO UE-220138: Sync -> Async
-			MUTABLE_CPUPROFILER_SCOPE(CacheCookedPlatformData);
-
-			PlatformData.ModelResources->InitCookData(*GetPublic());
-
-			// Value Id to file mapping to reconstruct the cached data
-			TMap<FValueId, MutablePrivate::FFile> ValueIdToFile;
-
-			{
-				MUTABLE_CPUPROFILER_SCOPE(BuildValueIdToFile);
-				TArray<MutablePrivate::FFile> BulkDataFiles;
-				FMemoryReaderView FilesReader(BulkDataFilesBytesDDC.GetView());
-				FilesReader << BulkDataFiles;
-
-				ValueIdToFile.Reserve(BulkDataFiles.Num());
-
-				uint32 FileIndex = 0;
-				for (MutablePrivate::FFile& File : BulkDataFiles)
-				{
-					MutablePrivate::FFile& DestFile = ValueIdToFile.Add(GetDerivedDataValueIdForResource(File.DataType, FileIndex, File.ResourceType, File.Flags));
-					DestFile = MoveTemp(File);
-
-					++FileIndex;
-				}
-				BulkDataFiles.Empty();
-			}
-
-			// Create a new pull request to retrieve all compiled data. Streamable bulk data included
-			FCacheGetRequest Request;
-			Request.Name = GetPathNameSafe(GetPublic());
-			Request.Key = CacheKey;
-			Request.Policy = ECachePolicy::Default;
-
-			FRequestOwner RequestOwner(EPriority::Blocking);
-			GetCache().Get(MakeArrayView(&Request, 1), RequestOwner,
-				[&PlatformData, &ValueIdToFile](FCacheGetResponse&& Response)
-				{
-					MUTABLE_CPUPROFILER_SCOPE(CacheBulkDataFromDDC);
-
-					if (ensure(Response.Status == EStatus::Ok))
-					{
-						// Get all values and convert them to FMutableCachedPlatformData's format
-						TConstArrayView<FValueWithId> Values = Response.Record.GetValues();
-
-						TArray64<uint8> TempData;
-						for (const FValueWithId& Value : Values)
-						{
-							check(Value.IsValid());
-
-							const MutablePrivate::FFile* File = ValueIdToFile.Find(Value.GetId());
-							if (!File) // Skip value. It is not a streamable binary blob.
-							{
-								continue;
-							}
-
-							const uint64 RawSize = Value.GetRawSize();
-							TempData.SetNumUninitialized(RawSize, EAllowShrinking::No);
-
-							// Decompress streamable binary blobs
-							const bool bDecompressedSuccessfully = Value.GetData().TryDecompressTo(MakeMemoryView(TempData.GetData(), RawSize));
-							check(bDecompressedSuccessfully);
-
-							// Filter and cache the data by DataType
-							switch (File->DataType)
-							{
-							case MutablePrivate::EStreamableDataType::Model:
-							{
-								for (const MutablePrivate::FBlock& Block : File->Blocks)
-								{
-									PlatformData.ModelStreamableData.Set(Block.Id, TempData.GetData() + Block.Offset, Block.Size);
-								}
-								break;
-							}
-							case MutablePrivate::EStreamableDataType::RealTimeMorph:
-							{
-								for (const MutablePrivate::FBlock& Block : File->Blocks)
-								{
-									PlatformData.MorphStreamableData.Set(Block.Id, TempData.GetData() + Block.Offset, Block.Size);
-								}
-								break;
-							}
-							case MutablePrivate::EStreamableDataType::Clothing:
-							{
-								for (const MutablePrivate::FBlock& Block : File->Blocks)
-								{
-									PlatformData.ClothingStreamableData.Set(Block.Id, TempData.GetData() + Block.Offset, Block.Size);
-								}
-								break;
-							}
-
-							default:
-								unimplemented();
-								break;
-							}
-						}
-					}
-				});
-			RequestOwner.Wait();
-
-			// Generate list of files and update streamable blocks ids and offsets
-			if (CVarMutableUseBulkData.GetValueOnAnyThread())
-			{
-				const uint32 NumBulkDataFilesPerBucket = MAX_uint8;
-				MutablePrivate::GenerateBulkDataFilesListWithFileLimit(PlatformData.Model, *PlatformData.ModelStreamableBulkData.Get(), NumBulkDataFilesPerBucket, PlatformData.BulkDataFiles);
-			}
-			else
-			{
-				MutablePrivate::GenerateBulkDataFilesListWithSizeLimit(PlatformData.Model, *PlatformData.ModelStreamableBulkData.Get(), Options.TargetPlatform, Options.PackagedDataBytesLimit, PlatformData.BulkDataFiles);
-			}
-		}
-
-		return true;
-	}
-
-	return false;
-}
-
-
 UE::DerivedData::FValueId GetDerivedDataValueIdForResource(MutablePrivate::EStreamableDataType StreamableDataType, uint32 FileId, uint16 ResourceType, uint16 Flags)
 {
 	UE::DerivedData::FValueId::ByteArray ValueIdBytes = {};
diff --git a/Engine/Plugins/Mutable/Source/CustomizableObject/Private/MuCO/MutableStreamRequest.cpp b/Engine/Plugins/Mutable/Source/CustomizableObject/Private/MuCO/MutableStreamRequest.cpp
--- a/Engine/Plugins/Mutable/Source/CustomizableObject/Private/MuCO/MutableStreamRequest.cpp
+++ b/Engine/Plugins/Mutable/Source/CustomizableObject/Private/MuCO/MutableStreamRequest.cpp
@@ -41,7 +41,7 @@ void FMutableStreamRequest::AddBlock(const FMutableStreamableBlock& Block, Mutab
 		int32 FileHandleIndex = OpenFilesIds.Find(Block.FileId);
 		if (FileHandleIndex == INDEX_NONE)
 		{
-			const FString& FullFileName = ModelStreamableBulkData->FullFilePath + UCustomizableObjectPrivate::GetDataTypeExtension(DataType);
+			const FString& FullFileName = ModelStreamableBulkData->FullFilePath + GetDataTypeExtension(DataType);
 			IAsyncReadFileHandle* ReadFileHandle = FPlatformFileManager::Get().GetPlatformFile().OpenAsyncRead(*FullFileName);
 			OpenFileHandles.Emplace(ReadFileHandle);
 
diff --git a/Engine/Plugins/Mutable/Source/CustomizableObject/Private/MuCO/UnrealMutableModelDiskStreamer.cpp b/Engine/Plugins/Mutable/Source/CustomizableObject/Private/MuCO/UnrealMutableModelDiskStreamer.cpp
--- a/Engine/Plugins/Mutable/Source/CustomizableObject/Private/MuCO/UnrealMutableModelDiskStreamer.cpp
+++ b/Engine/Plugins/Mutable/Source/CustomizableObject/Private/MuCO/UnrealMutableModelDiskStreamer.cpp
@@ -108,7 +108,7 @@ bool FUnrealMutableModelBulkReader::PrepareStreamingForObject(UCustomizableObjec
 
 #if WITH_EDITOR
 		FString FullFileName = CustomizableObject->GetPrivate()->GetCompiledDataFileName(nullptr, true);
-		NewData.BulkFilePrefix = *FullFileName + UCustomizableObjectPrivate::GetDataTypeExtension(MutablePrivate::EStreamableDataType::Model);
+		NewData.BulkFilePrefix = *FullFileName + GetDataTypeExtension(MutablePrivate::EStreamableDataType::Model);
 #else
 		if (ModelStreamableBulkData->StreamableBulkData.IsEmpty())
 		{
diff --git a/Engine/Plugins/Mutable/Source/CustomizableObjectEditor/Private/MuCOE/CustomizableObjectCompileRunnable.cpp b/Engine/Plugins/Mutable/Source/CustomizableObjectEditor/Private/MuCOE/CustomizableObjectCompileRunnable.cpp
--- a/Engine/Plugins/Mutable/Source/CustomizableObjectEditor/Private/MuCOE/CustomizableObjectCompileRunnable.cpp
+++ b/Engine/Plugins/Mutable/Source/CustomizableObjectEditor/Private/MuCOE/CustomizableObjectCompileRunnable.cpp
@@ -351,7 +351,7 @@ FCustomizableObjectSaveDDRunnable::FCustomizableObjectSaveDDRunnable(const TShar
 	UCustomizableObject* CustomizableObject = InRequest->GetCustomizableObject();
 	CustomizableObjectName = GetNameSafe(CustomizableObject);
 
-	CustomizableObjectHeader.InternalVersion = UCustomizableObjectPrivate::GetECustomizableObjectVersionEnumHash();
+	CustomizableObjectHeader.InternalVersion = GetECustomizableObjectVersionEnumHash();
 	CustomizableObjectHeader.VersionId = CustomizableObject->GetPrivate()->GetVersionId();
 
 	// Cache ModelResources
@@ -565,7 +565,7 @@ void FCustomizableObjectSaveDDRunnable::StoreCachedPlatformDataToDisk()
 
 	// Create folder...
 	IFileManager& FileManager = IFileManager::Get();
-	FileManager.MakeDirectory(*UCustomizableObjectPrivate::GetCompiledDataFolderPath(), true);
+	FileManager.MakeDirectory(*GetCompiledDataFolderPath(), true);
 
 	// Delete files and create new memory writers for each streamable data type.
 	bool bFilesDeleted = true;
@@ -574,7 +574,7 @@ void FCustomizableObjectSaveDDRunnable::StoreCachedPlatformDataToDisk()
 	const int32 NumDataTypes = static_cast<int32>(MutablePrivate::EStreamableDataType::DataTypeCount);
 	for (int32 DataType = 0; DataType < NumDataTypes; ++DataType)
 	{
-		const FString FilePath = FullFileName + UCustomizableObjectPrivate::GetDataTypeExtension(static_cast<MutablePrivate::EStreamableDataType>(DataType));
+		const FString FilePath = FullFileName + GetDataTypeExtension(static_cast<MutablePrivate::EStreamableDataType>(DataType));
 		if (FileManager.FileExists(*FilePath)
 			&& !FileManager.Delete(*FilePath, true, false, true))
 		{
diff --git a/Engine/Plugins/Mutable/Source/CustomizableObjectEditor/Private/MuCOE/CustomizableObjectCompiler.cpp b/Engine/Plugins/Mutable/Source/CustomizableObjectEditor/Private/MuCOE/CustomizableObjectCompiler.cpp
--- a/Engine/Plugins/Mutable/Source/CustomizableObjectEditor/Private/MuCOE/CustomizableObjectCompiler.cpp
+++ b/Engine/Plugins/Mutable/Source/CustomizableObjectEditor/Private/MuCOE/CustomizableObjectCompiler.cpp
@@ -31,6 +31,10 @@
 #include "MuR/ModelPrivate.h"
 #include "MuCOE/CustomizableObjectVersionBridge.h"
 
+#include "DerivedDataCache.h"
+#include "DerivedDataCacheKey.h"
+#include "DerivedDataRequestOwner.h"
+
 class UTexture2D;
 
 #define LOCTEXT_NAMESPACE "CustomizableObjectEditor"
@@ -53,6 +57,26 @@
 		bFinished = false;
 	}
 
+	if (LoadModelDataFromDDCEvent.IsValid())
+	{
+		bFinished = false;
+
+		if (LoadModelDataFromDDCEvent->IsCompleted())
+		{
+			FinishLoadingModelDataFromDDC();
+		}
+	}
+
+	if (LoadStreamableDataFromDDCEvent.IsValid())
+	{
+		bFinished = false;
+
+		if (LoadStreamableDataFromDDCEvent->IsCompleted())
+		{
+			FinishLoadingStreamableDataFromDDC();
+		}
+	}
+
 	if (AsynchronousStreamableHandlePtr && AsynchronousStreamableHandlePtr->IsActive())
 	{
 		bFinished = false;
@@ -105,6 +129,51 @@
 }
 
 
+void FCustomizableObjectCompiler::PreloadReferencerAssets()
+{
+	TRACE_BEGIN_REGION(UE_MUTABLE_PRELOAD_REGION);
+	UE_LOG(LogMutable, Verbose, TEXT("PROFILE: [ %16.8f ] Preload asynchronously assets start."), FPlatformTime::Seconds());
+
+	TArray<FAssetData> ReferencingAssets;
+	GetReferencingPackages(*CurrentObject, ReferencingAssets);
+
+	TArray<FSoftObjectPath> ArrayAssetToStream;
+	for (FAssetData& Element : ReferencingAssets)
+	{
+		ArrayAssetToStream.Add(Element.GetSoftObjectPath());
+	}
+
+	bool bAssetsLoaded = true;
+
+	const bool bAsync = CurrentRequest->bAsync;
+	if (ArrayAssetToStream.Num() > 0)
+	{
+		// Customizations are marked as editoronly on load and are not packaged into the runtime game by default.
+		// The ones that need to be kept will be copied into SoftObjectPath on the object during save.
+
+		UCustomizableObjectSystem* System = UCustomizableObjectSystem::GetInstanceChecked();
+		const TSharedRef<FMutableStreamableManager>& Streamable = System->GetPrivate()->StreamableManager;
+
+		if (bAsync && !CurrentOptions.bIsCooking)
+		{
+			AddCompileNotification(LOCTEXT("LoadingReferencerAssets", "Loading assets"));
+
+			AsynchronousStreamableHandlePtr = Streamable->RequestAsyncLoad(ArrayAssetToStream, FStreamableDelegate::CreateRaw(this, &FCustomizableObjectCompiler::PreloadingReferencerAssetsCallback, bAsync));
+			bAssetsLoaded = false;
+		}
+		else
+		{
+			Streamable->RequestSyncLoad(ArrayAssetToStream);
+		}
+	}
+
+	if (bAssetsLoaded)
+	{
+		PreloadingReferencerAssetsCallback(bAsync);
+	}
+}
+
+
 void FCustomizableObjectCompiler::PreloadingReferencerAssetsCallback(bool bAsync)
 {
 	check(IsInGameThread());
@@ -311,54 +380,10 @@
 	}
 
 	// DDC check
-	if (TryLoadCompiledDataFromDDC(*CurrentObject))
+	if (!TryLoadCompiledDataFromDDC(*CurrentObject))
 	{
-		UE_LOG(LogMutable, Verbose, TEXT("PROFILE: [ %16.8f ] Finishing Compilation task for CO [%s]."), FPlatformTime::Seconds(), *CurrentObject->GetName());
-		TRACE_END_REGION(UE_MUTABLE_COMPILE_REGION);
-
-		UE_LOG(LogMutable, Display, TEXT("Compiled data loaded from DDC"));
-		CompleteRequest(ECompilationStatePrivate::Completed, ECompilationResultPrivate::Success);
-		return;
+		// DDC is disabled, proceed with compilation
+		PreloadReferencerAssets();
 	}
-	
-	TRACE_BEGIN_REGION(UE_MUTABLE_PRELOAD_REGION);
-	UE_LOG(LogMutable, Verbose, TEXT("PROFILE: [ %16.8f ] Preload asynchronously assets start."), FPlatformTime::Seconds());
-
-	TArray<FAssetData> ReferencingAssets;
-	GetReferencingPackages(*CurrentObject, ReferencingAssets);
-
-	TArray<FSoftObjectPath> ArrayAssetToStream;
-	for (FAssetData& Element : ReferencingAssets)
-	{
-		ArrayAssetToStream.Add(Element.GetSoftObjectPath());
-	}
-	
-	bool bAssetsLoaded = true;
-
-	const bool bAsync = InCompileRequest->bAsync;
-	if (ArrayAssetToStream.Num() > 0)
-	{
-		// Customizations are marked as editoronly on load and are not packaged into the runtime game by default.
-		// The ones that need to be kept will be copied into SoftObjectPath on the object during save.
-
-		const TSharedRef<FMutableStreamableManager>& Streamable = System->GetPrivate()->StreamableManager;
-
-		if (bAsync && !CurrentOptions.bIsCooking)
-		{
-			AddCompileNotification(LOCTEXT("LoadingReferencerAssets", "Loading assets"));
-
-			AsynchronousStreamableHandlePtr = Streamable->RequestAsyncLoad(ArrayAssetToStream, FStreamableDelegate::CreateRaw(this, &FCustomizableObjectCompiler::PreloadingReferencerAssetsCallback, bAsync));
-			bAssetsLoaded = false;
-		}
-		else
-		{
-			Streamable->RequestSyncLoad(ArrayAssetToStream);
-		}
-	}
-
-	if (bAssetsLoaded)
-	{
-		PreloadingReferencerAssetsCallback(bAsync);
-	}
 }
 
@@ -914,7 +939,7 @@
 		
 		PlatformData->ModelStreamableBulkData = MakeShared<FModelStreamableBulkData>();
 
-		ModelResources->CodeVersion = UCustomizableObjectPrivate::GetECustomizableObjectVersionEnumHash();
+		ModelResources->CodeVersion = GetECustomizableObjectVersionEnumHash();
 
 		ModelResources->ReferenceSkeletalMeshesData = MoveTemp(GenerationContext.ReferenceSkeletalMeshesData);
 
@@ -1240,6 +1265,15 @@
 			continue;
 		}
 		
+		if (CompileRequest->Options.bIsCooking)
+		{
+			TSharedRef<FCompilationRequest> SelectedCompileRequest = CompileRequest;
+			It.RemoveCurrent(); // Preserve order.
+
+			Compile(SelectedCompileRequest);
+			return true;
+		}
+
 		switch (Object->GetPrivate()->Status.Get())
 		{
 		case FCustomizableObjectStatusTypes::EState::Loading:
@@ -1280,12 +1314,265 @@
 		// Compilation not allowed to query DDC requests. 
 		return false;
 	}
-	
+
+	MUTABLE_CPUPROFILER_SCOPE(TryLoadCompiledDataFromDDC);
+
 	CurrentRequest->BuildDerivedDataCacheKey();
 
 	FCacheKey CacheKey = CurrentRequest->GetDerivedDataCacheKey();
+	check(CacheKey.Hash.IsZero() == false);
 
-	if (CustomizableObject.GetPrivate()->LoadCompiledDataFromDDC(CurrentOptions, DefaultPolicy, &CacheKey, *PlatformData.Get()))
+	using namespace UE::DerivedData;
+
+	DDCHeapMemory = MakeShared<FDDCHeapMemory>();
+
+	/* Overview.
+	*	1. Create an initial pull request to look for the compiled data in the DDC. Skip streamable binary blobs.
+	*	2. Try to load the compiled data.
+	*	3. (Cooking) Create a second request to pull all streamable blobs and cache the compiled data.
+	*/
+
+	// Set the request policy to Default + SkipData to avoid pulling the streamable files until we know the compiled data can be used.
+	FCacheRecordPolicyBuilder PolicyBuilder(DefaultPolicy | ECachePolicy::SkipData);
+
+	// Overwrite the request policy for the resources we want to pull
+	PolicyBuilder.AddValuePolicy(MutablePrivate::GetDerivedDataModelResourcesId(), DefaultPolicy);
+	PolicyBuilder.AddValuePolicy(MutablePrivate::GetDerivedDataModelId(), DefaultPolicy);
+	PolicyBuilder.AddValuePolicy(MutablePrivate::GetDerivedDataModelStreamableBulkDataId(), DefaultPolicy);
+	PolicyBuilder.AddValuePolicy(MutablePrivate::GetDerivedDataBulkDataFilesId(), DefaultPolicy);
+
+	FCacheGetRequest Request;
+	Request.Name = GetPathNameSafe(CurrentObject);
+	Request.Key = CacheKey;
+	Request.Policy = PolicyBuilder.Build();
+
+	LoadModelDataFromDDCEvent = MakeShared<UE::Tasks::FTaskEvent>(TEXT("TryGetBaseDataFromDDC"));
+
+	// Sync request to retrieve the compiled data for validation. Streamable resources are excluded.
+	FRequestOwner RequestOwner(!CurrentRequest->bAsync ? EPriority::Blocking : EPriority::Highest);
+	GetCache().Get(MakeArrayView(&Request, 1), RequestOwner,
+		[DDCHeapMemory = DDCHeapMemory, CompletionEvent = LoadModelDataFromDDCEvent](FCacheGetResponse&& Response) mutable
+		{
+			MUTABLE_CPUPROFILER_SCOPE(RetrieveModelDataFromDDC);
+			
+			if (Response.Status == EStatus::Ok)
+			{
+				const FCompressedBuffer& ModelCompressedBuffer = Response.Record.GetValue(MutablePrivate::GetDerivedDataModelId()).GetData();
+				DDCHeapMemory->ModelBytesDDC = ModelCompressedBuffer.Decompress();
+
+				const FCompressedBuffer& ModelResourcesCompressedBuffer = Response.Record.GetValue(MutablePrivate::GetDerivedDataModelResourcesId()).GetData();
+				DDCHeapMemory->ModelResourcesBytesDDC = ModelResourcesCompressedBuffer.Decompress();
+
+				const FCompressedBuffer& ModelStreamablesCompressedBuffer = Response.Record.GetValue(MutablePrivate::GetDerivedDataModelStreamableBulkDataId()).GetData();
+				DDCHeapMemory->ModelStreamablesBytesDDC = ModelStreamablesCompressedBuffer.Decompress();
+
+				const FCompressedBuffer& BulkDataFilesCompressedBuffer = Response.Record.GetValue(MutablePrivate::GetDerivedDataBulkDataFilesId()).GetData();
+				DDCHeapMemory->BulkDataFilesBytesDDC = BulkDataFilesCompressedBuffer.Decompress();
+			}
+			
+			CompletionEvent->Trigger();
+		});
+
+	if (!CurrentRequest->bAsync)
 	{
+		// Force sync
+		RequestOwner.Wait();
+		FinishLoadingModelDataFromDDC();
+		FinishLoadingStreamableDataFromDDC();
+	}
+	else
+	{
+		RequestOwner.KeepAlive();
+	}
+
+	return true;
+}
+
+
+void FCustomizableObjectCompiler::FinishLoadingModelDataFromDDC()
+{
+	MUTABLE_CPUPROFILER_SCOPE(FinishLoadingModelDataFromDDC);
+
+	LoadModelDataFromDDCEvent.Reset();
+
+	using namespace UE::DerivedData;
+
+	bool bHasValidData = false;
+
+	// Check if it is possible to use the data stored in DDC.
+	if (!DDCHeapMemory->ModelBytesDDC.IsNull() && !DDCHeapMemory->ModelResourcesBytesDDC.IsNull()
+		&& !DDCHeapMemory->ModelStreamablesBytesDDC.IsNull() && !DDCHeapMemory->BulkDataFilesBytesDDC.IsNull())
+	{
+		MUTABLE_CPUPROFILER_SCOPE(LoadModelDataFromDDC);
+
+		// Load the compiled data to validate it.
+		FMemoryReaderView ModelResourcesReader(DDCHeapMemory->ModelResourcesBytesDDC.GetView());
+
+		if (TObjectPtr<UModelResources> LocalModelResources = LoadModelResources_Internal(ModelResourcesReader, CurrentObject, CurrentOptions.TargetPlatform, CurrentOptions.bIsCooking))
+		{
+			PlatformData->ModelResources = TStrongObjectPtr<UModelResources>(LocalModelResources);
+
+			if (CurrentOptions.bIsCooking)
+			{
+				PlatformData->ModelResources->InitCookData(*CurrentObject);
+			}
+
+			FMemoryReaderView ModelStreamablesReader(DDCHeapMemory->ModelStreamablesBytesDDC.GetView());
+			PlatformData->ModelStreamableBulkData = LoadModelStreamableBulk_Internal(ModelStreamablesReader);
+			PlatformData->ModelStreamableBulkData->bIsStoredInDDC = true;
+			PlatformData->ModelStreamableBulkData->DDCKey = CurrentRequest->GetDerivedDataCacheKey();
+			PlatformData->ModelStreamableBulkData->DDCDefaultPolicy = ECachePolicy::Default;
+
+			FMemoryReaderView ModelReader(DDCHeapMemory->ModelBytesDDC.GetView());
+			PlatformData->Model = LoadModel_Internal(ModelReader);
+
+			bHasValidData = true;
+		}
+	}
+
+	LoadStreamableDataFromDDCEvent = MakeShared<UE::Tasks::FTaskEvent>(TEXT("LoadFromDDCCompletionEvent"));
+
+	// Request loading all streamable data when cooking
+	if (bHasValidData && CurrentOptions.bIsCooking)
+	{
+		// Create a new pull request to retrieve all compiled data. Streamable bulk data included
+		FCacheGetRequest Request;
+		Request.Name = GetPathNameSafe(CurrentObject);
+		Request.Key = CurrentRequest->GetDerivedDataCacheKey();
+		Request.Policy = ECachePolicy::Default;
+
+		FRequestOwner RequestOwner(!CurrentRequest->bAsync ? EPriority::Blocking : EPriority::Highest);
+		GetCache().Get(MakeArrayView(&Request, 1), RequestOwner,
+			[DDCHeapMemory = DDCHeapMemory, PlatformData = PlatformData, CompletionEvent = LoadStreamableDataFromDDCEvent, CurrentRequest = CurrentRequest](FCacheGetResponse&& Response) mutable
+			{
+				MUTABLE_CPUPROFILER_SCOPE(GetStreamableDataFromDDC);
+
+				if (ensure(Response.Status == EStatus::Ok))
+				{
+					UCustomizableObject* CustomizableObject = CurrentRequest->GetCustomizableObject();
+					const FCompilationOptions& Options = CurrentRequest->Options;
+
+					// Value Id to file mapping to reconstruct the cached data
+					TMap<FValueId, MutablePrivate::FFile> ValueIdToFile;
+
+					{
+						MUTABLE_CPUPROFILER_SCOPE(BuildValueIdToFile);
+						TArray<MutablePrivate::FFile> BulkDataFiles;
+						FMemoryReaderView FilesReader(DDCHeapMemory->BulkDataFilesBytesDDC.GetView());
+						FilesReader << BulkDataFiles;
+
+						ValueIdToFile.Reserve(BulkDataFiles.Num());
+
+						uint32 FileIndex = 0;
+						for (MutablePrivate::FFile& File : BulkDataFiles)
+						{
+							FValueId ValueId = GetDerivedDataValueIdForResource(File.DataType, FileIndex, File.ResourceType, File.Flags);
+							ValueIdToFile.Add(ValueId, MoveTemp(File));
+
+							++FileIndex;
+						}
+						BulkDataFiles.Empty();
+					}
+
+					// Get all values and convert them to FMutableCachedPlatformData's format
+					TConstArrayView<FValueWithId> Values = Response.Record.GetValues();
+
+					TArray64<uint8> TempData;
+					for (const FValueWithId& Value : Values)
+					{
+						check(Value.IsValid());
+
+						const MutablePrivate::FFile* File = ValueIdToFile.Find(Value.GetId());
+						if (!File) // Skip value. It is not a streamable binary blob.
+						{
+							continue;
+						}
+
+						const uint64 RawSize = Value.GetRawSize();
+						TempData.SetNumUninitialized(RawSize, EAllowShrinking::No);
+
+						// Decompress streamable binary blobs
+						const bool bDecompressedSuccessfully = Value.GetData().TryDecompressTo(MakeMemoryView(TempData.GetData(), RawSize));
+						check(bDecompressedSuccessfully);
+
+						// Filter and cache the data by DataType
+						switch (File->DataType)
+						{
+						case MutablePrivate::EStreamableDataType::Model:
+						{
+							for (const MutablePrivate::FBlock& Block : File->Blocks)
+							{
+								PlatformData->ModelStreamableData.Set(Block.Id, TempData.GetData() + Block.Offset, Block.Size);
+							}
+							break;
+						}
+						case MutablePrivate::EStreamableDataType::RealTimeMorph:
+						{
+							for (const MutablePrivate::FBlock& Block : File->Blocks)
+							{
+								PlatformData->MorphStreamableData.Set(Block.Id, TempData.GetData() + Block.Offset, Block.Size);
+							}
+							break;
+						}
+						case MutablePrivate::EStreamableDataType::Clothing:
+						{
+							for (const MutablePrivate::FBlock& Block : File->Blocks)
+							{
+								PlatformData->ClothingStreamableData.Set(Block.Id, TempData.GetData() + Block.Offset, Block.Size);
+							}
+							break;
+						}
+
+						default:
+							unimplemented();
+							break;
+						}
+					}
+
+					// Generate list of files and update streamable blocks ids and offsets
+					if (CVarMutableUseBulkData.GetValueOnAnyThread())
+					{
+						const uint32 NumBulkDataFilesPerBucket = MAX_uint8;
+						MutablePrivate::GenerateBulkDataFilesListWithFileLimit(PlatformData->Model, *PlatformData->ModelStreamableBulkData.Get(),
+							NumBulkDataFilesPerBucket, PlatformData->BulkDataFiles);
+					}
+					else
+					{
+						MutablePrivate::GenerateBulkDataFilesListWithSizeLimit(PlatformData->Model, *PlatformData->ModelStreamableBulkData.Get(),
+							Options.TargetPlatform, Options.PackagedDataBytesLimit, PlatformData->BulkDataFiles);
+					}
+				}
+				else
+				{
+					PlatformData->Model.Reset();
+					PlatformData->ModelResources.Reset();
+					PlatformData->ModelStreamableBulkData.Reset();
+				}
+
+				CompletionEvent->Trigger();
+			});
+
+		if (!CurrentRequest->bAsync)
+		{
+			// Force sync
+			RequestOwner.Wait();
+		}
+		else
+		{
+			RequestOwner.KeepAlive();
+		}
+	}
+	else
+	{
+		LoadStreamableDataFromDDCEvent->Trigger();
+	}
+}
+
+void FCustomizableObjectCompiler::FinishLoadingStreamableDataFromDDC()
+{
+	LoadStreamableDataFromDDCEvent.Reset();
+
+	if (PlatformData->Model && PlatformData->ModelResources && PlatformData->ModelStreamableBulkData)
+	{
 		if (CurrentOptions.bIsCooking)
 		{
@@ -1293,10 +1580,16 @@
 			CachedPlatformData = MoveTemp(*PlatformData.Get());
 		}
 
-		return true;
+		UE_LOG(LogMutable, Verbose, TEXT("PROFILE: [ %16.8f ] Finishing Compilation task for CO [%s]."), FPlatformTime::Seconds(), *CurrentObject->GetName());
+		TRACE_END_REGION(UE_MUTABLE_COMPILE_REGION);
+
+		UE_LOG(LogMutable, Display, TEXT("Compiled data loaded from DDC"));
+		CompleteRequest(ECompilationStatePrivate::Completed, ECompilationResultPrivate::Success);
 	}
-
-	return false;
+	else
+	{
+		PreloadReferencerAssets();
+	}
 }
 
 
@@ -1403,11 +1696,7 @@
 	{
 		TObjectPtr<UModelResources> ModelResources = PlatformData->ModelResources.Get();
 
-		ModelResources->Skeletons.Reserve(CompilationContext->ReferencedSkeletons.Num());
-		for (const USkeleton* Skeleton : CompilationContext->ReferencedSkeletons)
-		{
-			ModelResources->Skeletons.Emplace(const_cast<USkeleton*>(Skeleton));
-		}
+		ModelResources->Skeletons = MoveTemp(CompilationContext->ReferencedSkeletons);
 
 		ModelResources->PhysicsAssets = MoveTemp(CompilationContext->PhysicsAssets);
 		ModelResources->AnimBpOverridePhysiscAssetsInfo = MoveTemp(CompilationContext->AnimBpOverridePhysicsAssetsInfo);
diff --git a/Engine/Plugins/Mutable/Source/CustomizableObjectEditor/Private/MuCOE/CustomizableObjectCompiler.h b/Engine/Plugins/Mutable/Source/CustomizableObjectEditor/Private/MuCOE/CustomizableObjectCompiler.h
--- a/Engine/Plugins/Mutable/Source/CustomizableObjectEditor/Private/MuCOE/CustomizableObjectCompiler.h
+++ b/Engine/Plugins/Mutable/Source/CustomizableObjectEditor/Private/MuCOE/CustomizableObjectCompiler.h
@@ -63,6 +63,9 @@ public:
 	void CompilerLog(const FText& Message, const UObject* Context = nullptr, const EMessageSeverity::Type MessageSeverity = EMessageSeverity::Warning, const bool bAddBaseObjectInfo = true, const ELoggerSpamBin SpamBin = ELoggerSpamBin::ShowAll);
 	void NotifyCompilationErrors() const;
 
+	void FinishLoadingModelDataFromDDC();
+	void FinishLoadingStreamableDataFromDDC();
+
 	void FinishCompilationTask();
 	void FinishSavingDerivedDataTask();
 
@@ -111,8 +114,10 @@ private:
 	void CompleteRequest(ECompilationStatePrivate State, ECompilationResultPrivate Result);
 	bool TryPopCompileRequest();
 
+	/** Attempts to load the compiled data from DDC. */
 	bool TryLoadCompiledDataFromDDC(UCustomizableObject& CustomizableObject);
 
+	void PreloadReferencerAssets();
 	void PreloadingReferencerAssetsCallback(bool bAsync);
 	
 	// Will output to Mutable Log the warning and error messages generated during the CO compilation
@@ -128,7 +133,20 @@ private:
 	ECompilationResultPrivate GetCompilationResult() const;
 
 	void ProcessCompileTasks();
-	
+
+	struct FDDCHeapMemory
+	{
+		FSharedBuffer ModelBytesDDC;
+		FSharedBuffer ModelResourcesBytesDDC;
+		FSharedBuffer ModelStreamablesBytesDDC;
+		FSharedBuffer BulkDataFilesBytesDDC;
+	};
+	TSharedPtr<FDDCHeapMemory> DDCHeapMemory;
+
+	/** Load compiled data from DDC completion events*/
+	TSharedPtr<UE::Tasks::FTaskEvent> LoadModelDataFromDDCEvent;
+	TSharedPtr<UE::Tasks::FTaskEvent> LoadStreamableDataFromDDCEvent;
+
 	/** Pointer to the Asynchronous Preloading process call back */
 	TSharedPtr<FStreamableHandle> AsynchronousStreamableHandlePtr;
 
diff --git a/Engine/Plugins/Mutable/Source/CustomizableObjectEditor/Private/MuCOE/CustomizableObjectEditorModule.cpp b/Engine/Plugins/Mutable/Source/CustomizableObjectEditor/Private/MuCOE/CustomizableObjectEditorModule.cpp
--- a/Engine/Plugins/Mutable/Source/CustomizableObjectEditor/Private/MuCOE/CustomizableObjectEditorModule.cpp
+++ b/Engine/Plugins/Mutable/Source/CustomizableObjectEditor/Private/MuCOE/CustomizableObjectEditorModule.cpp
@@ -151,7 +151,7 @@ void DeleteUnusedMutableStreamedDataEditorFiles()
 {
 	const FDateTime CurrentTime = FDateTime::Now();
 
-	const FString CompiledDataFolder = UCustomizableObjectPrivate::GetCompiledDataFolderPath();
+	const FString CompiledDataFolder = GetCompiledDataFolderPath();
 	const FString FileExtension = TEXT(".mut");
 
 	TArray<FString> Files;
diff --git a/Engine/Plugins/Mutable/Source/CustomizableObjectEditor/Private/MuCOE/GenerateMutableSource/GenerateMutableSource.cpp b/Engine/Plugins/Mutable/Source/CustomizableObjectEditor/Private/MuCOE/GenerateMutableSource/GenerateMutableSource.cpp
--- a/Engine/Plugins/Mutable/Source/CustomizableObjectEditor/Private/MuCOE/GenerateMutableSource/GenerateMutableSource.cpp
+++ b/Engine/Plugins/Mutable/Source/CustomizableObjectEditor/Private/MuCOE/GenerateMutableSource/GenerateMutableSource.cpp
@@ -1386,7 +1386,7 @@ void PopulateReferenceSkeletalMeshesData(FMutableGraphGenerationContext& Generat
 		const FMutableComponentInfo& ComponentInfo = GenerationContext.CompilationContext.ComponentInfos[ComponentIndex];
 		const FName& ComponentName = ComponentInfo.ComponentName;
 		
-		USkeletalMesh* RefSkeletalMesh = ComponentInfo.RefSkeletalMesh;
+		USkeletalMesh* RefSkeletalMesh = ComponentInfo.RefSkeletalMesh.Get();
 		check(RefSkeletalMesh);
 
 		FMutableRefSkeletalMeshData& Data = GenerationContext.ReferenceSkeletalMeshesData[ComponentIndex];
@@ -1814,7 +1814,7 @@ FMutableComponentInfo::FMutableComponentInfo(FName InComponentName, USkeletalMes
 		return;
 	}
 
-	RefSkeletalMesh = InRefSkeletalMesh;
+	RefSkeletalMesh = TStrongObjectPtr<USkeletalMesh>(InRefSkeletalMesh);
 	RefSkeleton = RefSkeletalMesh->GetSkeleton();
 
 	const int32 NumBones = RefSkeleton->GetReferenceSkeleton().GetRawBoneNum();
diff --git a/Engine/Plugins/Mutable/Source/CustomizableObjectEditor/Private/MuCOE/GenerateMutableSource/GenerateMutableSource.h b/Engine/Plugins/Mutable/Source/CustomizableObjectEditor/Private/MuCOE/GenerateMutableSource/GenerateMutableSource.h
--- a/Engine/Plugins/Mutable/Source/CustomizableObjectEditor/Private/MuCOE/GenerateMutableSource/GenerateMutableSource.h
+++ b/Engine/Plugins/Mutable/Source/CustomizableObjectEditor/Private/MuCOE/GenerateMutableSource/GenerateMutableSource.h
@@ -30,6 +30,7 @@
 #include "MuR/Mesh.h"
 #include "MuR/Ptr.h"
 #include "UObject/Package.h"
+#include "UObject/StrongObjectPtr.h"
 
 class UCustomizableObjectNodeGroupProjectorParameter;
 class UCustomizableObjectNodeMaterialBase;
@@ -359,13 +360,13 @@ struct FMutableComponentInfo
 	FMutableLODSettings LODSettings;
 	
 	// Each component must have a reference SkeletalMesh with a valid Skeleton
-	USkeletalMesh* RefSkeletalMesh = nullptr;
-	USkeleton* RefSkeleton = nullptr;
+	TStrongObjectPtr<USkeletalMesh> RefSkeletalMesh;
+	TObjectPtr<USkeleton> RefSkeleton;
 	
 	UCustomizableObjectNodeComponentMesh* NodeComponentMesh = nullptr;
 	
 	// Map to check skeleton compatibility
-	TMap<const USkeleton*, bool> SkeletonCompatibility;
+	TMap<SIZE_T, bool> SkeletonCompatibility;
 	
 	// Hierarchy hash from parent-bone to root bone, used to check if additional skeletons are compatible with
 	// the RefSkeleton
@@ -459,7 +460,7 @@ struct FMutableCompilationContext
 	FCompilationOptions Options;
 
 	/** */
-	TArray<const USkeleton*> ReferencedSkeletons;
+	TArray<TSoftObjectPtr<USkeleton>> ReferencedSkeletons;
 
 	/** Global morph selection overrides. */
 	TArray<FRealTimeMorphSelectionOverride> RealTimeMorphTargetsOverrides;
diff --git a/Engine/Plugins/Mutable/Source/CustomizableObjectEditor/Private/MuCOE/GenerateMutableSource/GenerateMutableSourceComponent.cpp b/Engine/Plugins/Mutable/Source/CustomizableObjectEditor/Private/MuCOE/GenerateMutableSource/GenerateMutableSourceComponent.cpp
--- a/Engine/Plugins/Mutable/Source/CustomizableObjectEditor/Private/MuCOE/GenerateMutableSource/GenerateMutableSourceComponent.cpp
+++ b/Engine/Plugins/Mutable/Source/CustomizableObjectEditor/Private/MuCOE/GenerateMutableSource/GenerateMutableSourceComponent.cpp
@@ -563,7 +563,7 @@ void FirstPass(UCustomizableObjectNodeComponentMesh& Node, FMutableGraphGenerati
 	{
 		for (int32 MeshIndex = 0; MeshIndex < GenerationContext.CompilationContext.ComponentInfos.Num(); ++MeshIndex)
 		{
-			RefSkeletalMesh = GenerationContext.CompilationContext.ComponentInfos[MeshIndex].RefSkeletalMesh;
+			RefSkeletalMesh = GenerationContext.CompilationContext.ComponentInfos[MeshIndex].RefSkeletalMesh.Get();
 			check(RefSkeletalMesh);
 
 			GenerationContext.bEnableLODStreaming = GenerationContext.bEnableLODStreaming &&
diff --git a/Engine/Plugins/Mutable/Source/CustomizableObjectEditor/Private/MuCOE/GenerateMutableSource/GenerateMutableSourceGroupProjector.cpp b/Engine/Plugins/Mutable/Source/CustomizableObjectEditor/Private/MuCOE/GenerateMutableSource/GenerateMutableSourceGroupProjector.cpp
--- a/Engine/Plugins/Mutable/Source/CustomizableObjectEditor/Private/MuCOE/GenerateMutableSource/GenerateMutableSourceGroupProjector.cpp
+++ b/Engine/Plugins/Mutable/Source/CustomizableObjectEditor/Private/MuCOE/GenerateMutableSource/GenerateMutableSourceGroupProjector.cpp
@@ -445,7 +445,7 @@ TOptional<FGroupProjectorTempData> GenerateMutableGroupProjector(UCustomizableOb
 
 			check(GroupProjectorTempData.PoseBoneDataArray.Num() == PoseIndex);
 			GroupProjectorTempData.PoseBoneDataArray.AddDefaulted(1);
-			UCustomizableObjectNodeAnimationPose::StaticRetrievePoseInformation(PoseAsset, CurrentComponentInfo->RefSkeletalMesh,
+			UCustomizableObjectNodeAnimationPose::StaticRetrievePoseInformation(PoseAsset, CurrentComponentInfo->RefSkeletalMesh.Get(),
 				GroupProjectorTempData.PoseBoneDataArray[PoseIndex].ArrayBoneName, GroupProjectorTempData.PoseBoneDataArray[PoseIndex].ArrayTransform);
 		}
 	
diff --git a/Engine/Plugins/Mutable/Source/CustomizableObjectEditor/Private/MuCOE/GenerateMutableSource/GenerateMutableSourceMesh.cpp b/Engine/Plugins/Mutable/Source/CustomizableObjectEditor/Private/MuCOE/GenerateMutableSource/GenerateMutableSourceMesh.cpp
--- a/Engine/Plugins/Mutable/Source/CustomizableObjectEditor/Private/MuCOE/GenerateMutableSource/GenerateMutableSourceMesh.cpp
+++ b/Engine/Plugins/Mutable/Source/CustomizableObjectEditor/Private/MuCOE/GenerateMutableSource/GenerateMutableSourceMesh.cpp
@@ -493,7 +493,7 @@ bool IsSkeletalMeshCompatibleWithRefSkeleton(FMutableComponentInfo& ComponentInf
 		return true;
 	}
 
-	if (bool* SkeletonCompatibility = ComponentInfo.SkeletonCompatibility.Find(Skeleton))
+	if (bool* SkeletonCompatibility = ComponentInfo.SkeletonCompatibility.Find(reinterpret_cast<SIZE_T>(Skeleton.Get())))
 	{
 		return *SkeletonCompatibility;
 	}
@@ -949,7 +949,7 @@ TSharedPtr<mu::FMesh> ConvertSkeletalMeshToMutable(
 		FMutableComponentInfo* MutComponentInfo = Context.GetComponentInfo(Source.Component);
 		check(MutComponentInfo);
 
-		USkeletalMesh* ComponentRefSkeletalMesh = MutComponentInfo->RefSkeletalMesh;
+		USkeletalMesh* ComponentRefSkeletalMesh = MutComponentInfo->RefSkeletalMesh.Get();
 		USkeleton* ComponentRefSkeleton = MutComponentInfo->RefSkeleton;
 		check(ComponentRefSkeletalMesh);
 		check(ComponentRefSkeleton);
@@ -958,7 +958,7 @@ TSharedPtr<mu::FMesh> ConvertSkeletalMeshToMutable(
 		{
 			FString ErrorMessage;
 			const bool bCompatible = IsSkeletalMeshCompatibleWithRefSkeleton(*MutComponentInfo, SkeletalMesh, ErrorMessage);
-			MutComponentInfo->SkeletonCompatibility.Add(Skeleton, bCompatible);
+			MutComponentInfo->SkeletonCompatibility.Add(reinterpret_cast<SIZE_T>(Skeleton), bCompatible);
 
 			if (!bCompatible)
 			{
@@ -974,7 +974,7 @@ TSharedPtr<mu::FMesh> ConvertSkeletalMeshToMutable(
 			MutableMesh->AddSkeletonID(RefSkeletonID);
 
 			// Add the skeleton to the list of referenced skeletons and add its index to the mesh
-			const int32 SkeletonID = Context.ReferencedSkeletons.AddUnique(Skeleton);
+			const int32 SkeletonID = Context.ReferencedSkeletons.AddUnique(const_cast<USkeleton*>(Skeleton));
 			MutableMesh->AddSkeletonID(SkeletonID);
 		}
 
@@ -4332,7 +4332,7 @@ mu::Ptr<mu::NodeMesh> GenerateMutableSourceMesh(const UEdGraphPin* Pin,
 				{
 					TArray<FName> ArrayBoneName;
 					TArray<FTransform> ArrayTransform;
-					UCustomizableObjectNodeAnimationPose::StaticRetrievePoseInformation(TypedNode->PoseAsset, GenerationContext.GetCurrentComponentInfo()->RefSkeletalMesh, ArrayBoneName, ArrayTransform);
+					UCustomizableObjectNodeAnimationPose::StaticRetrievePoseInformation(TypedNode->PoseAsset, GenerationContext.GetCurrentComponentInfo()->RefSkeletalMesh.Get(), ArrayBoneName, ArrayTransform);
 					mu::Ptr<mu::NodeMeshApplyPose> NodeMeshApplyPose = CreateNodeMeshApplyPose(GenerationContext, InputMeshNode, ArrayBoneName, ArrayTransform);
 
 					if (NodeMeshApplyPose)
diff --git a/Engine/Plugins/Mutable/Source/CustomizableObjectEditor/Private/MuCOE/GenerateMutableSource/GenerateMutableSourceTable.cpp b/Engine/Plugins/Mutable/Source/CustomizableObjectEditor/Private/MuCOE/GenerateMutableSource/GenerateMutableSourceTable.cpp
--- a/Engine/Plugins/Mutable/Source/CustomizableObjectEditor/Private/MuCOE/GenerateMutableSource/GenerateMutableSourceTable.cpp
+++ b/Engine/Plugins/Mutable/Source/CustomizableObjectEditor/Private/MuCOE/GenerateMutableSource/GenerateMutableSourceTable.cpp
@@ -568,7 +568,7 @@ bool FillTableColumn(const UCustomizableObjectNodeTable* TableNode, mu::Ptr<mu::
 
 				TArray<FName> ArrayBoneName;
 				TArray<FTransform> ArrayTransform;
-				UCustomizableObjectNodeAnimationPose::StaticRetrievePoseInformation(PoseAsset, GenerationContext.GetCurrentComponentInfo()->RefSkeletalMesh, ArrayBoneName, ArrayTransform);
+				UCustomizableObjectNodeAnimationPose::StaticRetrievePoseInformation(PoseAsset, GenerationContext.GetCurrentComponentInfo()->RefSkeletalMesh.Get(), ArrayBoneName, ArrayTransform);
 
 				TSharedPtr<mu::FMesh> MutableMesh = MakeShared<mu::FMesh>();
 				TSharedPtr<mu::FSkeleton> MutableSkeleton = MakeShared<mu::FSkeleton>();
diff --git a/Engine/Plugins/Mutable/Source/MutableValidation/Private/MuV/CustomizableObjectCompilationUtility.cpp b/Engine/Plugins/Mutable/Source/MutableValidation/Private/MuV/CustomizableObjectCompilationUtility.cpp
--- a/Engine/Plugins/Mutable/Source/MutableValidation/Private/MuV/CustomizableObjectCompilationUtility.cpp
+++ b/Engine/Plugins/Mutable/Source/MutableValidation/Private/MuV/CustomizableObjectCompilationUtility.cpp
@@ -62,7 +62,7 @@ bool FCustomizableObjectCompilationUtility::CompileCustomizableObject(UCustomiza
 
 	// Delete data from previous compilations just to be sure it does not affect the new compilation operation
 	IFileManager& FileManager = IFileManager::Get();
-	const FString CompiledDataDirectory = CustomizableObject->GetPrivate()->GetCompiledDataFolderPath();
+	const FString CompiledDataDirectory = GetCompiledDataFolderPath();
 	if (CustomizableObject->IsCompiled() && FileManager.DirectoryExists(*CompiledDataDirectory))
 	{
 		UE_LOG(LogMutable,Display,TEXT("Deleting old mutable compiled data from '%s'..."), *CompiledDataDirectory);
