[Mutable] Fix typo in the number of parallel batches executed by ImageInvert operations.

#rnx

diff --git a/Engine/Plugins/Mutable/Source/MutableRuntime/Private/MuR/OpImageInvert.h b/Engine/Plugins/Mutable/Source/MutableRuntime/Private/MuR/OpImageInvert.h
--- a/Engine/Plugins/Mutable/Source/MutableRuntime/Private/MuR/OpImageInvert.h
+++ b/Engine/Plugins/Mutable/Source/MutableRuntime/Private/MuR/OpImageInvert.h
@@ -49,7 +49,7 @@ namespace OpImageInvertInternal
 		const int32 BytesPerElem = GetImageFormatData(AImage->GetFormat()).BytesPerBlock;
 		const int32 NumBatches = AImage->DataStorage.GetNumBatches(NumBatchElems, BytesPerElem);
 
-		OpImageInvertInternal::InvokeParallelFor(NumBatchElems,
+		OpImageInvertInternal::InvokeParallelFor(NumBatches,
 		[
 			AImage, NumBatchElems, BytesPerElem
 		](int32 BatchId)
