[Mutable] Fixed the following macro bugs:
+ Now macros with object group nodes and mesh section nodes can be reused.
+ Added macro support to morph stack nodes, morph factor nodes, material asset pins and posse asset pins of table nodes.
+ Fixed tag's widget where the interal tag of the mesh section node was not visible.

#rnx

diff --git a/Engine/Plugins/Mutable/Source/CustomizableObjectEditor/Private/MuCOE/GenerateMutableSource/GenerateMutableSource.cpp b/Engine/Plugins/Mutable/Source/CustomizableObjectEditor/Private/MuCOE/GenerateMutableSource/GenerateMutableSource.cpp
--- a/Engine/Plugins/Mutable/Source/CustomizableObjectEditor/Private/MuCOE/GenerateMutableSource/GenerateMutableSource.cpp
+++ b/Engine/Plugins/Mutable/Source/CustomizableObjectEditor/Private/MuCOE/GenerateMutableSource/GenerateMutableSource.cpp
@@ -329,18 +329,7 @@ int32 FMutableGraphGenerationContext::AddStreamedResource(UAssetUserData& AssetU
 
 const FGuid FMutableGraphGenerationContext::GetNodeIdUnique(const UCustomizableObjectNode* Node)
 {
-	FGuid NodeID = Node->NodeGuid;
-
-	if (Node->IsInMacro())
-	{
-		check(MacroNodesStack.Num()); // To ensure that we only enter here when compiling
-		
-		for (int32 MacroIndex = 0; MacroIndex < MacroNodesStack.Num(); ++MacroIndex)
-		{
-			NodeID = FGuid::Combine(NodeID, MacroNodesStack[MacroIndex]->NodeGuid);
-		}
-	}
-
+	FGuid NodeID = GetNodeIdUnchecked(Node);
 	TArray<const UObject*>* ArrayResult = NodeIdsMap.Find(NodeID);
 
 	if (ArrayResult == nullptr)
@@ -356,6 +345,24 @@ const FGuid FMutableGraphGenerationContext::GetNodeIdUnique(const UCustomizableO
 }
 
 
+const FGuid FMutableGraphGenerationContext::GetNodeIdUnchecked(const UCustomizableObjectNode* Node)
+{
+	FGuid NodeID = Node->NodeGuid;
+
+	if (Node->IsInMacro())
+	{
+		check(MacroNodesStack.Num()); // To ensure that we only enter here when compiling
+
+		for (int32 MacroIndex = 0; MacroIndex < MacroNodesStack.Num(); ++MacroIndex)
+		{
+			NodeID = FGuid::Combine(NodeID, MacroNodesStack[MacroIndex]->NodeGuid);
+		}
+	}
+
+	return NodeID;
+}
+
+
 void FMutableGraphGenerationContext::GenerateSharedSurfacesUniqueIds()
 {
 	int32 UniqueId = 0;
@@ -393,6 +400,11 @@ void FMutableGraphGenerationContext::GenerateSharedSurfacesUniqueIds()
 					continue;
 				}
 
+				if (SharedSurfaces[AuxSurfaceIndex].MacroContext != CurrentSharedSurface.MacroContext)
+				{
+					continue;
+				}
+
 				SharedSurfaces[AuxSurfaceIndex].NodeSurfaceNew->SharedSurfaceId = UniqueId;
 				VisitedSurfaces[AuxSurfaceIndex] = true;
 			}
@@ -683,6 +695,39 @@ UTexture2D* FindReferenceImage(const UEdGraphPin* Pin, FMutableGraphGenerationCo
 			Result = TypedNodeTable->GetColumnDefaultAssetByType<UTexture2D>(Pin);
 		}
 	}
+
+	else if (const UCustomizableObjectNodeMacroInstance* TypedNodeMacro = Cast<UCustomizableObjectNodeMacroInstance>(Node))
+	{
+		if (const UEdGraphPin* OutputPin = TypedNodeMacro->GetMacroIOPin(ECOMacroIOType::COMVT_Output, Pin->PinName))
+		{
+			if (const UEdGraphPin* ConnectedPin = FollowInputPin(*OutputPin))
+			{
+				GenerationContext.MacroNodesStack.Push(TypedNodeMacro);
+				Result = FindReferenceImage(ConnectedPin, GenerationContext);
+				GenerationContext.MacroNodesStack.Pop();
+			}
+		}
+	}
+
+	else if (const UCustomizableObjectNodeTunnel* TypedNodeTunnel = Cast<UCustomizableObjectNodeTunnel>(Node))
+	{
+		check(TypedNodeTunnel->bIsInputNode);
+		check(GenerationContext.MacroNodesStack.Num());
+
+		const UCustomizableObjectNodeMacroInstance* MacroInstanceNode = GenerationContext.MacroNodesStack.Pop();
+		check(MacroInstanceNode);
+
+		if (const UEdGraphPin* InputPin = MacroInstanceNode->FindPin(Pin->PinName, EEdGraphPinDirection::EGPD_Input))
+		{
+			if (const UEdGraphPin* FollowPin = FollowInputPin(*InputPin))
+			{
+				Result = FindReferenceImage(FollowPin, GenerationContext);
+			}
+		}
+
+		// Push the Macro again even if the result is null
+		GenerationContext.MacroNodesStack.Push(MacroInstanceNode);
+	}
 	
 	return Result;
 }
@@ -830,6 +875,7 @@ uint32 AddUniqueSurfaceMetadata(const FMutableSourceSurfaceMetadata& SourceMetad
 	return SurfaceMetadataUniqueHash;
 }
 
+
 // Convert a CustomizableObject Source Graph into a mutable source graph  
 mu::Ptr<mu::NodeObject> GenerateMutableSource(const UEdGraphPin * Pin, FMutableGraphGenerationContext & GenerationContext)
 {
@@ -958,7 +1004,7 @@ mu::Ptr<mu::NodeObject> GenerateMutableSource(const UEdGraphPin * Pin, FMutableG
 
 				if (const UCustomizableObjectNodeMacroInstance* TypedNodeMacro = Cast<UCustomizableObjectNodeMacroInstance>(ChildNodePin->GetOwningNode()))
 				{
-					if (const UEdGraphPin* OutputPin = TypedNodeMacro->GetMacroIOPin(ECOMacroIOType::COMVT_Output, Pin->PinName))
+					if (const UEdGraphPin* OutputPin = TypedNodeMacro->GetMacroIOPin(ECOMacroIOType::COMVT_Output, ChildNodePin->PinName))
 					{
 						// Get the modifier pin inside the macro
 						if (UEdGraphPin* FollowPin = FollowInputPin(*OutputPin))
@@ -1049,18 +1095,26 @@ mu::Ptr<mu::NodeObject> GenerateMutableSource(const UEdGraphPin * Pin, FMutableG
 		// All sockets from all mesh parts plugged into this group node will have the following priority when there's a socket name clash
 		GenerationContext.SocketPriorityStack.Push(TypedNodeGroup->SocketPriority);
 
+		// Getting node Id. We may need a different ID if the group node is inside a Macro and this Macro is used more than once.
+		// Group nodes of Macros Can Not have external children so we can change their ids here safetely.
+		FGuid GroupNodeId = GenerationContext.GetNodeIdUnchecked(TypedNodeGroup);
+
 		GroupNode->SetMessageContext(Node);
 		GroupNode->SetName(TypedNodeGroup->GetGroupName(&GenerationContext.MacroNodesStack));
-		GroupNode->SetUid(TypedNodeGroup->NodeGuid.ToString());
+		GroupNode->SetUid(GroupNodeId.ToString());
 		
 		TArray<UCustomizableObjectNodeGroupProjectorParameter*> GroupProjectors;
 		if (UEdGraphPin* ProjectorsPin = TypedNodeGroup->GroupProjectorsPin())
 		{
 			for (const UEdGraphPin* ConnectedPin : FollowInputPinArray(*ProjectorsPin))
 			{
-				if (UCustomizableObjectNodeGroupProjectorParameter* GroupProjectorNode = Cast<UCustomizableObjectNodeGroupProjectorParameter>(ConnectedPin->GetOwningNode()))
+				// Checking if it's linked to a Macro or tunnel node
+				if (const UEdGraphPin* GroupProjectorPin = GraphTraversal::FindIOPinSourceThroughMacroContext(*ConnectedPin, &GenerationContext.MacroNodesStack))
 				{
-					GroupProjectors.Add(GroupProjectorNode);
+					if (UCustomizableObjectNodeGroupProjectorParameter* GroupProjectorNode = Cast<UCustomizableObjectNodeGroupProjectorParameter>(GroupProjectorPin->GetOwningNode()))
+					{
+						GroupProjectors.Add(GroupProjectorNode);
+					}
 				}
 			}
 		}
@@ -1086,13 +1140,25 @@ mu::Ptr<mu::NodeObject> GenerateMutableSource(const UEdGraphPin * Pin, FMutableG
 
 		// External children
 		TArray<UCustomizableObjectNodeObject*> ExternalChildNodes;
-		GenerationContext.GroupIdToExternalNodeMap.MultiFind(TypedNodeGroup->NodeGuid, ExternalChildNodes);
+
+		GenerationContext.GroupIdToExternalNodeMap.MultiFind(GroupNodeId, ExternalChildNodes);
 		FMutableGraphGenerationContext::FParamInfo ParamInfo(TypedNodeGroup->GetGroupName(&GenerationContext.MacroNodesStack),
 															 TypedNodeGroup->GroupType == ECustomizableObjectGroupType::COGT_TOGGLE);
-		GenerationContext.GuidToParamNameMap.Add(TypedNodeGroup->NodeGuid, ParamInfo);
+		GenerationContext.GuidToParamNameMap.Add(GroupNodeId,ParamInfo);
 
 		// Children
-		const TArray<UEdGraphPin*> ConnectedChildrenPins = FollowInputPinArray(*TypedNodeGroup->ObjectsPin());
+		TArray<const UEdGraphPin*> ConnectedChildrenPins;
+
+		for (const UEdGraphPin* LinkedPin : FollowInputPinArray(*TypedNodeGroup->ObjectsPin()))
+		{
+			// Checking if it's linked to a macro node and the macro node is linked to something
+			if (const UEdGraphPin* MacroCheckPin = GraphTraversal::FindIOPinSourceThroughMacroContext(*LinkedPin, &GenerationContext.MacroNodesStack))
+			{
+				// We use the original pins to generate the children objects. GenerateMutableSourceX functions already go through macro nodes.
+				ConnectedChildrenPins.Add(LinkedPin);
+			}
+		}
+
 		const int32 NumChildren = ConnectedChildrenPins.Num();
 		const int32 TotalNumChildren = NumChildren + ExternalChildNodes.Num();
 
@@ -1115,7 +1181,11 @@ mu::Ptr<mu::NodeObject> GenerateMutableSource(const UEdGraphPin * Pin, FMutableG
 			bool bLastChildNode = (ChildIndex == NumChildren - 1) && (ExternalChildNodes.Num() == 0);
 			bool bConnectAtLeastTheLastChild = bLastChildNode && !bAtLeastOneConnected;
 
-			UCustomizableObjectNodeObject* CustomizableObjectNodeObject = Cast<UCustomizableObjectNodeObject>(ConnectedChildrenPins[ChildIndex]->GetOwningNode());
+			// Check if this node is linked to a MacroNode
+			const UEdGraphPin* MacroContextPin = GraphTraversal::FindIOPinSourceThroughMacroContext(*ConnectedChildrenPins[ChildIndex], &GenerationContext.MacroNodesStack);
+			check(MacroContextPin);
+
+			UCustomizableObjectNodeObject* CustomizableObjectNodeObject = Cast<UCustomizableObjectNodeObject>(MacroContextPin->GetOwningNode());
 
 			const FString* SelectedOptionName = GenerationContext.CompilationContext.Options.ParamNamesToSelectedOptions.Find(TypedNodeGroup->GetGroupName(&GenerationContext.MacroNodesStack)); // If the param is in the map restrict to only the selected option
 			mu::Ptr<mu::NodeObject> ChildNode;
@@ -1168,56 +1238,65 @@ mu::Ptr<mu::NodeObject> GenerateMutableSource(const UEdGraphPin * Pin, FMutableG
 		{
 			if (const UEdGraphPin* ConnectedPin = FollowOutputPin(*Pin))
 			{
-				const FGuid* ParentId = nullptr;
-				FString ValueName;
+				// Checking if it's linked to a Macro or tunnel node
+				const UEdGraphPin* NodeLinkedPin = GraphTraversal::FindIOPinSourceThroughMacroContext(*ConnectedPin, &GenerationContext.MacroNodesStack);
 
-				if (UCustomizableObjectNodeObject* NodeObject = Cast<UCustomizableObjectNodeObject>(ConnectedPin->GetOwningNode()))
+				if (NodeLinkedPin)
 				{
-					ParentId = GenerationContext.GroupIdToExternalNodeMap.FindKey(NodeObject);
-					ValueName = NodeObject->GetObjectName(&GenerationContext.MacroNodesStack);
+					const FGuid* ParentId = nullptr;
+					FString ValueName;
 
-					// Group objects in the same graph aren't in the GroupIdToExternalNodeMap, so follow the pins instead
-					if (!ParentId && NodeObject->OutputPin())
+					if (UCustomizableObjectNodeObject* NodeObject = Cast<UCustomizableObjectNodeObject>(NodeLinkedPin->GetOwningNode()))
 					{
-						if (const UEdGraphPin* ConnectedPinToObject = FollowOutputPin(*NodeObject->OutputPin()))
+						ParentId = GenerationContext.GroupIdToExternalNodeMap.FindKey(NodeObject);
+						ValueName = NodeObject->GetObjectName(&GenerationContext.MacroNodesStack);
+
+						// Group objects in the same graph aren't in the GroupIdToExternalNodeMap, so follow the pins instead
+						if (!ParentId && NodeObject->OutputPin())
 						{
-							if (UCustomizableObjectNodeObjectGroup* ParentGroupNode = Cast<UCustomizableObjectNodeObjectGroup>(ConnectedPinToObject->GetOwningNode()))
+							if (const UEdGraphPin* ConnectedPinToObject = FollowOutputPin(*NodeObject->OutputPin()))
 							{
-								ParentId = &ParentGroupNode->NodeGuid;
+								if (UCustomizableObjectNodeObjectGroup* ParentGroupNode = Cast<UCustomizableObjectNodeObjectGroup>(ConnectedPinToObject->GetOwningNode()))
+								{
+									FGuid ParentGroupNodeId = GenerationContext.GetNodeIdUnchecked(ParentGroupNode);
+									ParentId = &ParentGroupNodeId;
+								}
 							}
 						}
 					}
-				}
-				else if (UCustomizableObjectNodeObjectGroup* NodeObjectGroup = Cast<UCustomizableObjectNodeObjectGroup>(ConnectedPin->GetOwningNode()))
-				{
-					ParentId = &NodeObjectGroup->NodeGuid;
-					ValueName = TypedNodeGroup->GetGroupName(&GenerationContext.MacroNodesStack);
-				}
+					else if (UCustomizableObjectNodeObjectGroup* NodeObjectGroup = Cast<UCustomizableObjectNodeObjectGroup>(NodeLinkedPin->GetOwningNode()))
+					{
+						FGuid ParentGroupNodeId = GenerationContext.GetNodeIdUnchecked(NodeObjectGroup);
+						ParentId = &ParentGroupNodeId;
 
-				if (ParentId)
-				{
-					FMutableGraphGenerationContext::FParamInfo* ParentParamInfo = GenerationContext.GuidToParamNameMap.Find(*ParentId);
+						ValueName = TypedNodeGroup->GetGroupName(&GenerationContext.MacroNodesStack);
+					}
 
-					if (ParentParamInfo)
+					if (ParentId)
 					{
-						FString ParentParamName = ParentParamInfo->ParamName;
-
-						ParameterUIData.ParamUIMetadata.ExtraInformation.Add(FString("__ParentParamName"), 
-							ParentParamInfo->bIsToggle ? ValueName : *ParentParamName);
+						FMutableGraphGenerationContext::FParamInfo* ParentParamInfo = GenerationContext.GuidToParamNameMap.Find(*ParentId);
 
-						if (bHideWhenNotSelected)
+						if (ParentParamInfo)
 						{
-							ParameterUIData.ParamUIMetadata.ExtraInformation.Add(FString("__DisplayWhenParentValueEquals"), 
-								ParentParamInfo->bIsToggle ? FString("1") : ValueName);
-						}
+							FString ParentParamName = ParentParamInfo->ParamName;
 
-						if (bCollapseUnderParent)
-						{
-							ParameterUIData.ParamUIMetadata.ExtraInformation.Add(FString("CollapseUnderParent"));
+							ParameterUIData.ParamUIMetadata.ExtraInformation.Add(FString("__ParentParamName"),
+								ParentParamInfo->bIsToggle ? ValueName : *ParentParamName);
+
+							if (bHideWhenNotSelected)
+							{
+								ParameterUIData.ParamUIMetadata.ExtraInformation.Add(FString("__DisplayWhenParentValueEquals"),
+									ParentParamInfo->bIsToggle ? FString("1") : ValueName);
+							}
 
-							FMutableParameterData ParentParameterUIData;
-							ParentParameterUIData.ParamUIMetadata.ExtraInformation.Add(FString("__HasCollapsibleChildren"));
-							GenerationContext.ParameterUIDataMap.Add(*ParentParamName, ParentParameterUIData);
+							if (bCollapseUnderParent)
+							{
+								ParameterUIData.ParamUIMetadata.ExtraInformation.Add(FString("CollapseUnderParent"));
+
+								FMutableParameterData ParentParameterUIData;
+								ParentParameterUIData.ParamUIMetadata.ExtraInformation.Add(FString("__HasCollapsibleChildren"));
+								GenerationContext.ParameterUIDataMap.Add(*ParentParamName, ParentParameterUIData);
+							}
 						}
 					}
 				}
@@ -1876,10 +1955,11 @@ void FMutableComponentInfo::AccumulateBonesToRemovePerLOD(const TArray<FLODReduc
 
 #undef LOCTEXT_NAMESPACE
 
-FMutableGraphGenerationContext::FSharedSurface::FSharedSurface(uint8 InLOD, const mu::Ptr<mu::NodeSurfaceNew>& InNodeSurfaceNew)
+FMutableGraphGenerationContext::FSharedSurface::FSharedSurface(uint8 InLOD, const mu::Ptr<mu::NodeSurfaceNew>& InNodeSurfaceNew, TArray<const UCustomizableObjectNodeMacroInstance*> CurrentMacroContext)
 {
 	LOD = InLOD;
 	NodeSurfaceNew = InNodeSurfaceNew;
+	MacroContext = CurrentMacroContext;
 }
 
 
diff --git a/Engine/Plugins/Mutable/Source/CustomizableObjectEditor/Private/MuCOE/GenerateMutableSource/GenerateMutableSource.h b/Engine/Plugins/Mutable/Source/CustomizableObjectEditor/Private/MuCOE/GenerateMutableSource/GenerateMutableSource.h
--- a/Engine/Plugins/Mutable/Source/CustomizableObjectEditor/Private/MuCOE/GenerateMutableSource/GenerateMutableSource.h
+++ b/Engine/Plugins/Mutable/Source/CustomizableObjectEditor/Private/MuCOE/GenerateMutableSource/GenerateMutableSource.h
@@ -15,6 +15,7 @@
 #include "MuCOE/ExtensionDataCompilerInterface.h"
 #include "MuCOE/Nodes/CustomizableObjectNodeMaterialBase.h"
 #include "MuCOE/Nodes/CustomizableObjectNodeObject.h"
+#include "MuCOE/Nodes/CustomizableObjectNodeTable.h"
 #include "MuT/NodeComponentNew.h"
 #include "MuT/NodeImageConstant.h"
 #include "MuT/NodeMeshConstant.h"
@@ -35,7 +36,6 @@
 class UCustomizableObjectNodeGroupProjectorParameter;
 class UCustomizableObjectNodeMaterialBase;
 class UCustomizableObjectNodeModifierClipWithMesh;
-class UCustomizableObjectNodeTable;
 struct FCustomizableObjectClothingAssetData;
 
 class FCustomizableObjectCompiler;
@@ -45,6 +45,7 @@ class UCompositeDataTable;
 class UCustomizableObjectNodeMacroInstance;
 class UCustomizableObjectNodeMaterial;
 class UCustomizableObjectNodeMeshMorph;
+class UCustomizableObjectNodeMeshMorphStackDefinition;
 class UCustomizableObjectNodeObjectGroup;
 class UCustomizableObjectNodeVariation;
 class UEdGraphNode;
@@ -541,10 +542,10 @@ struct FGeneratedMutableDataTableKey
 	// Compilation Restrictions:
 	// Name of the column that determines de version control
 	FName VersionColumn;
-	
+
 	// Name of the column that determines the compilation filter
 	FName CompilationFilterColumn;
-	
+
 	// Compilation Filter
 	TArray<FString> CompilationFilters;
 
@@ -746,6 +747,11 @@ struct FMutableGraphGenerationContext
 	// Check if the Id of the node Node already exists, if it's new adds it to NodeIds array, otherwise, returns new Id
 	const FGuid GetNodeIdUnique(const UCustomizableObjectNode* Node);
 
+	/** Same as GetNodeIdUnique but does not triggers any warning on repeated IDs.
+	* Use only if GetNodeIdUnique is used in another part of the code.
+	*/
+	const FGuid GetNodeIdUnchecked(const UCustomizableObjectNode* Node);
+
 	/** Generates shared surface IDs for all surface nodes. If one or more nodes are equal, they will use the same SharedSurfaceId */
 	void GenerateSharedSurfacesUniqueIds();
 
@@ -850,7 +856,7 @@ struct FMutableGraphGenerationContext
 
 	struct FSharedSurface
 	{
-		FSharedSurface(uint8 InLOD, const mu::Ptr<mu::NodeSurfaceNew>& InNodeSurfaceNew);
+		FSharedSurface(uint8 InLOD, const mu::Ptr<mu::NodeSurfaceNew>& InNodeSurfaceNew, TArray<const UCustomizableObjectNodeMacroInstance*> CurrentMacroContext);
 
 		bool operator==(const FSharedSurface& o) const;
 
@@ -859,6 +865,7 @@ struct FMutableGraphGenerationContext
 
 		bool bMakeUnique = false;
 		TArray<SIZE_T> NodeModifierIDs;
+		TArray<const UCustomizableObjectNodeMacroInstance*> MacroContext;
 	};
 
 	// Material to SharedSurfaceId
diff --git a/Engine/Plugins/Mutable/Source/CustomizableObjectEditor/Private/MuCOE/GenerateMutableSource/GenerateMutableSourceGroupProjector.cpp b/Engine/Plugins/Mutable/Source/CustomizableObjectEditor/Private/MuCOE/GenerateMutableSource/GenerateMutableSourceGroupProjector.cpp
--- a/Engine/Plugins/Mutable/Source/CustomizableObjectEditor/Private/MuCOE/GenerateMutableSource/GenerateMutableSourceGroupProjector.cpp
+++ b/Engine/Plugins/Mutable/Source/CustomizableObjectEditor/Private/MuCOE/GenerateMutableSource/GenerateMutableSourceGroupProjector.cpp
@@ -318,6 +318,7 @@ TOptional<FGroupProjectorTempData> GenerateMutableGroupProjector(UCustomizableOb
 	if (GroupProjectorTempData.NodeProjectorParameterPtr)
 	{
 		// Use the projector parameter uid + num to identify parameters derived from this node
+		// TODO(Max UE-247783): Get Node UID taking into account the macro context
 		FGuid NumLayersParamUid = ProjParamNode->NodeGuid;
 		NumLayersParamUid.D += 1;
 		FGuid SelectedPoseParamUid = ProjParamNode->NodeGuid;
diff --git a/Engine/Plugins/Mutable/Source/CustomizableObjectEditor/Private/MuCOE/GenerateMutableSource/GenerateMutableSourceMesh.cpp b/Engine/Plugins/Mutable/Source/CustomizableObjectEditor/Private/MuCOE/GenerateMutableSource/GenerateMutableSourceMesh.cpp
--- a/Engine/Plugins/Mutable/Source/CustomizableObjectEditor/Private/MuCOE/GenerateMutableSource/GenerateMutableSourceMesh.cpp
+++ b/Engine/Plugins/Mutable/Source/CustomizableObjectEditor/Private/MuCOE/GenerateMutableSource/GenerateMutableSourceMesh.cpp
@@ -2969,41 +2969,48 @@ mu::Ptr<mu::NodeScalar> GenerateMorphFactor(const UCustomizableObjectNode* Node,
 
 	if (const UEdGraphPin* ConnectedPin = FollowInputPin(FactorPin))
 	{
-		UEdGraphNode* floatNode = ConnectedPin->GetOwningNode();
+		// Checking if it's linked to a Macro or tunnel node
+		const UEdGraphPin* FloatPin = GraphTraversal::FindIOPinSourceThroughMacroContext(*ConnectedPin, &GenerationContext.MacroNodesStack);
 		bool validStaticFactor = true;
-		
-		if (const UCustomizableObjectNodeFloatParameter* floatParameterNode = Cast<UCustomizableObjectNodeFloatParameter>(floatNode))
+
+		if (FloatPin)
 		{
-			if (floatParameterNode->DefaultValue < -1.0f || floatParameterNode->DefaultValue > 1.0f)
-			{
-				validStaticFactor = false;
-				FString msg = FString::Printf(TEXT("Mesh morph nodes only accept factors between -1.0 and 1.0 inclusive but the default value of the float parameter node is (%f). Factor will be ignored."), floatParameterNode->DefaultValue);
-				GenerationContext.Log(FText::FromString(msg), Node);
-			}
-			if (floatParameterNode->ParamUIMetadata.MinimumValue < -1.0f)
-			{
-				validStaticFactor = false;
-				FString msg = FString::Printf(TEXT("Mesh morph nodes only accept factors between -1.0 and 1.0 inclusive but the minimum UI value for the input float parameter node is (%f). Factor will be ignored."), floatParameterNode->ParamUIMetadata.MinimumValue);
-				GenerationContext.Log(FText::FromString(msg), Node);
-			}
-			if (floatParameterNode->ParamUIMetadata.MaximumValue > 1.0f)
+			UEdGraphNode* floatNode = FloatPin->GetOwningNode();
+
+			if (const UCustomizableObjectNodeFloatParameter* floatParameterNode = Cast<UCustomizableObjectNodeFloatParameter>(floatNode))
 			{
-				validStaticFactor = false;
-				FString msg = FString::Printf(TEXT("Mesh morph nodes only accept factors between -1.0 and 1.0 inclusive but the maximum UI value for the input float parameter node is (%f). Factor will be ignored."), floatParameterNode->ParamUIMetadata.MaximumValue);
-				GenerationContext.Log(FText::FromString(msg), Node);
+				if (floatParameterNode->DefaultValue < -1.0f || floatParameterNode->DefaultValue > 1.0f)
+				{
+					validStaticFactor = false;
+					FString msg = FString::Printf(TEXT("Mesh morph nodes only accept factors between -1.0 and 1.0 inclusive but the default value of the float parameter node is (%f). Factor will be ignored."), floatParameterNode->DefaultValue);
+					GenerationContext.Log(FText::FromString(msg), Node);
+				}
+				if (floatParameterNode->ParamUIMetadata.MinimumValue < -1.0f)
+				{
+					validStaticFactor = false;
+					FString msg = FString::Printf(TEXT("Mesh morph nodes only accept factors between -1.0 and 1.0 inclusive but the minimum UI value for the input float parameter node is (%f). Factor will be ignored."), floatParameterNode->ParamUIMetadata.MinimumValue);
+					GenerationContext.Log(FText::FromString(msg), Node);
+				}
+				if (floatParameterNode->ParamUIMetadata.MaximumValue > 1.0f)
+				{
+					validStaticFactor = false;
+					FString msg = FString::Printf(TEXT("Mesh morph nodes only accept factors between -1.0 and 1.0 inclusive but the maximum UI value for the input float parameter node is (%f). Factor will be ignored."), floatParameterNode->ParamUIMetadata.MaximumValue);
+					GenerationContext.Log(FText::FromString(msg), Node);
+				}
 			}
-		}
-		
-		else if (const UCustomizableObjectNodeFloatConstant* floatConstantNode = Cast<UCustomizableObjectNodeFloatConstant>(floatNode))
-		{
-			if (floatConstantNode->Value < -1.0f || floatConstantNode->Value > 1.0f)
+
+			else if (const UCustomizableObjectNodeFloatConstant* floatConstantNode = Cast<UCustomizableObjectNodeFloatConstant>(floatNode))
 			{
-				validStaticFactor = false;
-				FString msg = FString::Printf(TEXT("Mesh morph nodes only accept factors between -1.0 and 1.0 inclusive but the value of the float constant node is (%f). Factor will be ignored."), floatConstantNode->Value);
-				GenerationContext.Log(FText::FromString(msg), Node);
+				if (floatConstantNode->Value < -1.0f || floatConstantNode->Value > 1.0f)
+				{
+					validStaticFactor = false;
+					FString msg = FString::Printf(TEXT("Mesh morph nodes only accept factors between -1.0 and 1.0 inclusive but the value of the float constant node is (%f). Factor will be ignored."), floatConstantNode->Value);
+					GenerationContext.Log(FText::FromString(msg), Node);
+				}
 			}
 		}
 
+		// If is a valid factor, continue the Generation
 		if (validStaticFactor)
 		{
 			FactorNode = GenerateMutableSourceFloat(ConnectedPin, GenerationContext);
@@ -3634,6 +3641,89 @@ mu::NodeMeshPtr GenerateMorphMesh(const UEdGraphPin* Pin,
 }
 
 
+TArray<FMorphNodeData> GenerateMeshMorphStackDefinition(const UEdGraphPin* Pin, UEdGraphPin* MeshPin, const TArray<FString>& MorphTargetNames, FMutableGraphGenerationContext& GenerationContext)
+{
+	UCustomizableObjectNode* Node = CastChecked<UCustomizableObjectNode>(Pin->GetOwningNode());
+	TArray<FMorphNodeData> OutMorphData;
+
+	if (const UCustomizableObjectNodeMeshMorphStackDefinition* TypedNodeMeshMorphStackDef = Cast<UCustomizableObjectNodeMeshMorphStackDefinition>(Node))
+	{
+		if (Node->IsNodeOutDatedAndNeedsRefresh())
+		{
+			Node->SetRefreshNodeWarning();
+		}
+
+		TArray<UEdGraphPin*> MorphPins = TypedNodeMeshMorphStackDef->GetAllNonOrphanPins();
+
+		for (const UEdGraphPin* MorphPin : MorphPins)
+		{
+			if (MorphPin->LinkedTo.Num() && MorphTargetNames.Contains(MorphPin->PinFriendlyName.ToString()))
+			{
+				// Generate Factor
+				mu::Ptr<mu::NodeScalar> FactorNode = GenerateMorphFactor(Node, *MorphPin, GenerationContext);
+
+				OutMorphData.Add({ Node, MorphPin->PinFriendlyName.ToString(), FactorNode, MeshPin });
+			}
+		}
+	}
+
+	else if (const UCustomizableObjectNodeMacroInstance* TypedNodeMacro = Cast<UCustomizableObjectNodeMacroInstance>(Node))
+	{
+		if (const UEdGraphPin* OutputPin = TypedNodeMacro->GetMacroIOPin(ECOMacroIOType::COMVT_Output, Pin->PinName))
+		{
+			if (const UEdGraphPin* FollowPin = FollowInputPin(*OutputPin))
+			{
+				GenerationContext.MacroNodesStack.Push(TypedNodeMacro);
+				OutMorphData = GenerateMeshMorphStackDefinition(FollowPin, MeshPin, MorphTargetNames, GenerationContext);
+				GenerationContext.MacroNodesStack.Pop();
+			}
+			else
+			{
+				FText Msg = FText::Format(LOCTEXT("MacroInstanceError_PinNotLinked", "Macro Output node Pin {0} not linked."), FText::FromName(Pin->PinName));
+				GenerationContext.Log(Msg, Node);
+			}
+		}
+		else
+		{
+			FText Msg = FText::Format(LOCTEXT("MacroInstanceError_PinNameNotFound", "Macro Output node does not contain a pin with name {0}."), FText::FromName(Pin->PinName));
+			GenerationContext.Log(Msg, Node);
+		}
+	}
+
+	else if (const UCustomizableObjectNodeTunnel* TypedNodeTunnel = Cast<UCustomizableObjectNodeTunnel>(Node))
+	{
+		check(TypedNodeTunnel->bIsInputNode);
+		check(GenerationContext.MacroNodesStack.Num());
+
+		const UCustomizableObjectNodeMacroInstance* MacroInstanceNode = GenerationContext.MacroNodesStack.Pop();
+		check(MacroInstanceNode);
+
+		if (const UEdGraphPin* InputPin = MacroInstanceNode->FindPin(Pin->PinName, EEdGraphPinDirection::EGPD_Input))
+		{
+			if (const UEdGraphPin* FollowPin = FollowInputPin(*InputPin))
+			{
+				OutMorphData = GenerateMeshMorphStackDefinition(FollowPin, MeshPin, MorphTargetNames, GenerationContext);
+			}
+		}
+		else
+		{
+			FText Msg = FText::Format(LOCTEXT("MacroTunnelError_PinNameNotFound", "Macro Instance Node does not contain a pin with name {0}."), FText::FromName(Pin->PinName));
+			GenerationContext.Log(Msg, Node);
+		}
+
+		// Push the Macro again even if the result is null
+		GenerationContext.MacroNodesStack.Push(MacroInstanceNode);
+	}
+
+	else
+	{
+		GenerationContext.Log(LOCTEXT("UnimplementedNode", "Node type not implemented yet."), Node);
+	}
+
+	return OutMorphData;
+}
+
+
 mu::Ptr<mu::NodeMesh> GenerateMutableSourceMesh(const UEdGraphPin* Pin,
 	FMutableGraphGenerationContext& GenerationContext,
 	const FMutableSourceMeshData& BaseMeshData,
@@ -3950,7 +4040,7 @@ mu::Ptr<mu::NodeMesh> GenerateMutableSourceMesh(const UEdGraphPin* Pin,
 			{
 				if (const UEdGraphPin* ConnectedStringPin = FollowInputPin(*MorphTargetNamePin))
 				{
-					if (const UEdGraphPin* SourceStringPin = GraphTraversal::FindStringSource(*ConnectedStringPin, &GenerationContext.MacroNodesStack))
+					if (const UEdGraphPin* SourceStringPin = GraphTraversal::FindIOPinSourceThroughMacroContext(*ConnectedStringPin, &GenerationContext.MacroNodesStack))
 					{
 						if (const UCustomizableObjectNodeStaticString* StringNode = Cast<UCustomizableObjectNodeStaticString>(SourceStringPin->GetOwningNode()))
 						{
@@ -3982,77 +4072,33 @@ mu::Ptr<mu::NodeMesh> GenerateMutableSourceMesh(const UEdGraphPin* Pin,
 
 	else if (const UCustomizableObjectNodeMeshMorphStackApplication* TypedNodeMeshMorphStackApp = Cast< UCustomizableObjectNodeMeshMorphStackApplication >(Node))
 	{
-		const TArray<FString> MorphNames = TypedNodeMeshMorphStackApp->GetMorphList();
-		
+		Result = nullptr;
+
 		if (const UEdGraphPin* ConnectedPin = FollowInputPin(*TypedNodeMeshMorphStackApp->GetStackPin()))
 		{
-			UEdGraphNode* OwningNode = ConnectedPin->GetOwningNode();
-			if (UCustomizableObjectNodeMeshMorphStackDefinition* TypedNodeMeshMorphStackDef = Cast<UCustomizableObjectNodeMeshMorphStackDefinition>(OwningNode))
-			{
-				// Checking if is out of data
-				if (TypedNodeMeshMorphStackDef->IsNodeOutDatedAndNeedsRefresh())
-				{
-					TypedNodeMeshMorphStackDef->SetRefreshNodeWarning();
-				}
-
-				mu::Ptr<mu::NodeMeshMorph> MeshNode = new mu::NodeMeshMorph();
-				Result = MeshNode;
+			const TArray<FString> MorphTargetNames = TypedNodeMeshMorphStackApp->GetMorphList();
+			TArray<FMorphNodeData> MorphDataArray = GenerateMeshMorphStackDefinition(ConnectedPin, TypedNodeMeshMorphStackApp->GetMeshPin(), MorphTargetNames, GenerationContext);
+			int32 AddedMorphs = 0;
 
-				TArray<UEdGraphPin*> MorphPins = TypedNodeMeshMorphStackDef->GetAllNonOrphanPins();
-
-				int32 AddedMorphs = 0;
-
-				for (int32 PinIndex = 0; PinIndex < MorphPins.Num(); ++PinIndex)
-				{
-					UEdGraphPin* MorphPin = MorphPins[PinIndex];
-
-					const UEdGraphSchema_CustomizableObject* Schema = GetDefault<UEdGraphSchema_CustomizableObject>();
-
-					// Checking if it's a valid pin
-					if (MorphPin->Direction == EEdGraphPinDirection::EGPD_Output 
-						|| MorphPin->PinType.PinCategory != Schema->PC_Float 
-						|| !MorphPins[PinIndex]->LinkedTo.Num())
-					{
-						continue;
-					}
-
-					// Checking if the morph exists in the application node
-					FString MorphName = MorphPin->PinFriendlyName.ToString();
-					if (!MorphNames.Contains(MorphName))
-					{
-						continue;
-					}
-
-					//Factor
-					mu::Ptr<mu::NodeScalar> FactorNode = GenerateMorphFactor(Node, *MorphPin, GenerationContext);
-
-					// Mesh Morph Stack Management. TODO(Max): should we add the stack application node here instead of the def? Or both?
-					FMorphNodeData NewMorphData = { TypedNodeMeshMorphStackDef, MorphName, FactorNode, TypedNodeMeshMorphStackApp->GetMeshPin() };
-					GenerationContext.MeshMorphStack.Push(NewMorphData);
-
-					AddedMorphs++;
-				}
-
-				if (const UEdGraphPin* MeshConnectedPin = FollowInputPin(*TypedNodeMeshMorphStackApp->GetMeshPin()))
-				{
-					Result = GenerateMutableSourceMesh(MeshConnectedPin, GenerationContext, BaseMeshData, false, bOnlyConnectedLOD);
-				}
+			for (const FMorphNodeData& MorphData : MorphDataArray)
+			{
+				GenerationContext.MeshMorphStack.Push(MorphData);
+				AddedMorphs++;
+			}
 
-				for (int32 MorphIndex = 0; MorphIndex < AddedMorphs; ++MorphIndex)
-				{
-					GenerationContext.MeshMorphStack.Pop(EAllowShrinking::Yes);
-				}
+			if (const UEdGraphPin* MeshConnectedPin = FollowInputPin(*TypedNodeMeshMorphStackApp->GetMeshPin()))
+			{
+				Result = GenerateMutableSourceMesh(MeshConnectedPin, GenerationContext, BaseMeshData, false, bOnlyConnectedLOD);
 			}
-			else
+
+			for (int32 MorphIndex = 0; MorphIndex < AddedMorphs; ++MorphIndex)
 			{
-				GenerationContext.Log(LOCTEXT("MorphStackGenerationFailed", "Stack definition Generation failed."), Node);
-				Result = nullptr;
+				GenerationContext.MeshMorphStack.Pop(EAllowShrinking::No);
 			}
 		}
 		else
 		{
 			GenerationContext.Log(LOCTEXT("MorphStackConnectionFailed", "Stack definition connection not found."), Node);
-			Result = nullptr;
 		}
 	}
 
@@ -4349,16 +4395,13 @@ mu::Ptr<mu::NodeMesh> GenerateMutableSourceMesh(const UEdGraphPin* Pin,
 				}
 				else if (const UEdGraphPin* TablePosePin = FollowInputPin(*TypedNode->GetTablePosePin()))
 				{
-					if (const UCustomizableObjectNodeTable* TypedNodeTable = Cast<UCustomizableObjectNodeTable>(TablePosePin->GetOwningNode()))
-					{
-						mu::Ptr<mu::NodeMeshApplyPose> NodeMeshApplyPose = new mu::NodeMeshApplyPose();
-						mu::Ptr<mu::NodeMesh> MeshTableNode = GenerateMutableSourceMesh(TablePosePin, GenerationContext, BaseMeshData, false, bOnlyConnectedLOD);
+					mu::Ptr<mu::NodeMeshApplyPose> NodeMeshApplyPose = new mu::NodeMeshApplyPose();
+					mu::Ptr<mu::NodeMesh> MeshTableNode = GenerateMutableSourceMesh(TablePosePin, GenerationContext, BaseMeshData, false, bOnlyConnectedLOD);
 
-						NodeMeshApplyPose->Base = InputMeshNode;
-						NodeMeshApplyPose->Pose = MeshTableNode;
+					NodeMeshApplyPose->Base = InputMeshNode;
+					NodeMeshApplyPose->Pose = MeshTableNode;
 
-						Result = NodeMeshApplyPose;
-					}
+					Result = NodeMeshApplyPose;
 				}
 				else
 				{
diff --git a/Engine/Plugins/Mutable/Source/CustomizableObjectEditor/Private/MuCOE/GenerateMutableSource/GenerateMutableSourceModifier.cpp b/Engine/Plugins/Mutable/Source/CustomizableObjectEditor/Private/MuCOE/GenerateMutableSource/GenerateMutableSourceModifier.cpp
--- a/Engine/Plugins/Mutable/Source/CustomizableObjectEditor/Private/MuCOE/GenerateMutableSource/GenerateMutableSourceModifier.cpp
+++ b/Engine/Plugins/Mutable/Source/CustomizableObjectEditor/Private/MuCOE/GenerateMutableSource/GenerateMutableSourceModifier.cpp
@@ -539,7 +539,7 @@ mu::Ptr<mu::NodeModifier> GenerateMutableSourceModifier(const UEdGraphPin * Pin,
 		{
 			if (const UEdGraphPin* ConnectedStringPin = FollowInputPin(*MorphTargetNamePin))
 			{
-				if (const UEdGraphPin* SourceStringPin = GraphTraversal::FindStringSource(*ConnectedStringPin, &GenerationContext.MacroNodesStack))
+				if (const UEdGraphPin* SourceStringPin = GraphTraversal::FindIOPinSourceThroughMacroContext(*ConnectedStringPin, &GenerationContext.MacroNodesStack))
 				{
 					if (const UCustomizableObjectNodeStaticString* StringNode = Cast<UCustomizableObjectNodeStaticString>(SourceStringPin->GetOwningNode()))
 					{
@@ -557,39 +557,46 @@ mu::Ptr<mu::NodeModifier> GenerateMutableSourceModifier(const UEdGraphPin * Pin,
 
 		if (const UEdGraphPin* ConnectedPin = FollowInputPin(*TypedNodeMorph->FactorPin()))
 		{
-			UEdGraphNode* floatNode = ConnectedPin->GetOwningNode();
+			// Checking if it's linked to a Macro or tunnel node
+			const UEdGraphPin* FloatPin = GraphTraversal::FindIOPinSourceThroughMacroContext(*ConnectedPin, &GenerationContext.MacroNodesStack);
 			bool validStaticFactor = true;
-			if (const UCustomizableObjectNodeFloatParameter* floatParameterNode = Cast<UCustomizableObjectNodeFloatParameter>(floatNode))
+
+			if (FloatPin)
 			{
-				if (floatParameterNode->DefaultValue < -1.0f || floatParameterNode->DefaultValue > 1.0f)
-				{
-					validStaticFactor = false;
-					FString msg = FString::Printf(TEXT("Mesh morph nodes only accept factors between -1.0 and 1.0 inclusive but the default value of the float parameter node is (%f). Factor will be ignored."), floatParameterNode->DefaultValue);
-					GenerationContext.Log(FText::FromString(msg), Node);
-				}
-				if (floatParameterNode->ParamUIMetadata.MinimumValue < -1.0f)
-				{
-					validStaticFactor = false;
-					FString msg = FString::Printf(TEXT("Mesh morph nodes only accept factors between -1.0 and 1.0 inclusive but the minimum UI value for the input float parameter node is (%f). Factor will be ignored."), floatParameterNode->ParamUIMetadata.MinimumValue);
-					GenerationContext.Log(FText::FromString(msg), Node);
-				}
-				if (floatParameterNode->ParamUIMetadata.MaximumValue > 1.0f)
+				UEdGraphNode* floatNode = FloatPin->GetOwningNode();
+				if (const UCustomizableObjectNodeFloatParameter* floatParameterNode = Cast<UCustomizableObjectNodeFloatParameter>(floatNode))
 				{
-					validStaticFactor = false;
-					FString msg = FString::Printf(TEXT("Mesh morph nodes only accept factors between -1.0 and 1.0 inclusive but the maximum UI value for the input float parameter node is (%f). Factor will be ignored."), floatParameterNode->ParamUIMetadata.MaximumValue);
-					GenerationContext.Log(FText::FromString(msg), Node);
+					if (floatParameterNode->DefaultValue < -1.0f || floatParameterNode->DefaultValue > 1.0f)
+					{
+						validStaticFactor = false;
+						FString msg = FString::Printf(TEXT("Mesh morph nodes only accept factors between -1.0 and 1.0 inclusive but the default value of the float parameter node is (%f). Factor will be ignored."), floatParameterNode->DefaultValue);
+						GenerationContext.Log(FText::FromString(msg), Node);
+					}
+					if (floatParameterNode->ParamUIMetadata.MinimumValue < -1.0f)
+					{
+						validStaticFactor = false;
+						FString msg = FString::Printf(TEXT("Mesh morph nodes only accept factors between -1.0 and 1.0 inclusive but the minimum UI value for the input float parameter node is (%f). Factor will be ignored."), floatParameterNode->ParamUIMetadata.MinimumValue);
+						GenerationContext.Log(FText::FromString(msg), Node);
+					}
+					if (floatParameterNode->ParamUIMetadata.MaximumValue > 1.0f)
+					{
+						validStaticFactor = false;
+						FString msg = FString::Printf(TEXT("Mesh morph nodes only accept factors between -1.0 and 1.0 inclusive but the maximum UI value for the input float parameter node is (%f). Factor will be ignored."), floatParameterNode->ParamUIMetadata.MaximumValue);
+						GenerationContext.Log(FText::FromString(msg), Node);
+					}
 				}
-			}
-			else if (const UCustomizableObjectNodeFloatConstant* floatConstantNode = Cast<UCustomizableObjectNodeFloatConstant>(floatNode))
-			{
-				if (floatConstantNode->Value < -1.0f || floatConstantNode->Value > 1.0f)
+				else if (const UCustomizableObjectNodeFloatConstant* floatConstantNode = Cast<UCustomizableObjectNodeFloatConstant>(floatNode))
 				{
-					validStaticFactor = false;
-					FString msg = FString::Printf(TEXT("Mesh morph nodes only accept factors between -1.0 and 1.0 inclusive but the value of the float constant node is (%f). Factor will be ignored."), floatConstantNode->Value);
-					GenerationContext.Log(FText::FromString(msg), Node);
+					if (floatConstantNode->Value < -1.0f || floatConstantNode->Value > 1.0f)
+					{
+						validStaticFactor = false;
+						FString msg = FString::Printf(TEXT("Mesh morph nodes only accept factors between -1.0 and 1.0 inclusive but the value of the float constant node is (%f). Factor will be ignored."), floatConstantNode->Value);
+						GenerationContext.Log(FText::FromString(msg), Node);
+					}
 				}
 			}
 
+			// If is a valid factor, continue the Generation
 			if (validStaticFactor)
 			{
 				mu::Ptr<mu::NodeScalar> FactorNode = GenerateMutableSourceFloat(ConnectedPin, GenerationContext);
diff --git a/Engine/Plugins/Mutable/Source/CustomizableObjectEditor/Private/MuCOE/GenerateMutableSource/GenerateMutableSourceSurface.cpp b/Engine/Plugins/Mutable/Source/CustomizableObjectEditor/Private/MuCOE/GenerateMutableSource/GenerateMutableSourceSurface.cpp
--- a/Engine/Plugins/Mutable/Source/CustomizableObjectEditor/Private/MuCOE/GenerateMutableSource/GenerateMutableSourceSurface.cpp
+++ b/Engine/Plugins/Mutable/Source/CustomizableObjectEditor/Private/MuCOE/GenerateMutableSource/GenerateMutableSourceSurface.cpp
@@ -127,7 +127,7 @@ mu::Ptr<mu::NodeSurface> GenerateMutableSourceSurface(const UEdGraphPin * Pin, F
 		// Add to the list of surfaces that could be reused between LODs for this NodeMaterial.
 		// TODO(Max)
 		TArray<FMutableGraphGenerationContext::FSharedSurface>& SharedSurfaces = GenerationContext.SharedSurfaceIds.FindOrAdd(TypedNodeMat, {});
-		FMutableGraphGenerationContext::FSharedSurface& SharedSurface = SharedSurfaces.Add_GetRef(FMutableGraphGenerationContext::FSharedSurface(GenerationContext.CurrentLOD, SurfNode));
+		FMutableGraphGenerationContext::FSharedSurface& SharedSurface = SharedSurfaces.Add_GetRef(FMutableGraphGenerationContext::FSharedSurface(GenerationContext.CurrentLOD, SurfNode, GenerationContext.MacroNodesStack));
 		SharedSurface.bMakeUnique = !TypedNodeMat->IsReuseMaterialBetweenLODs();
 
 		GenerationContext.CurrentReferencedMaterialIndex = GenerationContext.ReferencedMaterials.AddUnique(Material);
@@ -215,28 +215,32 @@ mu::Ptr<mu::NodeSurface> GenerateMutableSourceSurface(const UEdGraphPin * Pin, F
 		// Checking if we should not use the material of the table node even if it is linked to the material node
 		if (TypedNodeMat->GetMaterialAssetPin())
 		{
-			if (const UEdGraphPin* MaterialAssetConnectedPin = FollowInputPin(*TypedNodeMat->GetMaterialAssetPin()))
+			if (const UEdGraphPin* ConnectedPin = FollowInputPin(*TypedNodeMat->GetMaterialAssetPin()))
 			{
-				if (const UCustomizableObjectNodeTable* TypedNodeTable = Cast< UCustomizableObjectNodeTable >(MaterialAssetConnectedPin->GetOwningNode()))
+				//Check if the pin goes through a macro or tunnel node
+				if (const UEdGraphPin* MaterialAssetPin = GraphTraversal::FindIOPinSourceThroughMacroContext(*ConnectedPin, &GenerationContext.MacroNodesStack))
 				{
-					if (UMaterialInstance* TableMaterial = TypedNodeTable->GetColumnDefaultAssetByType<UMaterialInstance>(MaterialAssetConnectedPin))
+					if (const UCustomizableObjectNodeTable* TypedNodeTable = Cast< UCustomizableObjectNodeTable >(MaterialAssetPin->GetOwningNode()))
 					{
-						// Checking if the reference material of the Table Node has the same parent as the material of the Material Node 
-						if (!Material || TableMaterial->GetMaterial() != Material->GetMaterial())
+						if (UMaterialInstance* TableMaterial = TypedNodeTable->GetColumnDefaultAssetByType<UMaterialInstance>(MaterialAssetPin))
 						{
-							bTableMaterialPinLinked = false;
+							// Checking if the reference material of the Table Node has the same parent as the material of the Material Node 
+							if (!Material || TableMaterial->GetMaterial() != Material->GetMaterial())
+							{
+								bTableMaterialPinLinked = false;
 
-							GenerationContext.Log(LOCTEXT("DifferentParentMaterial", "The Default Material Instance of the Data Table must have the same Parent Material."), TypedNodeMat->GetMaterialNode());
+								GenerationContext.Log(LOCTEXT("DifferentParentMaterial", "The Default Material Instance of the Data Table must have the same Parent Material."), TypedNodeMat->GetMaterialNode());
+							}
 						}
-					}
-					else
-					{
-						FString TableColumnName = TypedNodeTable->GetPinColumnName(MaterialAssetConnectedPin);
+						else
+						{
+							FString TableColumnName = TypedNodeTable->GetPinColumnName(MaterialAssetPin);
 
-						FText Msg = FText::Format(LOCTEXT("DefaultValueNotFound", "Couldn't find a default value in the data table's struct for the column {0}. The default value is null or not a Material Instance."), FText::FromString(TableColumnName));
-						GenerationContext.Log(Msg, Node);
+							FText Msg = FText::Format(LOCTEXT("DefaultValueNotFound", "Couldn't find a default value in the data table's struct for the column {0}. The default value is null or not a Material Instance."), FText::FromString(TableColumnName));
+							GenerationContext.Log(Msg, Node);
 
-						bTableMaterialPinLinked = false;
+							bTableMaterialPinLinked = false;
+						}
 					}
 				}
 			}
diff --git a/Engine/Plugins/Mutable/Source/CustomizableObjectEditor/Private/MuCOE/GraphTraversal.cpp b/Engine/Plugins/Mutable/Source/CustomizableObjectEditor/Private/MuCOE/GraphTraversal.cpp
--- a/Engine/Plugins/Mutable/Source/CustomizableObjectEditor/Private/MuCOE/GraphTraversal.cpp
+++ b/Engine/Plugins/Mutable/Source/CustomizableObjectEditor/Private/MuCOE/GraphTraversal.cpp
@@ -10,7 +10,6 @@
 #include "MuCOE/Nodes/CustomizableObjectNodeExposePin.h"
 #include "MuCOE/Nodes/CustomizableObjectNodeModifierExtendMeshSection.h"
 #include "MuCOE/Nodes/CustomizableObjectNodeExternalPin.h"
-#include "MuCOE/Nodes/CustomizableObjectNodeMacroInstance.h"
 #include "MuCOE/Nodes/CustomizableObjectNodeMaterial.h"
 #include "MuCOE/Nodes/CustomizableObjectNodeMaterialVariation.h"
 #include "MuCOE/Nodes/CustomizableObjectNodeMeshMorph.h"
@@ -25,7 +24,7 @@
 #include "MuCOE/Nodes/CustomizableObjectNodeStaticMesh.h"
 #include "MuCOE/Nodes/CustomizableObjectNodeStaticString.h"
 #include "MuCOE/Nodes/CustomizableObjectNodeTable.h"
-#include "MuCOE/Nodes/CustomizableObjectNodeTunnel.h"
+
 #include "MuCOE/Nodes/CustomizableObjectNodeAnimationPose.h"
 #include "MuCOE/Nodes/CustomizableObjectNodeObjectGroup.h"
 #include "MuCOE/Nodes/CustomizableObjectNodeReroute.h"
@@ -545,6 +544,62 @@ void GraphTraversal::VisitNodes(UCustomizableObjectNode& StartNode, const TFunct
 }
 
 
+const UEdGraphPin* GraphTraversal::FindIOPinSourceThroughMacroContext(const UEdGraphPin& Pin, TArray<const UCustomizableObjectNodeMacroInstance*>* MacroContext)
+{
+	const UEdGraphPin* ReturnPin = nullptr;
+	const UEdGraphNode* Node = Pin.GetOwningNode();
+	check(Node);
+
+	if (const UCustomizableObjectNodeMacroInstance* NodeMacro = Cast<UCustomizableObjectNodeMacroInstance>(Node))
+	{
+		if (const UEdGraphPin* OutputPin = NodeMacro->GetMacroIOPin(Pin.Direction == EGPD_Output ? ECOMacroIOType::COMVT_Output : ECOMacroIOType::COMVT_Input, Pin.PinName))
+		{
+			if (const UEdGraphPin* ConnectedPin = FollowInputPin(*OutputPin))
+			{
+				TArray<const UCustomizableObjectNodeMacroInstance*> NewMacroContext;
+
+				if (!MacroContext)
+				{
+					MacroContext = &NewMacroContext;
+				}
+
+				MacroContext->Push(NodeMacro);
+				ReturnPin = GraphTraversal::FindIOPinSourceThroughMacroContext(*ConnectedPin, MacroContext);
+				MacroContext->Pop();
+			}
+		}
+	}
+
+	else if (const UCustomizableObjectNodeTunnel* NodeTunnel = Cast<UCustomizableObjectNodeTunnel>(Node))
+	{
+		if (MacroContext && MacroContext->Num())
+		{
+			const UCustomizableObjectNodeMacroInstance* MacroInstanceNode = MacroContext->Pop();
+
+			if (MacroInstanceNode)
+			{
+				if (const UEdGraphPin* FollowPin = MacroInstanceNode->FindPin(Pin.PinName, NodeTunnel->bIsInputNode ? EEdGraphPinDirection::EGPD_Input : EEdGraphPinDirection::EGPD_Output))
+				{
+					if (const UEdGraphPin* ConnectedPin = NodeTunnel->bIsInputNode ? FollowInputPin(*FollowPin) : FollowOutputPin(*FollowPin))
+					{
+						ReturnPin = GraphTraversal::FindIOPinSourceThroughMacroContext(*ConnectedPin, MacroContext);
+					}
+				}
+			}
+
+			MacroContext->Push(MacroInstanceNode);
+		}
+	}
+
+	else
+	{
+		ReturnPin = &Pin;
+	}
+
+	return ReturnPin;
+}
+
+
 UCustomizableObjectNodeObject* GraphTraversal::GetFullGraphRootNode(const UCustomizableObject* Object, TArray<const UCustomizableObject*>& VisitedObjects)
 {
 	if (Object != nullptr)
@@ -913,68 +968,3 @@ void NodePinConnectionListChanged(const TArray<UEdGraphPin*>& Pins)
 }
 
 
-const UEdGraphPin* GraphTraversal::FindStringSource(const UEdGraphPin& Pin, TArray<const UCustomizableObjectNodeMacroInstance*>* MacroContext)
-{
-	const UEdGraphPin* ReturnPin = nullptr;
-	const UEdGraphNode* Node = Pin.GetOwningNode();
-	check(Node);
-
-	if (Cast<UCustomizableObjectNodeStaticString>(Node))
-	{
-		return &Pin;
-	}
-
-	else if (const UCustomizableObjectNodeMacroInstance* NodeMacro = Cast<UCustomizableObjectNodeMacroInstance>(Node))
-	{
-		if (const UEdGraphPin* OutputPin = NodeMacro->GetMacroIOPin(ECOMacroIOType::COMVT_Output, Pin.PinName))
-		{
-			if (const UEdGraphPin* ConnectedPin = FollowInputPin(*OutputPin))
-			{
-				TArray<const UCustomizableObjectNodeMacroInstance*> NewMacroContext;
-
-				if (!MacroContext)
-				{
-					MacroContext = &NewMacroContext;
-				}
-
-				MacroContext->Push(NodeMacro);
-				ReturnPin = GraphTraversal::FindStringSource(*ConnectedPin, MacroContext);
-				MacroContext->Pop();
-
-				return ReturnPin;
-			}
-		}
-	}
-
-	else if (const UCustomizableObjectNodeTunnel* NodeTunnel = Cast<UCustomizableObjectNodeTunnel>(Node))
-	{
-		check(NodeTunnel->bIsInputNode);
-
-		if (MacroContext && MacroContext->Num())
-		{
-			const UCustomizableObjectNodeMacroInstance* MacroInstanceNode = MacroContext->Pop();
-
-			if (MacroInstanceNode)
-			{
-				if (const UEdGraphPin* InputPin = MacroInstanceNode->FindPin(Pin.PinName, EEdGraphPinDirection::EGPD_Input))
-				{
-					if (const UEdGraphPin* ConnectedPin = FollowInputPin(*InputPin))
-					{
-						ReturnPin = GraphTraversal::FindStringSource(*ConnectedPin, MacroContext);
-					}
-				}
-			}
-
-			MacroContext->Push(MacroInstanceNode);
-
-			return ReturnPin;
-		}
-	}
-
-	else
-	{
-		unimplemented(); // Case missing.
-	}
-
-	return nullptr;
-}
diff --git a/Engine/Plugins/Mutable/Source/CustomizableObjectEditor/Private/MuCOE/GraphTraversal.h b/Engine/Plugins/Mutable/Source/CustomizableObjectEditor/Private/MuCOE/GraphTraversal.h
--- a/Engine/Plugins/Mutable/Source/CustomizableObjectEditor/Private/MuCOE/GraphTraversal.h
+++ b/Engine/Plugins/Mutable/Source/CustomizableObjectEditor/Private/MuCOE/GraphTraversal.h
@@ -3,6 +3,8 @@
 #pragma once
 
 #include "Containers/Array.h"
+#include "MuCOE/Nodes/CustomizableObjectNodeTunnel.h"
+#include "MuCOE/Nodes/CustomizableObjectNodeMacroInstance.h"
 
 class UCustomizableObject;
 class UCustomizableObjectNodeObject;
@@ -95,11 +97,11 @@ namespace GraphTraversal
 	void VisitNodes(UCustomizableObjectNode& StartNode, const TFunction<void(UCustomizableObjectNode&)>& VisitFunction,
 		const TMultiMap<FGuid, UCustomizableObjectNodeObject*>* ObjectGroupMap = nullptr, TArray<const UCustomizableObjectNodeMacroInstance*>* MacroContext = nullptr);
 
-	/** Given an output pin, return the output pin where the string node is located.
-	* @param Pin: Input pin where the search starts.
+	/** Returns a pin linked to the original pin. This function goes through Macro and Tunnel Nodes.
+	* @param Pin: Input/Output pin where the search starts.
 	* @param MacroContext: Copy of the current macro context.
 	*/
-	const UEdGraphPin* FindStringSource(const UEdGraphPin& Pin, TArray<const UCustomizableObjectNodeMacroInstance*>* MacroContext = nullptr);
+	const UEdGraphPin* FindIOPinSourceThroughMacroContext(const UEdGraphPin& Pin, TArray<const UCustomizableObjectNodeMacroInstance*>* MacroContext = nullptr);
 }
 
 /** Given an output pin, return the output pin where the mesh is located. */
diff --git a/Engine/Plugins/Mutable/Source/CustomizableObjectEditor/Private/MuCOE/Nodes/CustomizableObjectNodeComponent.cpp b/Engine/Plugins/Mutable/Source/CustomizableObjectEditor/Private/MuCOE/Nodes/CustomizableObjectNodeComponent.cpp
--- a/Engine/Plugins/Mutable/Source/CustomizableObjectEditor/Private/MuCOE/Nodes/CustomizableObjectNodeComponent.cpp
+++ b/Engine/Plugins/Mutable/Source/CustomizableObjectEditor/Private/MuCOE/Nodes/CustomizableObjectNodeComponent.cpp
@@ -70,7 +70,7 @@ FName UCustomizableObjectNodeComponent::GetComponentName(TArray<const UCustomiza
 	{
 		if (const UEdGraphPin* LinkedPin = FollowInputPin(*NamePin))
 		{
-			const UEdGraphPin* StringPin = GraphTraversal::FindStringSource(*LinkedPin, MacroContext);
+			const UEdGraphPin* StringPin = GraphTraversal::FindIOPinSourceThroughMacroContext(*LinkedPin, MacroContext);
 
 			if (const UCustomizableObjectNodeStaticString* StringNode = StringPin ? Cast<UCustomizableObjectNodeStaticString>(StringPin->GetOwningNode()) : nullptr)
 			{
diff --git a/Engine/Plugins/Mutable/Source/CustomizableObjectEditor/Private/MuCOE/Nodes/CustomizableObjectNodeComponentMeshAddTo.cpp b/Engine/Plugins/Mutable/Source/CustomizableObjectEditor/Private/MuCOE/Nodes/CustomizableObjectNodeComponentMeshAddTo.cpp
--- a/Engine/Plugins/Mutable/Source/CustomizableObjectEditor/Private/MuCOE/Nodes/CustomizableObjectNodeComponentMeshAddTo.cpp
+++ b/Engine/Plugins/Mutable/Source/CustomizableObjectEditor/Private/MuCOE/Nodes/CustomizableObjectNodeComponentMeshAddTo.cpp
@@ -156,7 +156,7 @@ FName UCustomizableObjectNodeComponentMeshAddTo::GetParentComponentName(TArray<c
 	{
 		if (const UEdGraphPin* LinkedPin = FollowInputPin(*ParentNamePin))
 		{
-			const UEdGraphPin* StringPin = GraphTraversal::FindStringSource(*LinkedPin, MacroContext);
+			const UEdGraphPin* StringPin = GraphTraversal::FindIOPinSourceThroughMacroContext(*LinkedPin, MacroContext);
 
 			if (const UCustomizableObjectNodeStaticString* StringNode = StringPin ? Cast<UCustomizableObjectNodeStaticString>(StringPin->GetOwningNode()) : nullptr)
 			{
diff --git a/Engine/Plugins/Mutable/Source/CustomizableObjectEditor/Private/MuCOE/Nodes/CustomizableObjectNodeMaterial.cpp b/Engine/Plugins/Mutable/Source/CustomizableObjectEditor/Private/MuCOE/Nodes/CustomizableObjectNodeMaterial.cpp
--- a/Engine/Plugins/Mutable/Source/CustomizableObjectEditor/Private/MuCOE/Nodes/CustomizableObjectNodeMaterial.cpp
+++ b/Engine/Plugins/Mutable/Source/CustomizableObjectEditor/Private/MuCOE/Nodes/CustomizableObjectNodeMaterial.cpp
@@ -1267,7 +1267,7 @@ TArray<FString> UCustomizableObjectNodeMaterial::GetEnableTags(TArray<const UCus
 
 		for (const UEdGraphPin* StringPin : FollowInputPinArray(*EnableTagsPin))
 		{
-			const UEdGraphPin* SourceStringPin = GraphTraversal::FindStringSource(*StringPin, MacroContext);
+			const UEdGraphPin* SourceStringPin = GraphTraversal::FindIOPinSourceThroughMacroContext(*StringPin, MacroContext);
 
 			if (SourceStringPin)
 			{
diff --git a/Engine/Plugins/Mutable/Source/CustomizableObjectEditor/Private/MuCOE/Nodes/CustomizableObjectNodeModifierBase.cpp b/Engine/Plugins/Mutable/Source/CustomizableObjectEditor/Private/MuCOE/Nodes/CustomizableObjectNodeModifierBase.cpp
--- a/Engine/Plugins/Mutable/Source/CustomizableObjectEditor/Private/MuCOE/Nodes/CustomizableObjectNodeModifierBase.cpp
+++ b/Engine/Plugins/Mutable/Source/CustomizableObjectEditor/Private/MuCOE/Nodes/CustomizableObjectNodeModifierBase.cpp
@@ -46,7 +46,7 @@ TArray<FString> UCustomizableObjectNodeModifierBase::GetNodeRequiredTags(TArray<
 
 		for (const UEdGraphPin* StringPin : FollowInputPinArray(*ReqTagsPin))
 		{
-			const UEdGraphPin* SourceStringPin = GraphTraversal::FindStringSource(*StringPin, MacroContext);
+			const UEdGraphPin* SourceStringPin = GraphTraversal::FindIOPinSourceThroughMacroContext(*StringPin, MacroContext);
 
 			if (SourceStringPin)
 			{
diff --git a/Engine/Plugins/Mutable/Source/CustomizableObjectEditor/Private/MuCOE/Nodes/CustomizableObjectNodeModifierExtendMeshSection.cpp b/Engine/Plugins/Mutable/Source/CustomizableObjectEditor/Private/MuCOE/Nodes/CustomizableObjectNodeModifierExtendMeshSection.cpp
--- a/Engine/Plugins/Mutable/Source/CustomizableObjectEditor/Private/MuCOE/Nodes/CustomizableObjectNodeModifierExtendMeshSection.cpp
+++ b/Engine/Plugins/Mutable/Source/CustomizableObjectEditor/Private/MuCOE/Nodes/CustomizableObjectNodeModifierExtendMeshSection.cpp
@@ -240,7 +240,7 @@ TArray<FString> UCustomizableObjectNodeModifierExtendMeshSection::GetEnableTags(
 
 		for (const UEdGraphPin* StringPin : FollowInputPinArray(*TagsPin))
 		{
-			const UEdGraphPin* SourceStringPin = GraphTraversal::FindStringSource(*StringPin, MacroContext);
+			const UEdGraphPin* SourceStringPin = GraphTraversal::FindIOPinSourceThroughMacroContext(*StringPin, MacroContext);
 
 			if (SourceStringPin)
 			{
diff --git a/Engine/Plugins/Mutable/Source/CustomizableObjectEditor/Private/MuCOE/Nodes/CustomizableObjectNodeObjectChild.cpp b/Engine/Plugins/Mutable/Source/CustomizableObjectEditor/Private/MuCOE/Nodes/CustomizableObjectNodeObjectChild.cpp
--- a/Engine/Plugins/Mutable/Source/CustomizableObjectEditor/Private/MuCOE/Nodes/CustomizableObjectNodeObjectChild.cpp
+++ b/Engine/Plugins/Mutable/Source/CustomizableObjectEditor/Private/MuCOE/Nodes/CustomizableObjectNodeObjectChild.cpp
@@ -109,7 +109,7 @@ FString UCustomizableObjectNodeObjectChild::GetObjectName(TArray<const UCustomiz
 	{
 		if (const UEdGraphPin* LinkedPin = FollowInputPin(*NamePin.Get()))
 		{
-			const UEdGraphPin* StringPin = GraphTraversal::FindStringSource(*LinkedPin, MacroContext);
+			const UEdGraphPin* StringPin = GraphTraversal::FindIOPinSourceThroughMacroContext(*LinkedPin, MacroContext);
 
 			if (const UCustomizableObjectNodeStaticString* StringNode = StringPin ? Cast<UCustomizableObjectNodeStaticString>(StringPin->GetOwningNode()) : nullptr)
 			{
diff --git a/Engine/Plugins/Mutable/Source/CustomizableObjectEditor/Private/MuCOE/Nodes/CustomizableObjectNodeObjectDetails.cpp b/Engine/Plugins/Mutable/Source/CustomizableObjectEditor/Private/MuCOE/Nodes/CustomizableObjectNodeObjectDetails.cpp
--- a/Engine/Plugins/Mutable/Source/CustomizableObjectEditor/Private/MuCOE/Nodes/CustomizableObjectNodeObjectDetails.cpp
+++ b/Engine/Plugins/Mutable/Source/CustomizableObjectEditor/Private/MuCOE/Nodes/CustomizableObjectNodeObjectDetails.cpp
@@ -387,6 +387,8 @@ void FCustomizableObjectNodeObjectDetails::FillParameterNamesArray()
 	// Array to store all the Child Object Nodes
 	TArray<const UCustomizableObjectNodeObject*> AllObjectNodes;
 
+	//TODO(Max) Use the Visit function of the GraphTraversal to take into account macro parameters?
+
 	for (const UCustomizableObject* Object : CustomObjectTree)
 	{
 		// Cheking Private to avoid a static analysis waring
diff --git a/Engine/Plugins/Mutable/Source/CustomizableObjectEditor/Private/MuCOE/Nodes/CustomizableObjectNodeObjectGroup.cpp b/Engine/Plugins/Mutable/Source/CustomizableObjectEditor/Private/MuCOE/Nodes/CustomizableObjectNodeObjectGroup.cpp
--- a/Engine/Plugins/Mutable/Source/CustomizableObjectEditor/Private/MuCOE/Nodes/CustomizableObjectNodeObjectGroup.cpp
+++ b/Engine/Plugins/Mutable/Source/CustomizableObjectEditor/Private/MuCOE/Nodes/CustomizableObjectNodeObjectGroup.cpp
@@ -203,7 +203,7 @@ FString UCustomizableObjectNodeObjectGroup::GetGroupName(TArray<const UCustomiza
 	{
 		if (const UEdGraphPin* LinkedPin = FollowInputPin(*NamePin.Get()))
 		{
-			const UEdGraphPin* StringPin = GraphTraversal::FindStringSource(*LinkedPin, MacroContext);
+			const UEdGraphPin* StringPin = GraphTraversal::FindIOPinSourceThroughMacroContext(*LinkedPin, MacroContext);
 
 			if (const UCustomizableObjectNodeStaticString* StringNode = StringPin ? Cast<UCustomizableObjectNodeStaticString>(StringPin->GetOwningNode()) : nullptr)
 			{
diff --git a/Engine/Plugins/Mutable/Source/CustomizableObjectEditor/Private/MuCOE/Nodes/CustomizableObjectNodeParameter.cpp b/Engine/Plugins/Mutable/Source/CustomizableObjectEditor/Private/MuCOE/Nodes/CustomizableObjectNodeParameter.cpp
--- a/Engine/Plugins/Mutable/Source/CustomizableObjectEditor/Private/MuCOE/Nodes/CustomizableObjectNodeParameter.cpp
+++ b/Engine/Plugins/Mutable/Source/CustomizableObjectEditor/Private/MuCOE/Nodes/CustomizableObjectNodeParameter.cpp
@@ -104,7 +104,7 @@ FString UCustomizableObjectNodeParameter::GetParameterName(TArray<const UCustomi
 	{
 		if (const UEdGraphPin* LinkedPin = FollowInputPin(*NamePin.Get()))
 		{
-			const UEdGraphPin* StringPin = GraphTraversal::FindStringSource(*LinkedPin, MacroContext);
+			const UEdGraphPin* StringPin = GraphTraversal::FindIOPinSourceThroughMacroContext(*LinkedPin, MacroContext);
 
 			if (const UCustomizableObjectNodeStaticString* StringNode = StringPin ? Cast<UCustomizableObjectNodeStaticString>(StringPin->GetOwningNode()) : nullptr)
 			{
diff --git a/Engine/Plugins/Mutable/Source/CustomizableObjectEditor/Private/MuCOE/Nodes/CustomizableObjectNodeVariation.cpp b/Engine/Plugins/Mutable/Source/CustomizableObjectEditor/Private/MuCOE/Nodes/CustomizableObjectNodeVariation.cpp
--- a/Engine/Plugins/Mutable/Source/CustomizableObjectEditor/Private/MuCOE/Nodes/CustomizableObjectNodeVariation.cpp
+++ b/Engine/Plugins/Mutable/Source/CustomizableObjectEditor/Private/MuCOE/Nodes/CustomizableObjectNodeVariation.cpp
@@ -91,7 +91,7 @@ FString UCustomizableObjectNodeVariation::GetVariationTag(int32 Index, TArray<co
 	{
 		if (const UEdGraphPin* ConnectedStringPin = FollowInputPin(*TagPin))
 		{
-			const UEdGraphPin* SourceStringPin = GraphTraversal::FindStringSource(*ConnectedStringPin, MacroContext);
+			const UEdGraphPin* SourceStringPin = GraphTraversal::FindIOPinSourceThroughMacroContext(*ConnectedStringPin, MacroContext);
 
 			if (SourceStringPin)
 			{
diff --git a/Engine/Plugins/Mutable/Source/CustomizableObjectEditor/Private/MuCOE/SMutableTagListWidget.cpp b/Engine/Plugins/Mutable/Source/CustomizableObjectEditor/Private/MuCOE/SMutableTagListWidget.cpp
--- a/Engine/Plugins/Mutable/Source/CustomizableObjectEditor/Private/MuCOE/SMutableTagListWidget.cpp
+++ b/Engine/Plugins/Mutable/Source/CustomizableObjectEditor/Private/MuCOE/SMutableTagListWidget.cpp
@@ -204,12 +204,12 @@ void SMutableTagComboBox::RefreshOptions()
 				continue;
 			}
 
-			const TArray<FString>& EnableTags = Typed->GetEnableTags();
-			if (EnableTags.Num())
+			const TArray<FString>* EnableTags = Typed->GetEnableTagsArray();
+			if (EnableTags)
 			{
 				TSharedPtr<SMutableSearchComboBox::FFilteredOption> NodeOption = AddNodeHierarchyOptions(Typed, AddedOptions);
 
-				for (const FString& OneTag : EnableTags)
+				for (const FString& OneTag : *EnableTags)
 				{
 					if (OneTag.IsEmpty())
 					{
