[Mutable] Fixed recompilation error due to a variable rename.

#rnx

diff --git a/Engine/Plugins/Mutable/Source/CustomizableObjectEditor/Private/MuCOE/GenerateMutableSource/GenerateMutableSourceMesh.cpp b/Engine/Plugins/Mutable/Source/CustomizableObjectEditor/Private/MuCOE/GenerateMutableSource/GenerateMutableSourceMesh.cpp
--- a/Engine/Plugins/Mutable/Source/CustomizableObjectEditor/Private/MuCOE/GenerateMutableSource/GenerateMutableSourceMesh.cpp
+++ b/Engine/Plugins/Mutable/Source/CustomizableObjectEditor/Private/MuCOE/GenerateMutableSource/GenerateMutableSourceMesh.cpp
@@ -4465,7 +4465,7 @@ mu::Ptr<mu::NodeMesh> GenerateMutableSourceMesh(const UEdGraphPin* Pin,
 				mu::Ptr<mu::FTable> GeneratedTable = nullptr;
 				GeneratedTable = GenerateMutableSourceTable(DataTable, TypedNodeTable, GenerationContext);
 
-				if (Table)
+				if (GeneratedTable)
 				{
 					mu::Ptr<mu::NodeMeshTable> MeshTableNode = new mu::NodeMeshTable();
 					
