[MutableRuntime]
* Fixed code that somehow got integrated wrong from 36

diff --git a/Engine/Plugins/Mutable/Source/MutableRuntime/Public/MuR/ImageDataStorage.h b/Engine/Plugins/Mutable/Source/MutableRuntime/Public/MuR/ImageDataStorage.h
--- a/Engine/Plugins/Mutable/Source/MutableRuntime/Public/MuR/ImageDataStorage.h
+++ b/Engine/Plugins/Mutable/Source/MutableRuntime/Public/MuR/ImageDataStorage.h
@@ -11,6 +11,8 @@
 
 #include <atomic>
 
+#define UE_API MUTABLERUNTIME_API
+
 namespace mu::MemoryCounters
 {
 	struct FImageMemoryCounter
@@ -206,3 +208,4 @@ public:
 
 }
 
+#undef UE_API
