[Mutable] Restrict Passthrough Mesh Component to only support Skeletal Meshes.

diff --git a/Engine/Plugins/Mutable/Config/BaseMutable.ini b/Engine/Plugins/Mutable/Config/BaseMutable.ini
--- a/Engine/Plugins/Mutable/Config/BaseMutable.ini
+++ b/Engine/Plugins/Mutable/Config/BaseMutable.ini
@@ -23,4 +23,5 @@
 +FunctionRedirects=(OldName="/Script/CustomizableObject.CustomizableObject.GetIntParameterOptionUIMetadata",NewName="/Script/CustomizableObject.CustomizableObject.GetEnumParameterValueUIMetadata")
 +FunctionRedirects=(OldName="/Script/CustomizableObject.CustomizableObject.GetIntParameterGroupType",NewName="/Script/CustomizableObject.CustomizableObject.GetEnumParameterGroupType")
 +FunctionRedirects=(OldName="/Script/CustomizableObject.CustomizableObject.GetIntParameterOptionDataTable",NewName="/Script/CustomizableObject.CustomizableObject.GetEnumParameterValueDataTable")
-+FunctionRedirects=(OldName="/Script/CustomizableObject.CustomizableObject.GetIntParameterDefaultValue",NewName="/Script/CustomizableObject.CustomizableObject.GetEnumParameterDefaultValue")
\ No newline at end of file
++FunctionRedirects=(OldName="/Script/CustomizableObject.CustomizableObject.GetIntParameterDefaultValue",NewName="/Script/CustomizableObject.CustomizableObject.GetEnumParameterDefaultValue")
++PropertyRedirects=(OldName="/Script/CustomizableObjectEditor.CustomizableObjectNodeComponentPassthroughMesh.Mesh",NewName="/Script/CustomizableObjectEditor.CustomizableObjectNodeComponentPassthroughMesh.SkeletalMesh")
\ No newline at end of file
diff --git a/Engine/Plugins/Mutable/Source/CustomizableObjectEditor/Private/MuCOE/GenerateMutableSource/GenerateMutableSourceComponent.cpp b/Engine/Plugins/Mutable/Source/CustomizableObjectEditor/Private/MuCOE/GenerateMutableSource/GenerateMutableSourceComponent.cpp
--- a/Engine/Plugins/Mutable/Source/CustomizableObjectEditor/Private/MuCOE/GenerateMutableSource/GenerateMutableSourceComponent.cpp
+++ b/Engine/Plugins/Mutable/Source/CustomizableObjectEditor/Private/MuCOE/GenerateMutableSource/GenerateMutableSourceComponent.cpp
@@ -236,14 +236,14 @@ mu::Ptr<mu::NodeComponent> GenerateMutableSourceComponent(const UEdGraphPin* Pin
 			return nullptr;
 		}
 
-		if (!TypedComponentPassthroughMesh->Mesh.IsValid())
+		if (!TypedComponentPassthroughMesh->SkeletalMesh.IsValid())
 		{
 			FString Msg = FString::Printf(TEXT("No mesh set for component node."));
 			GenerationContext.Log(FText::FromString(Msg), TypedComponentPassthroughMesh, EMessageSeverity::Warning);
 			return nullptr;
 		}
 		
-		USkeletalMesh* SkeletalMesh = Cast<USkeletalMesh>(GenerationContext.LoadObject(TSoftObjectPtr<USkeletalMesh>(TypedComponentPassthroughMesh->Mesh)));
+		USkeletalMesh* SkeletalMesh = Cast<USkeletalMesh>(GenerationContext.LoadObject(TSoftObjectPtr<USkeletalMesh>(TypedComponentPassthroughMesh->SkeletalMesh)));
 		if (!SkeletalMesh)
 		{
 			FString Msg = FString::Printf(TEXT("Only SkeletalMeshes are supported in this node, for now."));
diff --git a/Engine/Plugins/Mutable/Source/CustomizableObjectEditor/Private/MuCOE/Nodes/CustomizableObjectNodeComponentPassthroughMesh.cpp b/Engine/Plugins/Mutable/Source/CustomizableObjectEditor/Private/MuCOE/Nodes/CustomizableObjectNodeComponentPassthroughMesh.cpp
--- a/Engine/Plugins/Mutable/Source/CustomizableObjectEditor/Private/MuCOE/Nodes/CustomizableObjectNodeComponentPassthroughMesh.cpp
+++ b/Engine/Plugins/Mutable/Source/CustomizableObjectEditor/Private/MuCOE/Nodes/CustomizableObjectNodeComponentPassthroughMesh.cpp
@@ -89,21 +89,21 @@ void UCustomizableObjectNodeComponentPassthroughMesh::AllocateDefaultPins(UCusto
 {
 	Super::AllocateDefaultPins(RemapPins);
 
-	if (!Mesh.IsValid())
+	if (!SkeletalMesh.IsValid())
 	{
 		return;
 	}
 	
 	// Support SkeletalMeshes only, for now.
-	USkeletalMesh* SkeletalMesh = Cast<USkeletalMesh>(MutablePrivate::LoadObject(Mesh));
-	if (!SkeletalMesh)
+	USkeletalMesh* StrongSkeletalMesh = MutablePrivate::LoadObject<USkeletalMesh>(SkeletalMesh);
+	if (!StrongSkeletalMesh)
 	{
 		return;
 	}
 	
-	if (const FSkeletalMeshModel* ImportedModel = SkeletalMesh->GetImportedModel())
+	if (const FSkeletalMeshModel* ImportedModel = StrongSkeletalMesh->GetImportedModel())
 	{
-		const int32 NumLODs = SkeletalMesh->GetLODNum();
+		const int32 NumLODs = StrongSkeletalMesh->GetLODNum();
 		for (int32 LODIndex = 0; LODIndex < NumLODs; ++LODIndex)
 		{
 			const int32 NumSections = ImportedModel->LODModels[LODIndex].Sections.Num();
@@ -182,8 +182,8 @@ void UCustomizableObjectNodeComponentPassthroughMesh::GetPinSection(const UEdGra
 
 UMaterialInterface* UCustomizableObjectNodeComponentPassthroughMesh::GetMaterialFor(const UEdGraphPin* Pin) const
 {
-	USkeletalMesh* SkeletalMesh = Cast<USkeletalMesh>(MutablePrivate::LoadObject(Mesh));
-	if (SkeletalMesh)
+	USkeletalMesh* StrongSkeletalMesh = MutablePrivate::LoadObject<USkeletalMesh>(SkeletalMesh);
+	if (StrongSkeletalMesh)
 	{
 		if (FSkeletalMaterial* SkeletalMaterial = GetSkeletalMaterialFor(*Pin))
 		{
@@ -224,8 +224,8 @@ UMaterialInterface* UCustomizableObjectNodeComponentPassthroughMesh::GetMaterial
 
 FSkeletalMaterial* UCustomizableObjectNodeComponentPassthroughMesh::GetSkeletalMaterialFor(const int32 LODIndex, const int32 SectionIndex, const FSkeletalMeshModel* ImportedModel) const
 {
-	USkeletalMesh* SkeletalMesh = Cast<USkeletalMesh>(MutablePrivate::LoadObject(Mesh));
-	if (!SkeletalMesh)
+	USkeletalMesh* StrongSkeletalMesh = Cast<USkeletalMesh>(MutablePrivate::LoadObject(SkeletalMesh));
+	if (!StrongSkeletalMesh)
 	{
 		return nullptr;
 	}
@@ -234,7 +234,7 @@ FSkeletalMaterial* UCustomizableObjectNodeComponentPassthroughMesh::GetSkeletalM
 	int32 SkeletalMeshMaterialIndex = INDEX_NONE;
 	
 	// Check if we have lod info map to get the correct material index
-	if (const FSkeletalMeshLODInfo* LodInfo = SkeletalMesh->GetLODInfo(LODIndex))
+	if (const FSkeletalMeshLODInfo* LodInfo = StrongSkeletalMesh->GetLODInfo(LODIndex))
 	{
 		if (LodInfo->LODMaterialMap.IsValidIndex(SectionIndex))
 		{
@@ -251,7 +251,7 @@ FSkeletalMaterial* UCustomizableObjectNodeComponentPassthroughMesh::GetSkeletalM
 	}
 	else
 	{
-		FSkeletalMeshModel* AuxImportedModel = SkeletalMesh->GetImportedModel();
+		FSkeletalMeshModel* AuxImportedModel = StrongSkeletalMesh->GetImportedModel();
 
 		if (AuxImportedModel)
 		{
@@ -263,9 +263,9 @@ FSkeletalMaterial* UCustomizableObjectNodeComponentPassthroughMesh::GetSkeletalM
 	}
 	}
 	
-	if (SkeletalMesh->GetMaterials().IsValidIndex(SkeletalMeshMaterialIndex))
+	if (StrongSkeletalMesh->GetMaterials().IsValidIndex(SkeletalMeshMaterialIndex))
 	{
-		return &SkeletalMesh->GetMaterials()[SkeletalMeshMaterialIndex];
+		return &StrongSkeletalMesh->GetMaterials()[SkeletalMeshMaterialIndex];
 	}
 	
 	return nullptr;
diff --git a/Engine/Plugins/Mutable/Source/CustomizableObjectEditor/Private/MuCOE/Nodes/CustomizableObjectNodeComponentPassthroughMesh.h b/Engine/Plugins/Mutable/Source/CustomizableObjectEditor/Private/MuCOE/Nodes/CustomizableObjectNodeComponentPassthroughMesh.h
--- a/Engine/Plugins/Mutable/Source/CustomizableObjectEditor/Private/MuCOE/Nodes/CustomizableObjectNodeComponentPassthroughMesh.h
+++ b/Engine/Plugins/Mutable/Source/CustomizableObjectEditor/Private/MuCOE/Nodes/CustomizableObjectNodeComponentPassthroughMesh.h
@@ -67,8 +67,8 @@ class CUSTOMIZABLEOBJECTEDITOR_API UCustomizableObjectNodeComponentPassthroughMe
 public:
 	GENERATED_BODY()
 
-	UPROPERTY(EditAnywhere, Category = PassthroughMesh, Meta = (DisplayName = Mesh, AllowedClasses = "/Script/Engine.StaticMesh, /Script/Engine.SkeletalMesh"))
-	FSoftObjectPath Mesh;
+	UPROPERTY(EditAnywhere, Category = PassthroughMesh)
+	TSoftObjectPtr<USkeletalMesh> SkeletalMesh;
 
 	// UObject interface.
 	virtual void PostEditChangeProperty(FPropertyChangedEvent& PropertyChangedEvent) override;
