[mutable] Fixed race condition affecting the mutable tasks that blocked the mutable thread.
- A new task was being run while the old one was still active due to us removing the reference to the old task at the wrong time.

#rnx

diff --git a/Engine/Plugins/Mutable/Source/CustomizableObject/Private/MuCO/FMutableTaskGraph.cpp b/Engine/Plugins/Mutable/Source/CustomizableObject/Private/MuCO/FMutableTaskGraph.cpp
--- a/Engine/Plugins/Mutable/Source/CustomizableObject/Private/MuCO/FMutableTaskGraph.cpp
+++ b/Engine/Plugins/Mutable/Source/CustomizableObject/Private/MuCO/FMutableTaskGraph.cpp
@@ -336,25 +336,6 @@ int32 FMutableTaskGraph::Tick()
 	
 	TryLaunchMutableTaskLowPriority(false);
 
-	// See if we can clear the last reference to the mutable-thread task
-	if (IsTaskCompleted(LastMutableTask))
-	{
-		FScopeLock Lock(&MutableTaskLock);
-		if (IsTaskCompleted(LastMutableTask))
-		{
-			LastMutableTask = {};
-		}
-	}
-
-	if (IsTaskCompleted(LastMutableTaskLowPriority))
-	{
-		FScopeLock Lock(&MutableTaskLock);
-		if (IsTaskCompleted(LastMutableTaskLowPriority))
-		{
-			LastMutableTaskLowPriorityID = INVALID_ID;
-			LastMutableTaskLowPriority = {};
-		}
-	}
 	
 	return !IsTaskCompleted(LastMutableTask) + !IsTaskCompleted(LastMutableTaskLowPriority) + QueueMutableTasksLowPriority.Num() + !GameThreadTasks.IsEmpty();
 }
