[Mutable] Fix streaming crash when no indices.

diff --git a/Engine/Plugins/Mutable/Source/CustomizableObject/Private/MuCO/UnrealConversionUtils.cpp b/Engine/Plugins/Mutable/Source/CustomizableObject/Private/MuCO/UnrealConversionUtils.cpp
--- a/Engine/Plugins/Mutable/Source/CustomizableObject/Private/MuCO/UnrealConversionUtils.cpp
+++ b/Engine/Plugins/Mutable/Source/CustomizableObject/Private/MuCO/UnrealConversionUtils.cpp
@@ -493,13 +493,11 @@ namespace UnrealConversionUtils
 		}
 		
 		check(LODResource.MultiSizeIndexContainer.GetDataTypeSize() == ElementSize);
-
-		const int32 NumIndices = LODResource.RenderSections.Last().BaseIndex + (LODResource.RenderSections.Last().NumTriangles * 3);
-
-		// Don't assume the buffer is empty, otherwise we may add extra elements using Insert(). 
-		LODResource.MultiSizeIndexContainer.GetIndexBuffer()->Empty(NumIndices);
-		LODResource.MultiSizeIndexContainer.GetIndexBuffer()->Insert(0, NumIndices);
 		
+		// Don't assume the buffer is empty, otherwise we may add extra elements using Insert(). 
+		LODResource.MultiSizeIndexContainer.GetIndexBuffer()->Empty(IndexCount);
+		LODResource.MultiSizeIndexContainer.GetIndexBuffer()->Insert(0, IndexCount);
+
 		FMemory::Memcpy(LODResource.MultiSizeIndexContainer.GetIndexBuffer()->GetPointerTo(0), DataPtr, IndexCount * ElementSize);
 		
 		for (int32 SectionIndex = 0; SectionIndex < LODResource.RenderSections.Num(); ++SectionIndex)
