[mutable] Reuse mutable generated resources during the baking process. Mutable does still not reuse them during the update (TODO)
- Upgraded TextureReuseCache key from a string to a struct.
- Improved handling of files by the mutable baker to be able to re-use resources during the same bake by sharing them inside the same component.
- Mutable will still not re-use the textures generated. That will be done in a separated jira.
- Cleared some unnecessary code from the Table source generation.
- Added options to be able to provide the baker with the desired prefixes to be used for the files generated during the bake.

#rnx

diff --git a/Engine/Plugins/Mutable/Source/CustomizableObject/Private/MuCO/CustomizableObjectInstance.cpp b/Engine/Plugins/Mutable/Source/CustomizableObject/Private/MuCO/CustomizableObjectInstance.cpp
--- a/Engine/Plugins/Mutable/Source/CustomizableObject/Private/MuCO/CustomizableObjectInstance.cpp
+++ b/Engine/Plugins/Mutable/Source/CustomizableObject/Private/MuCO/CustomizableObjectInstance.cpp
@@ -6477,7 +6477,7 @@ void UCustomizableInstancePrivate::BuildMaterials(const TSharedRef<FUpdateContex
 
 							const uint32 MaterialPlaceHolderHash = GetTypeHash(MutableMaterialPlaceholder);
 							MIDName = FString::Printf(TEXT("%s%s_%u"),TEXT("MID_"),*MIDName, MaterialPlaceHolderHash);
-							MIDName = MakeUniqueObjectName(GetTransientPackage(), UMaterialInstanceDynamic::StaticClass(), *MIDName).ToString();
+							MIDName = MakeUniqueObjectName(GetTransientPackage(), UMaterialInstanceDynamic::StaticClass(), *MIDName, EUniqueObjectNameOptions::GloballyUnique).ToString();
 						}
 						
 						MaterialInstance = UMaterialInstanceDynamic::Create(MaterialTemplate, GetTransientPackage(), FName(MIDName));
diff --git a/Engine/Plugins/Mutable/Source/CustomizableObject/Public/MuCO/CustomizableObjectInstance.h b/Engine/Plugins/Mutable/Source/CustomizableObject/Public/MuCO/CustomizableObjectInstance.h
--- a/Engine/Plugins/Mutable/Source/CustomizableObject/Public/MuCO/CustomizableObjectInstance.h
+++ b/Engine/Plugins/Mutable/Source/CustomizableObject/Public/MuCO/CustomizableObjectInstance.h
@@ -45,11 +45,14 @@ enum class EPackageSaveResolutionType : uint8
 	/** The package got saved as a new file. */
 	NewFile,
 
+	/** The package was already present on disk so it was used instead of creating a new one. */
+	ReusedFile,
+	
 	/** The package was already present on disk so the old package was deleted and a new one was saved on its place */
-	Overriden,
+	Overriden UMETA(Hidden),			// Deprecated
 
 	/** Error type : An override was required but due to an error or lack of user permission it could not be done. */
-	UnableToOverride
+	UnableToOverride UMETA(Hidden),		// Deprecated
 };
 
 
@@ -105,10 +108,19 @@ DECLARE_DYNAMIC_DELEGATE_OneParam(FBakeOperationCompletedDelegate, const FCustom
  * Configuration data structure designed to serve as variable container for the customizable object instance baking methods.
  */
 USTRUCT(BlueprintType, Blueprintable)
-struct FBakingConfiguration
+struct CUSTOMIZABLEOBJECT_API FBakingConfiguration
 {
 	GENERATED_BODY()
 
+	// Note: We need to explicitly disable warnings on these constructors/operators for clang to be happy with deprecated variables
+	PRAGMA_DISABLE_DEPRECATION_WARNINGS
+	FBakingConfiguration() = default;
+	FBakingConfiguration(const FBakingConfiguration&) = default;
+	FBakingConfiguration(FBakingConfiguration&&) = default;
+	FBakingConfiguration& operator=(const FBakingConfiguration&) = default;
+	FBakingConfiguration& operator=(FBakingConfiguration&&) = default;
+	PRAGMA_ENABLE_DEPRECATION_WARNINGS
+	
 	/**
 	 * The path where to save the baked resources. EX /Game/MyBakingTest
 	 */
@@ -136,9 +148,34 @@ struct FBakingConfiguration
 	/**
 	 * Flag that determines if we should override already exported files or if we should not. If we encounter files to override and we have not permission to override them then the baking operation will fail.
 	 */
-	UPROPERTY(BlueprintReadWrite, Blueprintable, Category = CustomizableObjectInstanceBaker)
-	bool bAllowOverridingOfFiles = false;
+	UE_DEPRECATED(5.6, "The files will get automatically re-used instead")
+	UPROPERTY(meta = (DeprecatedProperty))
+	bool bAllowOverridingOfFiles_DEPRECATED = false;
+	
+	UPROPERTY(BlueprintReadWrite, Blueprintable, Category = "CustomizableObjectInstanceBaker|ResourceNameGeneration")
+	FString SkeletalMeshAssetPrefix = TEXT("SK_");
+
+	UPROPERTY(BlueprintReadWrite, Blueprintable, Category = "CustomizableObjectInstanceBaker|ResourceNameGeneration")
+	FString SkeletonAssetPrefix = TEXT("SKEL_");
 
+	UPROPERTY(BlueprintReadWrite, Blueprintable, Category = "CustomizableObjectInstanceBaker|ResourceNameGeneration")
+	FString PhysicsAssetPrefix = TEXT("PHYS_");
+	
+	UPROPERTY(BlueprintReadWrite, Blueprintable, Category = "CustomizableObjectInstanceBaker|ResourceNameGeneration")
+	FString TextureAssetPrefix = TEXT("T_");
+
+	UPROPERTY(BlueprintReadWrite, Blueprintable, Category = "CustomizableObjectInstanceBaker|ResourceNameGeneration")
+	FString MaterialAssetPrefix = TEXT("M_");
+
+	UPROPERTY(BlueprintReadWrite, Blueprintable, Category = "CustomizableObjectInstanceBaker|ResourceNameGeneration")
+	FString MaterialInstanceAssetPrefix = TEXT("MI_");
+	
+	UPROPERTY(BlueprintReadWrite, Blueprintable, Category = "CustomizableObjectInstanceBaker|ResourceNameGeneration")
+	FString MaterialDynamicInstanceAssetPrefix = TEXT("MID_");
+
+	UPROPERTY(BlueprintReadWrite, Blueprintable, Category = "CustomizableObjectInstanceBaker|ResourceNameGeneration")
+	FString MaterialConstantInstanceAssetPrefix = TEXT("MIC_");
+	
 	/**
 	 * Callback executed once the baking operation gets completed. It will return the end success state and also some data about the assets saved.
 	 */
diff --git a/Engine/Plugins/Mutable/Source/CustomizableObjectEditor/Internal/MuCOE/CustomizableObjectInstanceBakingUtils.h b/Engine/Plugins/Mutable/Source/CustomizableObjectEditor/Internal/MuCOE/CustomizableObjectInstanceBakingUtils.h
--- a/Engine/Plugins/Mutable/Source/CustomizableObjectEditor/Internal/MuCOE/CustomizableObjectInstanceBakingUtils.h
+++ b/Engine/Plugins/Mutable/Source/CustomizableObjectEditor/Internal/MuCOE/CustomizableObjectInstanceBakingUtils.h
@@ -9,28 +9,30 @@
  * @param InInstance The instance we want to update so we can later bake it's resources.
  * @param InInstanceUpdateDelegate Delegate to be called after the instance gets updated.
  */
-CUSTOMIZABLEOBJECTEDITOR_API
-void ScheduleInstanceUpdateForBaking(UCustomizableObjectInstance& InInstance, FInstanceUpdateNativeDelegate& InInstanceUpdateDelegate);
+CUSTOMIZABLEOBJECTEDITOR_API void ScheduleInstanceUpdateForBaking(UCustomizableObjectInstance& InInstance, FInstanceUpdateNativeDelegate& InInstanceUpdateDelegate);
+
+
+UE_DEPRECATED(5.6, "Use the function with the same name that uses an FBakingConfiguration as input instead.")
+CUSTOMIZABLEOBJECTEDITOR_API bool BakeCustomizableObjectInstance (UCustomizableObjectInstance& InInstance,
+ const FString& FileName, 
+ const FString& AssetPath,
+ const bool bExportAllResources,
+ const bool bGenerateConstantMaterialInstances, 
+ bool bHasPermissionToOverride,
+ bool bIsUnattendedExecution, 
+ TArray<TPair<EPackageSaveResolutionType,UPackage*>>& OutSavedPackages);
 
 
 /**
  * Serializes onto disk the resources used by the targeted Customizable Object Instance. The operation can be configured in the FInstanceBakingSettings settings object.
  * @param InInstance The mutable COI instance whose resources we want to bake onto disk
- * @param FileName The base name for the all the files generated during the baking process
- * @param AssetPath The path where to save the assets
- * @param bExportAllResources Determines if we want a full or partial export
- * @param bGenerateConstantMaterialInstances Determines if we want to generate constant material instances or not
- * @param bHasPermissionToOverride Do we want to allow the overriding of assets? This value will get updated internally if not Unattended. 
+ * @param Configuration The baking configuration object with all the settings to be used for this baking operation.
  * @param bIsUnattendedExecution Determines if we want to run this method and show promps for user interaction or if we want to automatically chose the more sensible option (without the need of user interaction)
  * @param OutSavedPackages A collection with the packages marked for save.
  * @return True if the baking operation could be completed without issues and false otherwise
  */
-CUSTOMIZABLEOBJECTEDITOR_API
-bool BakeCustomizableObjectInstance (UCustomizableObjectInstance& InInstance,
- const FString& FileName, 
- const FString& AssetPath,
- const bool bExportAllResources,
- const bool bGenerateConstantMaterialInstances, 
- bool bHasPermissionToOverride,
- bool bIsUnattendedExecution, 
- TArray<TPair<EPackageSaveResolutionType,UPackage*>>& OutSavedPackages);
\ No newline at end of file
+CUSTOMIZABLEOBJECTEDITOR_API bool BakeCustomizableObjectInstance(
+	UCustomizableObjectInstance& InInstance,
+	const FBakingConfiguration& Configuration,
+	bool bIsUnattendedExecution,
+	TMap<UPackage*,EPackageSaveResolutionType>& OutSavedPackages);
\ No newline at end of file
diff --git a/Engine/Plugins/Mutable/Source/CustomizableObjectEditor/Private/MuCOE/CustomizableObjectEditorViewportClient.cpp b/Engine/Plugins/Mutable/Source/CustomizableObjectEditor/Private/MuCOE/CustomizableObjectEditorViewportClient.cpp
--- a/Engine/Plugins/Mutable/Source/CustomizableObjectEditor/Private/MuCOE/CustomizableObjectEditorViewportClient.cpp
+++ b/Engine/Plugins/Mutable/Source/CustomizableObjectEditor/Private/MuCOE/CustomizableObjectEditorViewportClient.cpp
@@ -1608,14 +1608,18 @@ void FCustomizableObjectEditorViewportClient::OnInstanceForBakingUpdate(const FU
 	
 	if (FolderDlg->ShowModal() != EAppReturnType::Cancel)
 	{
-		TArray<TPair<EPackageSaveResolutionType,UPackage*>> SavedPackages;
+		FBakingConfiguration Configuration;
+		Configuration.OutputFilesBaseName = FolderDlg->GetFileName();
+		Configuration.OutputPath = FolderDlg->GetAssetPath();
+		Configuration.bExportAllResourcesOnBake = FolderDlg->GetExportAllResources();
+		Configuration.bGenerateConstantMaterialInstancesOnBake = FolderDlg->GetGenerateConstantMaterialInstances();
+
+		// TODO: UE-266716 Add a way to setup the prefixes from the editor dialog
+		
+		TMap<UPackage*, EPackageSaveResolutionType> SavedPackages;
 		BakeCustomizableObjectInstance(
 			*BakeTempInstance,
-			FolderDlg->GetFileName(),
-			FolderDlg->GetAssetPath(),
-			FolderDlg->GetExportAllResources(),
-			FolderDlg->GetGenerateConstantMaterialInstances(),
-			false,
+			Configuration,
 			false,
 			SavedPackages);
 	}
diff --git a/Engine/Plugins/Mutable/Source/CustomizableObjectEditor/Private/MuCOE/CustomizableObjectInstanceBaker.cpp b/Engine/Plugins/Mutable/Source/CustomizableObjectEditor/Private/MuCOE/CustomizableObjectInstanceBaker.cpp
--- a/Engine/Plugins/Mutable/Source/CustomizableObjectEditor/Private/MuCOE/CustomizableObjectInstanceBaker.cpp
+++ b/Engine/Plugins/Mutable/Source/CustomizableObjectEditor/Private/MuCOE/CustomizableObjectInstanceBaker.cpp
@@ -36,8 +36,8 @@ void UCustomizableObjectInstanceBaker::BakeInstance(UCustomizableObjectInstance*
 		return;
 	}
 
-	// If no name is provided then compose one
-	Configuration.OutputFilesBaseName = !Configuration.OutputFilesBaseName.IsEmpty() ? Configuration.OutputFilesBaseName : InstanceToBake.GetName();	
+	// If no name is provided then compose one (allow not having a name)
+	// Configuration.OutputFilesBaseName = !Configuration.OutputFilesBaseName.IsEmpty() ? Configuration.OutputFilesBaseName : InstanceToBake.GetName();	
 
 	// Schedule update of the instance for baking
 	UpdateInstance();
@@ -80,21 +80,12 @@ void UCustomizableObjectInstanceBaker::Bake() const
 	GIsRunningUnattendedScript = true;
 	
 	// Array with all the packages that did get saved. Packages that failed the saving procedure will not appear in this array.
-	TArray<TPair<EPackageSaveResolutionType,UPackage*>> SavedPackages;
+	TMap<UPackage*, EPackageSaveResolutionType> SavedPackages;
 
 	// Ensures no interaction from the user will be required (alongside GIsRunningUnattendedScript)
 	constexpr bool bIsUnhandledExecution = true;
 	
-	const bool bBakeWasSuccessful = BakeCustomizableObjectInstance
-	(*InstanceToBake,
-		Configuration.OutputFilesBaseName,
-		Configuration.OutputPath,
-		Configuration.bExportAllResourcesOnBake,
-		Configuration.bGenerateConstantMaterialInstancesOnBake,
-		Configuration.bAllowOverridingOfFiles,
-		bIsUnhandledExecution,					
-		SavedPackages
-	);
+	const bool bBakeWasSuccessful = BakeCustomizableObjectInstance(*InstanceToBake, Configuration, bIsUnhandledExecution, SavedPackages);
 	
 	// Generate the array of structs to output based on the packages reported as saved
 	TArray<FBakedResourceData> SavedPackagesData;
@@ -102,11 +93,11 @@ void UCustomizableObjectInstanceBaker::Bake() const
 		// Display a list with the paths of the assets saved
 		UE_LOG(LogMutable, Log, TEXT("Saved assets (%d) :"), SavedPackages.Num());
 		SavedPackagesData.Reserve(SavedPackages.Num());
-		for (const TPair<EPackageSaveResolutionType, UPackage*>& SavedPackagePair : SavedPackages)
+		for (const TPair<UPackage*, EPackageSaveResolutionType>& SavedPackagePair : SavedPackages)
 		{
 			FBakedResourceData BakeData;
-			BakeData.SaveType = SavedPackagePair.Key;
-			BakeData.AssetPath = SavedPackagePair.Value->GetPathName();
+			BakeData.AssetPath = SavedPackagePair.Key->GetPathName();
+			BakeData.SaveType = SavedPackagePair.Value;
 
 			// Sanity checks
 			check(BakeData.SaveType != EPackageSaveResolutionType::None);
diff --git a/Engine/Plugins/Mutable/Source/CustomizableObjectEditor/Private/MuCOE/CustomizableObjectInstanceBakingUtils.cpp b/Engine/Plugins/Mutable/Source/CustomizableObjectEditor/Private/MuCOE/CustomizableObjectInstanceBakingUtils.cpp
--- a/Engine/Plugins/Mutable/Source/CustomizableObjectEditor/Private/MuCOE/CustomizableObjectInstanceBakingUtils.cpp
+++ b/Engine/Plugins/Mutable/Source/CustomizableObjectEditor/Private/MuCOE/CustomizableObjectInstanceBakingUtils.cpp
@@ -7,7 +7,6 @@
 
 #include "Misc/MessageDialog.h"
 #include "FileHelpers.h"
-#include "ObjectTools.h"
 #include "UnrealBakeHelpers.h"
 #include "Animation/Skeleton.h"
 #include "Engine/SkeletalMesh.h"
@@ -26,16 +25,6 @@
 
 #define LOCTEXT_NAMESPACE "CustomizableObjectEditor" 
 
-// Recommended prefixes for various types of UE assets
-const FString SkeletalMeshAssetPrefix = TEXT("SK_");
-const FString MaterialAssetPrefix = TEXT("M_");
-const FString TextureAssetPrefix = TEXT("T_");
-const FString MaterialInstanceAssetPrefix = TEXT("MI_");
-const FString MaterialDynamicInstanceAssetPrefix = TEXT("MID_");
-const FString MaterialConstantInstanceAssetPrefix = TEXT("MIC_");
-const FString SkeletonAssetPrefix = TEXT("SKEL_");
-const FString PhysicsAssetPrefix = TEXT("PHYS_");
-
 
 /**
  * Remove the prefix from the provided string.
@@ -54,6 +43,7 @@ void RemovePrefix(FString& InOutString, const TArray<FString>& ToCheckPrefixes)
 	}
 }
 
+
 /**
  * Compose the resource name to be used by a resource.
  * @param InPrefix The prefix we want to use for a given resource.
@@ -83,16 +73,18 @@ void ShowErrorNotification(const FText& InMessage, const FText& InTitle = LOCTEX
 	}
 }
 
+
 /**
  * Utility functions for the baking operation.
  */
 
+
 /**
  * Validates the filename chosen for the baking data
- * @param FileName The filename chosen by the user
+ * @param InFileName The filename chosen by the user
  * @return True if validation was successful, false otherwise
  */
-bool ValidateProvidedFileName(const FString& FileName)
+bool ValidateProvidedFileName(const FString& InFileName)
 {
 	// Check for invalid characters in the name of the object to be serialized
 	TCHAR InvalidCharacter = '0';
@@ -102,7 +94,7 @@ bool ValidateProvidedFileName(const FString& FileName)
 		{
 			TCHAR Char = InvalidCharacters[InvalidCharIndex];
 			FString SearchedChar = FString::Chr(Char);
-			if (FileName.Contains(SearchedChar))
+			if (InFileName.Contains(SearchedChar))
 			{
 				InvalidCharacter = InvalidCharacters[InvalidCharIndex];
 				break;
@@ -154,23 +146,44 @@ bool ValidateProvidedAssetPath(const FString& FileName, const FString& AssetPath
 }
 
 
+/**
+ * Given a collection of UObjects return the index of the one that matches the name also provided
+ * @param InName The name of the object we are looking for
+ * @param Objects The collection of objects we want to check for an element with the given name
+ * @return The index of the element whose name matches the one provided or INDEX_NONE if the element could not be found
+ */
+int32 GetObjectWithNameIndex(const FString& InName, const TArray<UObject*>& Objects )
+{
+	int32 Index = 0;
+	for (UObject* CachedResource : Objects)
+	{
+		if (CachedResource && CachedResource->GetName() == InName)
+		{
+			return Index;
+		}
+		Index++;
+	}
+
+	return INDEX_NONE;
+}
+
 
 /**
  * Outputs a string that we know it is unique.
  * @param ResourceName The name of the resource we have provided. This should have the name of the current resource and will have the unique name for the resource once the method exits
- * @param InCachedResourceNames Collection with all the already processed resources name's
+ * @param CachedResources Collection with all the already processed resources name's
  */
-void MakeResourceNameUnique(FString& ResourceName, const TArray<FString>& InCachedResourceNames)
+void MakeResourceNameUnique(FString& ResourceName, const TArray<UObject*>& CachedResources)
 {
 	// Look for the resource name provided to see if we have already worked with it.
-	int32 FindResult = InCachedResourceNames.Find(ResourceName);
+	int32 FindResult = GetObjectWithNameIndex(ResourceName, CachedResources);
 	if (FindResult != INDEX_NONE)
 	{
 		// Add an integer suffix to create the unique name
 		uint32 Count = 0;
 		while (FindResult != INDEX_NONE)
 		{
-			FindResult = InCachedResourceNames.Find(ResourceName + "_" + FString::FromInt(Count));
+			FindResult = GetObjectWithNameIndex(ResourceName + "_" + FString::FromInt(Count), CachedResources);
 			Count++;
 		}
 
@@ -180,130 +193,62 @@ void MakeResourceNameUnique(FString& ResourceName, const TArray<FString>& InCach
 
 
 /**
- * Ensures the resource we want to save is ready to be saved. It handles closing it's editor and warning the user about possible overriding of resources.
- * @param InAssetSavePath The directory path where to save the baked object
- * @param InObjName The name of the object to be baked
- * @param bOverridePermissionGranted Control flag that determines if the user has given or not permission to override resources already in disk
- * @param bIsUnattended
- * @param OutSaveResolution
- * @return True if the operation was successful, false otherwise
+ * Get the package at a given path and with a given name
+ * @param InPackageSavePath The path were we plan to save the package
+ * @param InObjName The name used by the package
+ * @return True if the asset is already at path
  */
-bool CreateAssetAtPath(const FString& InAssetSavePath, const FString& InObjName, bool& bOverridePermissionGranted, const bool& bIsUnattended, EPackageSaveResolutionType& OutSaveResolution)
+UPackage* GetPackage(const FString& InPackageSavePath, const FString& InObjName)
 {
-	// Before the value provided by "bOverridePermissionGranted" was being updated due to user request but not it is not. It will stay as is if unattended an
-	// will get updated if this gets to be an attended execution.
-	
-	const FString PackagePath = InAssetSavePath + "/" + InObjName;
+	// Check if the package already exists
+	const FString PackagePath = InPackageSavePath + "/" + InObjName;
 	UPackage* ExistingPackage = FindPackage(nullptr, *PackagePath);
-
 	if (!ExistingPackage)
 	{
 		const FString PackageFilePath = PackagePath + "." + InObjName;
-
+		
 		FString PackageFileName;
 		if (FPackageName::DoesPackageExist(PackageFilePath, &PackageFileName))
 		{
 			ExistingPackage = LoadPackage(nullptr, *PackageFileName, LOAD_EditorOnly);
 		}
-		else
-		{
-			// if package does not exist
-			
-			if (!bIsUnattended)
-			{
-				// If the run is attended (the user is participating in it) then we will take care in consideration his decision what he wants.
-				bOverridePermissionGranted = false;
-			}
-
-			OutSaveResolution = EPackageSaveResolutionType::NewFile;
-			return true;
-		}
 	}
 
-	if (ExistingPackage)
-	{
-		// Checking if the asset is open in an editor
-		TArray<IAssetEditorInstance*> ObjectEditors = GEditor->GetEditorSubsystem<UAssetEditorSubsystem>()->FindEditorsForAssetAndSubObjects(ExistingPackage);
-		if (ObjectEditors.Num())
-		{
-			for (IAssetEditorInstance* ObjectEditorInstance : ObjectEditors)
-			{
-				// Close the editors that contains this asset
-				if (!ObjectEditorInstance->CloseWindow(EAssetEditorCloseReason::AssetEditorHostClosed))
-				{
-					const FText Caption = LOCTEXT("CustomizableObjectInstanceBakingUtils_OpenExistingFile", "Open File");
-					const FText Message = FText::Format(LOCTEXT("CustomizableObjectInstanceBakingUtils_CantCloseAsset", "This Object \"{0}\" is open in an editor and can't be closed automatically. Please close the editor and try to bake it again"), FText::FromString(InObjName));
-
-					ShowErrorNotification(Message, Caption);
-					
-					return false;
-				}
-			}
-		}
-
-		// If the execution requires user interaction, and we have no permission to override the existing file ask him if he wants or not to override data
-		if (!bIsUnattended && !bOverridePermissionGranted)
-		{
-			check (!FApp::IsUnattended())
-			const FText Caption = LOCTEXT("CustomizableObjectInstanceBakingUtils_AlreadyExistingBakedFiles", "Already existing baked files");
-			const FText Message = FText::Format(LOCTEXT("CustomizableObjectInstanceBakingUtils_OverwriteBakedInstance", "Instance baked files already exist in selected destination \"{0}\", this action will overwrite them."), FText::AsCultureInvariant(InAssetSavePath));
-
-			if (FMessageDialog::Open(EAppMsgType::OkCancel, Message, Caption) == EAppReturnType::Cancel)
-			{
-				return false;		// if the user cancels then we will still have no rights for overriding data
-			}
+	return ExistingPackage;
+}
 
-			UE_LOG(LogMutable, Error, TEXT("%s - %s"), *Caption.ToString(), *Message.ToString());
-			
-			// If the user accepts the prompt then we will consider we have a green light to override the asset
-			bOverridePermissionGranted = true;
-		}
-		
-		// At this point we may or may not have permission to delete the existing asset
-		
-		// Delete the old asset if we have permission to do so
-		if (UObject* ExistingObject = StaticFindObject(UObject::StaticClass(), ExistingPackage, *InObjName))
-		{
-			// Based on if we have or not permission to override the file do or do not so
-			if (bOverridePermissionGranted)
-			{
-				ExistingPackage->FullyLoad();
 
-				TArray<UObject*> ObjectsToDelete;
-				ObjectsToDelete.Add(ExistingObject);
-				
-				const FText Message = FText::Format(LOCTEXT("CustomizableObjectInstanceBakingUtils_AssetOverriden", "The COI asset \"{0}\" already exists and will be overriden due to user demand."), FText::FromString(ExistingPackage->GetName()));
-				UE_LOG(LogMutable,Warning,TEXT("%s"), *Message.ToString());
+/**
+ * Get the save resolution type for a given package. Will return or New or Reuse type depending if the package already exists on disk or not
+ * @param InPackageSavePath The path where to save the package
+ * @param InPackageName The name to be given to the package that we want to save or reuse
+ * @return 
+ */
+EPackageSaveResolutionType GetPackageSaveResolution(const FString& InPackageSavePath, const FString& InPackageName)
+{
+	EPackageSaveResolutionType SaveResolutionType = EPackageSaveResolutionType::None;
+	if (GetPackage(InPackageSavePath, InPackageName))
+	{
+		// File found, instead of creating a new package use this one and update the contents it has
+		SaveResolutionType = EPackageSaveResolutionType::ReusedFile;
+	}
+	else
+	{
+		// The package could not be found. It is safe to create a package at this path and with the given name
+		SaveResolutionType = EPackageSaveResolutionType::NewFile;
+	}
 
-				// Notify the caller we did proceed with the override (performed later)
-				OutSaveResolution = EPackageSaveResolutionType::Overriden;
-				
-				// Delete objects in the package with the same name as the one we want to create
-				const uint32 NumObjectsDeleted = ObjectTools::ForceDeleteObjects(ObjectsToDelete, false);
-				return NumObjectsDeleted == ObjectsToDelete.Num();
-			}
-			else
-			{
-				// Notify the caller that the override will not be performed
-				OutSaveResolution = EPackageSaveResolutionType::UnableToOverride;
-				
-				// Report that the file will not get overriden since we have no permission to do so
-				const FText UnableToOverrideMessage = FText::Format(LOCTEXT("CustomizableObjectInstanceBakingUtils_AssetCanNotBeOverriden", "Could not replace the COI asset \"{0}\" as it already exists."), FText::FromString(ExistingPackage->GetName()));
-				UE_LOG(LogMutable,Error,TEXT("%s"), *UnableToOverrideMessage.ToString());
+	// Just to be sure this case never happens even with future changes
+	check(SaveResolutionType == EPackageSaveResolutionType::ReusedFile || SaveResolutionType == EPackageSaveResolutionType::NewFile);
 
-				return false;
-			}
-		}
-		else
-		{
-			// Notify the caller that no override was required
-			OutSaveResolution = EPackageSaveResolutionType::NewFile;
-		}
-	}
+	const UEnum* OptimizationLevelEnum = StaticEnum<EPackageSaveResolutionType>();
+	check(OptimizationLevelEnum);
+	UE_LOG(LogMutable, Verbose, TEXT("SaveResolution for the \"%s\" package is \"%s\""), *InPackageSavePath, *OptimizationLevelEnum->GetValueAsName(SaveResolutionType).ToString())
 
-	return true;
+	return SaveResolutionType;
 }
 
+
 static bool bIsBakeOperationAlreadyScheduled = false;
 
 
@@ -364,17 +309,87 @@ void ScheduleInstanceUpdateForBaking(UCustomizableObjectInstance& InInstance, FI
 }
 
 
+/**
+ * Checks if the texture was generated by mutable or if it was not
+ * @param InTexture The texture we want to check
+ * @return true if the texture was generated by mutable, false otherwise
+ */
+bool IsAMutableTexture (const UTexture2D* InTexture)
+{
+	bool bIsMutableTexture = false;
+	for (UAssetUserData* UserData : *InTexture->GetAssetUserDataArray())
+	{
+		if (Cast<UMutableTextureMipDataProviderFactory>(UserData))
+		{
+			bIsMutableTexture = true;
+			break;
+		}
+	}
+	return bIsMutableTexture;
+}
+
+
+/**
+ * Add the package and save resolution only once to prevent the overriding of packages marked as New with Reuse ones if processed multiple times
+ * @param NewObjectPair The Pair of data we want to add to the map. Only if the package is not present in the map the action will get performed.
+ * @param OutSavedPackages The map you want to add data into.
+ */
+void AddUniqueToSavePackage(TPair<UPackage*, EPackageSaveResolutionType> NewObjectPair, TMap<UPackage*,EPackageSaveResolutionType>& OutSavedPackages)
+{
+	if (!OutSavedPackages.Find(NewObjectPair.Key))
+	{
+		OutSavedPackages.Add(NewObjectPair);
+	}
+}
+
+
+bool BakeCustomizableObjectInstance(UCustomizableObjectInstance& InInstance, const FString& FileName, const FString& AssetPath,
+	const bool bExportAllResources, const bool bGenerateConstantMaterialInstances, bool bHasPermissionToOverride, bool bIsUnattendedExecution,
+	TArray<TPair<EPackageSaveResolutionType, UPackage*>>& OutSavedPackages)
+{
+	FBakingConfiguration BakingConfiguration;
+	BakingConfiguration.OutputPath = AssetPath;
+	BakingConfiguration.OutputFilesBaseName = FileName;
+	BakingConfiguration.bExportAllResourcesOnBake = bExportAllResources;
+	BakingConfiguration.bGenerateConstantMaterialInstancesOnBake = bGenerateConstantMaterialInstances;
+
+	TMap<UPackage*, EPackageSaveResolutionType> SavedPackages;
+	bool bWasOperationSuccessful = BakeCustomizableObjectInstance(InInstance, BakingConfiguration, bIsUnattendedExecution, SavedPackages);
+
+	// Generate the array this method expects as an output based on the "SavedPackages" map
+	OutSavedPackages.Reset(SavedPackages.Num());
+	for (TPair<UPackage*,EPackageSaveResolutionType> SavedPackagesEntry : SavedPackages)
+	{
+		TPair<EPackageSaveResolutionType, UPackage*> ArrayEntry = {SavedPackagesEntry.Value, SavedPackagesEntry.Key};
+		OutSavedPackages.Add(ArrayEntry);
+	}
+
+	return bWasOperationSuccessful;
+}
+
+
 bool BakeCustomizableObjectInstance(
 	UCustomizableObjectInstance& InInstance,
-	const FString& FileName,
-	const FString& AssetPath,
-	const bool bExportAllResources,
-	const bool bGenerateConstantMaterialInstances,
-	bool bHasPermissionToOverride,
+	const FBakingConfiguration& Configuration,
 	bool bIsUnattendedExecution,
-	TArray<TPair<EPackageSaveResolutionType,UPackage*>>& OutSavedPackages)
+	TMap<UPackage*,EPackageSaveResolutionType>& OutSavedPackages)
 {
 	OutSavedPackages.Reset();
+
+	// Extract some baking settings from the configuration file provided
+	const FString& FileName = Configuration.OutputFilesBaseName;
+	const FString& AssetPath = Configuration.OutputPath;
+	
+	// Resource prefixes to be used. If an invalid prefix has been provided by the configuration object use the UE default one instead
+	const FString SkeletalMeshAssetPrefix = Configuration.SkeletalMeshAssetPrefix.IsEmpty() ? TEXT("SK_") :  Configuration.SkeletalMeshAssetPrefix;
+	const FString SkeletonAssetPrefix = Configuration.SkeletonAssetPrefix.IsEmpty() ? TEXT("SKEL_") : Configuration.SkeletonAssetPrefix;
+	const FString PhysicsAssetPrefix = Configuration.PhysicsAssetPrefix.IsEmpty() ? TEXT("PHYS_") : Configuration.PhysicsAssetPrefix;
+	const FString MaterialAssetPrefix = Configuration.MaterialAssetPrefix.IsEmpty() ? TEXT("M_") : Configuration.MaterialAssetPrefix;
+	const FString TextureAssetPrefix = Configuration.TextureAssetPrefix.IsEmpty() ? TEXT("T_") : Configuration.TextureAssetPrefix;
+	const FString MaterialInstanceAssetPrefix = Configuration.MaterialInstanceAssetPrefix.IsEmpty() ? TEXT("MI_") : Configuration.MaterialInstanceAssetPrefix;
+	const FString MaterialDynamicInstanceAssetPrefix = Configuration.MaterialDynamicInstanceAssetPrefix.IsEmpty() ? TEXT("MID_") : Configuration.MaterialDynamicInstanceAssetPrefix;
+	const FString MaterialConstantInstanceAssetPrefix = Configuration.MaterialConstantInstanceAssetPrefix.IsEmpty() ? TEXT("MIC_") : Configuration.MaterialConstantInstanceAssetPrefix;
+
 	
 	// Ensure that the state of the COI provided is valid --------------------------------------------------------------------------------------------
 	UCustomizableObject* InstanceCO = InInstance.GetCustomizableObject();
@@ -444,25 +459,24 @@ bool BakeCustomizableObjectInstance(
 		.Log();
 	}
 	
-	// Set the overriding flag to true or false:
-	//	- We ask the user at least once about if he is willing to override old baked data (attended operation) and this makes the flag change 
-	//	- We never ask the user (and therefore the value in bUsedGrantedOverridingRights never changes) when we work in Unattended mode.
-	bool bUsedGrantedOverridingRights = bHasPermissionToOverride;
-	if (FApp::IsUnattended() || GIsRunningUnattendedScript)
-	{
-		bIsUnattendedExecution = true;
-	}
-
 	// Prefixes used for UMaterial assets. 
-	const TArray<FString> MaterialResourcePrefixes = {MaterialAssetPrefix, MaterialInstanceAssetPrefix, MaterialDynamicInstanceAssetPrefix, MaterialConstantInstanceAssetPrefix};
+	const TArray<FString> MaterialResourcePrefixes =
+	{
+		MaterialAssetPrefix,
+		MaterialInstanceAssetPrefix,
+		MaterialDynamicInstanceAssetPrefix,
+		MaterialConstantInstanceAssetPrefix
+	};
 
 	// Array with the already processed resource names and resources (UObjects)
-	TArray<FString> CachedResourceNames;
-	TArray<UObject*> CachedObjects;
+	TArray<UObject*> DuplicatedObjects;
+	TArray<UObject*> HandledSourceObjects;
+	
 	
-	const int32 NumObjectComponents = InstanceCO->GetComponentCount();
-	for (int32 ObjectComponentIndex = 0; ObjectComponentIndex < NumObjectComponents; ++ObjectComponentIndex)
+	for (int32 ObjectComponentIndex = 0; ObjectComponentIndex < InstanceCO->GetComponentCount(); ++ObjectComponentIndex)
 	{
+		check(DuplicatedObjects.Num() == HandledSourceObjects.Num())
+		
 		const FName ComponentName = InstanceCO->GetPrivate()->GetComponentName(FCustomizableObjectComponentIndex(ObjectComponentIndex));
 		USkeletalMesh* Mesh = InInstance.GetComponentMeshSkeletalMesh(ComponentName);
 
@@ -470,9 +484,9 @@ bool BakeCustomizableObjectInstance(
 		{
 			continue;
 		}
-		
+
+		// Append the component name to the ObjectFileName
 		FString ObjectBaseName = FileName;
-		
 		if (FileName.IsEmpty())
 		{
 			ObjectBaseName = ComponentName.ToString();
@@ -481,29 +495,27 @@ bool BakeCustomizableObjectInstance(
 		{
 			ObjectBaseName = FileName + TEXT("_") + ComponentName.ToString();
 		}
-
 		ObjectBaseName += TEXT("_");
 		
 		TMap<UObject*, UObject*> ReplacementMap;
 
-		if (bExportAllResources)
+		if (Configuration.bExportAllResourcesOnBake)
 		{
-			UMaterialInstance* MaterialInstance;
 			UMaterial* Material;
-			UTexture* Texture;
 			FString PackageName;
-			UObject* DuplicatedObject;
-			
+
+			// Each element of the array corresponds to a material of the mesh
+			// The key is the parameter index and the value the texture to be used there
 			TArray<TMap<int, UTexture*>> TextureReplacementMaps;
 
 			// Duplicate Textures found in the Material Instances of the SkeletalMesh so we can later assign them to the
 			// duplicates of those material instances. At the end of the baking we will have a series of materials with the 
 			// parameters set as the material instances they are based of.
-			for (int32 m = 0; m < Mesh->GetMaterials().Num(); ++m)
+			for (int32 MaterialIndex = 0; MaterialIndex < Mesh->GetMaterials().Num(); ++MaterialIndex)
 			{
-				UMaterialInterface* Interface = Mesh->GetMaterials()[m].MaterialInterface;
+				UMaterialInterface* Interface = Mesh->GetMaterials()[MaterialIndex].MaterialInterface;
 				Material = Interface->GetMaterial();
-				MaterialInstance = Cast<UMaterialInstance>(Interface);
+				UMaterialInstance* MaterialInstance = Cast<UMaterialInstance>(Interface);
 
 				TextureReplacementMaps.AddDefaulted();
 				
@@ -515,6 +527,7 @@ bool BakeCustomizableObjectInstance(
 					
 					for (int32 ParameterInfoIndex = 0; ParameterInfoIndex < ParameterInfoObjects.Num(); ParameterInfoIndex++)
 					{
+						UTexture* Texture;
 						const FName& ParameterName = ParameterInfoObjects[ParameterInfoIndex].Name;
 						if (MaterialInstance->GetTextureParameterValue(ParameterName, Texture))
 						{
@@ -524,35 +537,42 @@ bool BakeCustomizableObjectInstance(
 								continue;
 							}
 							
-							// Check that the resource has not yet been processed, and if not, skip it
-							if (CachedObjects.Contains(SourceTexture))
-							{
-								continue;
-							}
-							
+							// Update the name adding to it the name of the component alongside other data
 							FString TextureName = SourceTexture->GetName();
 							RemovePrefix(TextureName, {TextureAssetPrefix});
 							ComposeResourceName(TextureAssetPrefix, ObjectBaseName, TextureName);
-							MakeResourceNameUnique(TextureName, CachedResourceNames);
-							
-							EPackageSaveResolutionType SaveType = EPackageSaveResolutionType::None;
-							if (!CreateAssetAtPath(AssetPath, TextureName, bUsedGrantedOverridingRights, bIsUnattendedExecution, SaveType))
-							{
-								return false;
-							}
-							
-							bool bIsMutableTexture = false;
-							for (UAssetUserData* UserData : *SourceTexture->GetAssetUserDataArray())
+
+							// Experimental
+							// If the texture is not from mutable and we know it has already been processed...
+							if (!IsAMutableTexture(SourceTexture) && HandledSourceObjects.Contains(SourceTexture))
 							{
-								if (Cast<UMutableTextureMipDataProviderFactory>(UserData))
+								// Instead of just creating a new name for the texture to be stored in disk, a better aproach could be using
+								// the already duplicated file and assigning it to the material we are working with.
+								// The issue this has is that the texture may be named as from one component and now will be used for multiple ones
+								for (UObject* DuplicatedObject : DuplicatedObjects)
 								{
-									bIsMutableTexture = true;
-									break;
+									if (DuplicatedObject &&
+										DuplicatedObject->IsA<UTexture>() &&
+										DuplicatedObject->GetName() == TextureName)
+									{
+										// TODO: Allow reusing a texture of one component in another one (DuplicatedObject->GetName() == TextureName)
+										//		to be able to do this I will need to stop using the name of the texture as the name has the component embedded on it
+										//		or, instead, remove the component name from the full name and then check
+											
+										UTexture* DupTexture = Cast<UTexture>(DuplicatedObject);
+										TextureReplacementMaps[MaterialIndex].Add(ParameterInfoIndex, DupTexture);
+
+										UE_LOG(LogMutable, Verbose, TEXT("Matching texture found with name %s in component %s"), *TextureName, *ComponentName.ToString())
+										break;
+									}
 								}
 							}
 							
+							// Here we will know if the asset already exists (ReusedFile) or if it is new (NewFile)
+							const EPackageSaveResolutionType SaveResolution = GetPackageSaveResolution(AssetPath, TextureName);
+							
 							// Duplicating mutable generated textures
-							if (bIsMutableTexture)
+							if (IsAMutableTexture(SourceTexture))
 							{
 								if (SourceTexture->GetPlatformData() && SourceTexture->GetPlatformData()->Mips.Num() > 0)
 								{
@@ -561,44 +581,46 @@ bool BakeCustomizableObjectInstance(
 									MaterialInstance->Parent->GetTextureParameterValue(FName(*ParameterName.GetPlainNameString()), OriginalTexture);
 
 									PackageName = AssetPath + FString("/") + TextureName;
-									TMap<UObject*, UObject*> DummyReplacementMap;
-									UTexture2D* DuplicatedTexture = FUnrealBakeHelpers::BakeHelper_CreateAssetTexture(SourceTexture, TextureName, PackageName, OriginalTexture, true, DummyReplacementMap, bUsedGrantedOverridingRights);
-									CachedResourceNames.Add(TextureName);
-									CachedObjects.Add(DuplicatedTexture);
-								
-									TPair<EPackageSaveResolutionType, UPackage*> PackageToSave {SaveType, DuplicatedTexture->GetPackage()};
-									OutSavedPackages.Add(PackageToSave);
+									UTexture2D* DuplicatedTexture = FUnrealBakeHelpers::BakeHelper_CreateAssetTexture(SourceTexture, TextureName, PackageName, OriginalTexture, nullptr, SaveResolution);
+									
+									HandledSourceObjects.Add(SourceTexture);
+									DuplicatedObjects.Add(DuplicatedTexture);
+									
+									AddUniqueToSavePackage({DuplicatedTexture->GetPackage(), SaveResolution},OutSavedPackages);
 
 									if (OriginalTexture != nullptr)
 									{
-										TextureReplacementMaps[m].Add(ParameterInfoIndex, DuplicatedTexture);
+										TextureReplacementMaps[MaterialIndex].Add(ParameterInfoIndex, DuplicatedTexture);
 									}
 								}
 							}
 							else
 							{
 								// Duplicate the non-mutable textures of the Material instance (pass-through textures)
-								
+
 								PackageName = AssetPath + FString("/") + TextureName;
-								TMap<UObject*, UObject*> DummyReplacementMap;
-								DuplicatedObject = FUnrealBakeHelpers::BakeHelper_DuplicateAsset(Texture, TextureName, PackageName, true, DummyReplacementMap, bUsedGrantedOverridingRights, false);
-								CachedResourceNames.Add(TextureName);
-								CachedObjects.Add(DuplicatedObject);
+								UObject* DuplicatedTexture = FUnrealBakeHelpers::BakeHelper_DuplicateAsset(Texture, TextureName, PackageName, nullptr, false, SaveResolution);
 
-								TPair<EPackageSaveResolutionType, UPackage*> PackageToSave {SaveType, DuplicatedObject->GetPackage()};
-								OutSavedPackages.Add(PackageToSave);
+								HandledSourceObjects.Add(Texture);
+								DuplicatedObjects.Add(DuplicatedTexture);
+								
+								AddUniqueToSavePackage({DuplicatedTexture->GetPackage(), SaveResolution},OutSavedPackages);
 
-								UTexture* DupTexture = Cast<UTexture>(DuplicatedObject);
-								TextureReplacementMaps[m].Add(ParameterInfoIndex, DupTexture);
+								UTexture* DupTexture = Cast<UTexture>(DuplicatedTexture);
+								TextureReplacementMaps[MaterialIndex].Add(ParameterInfoIndex, DupTexture);
 							}
 						}
 					}
 				}
 			}
-
-
+			
+			// At this point we have an array where each element represents one material index.
+			// Each element is formed by a map where the key is the parameter index and the value the texture to be used in that index so
+			// Later we will be able to update all materials of the component with the duplicated textures we have just created
+			
 			// Duplicate the materials used by each material instance so that the replacement map has proper information 
 			// when duplicating the material instances
+			// Each material will get filled with the data of the interface it is related to.
 			for (int32 MaterialIndex = 0; MaterialIndex < Mesh->GetMaterials().Num(); ++MaterialIndex)
 			{
 				UMaterialInterface* Interface = Mesh->GetMaterials()[MaterialIndex].MaterialInterface;
@@ -606,34 +628,41 @@ bool BakeCustomizableObjectInstance(
 				
 				if (Material)
 				{
-					// Check that the resource has not yet been processed, and if not, skip it
-					if (CachedObjects.Contains(Material))
+					// If the material has already been processed (so it is already part of the ReplacementMap) just skip it as we only need
+					// to update the Material Interface once (Material Interface <- Material)
+					if (HandledSourceObjects.Contains(Interface))
 					{
+						// should be the case if the interface has already been processed
+						check(ReplacementMap.Contains(Interface));	
 						continue;
 					}
-
-					FString MaterialName = Material->GetName();
-					RemovePrefix(MaterialName, {MaterialAssetPrefix});
+					
+					FString MaterialName = Interface->GetName();
+					RemovePrefix(MaterialName, MaterialResourcePrefixes);
 					ComposeResourceName(MaterialAssetPrefix, ObjectBaseName, MaterialName);
-					MakeResourceNameUnique(MaterialName, CachedResourceNames);
-
-					EPackageSaveResolutionType SaveType = EPackageSaveResolutionType::None;
-					if (!CreateAssetAtPath(AssetPath, MaterialName, bUsedGrantedOverridingRights, bIsUnattendedExecution, SaveType))
-					{
-						return false;
-					}
+					
+					// Give it an unique name or may happen that another material already stored on disk during this bake will get updated
+					// with data form this material. We want to avoid that from happening as we do not want to require the name to always be a hash of
+					// the contents (they are not always produced by mutable in this case)
+					MakeResourceNameUnique(MaterialName, DuplicatedObjects);
+					
+					const EPackageSaveResolutionType SaveResolution = GetPackageSaveResolution(AssetPath, MaterialName);
 
 					PackageName = AssetPath + FString("/") + MaterialName;
-					TMap<UObject*, UObject*> FakeReplacementMap;
-					DuplicatedObject = FUnrealBakeHelpers::BakeHelper_DuplicateAsset(Material, MaterialName, PackageName, 
-						false, FakeReplacementMap, bUsedGrantedOverridingRights, bGenerateConstantMaterialInstances);
-					CachedResourceNames.Add(MaterialName);
-					CachedObjects.Add(DuplicatedObject);
-					ReplacementMap.Add(Interface, DuplicatedObject);
+					UObject* DuplicatedObject = FUnrealBakeHelpers::BakeHelper_DuplicateAsset(Material, MaterialName, PackageName, 
+						 nullptr, Configuration.bGenerateConstantMaterialInstancesOnBake, SaveResolution);
 
-					TPair<EPackageSaveResolutionType, UPackage*> PackageToSave {SaveType, DuplicatedObject->GetPackage()};
-					OutSavedPackages.Add(PackageToSave);
+					HandledSourceObjects.Add(Material);
+					DuplicatedObjects.Add(DuplicatedObject);
 
+					// Trying to add the same interface twice would mean something is not being done correctly in this method
+					check(!ReplacementMap.Contains(Interface));		
+					// Tell the system that the Interface object will be replaced by the Duplicated one
+					ReplacementMap.Add(Interface, DuplicatedObject);
+					
+					AddUniqueToSavePackage({DuplicatedObject->GetPackage(), SaveResolution},OutSavedPackages);
+
+					// Copy the texture parameters from the interface to the material
 					if (UMaterial* DuplicatedMaterial = Cast<UMaterial>(DuplicatedObject))
 					{
 						FUnrealBakeHelpers::CopyAllMaterialParameters<UMaterial>(*DuplicatedMaterial, *Interface, TextureReplacementMaps[MaterialIndex]);
@@ -643,25 +672,25 @@ bool BakeCustomizableObjectInstance(
 		}
 		else
 		{
+			// Export only the mutable generated resources
+			
 			// Duplicate the material instances
 			for (int32 MaterialIndex = 0; MaterialIndex < Mesh->GetMaterials().Num(); ++MaterialIndex)
 			{
 				UMaterialInterface* Interface = Mesh->GetMaterials()[MaterialIndex].MaterialInterface;
-				UMaterial* ParentMaterial = Interface->GetMaterial();
 
-				// Check that the resource has not yet been processed, and if not, skip it
-				if (CachedObjects.Contains(Interface))
+				// Check that the resource has not yet been processed, and if so, skip it as we only want to duplicate each interface once
+				if (HandledSourceObjects.Contains(Interface))
 				{
 					continue;
 				}
 
 				FString MaterialName;
 				{
-					MaterialName = ParentMaterial ? ParentMaterial->GetName() : MaterialAssetPrefix;
-
+					MaterialName = Interface->GetName();		// Use the name of the interface as it is unique
 					RemovePrefix(MaterialName, MaterialResourcePrefixes);
 					
-					if (bGenerateConstantMaterialInstances && Interface->IsA<UMaterialInstance>())
+					if (Configuration.bGenerateConstantMaterialInstancesOnBake && Interface->IsA<UMaterialInstance>())
 					{
 						// Change the prefix to the Material Constant Instance since in this situation the new asset based on the Interface
 						// will be of the Material Constant Instance type
@@ -686,32 +715,29 @@ bool BakeCustomizableObjectInstance(
 							checkNoEntry();		// Invalid material type.
 						}
 					}
-					
-					MakeResourceNameUnique(MaterialName, CachedResourceNames);
-				}
-
-				EPackageSaveResolutionType SaveType = EPackageSaveResolutionType::None;
-				if (!CreateAssetAtPath(AssetPath, MaterialName, bUsedGrantedOverridingRights, bIsUnattendedExecution, SaveType))
-				{
-					return false;
 				}
 
+				// One material could be used by multiple material instances so, to be sure, make sure the name is unique so we do not reuse it during the bake
+				// We could be more specific in which cases we do this but I left it this way to keep the complexity as low as possible
+				MakeResourceNameUnique(MaterialName, DuplicatedObjects);
+				
+				const EPackageSaveResolutionType SaveResolution = GetPackageSaveResolution(AssetPath, MaterialName);
+				
 				FString MatPkgName = AssetPath + FString("/") + MaterialName;
-				UObject* DupMat = FUnrealBakeHelpers::BakeHelper_DuplicateAsset(Interface, MaterialName, 
-					MatPkgName, false, ReplacementMap, bUsedGrantedOverridingRights, bGenerateConstantMaterialInstances);
-				CachedObjects.Add(DupMat);
-				CachedResourceNames.Add(MaterialName);
-
-				TPair<EPackageSaveResolutionType, UPackage*> PackageToSave {SaveType, DupMat->GetPackage()};
-				OutSavedPackages.Add(PackageToSave);
+				UObject* DuplicatedMaterial = FUnrealBakeHelpers::BakeHelper_DuplicateAsset(Interface, MaterialName, 
+					MatPkgName, &ReplacementMap, Configuration.bGenerateConstantMaterialInstancesOnBake, SaveResolution);
+				HandledSourceObjects.Add(Interface);
+				DuplicatedObjects.Add(DuplicatedMaterial);
+				
+				AddUniqueToSavePackage({DuplicatedMaterial->GetPackage(), SaveResolution},OutSavedPackages);
 
 				// Only need to duplicate the generate textures if the original material is a dynamic instance
 				// If the material has Mutable textures, then it will be a dynamic material instance for sure
 				if (UMaterialInstance* MaterialInstance = Cast<UMaterialInstance>(Interface))
 				{
 					// Duplicate generated textures
-					UMaterialInstanceDynamic* InstDynamic = Cast<UMaterialInstanceDynamic>(DupMat);
-					UMaterialInstanceConstant* InstConstant = Cast<UMaterialInstanceConstant>(DupMat);
+					UMaterialInstanceDynamic* InstDynamic = Cast<UMaterialInstanceDynamic>(DuplicatedMaterial);
+					UMaterialInstanceConstant* InstConstant = Cast<UMaterialInstanceConstant>(DuplicatedMaterial);
 
 					if (InstDynamic || InstConstant)
 					{
@@ -723,9 +749,11 @@ bool BakeCustomizableObjectInstance(
 								{
 									if (UTexture2D* SourceTexture = Cast<UTexture2D>(MaterialInstance->TextureParameterValues[TextureIndex].ParameterValue))
 									{
-										if (CachedObjects.Contains(SourceTexture))
+										// If this source texture has already been processed during this bake operation, instead of creating a duplicate,
+										// just use the resource cached in the DuplicatedObjects array
+										if (HandledSourceObjects.Contains(SourceTexture))
 										{
-											UTexture* PrevTexture = Cast<UTexture>(CachedObjects[CachedResourceNames.Find(SourceTexture->GetName())]);
+											UTexture* PrevTexture = Cast<UTexture>(DuplicatedObjects[HandledSourceObjects.Find(SourceTexture)]);
 											check(PrevTexture);
 											
 											if (InstDynamic)
@@ -740,34 +768,35 @@ bool BakeCustomizableObjectInstance(
 											continue;
 										}
 
+										// This should never happen as we handle this in the previous conditional block where we early out
+										check(!HandledSourceObjects.Contains(SourceTexture))
+										
+										// The source texture has not yet been processed so duplicate it
+										
 										// Generate a new name for the duplicated texture so it does not collide with the original one
 										FString TextureName = SourceTexture->GetName();
 										RemovePrefix(TextureName, {TextureAssetPrefix});
 										ComposeResourceName(TextureAssetPrefix, ObjectBaseName, TextureName);
-										MakeResourceNameUnique(TextureName, CachedResourceNames);
 										
-										EPackageSaveResolutionType TextureSaveType = EPackageSaveResolutionType::None;
-										if (!CreateAssetAtPath(AssetPath, TextureName, bUsedGrantedOverridingRights, bIsUnattendedExecution, TextureSaveType))
-										{
-											return false;
-										}
+										// Check if the file is already in disk, so we can later decide what we do with that file
+										const EPackageSaveResolutionType TextureSaveResolution = GetPackageSaveResolution(AssetPath, TextureName);
 
 										FString TexPkgName = AssetPath + FString("/") + TextureName;
-										TMap<UObject*, UObject*> FakeReplacementMap;
-										UTexture2D* DupTex = FUnrealBakeHelpers::BakeHelper_CreateAssetTexture(SourceTexture, TextureName, TexPkgName, nullptr, false, FakeReplacementMap, bUsedGrantedOverridingRights);
-										CachedObjects.Add(DupTex);
-										CachedResourceNames.Add(TextureName);
+										UTexture2D* DuplicatedTexture = FUnrealBakeHelpers::BakeHelper_CreateAssetTexture(SourceTexture, TextureName, TexPkgName, nullptr, nullptr, SaveResolution);
 
-										TPair<EPackageSaveResolutionType, UPackage*> TexturePackageToSave {TextureSaveType, DupTex->GetPackage()};
-										OutSavedPackages.Add(TexturePackageToSave);
+										// Cache the file we just have handled and it's duplicate
+										HandledSourceObjects.Add(SourceTexture);
+										DuplicatedObjects.Add(DuplicatedTexture);
 										
+										AddUniqueToSavePackage({DuplicatedTexture->GetPackage(), TextureSaveResolution},OutSavedPackages);
+
 										if (InstDynamic)
 										{
-											InstDynamic->SetTextureParameterValue(MaterialInstance->TextureParameterValues[TextureIndex].ParameterInfo.Name, DupTex);
+											InstDynamic->SetTextureParameterValue(MaterialInstance->TextureParameterValues[TextureIndex].ParameterInfo.Name, DuplicatedTexture);
 										}
 										else if(InstConstant)
 										{
-											InstConstant->SetTextureParameterValueEditorOnly(MaterialInstance->TextureParameterValues[TextureIndex].ParameterInfo.Name, DupTex);
+											InstConstant->SetTextureParameterValueEditorOnly(MaterialInstance->TextureParameterValues[TextureIndex].ParameterInfo.Name, DuplicatedTexture);
 										}
 									}
 									else
@@ -795,6 +824,9 @@ bool BakeCustomizableObjectInstance(
 			}
 		}
 
+		// TODO: As it happens with the textures we may be able to re-use the other resources generated for one component into the baked data of the other
+		//		for now we will keep them separated for the assets being processed bellow
+		
 		// Get the clean name of the Skeletal mesh (without the embedded CO name)
 		FString SkeletalMeshName = Mesh->GetName();
 		if (Mesh->HasAnyFlags(RF_Transient))
@@ -804,7 +836,7 @@ bool BakeCustomizableObjectInstance(
 		}
 		RemovePrefix(SkeletalMeshName, {SkeletalMeshAssetPrefix});
 		ComposeResourceName(SkeletalMeshAssetPrefix, ObjectBaseName, SkeletalMeshName);
-		MakeResourceNameUnique(SkeletalMeshName, CachedResourceNames);
+		MakeResourceNameUnique(SkeletalMeshName, DuplicatedObjects);
 		
 		// Skeletal Mesh's Skeleton
 		if (USkeleton* Skeleton = Mesh->GetSkeleton())
@@ -812,26 +844,23 @@ bool BakeCustomizableObjectInstance(
 			const bool bTransient = Skeleton->GetPackage() == GetTransientPackage();
 
 			// Duplicate only if transient or export all assets.
-			if (bTransient || bExportAllResources)
+			if (bTransient || Configuration.bExportAllResourcesOnBake)
 			{
 				FString SkeletonName = Mesh->GetName();
 				RemovePrefix(SkeletonName, {SkeletalMeshAssetPrefix});
 				ComposeResourceName(SkeletonAssetPrefix, ObjectBaseName, SkeletonName);
-				MakeResourceNameUnique(SkeletonName, CachedResourceNames);
+				MakeResourceNameUnique(SkeletonName, DuplicatedObjects);
 				
-				EPackageSaveResolutionType SaveType = EPackageSaveResolutionType::None;
-				if (!CreateAssetAtPath(AssetPath, SkeletonName, bUsedGrantedOverridingRights, bIsUnattendedExecution, SaveType))
-				{
-					return false;
-				}
-
+				const EPackageSaveResolutionType SkeletonSaveResolution = GetPackageSaveResolution(AssetPath, SkeletonName);
 				FString SkeletonPkgName = AssetPath + FString("/") + SkeletonName;
 				UObject* DuplicatedSkeleton = FUnrealBakeHelpers::BakeHelper_DuplicateAsset(Skeleton, SkeletonName, 
-					SkeletonPkgName, false, ReplacementMap, bUsedGrantedOverridingRights, false);
+					SkeletonPkgName, &ReplacementMap, false, SkeletonSaveResolution);
+
+				HandledSourceObjects.Add(Skeleton);
+				DuplicatedObjects.Add(DuplicatedSkeleton);
+
+				AddUniqueToSavePackage({DuplicatedSkeleton->GetPackage(), SkeletonSaveResolution},OutSavedPackages);
 
-				CachedObjects.Add(DuplicatedSkeleton);
-				TPair<EPackageSaveResolutionType, UPackage*> PackageToSave {SaveType, DuplicatedSkeleton->GetPackage()};
-				OutSavedPackages.Add(PackageToSave);
 				ReplacementMap.Add(Skeleton, DuplicatedSkeleton);
 			}
 		}
@@ -843,26 +872,23 @@ bool BakeCustomizableObjectInstance(
 			const bool bTransient = PhysicsAsset->GetPackage() == GetTransientPackage();
 
 			// Duplicate only if transient or export all assets.
-			if (bTransient || bExportAllResources)
+			if (bTransient || Configuration.bExportAllResourcesOnBake)
 			{
 				FString PhysicsAssetName = Mesh->GetName();
 				RemovePrefix(PhysicsAssetName, {SkeletalMeshAssetPrefix});
 				ComposeResourceName(PhysicsAssetPrefix, ObjectBaseName, PhysicsAssetName);
-				MakeResourceNameUnique(PhysicsAssetName, CachedResourceNames);
+				MakeResourceNameUnique(PhysicsAssetName, DuplicatedObjects);
 				
-				EPackageSaveResolutionType SaveType = EPackageSaveResolutionType::None;
-				if (!CreateAssetAtPath(AssetPath, PhysicsAssetName, bUsedGrantedOverridingRights, bIsUnattendedExecution, SaveType))
-				{
-					return false;
-				}
-
+				const EPackageSaveResolutionType PhysicsAssetSaveResolution = GetPackageSaveResolution(AssetPath, PhysicsAssetName);
 				FString PhysicsAssetPkgName = AssetPath + FString("/") + PhysicsAssetName;
 				UObject* DuplicatedPhysicsAsset = FUnrealBakeHelpers::BakeHelper_DuplicateAsset(PhysicsAsset, PhysicsAssetName, 
-					PhysicsAssetPkgName, false, ReplacementMap, bUsedGrantedOverridingRights, false);
+					PhysicsAssetPkgName, &ReplacementMap, false, PhysicsAssetSaveResolution);
+
+				HandledSourceObjects.Add(PhysicsAsset);
+				DuplicatedObjects.Add(DuplicatedPhysicsAsset);
+
+				AddUniqueToSavePackage({DuplicatedPhysicsAsset->GetPackage(), PhysicsAssetSaveResolution},OutSavedPackages);
 
-				CachedObjects.Add(DuplicatedPhysicsAsset);
-				TPair<EPackageSaveResolutionType, UPackage*> PackageToSave {SaveType, DuplicatedPhysicsAsset->GetPackage()};
-				OutSavedPackages.Add(PackageToSave);
 				ReplacementMap.Add(PhysicsAsset, DuplicatedPhysicsAsset);
 
 				bNewPhysicsAssetCreated = true;
@@ -871,23 +897,18 @@ bool BakeCustomizableObjectInstance(
 
 
 		// Skeletal Mesh
-		EPackageSaveResolutionType SaveType = EPackageSaveResolutionType::None;
-		if (!CreateAssetAtPath(AssetPath, SkeletalMeshName, bUsedGrantedOverridingRights, bIsUnattendedExecution, SaveType))
-		{
-			return false;
-		}
-
+		const EPackageSaveResolutionType MeshSaveResolution = GetPackageSaveResolution(AssetPath, SkeletalMeshName);
 		FString PkgName = AssetPath + FString("/") + SkeletalMeshName;
-		UObject* DupObject = FUnrealBakeHelpers::BakeHelper_DuplicateAsset(Mesh, SkeletalMeshName, PkgName, 
-			false, ReplacementMap, bUsedGrantedOverridingRights, false);
-		CachedObjects.Add(DupObject);
-
-		TPair<EPackageSaveResolutionType, UPackage*> PackageToSave {SaveType, DupObject->GetPackage()};
-		OutSavedPackages.Add(PackageToSave);
+		UObject* DuplicatedMesh = FUnrealBakeHelpers::BakeHelper_DuplicateAsset(Mesh, SkeletalMeshName, PkgName, 
+			&ReplacementMap, false, MeshSaveResolution);
+		HandledSourceObjects.Add(Mesh);
+		DuplicatedObjects.Add(DuplicatedMesh);
+		
+		AddUniqueToSavePackage({DuplicatedMesh->GetPackage(), MeshSaveResolution},OutSavedPackages);
 
 		Mesh->Build();
 
-		if (USkeletalMesh* SkeletalMesh = Cast<USkeletalMesh>(DupObject))
+		if (USkeletalMesh* SkeletalMesh = Cast<USkeletalMesh>(DuplicatedMesh))
 		{
 			SkeletalMesh->ResetLODInfo();
 			for (int32 LODIndex = 0; LODIndex < Mesh->GetLODNum(); ++LODIndex)
@@ -954,8 +975,10 @@ bool BakeCustomizableObjectInstance(
 			SkeletalMesh->Build();
 		}
 
+		check(DuplicatedObjects.Num() == HandledSourceObjects.Num());
+		
 		// Remove duplicated UObjects from Root (previously added to avoid objects from being GC in the middle of the bake process)
-		for (UObject* Obj : CachedObjects)
+		for (UObject* Obj : DuplicatedObjects)
 		{
 			Obj->RemoveFromRoot();
 		}
@@ -969,9 +992,12 @@ bool BakeCustomizableObjectInstance(
 		// Prepare the list of assets we want to provide to "PromptForCheckoutAndSave" for saving
 		TArray<UPackage*> PackagesToSaveProxy;
 		PackagesToSaveProxy.Reserve(OutSavedPackages.Num());
-		for (TPair<EPackageSaveResolutionType, UPackage*> DataToSave : OutSavedPackages)
+		for (TPair<UPackage*, EPackageSaveResolutionType> DataToSave : OutSavedPackages)
 		{
-			PackagesToSaveProxy.Push(DataToSave.Value);
+			// Ensure we have no duplicates here
+			check(!PackagesToSaveProxy.Contains(DataToSave.Key));
+			
+			PackagesToSaveProxy.Push(DataToSave.Key);
 		}
 
 		// List of packages that could not be saved
@@ -979,10 +1005,14 @@ bool BakeCustomizableObjectInstance(
 		const bool bWasSavingSuccessful = FEditorFileUtils::PromptForCheckoutAndSave(PackagesToSaveProxy, false, !bIsUnattendedExecution, &FailedToSavePackages, false, false) == FEditorFileUtils::EPromptReturnCode::PR_Success;
 
 		// Remove all packages that were going to be saved but failed to do so
-		const int32 RemovedPackagesCount = OutSavedPackages.RemoveAll([&](const TPair<EPackageSaveResolutionType, UPackage*> ToSavePackage)
+		uint32 RemovedPackagesCount = 0;
+		for (UPackage* Package : FailedToSavePackages)
 		{
-			return FailedToSavePackages.Contains(ToSavePackage.Value);
-		});
+			if (OutSavedPackages.Remove(Package))
+			{
+				RemovedPackagesCount++;
+			}
+		}
 		OutSavedPackages.Shrink();
 
 		return RemovedPackagesCount > 0 ? false : bWasSavingSuccessful;
diff --git a/Engine/Plugins/Mutable/Source/CustomizableObjectEditor/Private/MuCOE/GenerateMutableSource/GenerateMutableSourceTable.cpp b/Engine/Plugins/Mutable/Source/CustomizableObjectEditor/Private/MuCOE/GenerateMutableSource/GenerateMutableSourceTable.cpp
--- a/Engine/Plugins/Mutable/Source/CustomizableObjectEditor/Private/MuCOE/GenerateMutableSource/GenerateMutableSourceTable.cpp
+++ b/Engine/Plugins/Mutable/Source/CustomizableObjectEditor/Private/MuCOE/GenerateMutableSource/GenerateMutableSourceTable.cpp
@@ -5,13 +5,10 @@
 #include "Animation/AnimInstance.h"
 #include "Animation/PoseAsset.h"
 #include "AssetRegistry/ARFilter.h"
-#include "AssetRegistry/AssetRegistryModule.h"
 #include "DataTableUtils.h"
 #include "Engine/CompositeDataTable.h"
 #include "Engine/StaticMesh.h"
-#include "Engine/Texture2DArray.h"
 #include "GameplayTagContainer.h"
-#include "Kismet2/StructureEditorUtils.h"
 #include "Materials/Material.h"
 #include "Materials/MaterialInstance.h"
 #include "MuCO/CustomizableObjectSystem.h"
@@ -24,7 +21,6 @@
 #include "MuCOE/Nodes/CustomizableObjectNodeAnimationPose.h"
 #include "MuCOE/GraphTraversal.h"
 #include "MuR/Mesh.h"
-#include "PhysicsEngine/PhysicsAsset.h"
 #include "Rendering/SkeletalMeshLODModel.h"
 #include "Rendering/SkeletalMeshModel.h"
 #include "MuCOE/CustomizableObjectVersionBridge.h"
@@ -1152,9 +1148,6 @@ UDataTable* GenerateDataTableFromStruct(const UCustomizableObjectNodeTable* Tabl
 		return Cast<UDataTable>(GeneratedDataTable);
 	}
 
-	FAssetRegistryModule& AssetRegistryModule = FModuleManager::LoadModuleChecked<FAssetRegistryModule>("AssetRegistry");
-	IAssetRegistry& AssetRegistry = AssetRegistryModule.GetRegistry();
-
 	if (TableNode->FilterPaths.IsEmpty())
 	{
 		// Preventing load all data tables of the project
diff --git a/Engine/Plugins/Mutable/Source/CustomizableObjectEditor/Private/MuCOE/UnrealBakeHelpers.cpp b/Engine/Plugins/Mutable/Source/CustomizableObjectEditor/Private/MuCOE/UnrealBakeHelpers.cpp
--- a/Engine/Plugins/Mutable/Source/CustomizableObjectEditor/Private/MuCOE/UnrealBakeHelpers.cpp
+++ b/Engine/Plugins/Mutable/Source/CustomizableObjectEditor/Private/MuCOE/UnrealBakeHelpers.cpp
@@ -17,28 +17,39 @@
 #include "Materials/MaterialInstanceConstant.h"
 #include "Materials/MaterialInstanceDynamic.h"
 #include "Factories/MaterialInstanceConstantFactoryNew.h"
-#include "Materials/MaterialExpressionTextureBase.h"
 
 #include "TextureResource.h"
 
-UObject* FUnrealBakeHelpers::BakeHelper_DuplicateAsset(UObject* Object, const FString& ObjName, const FString& PkgName, bool ResetDuplicatedFlags, 
-													   TMap<UObject*, UObject*>& ReplacementMap, bool OverwritePackage, 
-													   const bool bGenerateConstantMaterialInstances)
+UObject* FUnrealBakeHelpers::BakeHelper_DuplicateAsset(UObject* Object, const FString& ObjName, const FString& PkgName,
+											  			TMap<UObject*, UObject*>* ReplacementMap, const bool bGenerateConstantMaterialInstances,
+											  			const EPackageSaveResolutionType SaveResolutionType)
 {
-	FString FinalObjName = ObjName;
-	FString FinalPkgName = PkgName;
-
-	if (!OverwritePackage)
+	FString FinalObjectName = ObjName;
+	FString FinalPackageName = PkgName;
+	
+	if (SaveResolutionType == EPackageSaveResolutionType::NewFile)
 	{
+		// To prevent our file hash from being modified we will append a "_" so the GetModuleChecked do not mess up our number.
+		// If not done the hash will get re-hashed as a way of making the asset name unique. With the char added we make the method append a value at the end instead
+		FinalObjectName += "_";
+		FinalPackageName += "_";
+		
 		FAssetToolsModule& AssetToolsModule = FModuleManager::GetModuleChecked<FAssetToolsModule>("AssetTools");
-		AssetToolsModule.Get().CreateUniqueAssetName(PkgName, "", FinalPkgName, FinalObjName);
+		AssetToolsModule.Get().CreateUniqueAssetName(FinalPackageName, "", FinalPackageName, FinalObjectName);
+
+		// If the last char is our safety char, remove it, otherwise this asset has got a new suffix added to make the name unique
+		FinalPackageName.RemoveFromEnd("_");
+		FinalObjectName.RemoveFromEnd("_");
 	}
 
-	UPackage* Package = CreatePackage(*FinalPkgName);
+	// Create or LOAD the package (if said package exists)
+	UPackage* Package = CreatePackage(*FinalPackageName);
+	check(Package);
 	Package->FullyLoad();
 	
 	UMaterialInstance* MatInstance = Cast<UMaterialInstance>(Object);
 
+	// Duplicated object we will return to the caller
 	UObject* DupObject = nullptr;
 	
 	// Only generate UMaterialInstanceConstant constant material instances if the original material is actually an instance, so check it here. Otherwise just duplicate
@@ -47,7 +58,7 @@ UObject* FUnrealBakeHelpers::BakeHelper_DuplicateAsset(UObject* Object, const FS
 		UMaterialInstanceConstantFactoryNew* MaterialFactory = NewObject<UMaterialInstanceConstantFactoryNew>();
 		MaterialFactory->InitialParent = MatInstance->Parent;
 
-		UMaterialInstanceConstant* MatInstanceConst = CastChecked<UMaterialInstanceConstant>(MaterialFactory->FactoryCreateNew(UMaterialInstanceConstant::StaticClass(), Package, FName(FinalObjName), RF_NoFlags, nullptr, GWarn));
+		UMaterialInstanceConstant* MatInstanceConst = CastChecked<UMaterialInstanceConstant>(MaterialFactory->FactoryCreateNew(UMaterialInstanceConstant::StaticClass(), Package, FName(FinalObjectName), RF_NoFlags, nullptr, GWarn));
 
 		TMap<int32, UTexture*> EmptyTextureReplacementMap;
 		CopyAllMaterialParameters<UMaterialInstanceConstant>(*MatInstanceConst, *MatInstance, EmptyTextureReplacementMap);
@@ -56,7 +67,7 @@ UObject* FUnrealBakeHelpers::BakeHelper_DuplicateAsset(UObject* Object, const FS
 	}
 	else
 	{
-		FObjectDuplicationParameters Params = InitStaticDuplicateObjectParams(Object, Package, *FinalObjName, RF_AllFlags, nullptr, EDuplicateMode::Normal);
+		FObjectDuplicationParameters Params = InitStaticDuplicateObjectParams(Object, Package, *FinalObjectName, RF_AllFlags, nullptr, EDuplicateMode::Normal);
 		DupObject = StaticDuplicateObjectEx(Params);
 	}
 
@@ -74,11 +85,14 @@ UObject* FUnrealBakeHelpers::BakeHelper_DuplicateAsset(UObject* Object, const FS
 		// Notify the asset registry
 		FAssetRegistryModule::AssetCreated(DupObject);
 
-		// Replace all references
-		ReplacementMap.Add(Object, DupObject);
+		if (ReplacementMap)
+		{
+			// Replace all references
+			ReplacementMap->Add(Object, DupObject);
 
-		constexpr EArchiveReplaceObjectFlags ReplaceFlags = (EArchiveReplaceObjectFlags::IgnoreOuterRef | EArchiveReplaceObjectFlags::IgnoreArchetypeRef);
-		FArchiveReplaceObjectRef<UObject> ReplaceAr(DupObject, ReplacementMap, ReplaceFlags);
+			constexpr EArchiveReplaceObjectFlags ReplaceFlags = (EArchiveReplaceObjectFlags::IgnoreOuterRef | EArchiveReplaceObjectFlags::IgnoreArchetypeRef);
+			FArchiveReplaceObjectRef<UObject> ReplaceAr(DupObject, *ReplacementMap, ReplaceFlags);
+		}
 	}
 
 	return DupObject;
@@ -92,7 +106,7 @@ namespace
 	{
 		MUTABLE_CPUPROFILER_SCOPE(CopyTextureProperties)
 
-			Texture->NeverStream = SourceTexture->NeverStream;
+		Texture->NeverStream = SourceTexture->NeverStream;
 
 		Texture->SRGB = SourceTexture->SRGB;
 		Texture->Filter = SourceTexture->Filter;
@@ -108,35 +122,48 @@ namespace
 }
 
 
-UTexture2D* FUnrealBakeHelpers::BakeHelper_CreateAssetTexture(UTexture2D* SrcTex, const FString& TexObjName, const FString& TexPkgName, const UTexture* OrgTex, bool ResetDuplicatedFlags, TMap<UObject*, UObject*>& ReplacementMap, bool OverwritePackage)
+UTexture2D* FUnrealBakeHelpers::BakeHelper_CreateAssetTexture(UTexture2D* SourceTexture, const FString& TexObjName, const FString& TexPkgName, const UTexture* OrgTex,
+																TMap<UObject*, UObject*>* ReplacementMap, const EPackageSaveResolutionType SaveResolutionType)
 {
-	const bool bIsMutableTexture = !SrcTex->Source.IsValid();
+	const bool bIsMutableTexture = !SourceTexture->Source.IsValid();
 	if (!bIsMutableTexture)
 	{
-		return Cast<UTexture2D>(BakeHelper_DuplicateAsset(SrcTex, TexObjName, TexPkgName, ResetDuplicatedFlags, ReplacementMap, OverwritePackage, false));
+		return Cast<UTexture2D>(BakeHelper_DuplicateAsset(SourceTexture, TexObjName, TexPkgName, ReplacementMap,false, SaveResolutionType));
 	}
 
-	int32 sx = SrcTex->GetPlatformData()->SizeX;
-	int32 sy = SrcTex->GetPlatformData()->SizeY;
-	EPixelFormat SrcPixelFormat = SrcTex->GetPlatformData()->PixelFormat;
-	ETextureSourceFormat PixelFormat = (SrcPixelFormat == PF_BC4 || SrcPixelFormat == PF_G8) ? TSF_G8 : TSF_BGRA8;
+	int32 SourceTextureSizeX = SourceTexture->GetPlatformData()->SizeX;
+	int32 SourceTextureSizeY = SourceTexture->GetPlatformData()->SizeY;
+	EPixelFormat SourceTexturePixelFormat = SourceTexture->GetPlatformData()->PixelFormat;
+	ETextureSourceFormat PixelFormat = (SourceTexturePixelFormat == PF_BC4 || SourceTexturePixelFormat == PF_G8) ? TSF_G8 : TSF_BGRA8;
 
 	// Begin duplicate texture
-	FString FinalObjName = TexObjName;
-	FString FinalPkgName = TexPkgName;
+	FString FinalObjectName = TexObjName;
+	FString FinalPackageName = TexPkgName;
 
-	if (!OverwritePackage)
+	if (SaveResolutionType == EPackageSaveResolutionType::NewFile)
 	{
+		// As a safety measure that prevents the update of the hashed integer at the end, add a final char to prevent the the
+		// CreateUniqueAssetName from changing the value for our resource hash.
+		
+		// To prevent our hash from being modified
+		FinalObjectName += "_";
+		FinalPackageName += "_";
+		
+		// This is destroying the hash at the end of the TexObjectName....
+		// note: We are adding a sufix so it gets added to the FinalPkgName to prevent the method from changing the value of the hash
 		FAssetToolsModule& AssetToolsModule = FModuleManager::GetModuleChecked<FAssetToolsModule>("AssetTools");
-		AssetToolsModule.Get().CreateUniqueAssetName(TexPkgName, "", FinalPkgName, FinalObjName);
+		AssetToolsModule.Get().CreateUniqueAssetName(FinalPackageName, TEXT(""), FinalPackageName, FinalObjectName);
+
+		// If the last char is our safety char, remove it, otherwise this asset has got a new sufix added to make the name unique
+		FinalPackageName.RemoveFromEnd("_");
+		FinalObjectName.RemoveFromEnd("_");
 	}
 
-	UPackage* Package = CreatePackage(*FinalPkgName);
+	UPackage* Package = CreatePackage(*FinalPackageName);
+	check(Package);
 	Package->FullyLoad();
 
-	UTexture2D* DupTex = NewObject<UTexture2D>(Package, *FinalObjName, RF_Public | RF_Standalone);
-
-	ReplacementMap.Add(SrcTex, DupTex);
+	UTexture2D* DupTex = NewObject<UTexture2D>(Package, *FinalObjectName, RF_Public | RF_Standalone);
 
 	// The garbage collector is called in the middle of the bake process, and this can destroy this temporary object. 
 	// We add them to the garbage root to prevent this. This will avoid them being unloaded while the editor is running, but this
@@ -148,12 +175,17 @@ UTexture2D* FUnrealBakeHelpers::BakeHelper_CreateAssetTexture(UTexture2D* SrcTex
 	FAssetRegistryModule::AssetCreated(DupTex);
 
 	// Replace all references
-	constexpr EArchiveReplaceObjectFlags ReplaceFlags = (EArchiveReplaceObjectFlags::IgnoreOuterRef | EArchiveReplaceObjectFlags::IgnoreArchetypeRef);
-	FArchiveReplaceObjectRef<UObject> ReplaceAr(DupTex, ReplacementMap, ReplaceFlags);
+	if (ReplacementMap)
+	{
+		ReplacementMap->Add(SourceTexture, DupTex);
+		
+		constexpr EArchiveReplaceObjectFlags ReplaceFlags = (EArchiveReplaceObjectFlags::IgnoreOuterRef | EArchiveReplaceObjectFlags::IgnoreArchetypeRef);
+		FArchiveReplaceObjectRef<UObject> ReplaceAr(DupTex, *ReplacementMap, ReplaceFlags);
+	}
 
 	// End duplicate texture
 
-	CopyTextureProperties(DupTex, SrcTex);
+	CopyTextureProperties(DupTex, SourceTexture);
 
 	DupTex->RemoveUserDataOfClass(UMutableTextureMipDataProviderFactory::StaticClass());
 	if (OrgTex)
@@ -162,31 +194,31 @@ UTexture2D* FUnrealBakeHelpers::BakeHelper_CreateAssetTexture(UTexture2D* SrcTex
 	}
 
 	// Mutable textures only have platform data. We need to build the source data for them to be assets.
-	DupTex->Source.Init(sx, sy, 1, 1, PixelFormat);
+	DupTex->Source.Init(SourceTextureSizeX, SourceTextureSizeY, 1, 1, PixelFormat);
 
-	int32 MipCount = SrcTex->GetPlatformData()->Mips.Num();
+	int32 MipCount = SourceTexture->GetPlatformData()->Mips.Num();
 	if (!MipCount)
 	{
-		UE_LOG(LogMutable, Warning, TEXT("Bake Instances: Empty texture found [%s]."), *SrcTex->GetName());
+		UE_LOG(LogMutable, Warning, TEXT("Bake Instances: Empty texture found [%s]."), *SourceTexture->GetName());
 		return DupTex;
 	}
 
 	// Create a mutable image from the platform data.
-	mu::EImageFormat MutableFormat = UnrealToMutablePixelFormat(SrcTex->GetPlatformData()->PixelFormat, SrcTex->HasAlphaChannel());
-	TSharedPtr<mu::FImage> PlatformImage = MakeShared<mu::FImage>(sx, sy, 1, MutableFormat, mu::EInitializationType::NotInitialized);
+	mu::EImageFormat MutableFormat = UnrealToMutablePixelFormat(SourceTexture->GetPlatformData()->PixelFormat, SourceTexture->HasAlphaChannel());
+	TSharedPtr<mu::FImage> PlatformImage = MakeShared<mu::FImage>(SourceTextureSizeX, SourceTextureSizeY, 1, MutableFormat, mu::EInitializationType::NotInitialized);
 
 	constexpr int32 MipIndex = 0;
-	const uint8* SourceData = reinterpret_cast<const uint8*>(SrcTex->GetPlatformData()->Mips[MipIndex].BulkData.LockReadOnly());
+	const uint8* SourceData = reinterpret_cast<const uint8*>(SourceTexture->GetPlatformData()->Mips[MipIndex].BulkData.LockReadOnly());
 	check(SourceData); // A mutable-generated texture should always contain platform data
-	int32 PlatformDataSize = SrcTex->GetPlatformData()->Mips[MipIndex].BulkData.GetBulkDataSize();
+	int32 PlatformDataSize = SourceTexture->GetPlatformData()->Mips[MipIndex].BulkData.GetBulkDataSize();
 	check(PlatformImage->GetDataSize()== PlatformDataSize);
 	FMemory::Memcpy(PlatformImage->GetLODData(0), SourceData, PlatformDataSize);
 	SourceData = nullptr;
-	SrcTex->GetPlatformData()->Mips[MipIndex].BulkData.Unlock();
+	SourceTexture->GetPlatformData()->Mips[MipIndex].BulkData.Unlock();
 
 	// Reformat the mutable image
 	mu::EImageFormat UncompressedMutableFormat = mu::EImageFormat::RGBA_UByte;
-	switch (SrcTex->GetPlatformData()->PixelFormat)
+	switch (SourceTexture->GetPlatformData()->PixelFormat)
 	{
 	case PF_G8:
 	case PF_L8:
@@ -219,7 +251,7 @@ UTexture2D* FUnrealBakeHelpers::BakeHelper_CreateAssetTexture(UTexture2D* SrcTex
 	const bool bNeedsRBSwizzle = PixelFormat == TSF_BGRA8;
 	if (bNeedsRBSwizzle)
 	{
-		for (int32 x = 0; x < sx * sy; ++x)
+		for (int32 x = 0; x < SourceTextureSizeX * SourceTextureSizeY; ++x)
 		{
 			uint8 temp = Dest[0];
 			Dest[0] = Dest[2];
@@ -232,14 +264,14 @@ UTexture2D* FUnrealBakeHelpers::BakeHelper_CreateAssetTexture(UTexture2D* SrcTex
 	DupTex->Source.UnlockMip(MipIndex);
 
 	bool bNeeds_TC_Grayscale = PixelFormat == TSF_G8 || PixelFormat == TSF_G16;
-	bool bDoNotCompress = SrcPixelFormat == PF_R8G8B8A8;
-	bool bIsNormalMap = SrcPixelFormat == PF_BC5;
+	bool bDoNotCompress = SourceTexturePixelFormat == PF_R8G8B8A8;
+	bool bIsNormalMap = SourceTexturePixelFormat == PF_BC5;
 
 	if (bNeeds_TC_Grayscale || bDoNotCompress || bIsNormalMap)
 	{
 		FTextureFormatSettings Settings;
 
-		Settings.SRGB = SrcTex->SRGB;
+		Settings.SRGB = SourceTexture->SRGB;
 
 		if (bNeeds_TC_Grayscale)
 		{
diff --git a/Engine/Plugins/Mutable/Source/CustomizableObjectEditor/Private/MuCOE/UnrealBakeHelpers.h b/Engine/Plugins/Mutable/Source/CustomizableObjectEditor/Private/MuCOE/UnrealBakeHelpers.h
--- a/Engine/Plugins/Mutable/Source/CustomizableObjectEditor/Private/MuCOE/UnrealBakeHelpers.h
+++ b/Engine/Plugins/Mutable/Source/CustomizableObjectEditor/Private/MuCOE/UnrealBakeHelpers.h
@@ -10,6 +10,7 @@
 #include "Materials/MaterialInstanceConstant.h"
 #include "Materials/MaterialInterface.h"
 
+enum class EPackageSaveResolutionType : uint8;
 class UObject;
 class USkeletalMesh;
 class UTexture2D;
@@ -22,16 +23,17 @@ class CUSTOMIZABLEOBJECTEDITOR_API FUnrealBakeHelpers
 {
 public:
 
-	static UObject* BakeHelper_DuplicateAsset(UObject* Object, const FString& ObjName, const FString& PkgName, bool ResetDuplicatedFlags,
-											  TMap<UObject*, UObject*>& ReplacementMap, bool OverwritePackage,
-											  const bool bGenerateConstantMaterialInstances);
+	static UObject* BakeHelper_DuplicateAsset(UObject* Object, const FString& ObjName, const FString& PkgName,
+											  TMap<UObject*, UObject*>* ReplacementMap, const bool bGenerateConstantMaterialInstances,
+											  const EPackageSaveResolutionType SaveResolutionType);
 
 	/**
 	 * Duplicates a texture asset. Duplicates Mutable and non Mutable textures.
 	 *
 	 * @param OrgTex Original source texture from which a Mutable texture has been generated. Only required when SrcTex is a Mutable texture.
 	 */
-	static UTexture2D* BakeHelper_CreateAssetTexture(UTexture2D* SrcTex, const FString& TexObjName, const FString& TexPkgName, const UTexture* OrgTex, bool ResetDuplicatedFlags, TMap<UObject*, UObject*>& ReplacementMap, bool OverwritePackage);
+	static UTexture2D* BakeHelper_CreateAssetTexture(UTexture2D* SourceTexture, const FString& TexObjName, const FString& TexPkgName, const UTexture* OrgTex,
+													 TMap<UObject*, UObject*>* ReplacementMap, const EPackageSaveResolutionType SaveResolutionType);
 
 	template<typename MaterialType>
 	static void CopyAllMaterialParameters(MaterialType& DestMaterial, UMaterialInterface& OriginMaterial, const TMap<int32, UTexture*>& TextureReplacementMap)
diff --git a/Engine/Plugins/Mutable/Source/MutableValidation/Private/MuV/COIBakingTestCommandlet.cpp b/Engine/Plugins/Mutable/Source/MutableValidation/Private/MuV/COIBakingTestCommandlet.cpp
--- a/Engine/Plugins/Mutable/Source/MutableValidation/Private/MuV/COIBakingTestCommandlet.cpp
+++ b/Engine/Plugins/Mutable/Source/MutableValidation/Private/MuV/COIBakingTestCommandlet.cpp
@@ -174,14 +174,16 @@ int32 UCOIBakingTestCommandlet::Main(const FString& Params)
 		{
 			UE_LOG(LogMutable,Display,TEXT("Starting Instance Baking operation."));
 			{
-				TArray<TPair<EPackageSaveResolutionType,UPackage*>> SavedPackages;
+				FBakingConfiguration Configuration;
+				Configuration.OutputFilesBaseName = FString::Printf(TEXT("%s_Bake"), *TargetInstance->GetName());
+				Configuration.OutputPath = LocalBakingDirectory;
+				Configuration.bExportAllResourcesOnBake = true;
+				Configuration.bGenerateConstantMaterialInstancesOnBake = true;
+
+				TMap<UPackage*, EPackageSaveResolutionType> SavedPackages;
 				bWasBakingSuccessful = BakeCustomizableObjectInstance(
 					*TargetInstance,
-					FString::Printf(TEXT("%s_Bake"), *TargetInstance->GetName()),
-					LocalBakingDirectory,
-					true,
-					true,
-					true,
+					Configuration,
 					true,
 					SavedPackages);
 			}
@@ -213,4 +215,3 @@ int32 UCOIBakingTestCommandlet::Main(const FString& Params)
 }
 
 
-
