[Mutable] Fix image transform op codegen generating large images when options LayoutToApply is valid.

- Now a crop to the current block size is injected.

#rnx

diff --git a/Engine/Plugins/Mutable/Source/MutableTools/Private/MuT/CodeGenerator_Image.cpp b/Engine/Plugins/Mutable/Source/MutableTools/Private/MuT/CodeGenerator_Image.cpp
--- a/Engine/Plugins/Mutable/Source/MutableTools/Private/MuT/CodeGenerator_Image.cpp
+++ b/Engine/Plugins/Mutable/Source/MutableTools/Private/MuT/CodeGenerator_Image.cpp
@@ -637,7 +637,8 @@ namespace mu
 		NewOptions.ImageLayoutStrategy = CompilerOptions::TextureLayoutStrategy::None;
 		NewOptions.LayoutToApply = nullptr;
 		NewOptions.LayoutBlockId = -1;
-		NewOptions.RectSize = {};
+		NewOptions.RectSize = {0, 0};
+
 
         if (Node->Base)
         {
@@ -652,12 +653,53 @@ namespace mu
         }
 		
 		FImageDesc BaseDesc = Base->GetImageDesc(true);
-		
         Op->Base = Base;
 		Op->SourceSizeX = BaseDesc.m_size.X;
 		Op->SourceSizeY = BaseDesc.m_size.Y;
-
+		
         Result.op = Op; 
+
+		// Compute the image crop for the layout block to apply. 
+		if (Options.LayoutToApply)
+		{
+			FIntVector2 TransformImageSize = FIntVector2(Op->SizeX, Op->SizeY);
+			if (Op->SizeX == 0 && Op->SizeY == 0)
+			{
+				TransformImageSize = FIntVector2(BaseDesc.m_size.X, BaseDesc.m_size.Y);	
+			}
+
+			int32 BlockIndex = Options.LayoutToApply->FindBlock(Options.LayoutBlockId);
+			check(BlockIndex >= 0);
+
+			// Rect in layout grid units
+			FIntVector2 RectMinInCells  = Options.LayoutToApply->Blocks[BlockIndex].Min;
+			FIntVector2 RectSizeInCells = Options.LayoutToApply->Blocks[BlockIndex].Size;
+			
+			FIntVector2 Grid = Options.LayoutToApply->GetGridSize();
+			Grid = FIntVector2(FMath::Max(1, Grid.X), FMath::Max(1, Grid.Y));
+
+			// Transform to pixels
+			FIntVector2 BlockImageMin = FIntVector2(
+					(RectMinInCells.X * TransformImageSize.X) / Grid.X,
+					(RectMinInCells.Y * TransformImageSize.Y) / Grid.Y);
+
+			FIntVector2 BlockImageSize = FIntVector2(
+					(RectSizeInCells.X * TransformImageSize.X) / Grid.X,
+					(RectSizeInCells.Y * TransformImageSize.Y) / Grid.Y);
+			
+			if (BlockImageSize != TransformImageSize)
+			{
+				Ptr<ASTOpImageCrop> CropOp = new ASTOpImageCrop();
+				CropOp->Source = Op;
+			
+				CropOp->Min.X  = (uint16)BlockImageMin.X;
+				CropOp->Min.Y  = (uint16)BlockImageMin.Y;
+				CropOp->Size.X = (uint16)BlockImageSize.X;
+				CropOp->Size.Y = (uint16)BlockImageSize.Y;
+
+        		Result.op = CropOp;
+			}
+		}
     }
 
 
