[Mutable] Fixed vector color codification.

#rnx

diff --git a/Engine/Plugins/Mutable/Source/CustomizableObject/Private/MuCO/CustomizableObjectInstance.cpp b/Engine/Plugins/Mutable/Source/CustomizableObject/Private/MuCO/CustomizableObjectInstance.cpp
--- a/Engine/Plugins/Mutable/Source/CustomizableObject/Private/MuCO/CustomizableObjectInstance.cpp
+++ b/Engine/Plugins/Mutable/Source/CustomizableObject/Private/MuCO/CustomizableObjectInstance.cpp
@@ -6499,6 +6499,7 @@ void UCustomizableInstancePrivate::BuildMaterials(const TSharedRef<FUpdateContex
 								{
 									FLinearColor Color = Param.Vector;
 
+									// HACK: We encode an invalid value (Nan) for table option "None.
 									// Decoding "None" color parameters that use the material color
 									if (FMath::IsNaN(Color.R))
 									{
diff --git a/Engine/Plugins/Mutable/Source/MutableTools/Private/MuT/ASTOpConstantColor.cpp b/Engine/Plugins/Mutable/Source/MutableTools/Private/MuT/ASTOpConstantColor.cpp
--- a/Engine/Plugins/Mutable/Source/MutableTools/Private/MuT/ASTOpConstantColor.cpp
+++ b/Engine/Plugins/Mutable/Source/MutableTools/Private/MuT/ASTOpConstantColor.cpp
@@ -21,6 +21,14 @@ namespace mu
 		if (OtherUntyped.GetOpType() == GetOpType())
 		{
 			const ASTOpConstantColor* other = static_cast<const ASTOpConstantColor*>(&OtherUntyped);
+
+			// HACK: We encode an invalid value (Nan) for table option "None".
+			// This nan comparison is needed because some compilers will return that 0.0f is equal to Nan...
+			if (FMath::IsNaN(Value[0]) || FMath::IsNaN(other->Value[0]))
+			{
+				return FMath::IsNaN(Value[0]) == FMath::IsNaN(other->Value[0]);
+			}
+
 			return Value == other->Value;
 		}
 		return false;
