[Mutable] Invalidate compiled data if it can't be stored to disk.
Remove unnecessary ensure.
Reset DDCHeapMemory once used.

#rnx

diff --git a/Engine/Plugins/Mutable/Source/CustomizableObjectEditor/Private/MuCOE/CustomizableObjectCompileRunnable.cpp b/Engine/Plugins/Mutable/Source/CustomizableObjectEditor/Private/MuCOE/CustomizableObjectCompileRunnable.cpp
--- a/Engine/Plugins/Mutable/Source/CustomizableObjectEditor/Private/MuCOE/CustomizableObjectCompileRunnable.cpp
+++ b/Engine/Plugins/Mutable/Source/CustomizableObjectEditor/Private/MuCOE/CustomizableObjectCompileRunnable.cpp
@@ -568,7 +568,7 @@ void FCustomizableObjectSaveDDRunnable::StoreCachedPlatformDataToDisk()
 	FileManager.MakeDirectory(*GetCompiledDataFolderPath(), true);
 
 	// Delete files and create new memory writers for each streamable data type.
-	bool bFilesDeleted = true;
+	bool bSuccess = true;
 	TArray<TUniquePtr<FArchive>> MemoryWriters;
 
 	const int32 NumDataTypes = static_cast<int32>(MutablePrivate::EStreamableDataType::DataTypeCount);
@@ -579,7 +579,7 @@ void FCustomizableObjectSaveDDRunnable::StoreCachedPlatformDataToDisk()
 			&& !FileManager.Delete(*FilePath, true, false, true))
 		{
 			UE_LOG(LogMutable, Error, TEXT("Failed to delete file for data type [%d]."), DataType);
-			bFilesDeleted = false;
+			bSuccess = false;
 			break;
 		}
 
@@ -591,7 +591,7 @@ void FCustomizableObjectSaveDDRunnable::StoreCachedPlatformDataToDisk()
 		}
 	}
 
-	if (bFilesDeleted)
+	if (bSuccess)
 	{
 		// Serialize Streamable resources
 		const auto WriteBulkDataToDisk = [&MemoryWriters](MutablePrivate::FFile& File, TArray64<uint8>& FileBulkData, uint32 FileIndex)
@@ -621,21 +621,28 @@ void FCustomizableObjectSaveDDRunnable::StoreCachedPlatformDataToDisk()
 			MemoryWriters[0]->Serialize(ModelStreamablesBytes.GetData(), ModelStreamablesBytes.Num() * sizeof(uint8));
 			MemoryWriters[0]->Serialize(ModelData.GetData(), ModelData.Num() * sizeof(uint8));
 		}
-	}
-	else
-	{
-		// Couldn't delete files. Delete model to invalidate compilation.
-		PlatformData->Model.Reset();
+
+		// Write to disk and close the files
+		for (TUniquePtr<FArchive>& MemoryWriter : MemoryWriters)
+		{
+			if (MemoryWriter)
+			{
+				MemoryWriter->Flush();
+
+				if (!MemoryWriter->Close())
+				{
+					UE_LOG(LogMutable, Error, TEXT("Failed to write file to disk. File [%s]."), *MemoryWriter->GetArchiveName());
+					bSuccess = false;
+					break;
+				}
+			}
+		}
 	}
 	
-	// Write to disk and close the files
-	for (TUniquePtr<FArchive>& MemoryWriter : MemoryWriters)
+	if (!bSuccess)
 	{
-		if (MemoryWriter)
-		{
-			MemoryWriter->Flush();
-			MemoryWriter->Close();
-		}
+		// Delete model to invalidate compilation.
+		PlatformData->Model.Reset();
 	}
 }
 
diff --git a/Engine/Plugins/Mutable/Source/CustomizableObjectEditor/Private/MuCOE/CustomizableObjectCompiler.cpp b/Engine/Plugins/Mutable/Source/CustomizableObjectEditor/Private/MuCOE/CustomizableObjectCompiler.cpp
--- a/Engine/Plugins/Mutable/Source/CustomizableObjectEditor/Private/MuCOE/CustomizableObjectCompiler.cpp
+++ b/Engine/Plugins/Mutable/Source/CustomizableObjectEditor/Private/MuCOE/CustomizableObjectCompiler.cpp
@@ -1447,7 +1447,7 @@ void FCustomizableObjectCompiler::FinishLoadingModelDataFromDDC()
 			{
 				MUTABLE_CPUPROFILER_SCOPE(GetStreamableDataFromDDC);
 
-				if (ensure(Response.Status == EStatus::Ok))
+				if (Response.Status == EStatus::Ok)
 				{
 					UCustomizableObject* CustomizableObject = CurrentRequest->GetCustomizableObject();
 					const FCompilationOptions& Options = CurrentRequest->Options;
@@ -1570,6 +1570,7 @@ void FCustomizableObjectCompiler::FinishLoadingModelDataFromDDC()
 
 void FCustomizableObjectCompiler::FinishLoadingStreamableDataFromDDC()
 {
+	DDCHeapMemory.Reset();
 	LoadStreamableDataFromDDCEvent.Reset();
 
 	if (PlatformData->Model && PlatformData->ModelResources && PlatformData->ModelStreamableBulkData)
