[Mutable] Add utility function for batch parallel execution invokation.

#rnx

diff --git a/Engine/Plugins/Mutable/Source/MutableRuntime/Internal/MuR/ParallelExecutionUtils.h b/Engine/Plugins/Mutable/Source/MutableRuntime/Internal/MuR/ParallelExecutionUtils.h
new file mode 100644
--- /dev/null
+++ b/Engine/Plugins/Mutable/Source/MutableRuntime/Internal/MuR/ParallelExecutionUtils.h
@@ -0,0 +1,12 @@
+// Copyright Epic Games, Inc. All Rights Reserved.
+
+#pragma once
+
+#include "Templates/Function.h"
+
+namespace mu::ParallelExecutionUtils
+{
+    void InvokeBatchParallelFor(int32 Num, TFunctionRef<void(int32)> Body);
+    void InvokeBatchParallelForSingleThread(int32 Num, TFunctionRef<void(int32)> Body);
+}
+
diff --git a/Engine/Plugins/Mutable/Source/MutableRuntime/Private/MuR/OpImageBinarise.h b/Engine/Plugins/Mutable/Source/MutableRuntime/Private/MuR/OpImageBinarise.h
--- a/Engine/Plugins/Mutable/Source/MutableRuntime/Private/MuR/OpImageBinarise.h
+++ b/Engine/Plugins/Mutable/Source/MutableRuntime/Private/MuR/OpImageBinarise.h
@@ -5,7 +5,7 @@
 #include "MuR/ImagePrivate.h"
 #include "MuR/Platform.h"
 
-#include "Async/ParallelFor.h"
+#include "MuR/ParallelExecutionUtils.h"
 
 
 namespace mu
@@ -32,7 +32,10 @@ namespace mu
 		const int32 NumBatches = AImage->DataStorage.GetNumBatches(BatchNumElems, BytesPerElem);
 		check(DestImage->DataStorage.GetNumBatches(BatchNumElems, 1) == NumBatches);
 
-		auto ProcessBatch = [DestImage, AImage, ThresholdUNorm, BatchNumElems, BytesPerElem](int32 BatchId)
+		ParallelExecutionUtils::InvokeBatchParallelFor(NumBatches,
+		[
+			DestImage, AImage, ThresholdUNorm, BatchNumElems, BytesPerElem
+		](int32 BatchId)
 		{
 			TArrayView<uint8> DestView = DestImage->DataStorage.GetBatch(BatchId, BatchNumElems, 1);	
 			TArrayView<const uint8> AView = AImage->DataStorage.GetBatch(BatchId, BatchNumElems, BytesPerElem);	
@@ -82,16 +85,7 @@ namespace mu
 			default:
 				check(false);
 			}
-		};
-
-		if (NumBatches == 1)
-		{
-			ProcessBatch(0);
-		}
-		else if (NumBatches > 1)
-		{
-			ParallelFor(NumBatches, ProcessBatch);
-		}
+		});
 	}
 
 }
diff --git a/Engine/Plugins/Mutable/Source/MutableRuntime/Private/MuR/OpImageColourMap.h b/Engine/Plugins/Mutable/Source/MutableRuntime/Private/MuR/OpImageColourMap.h
--- a/Engine/Plugins/Mutable/Source/MutableRuntime/Private/MuR/OpImageColourMap.h
+++ b/Engine/Plugins/Mutable/Source/MutableRuntime/Private/MuR/OpImageColourMap.h
@@ -5,7 +5,7 @@
 #include "MuR/ImagePrivate.h"
 #include "MuR/Platform.h"
 
-#include "Async/ParallelFor.h"
+#include "MuR/ParallelExecutionUtils.h"
 
 namespace mu
 {
@@ -42,7 +42,10 @@ namespace mu
 		const int32 NumBatches = DestImage->DataStorage.GetNumBatchesLODRange(NumBatchElems, BytesPerElem, LODBegin, LODEnd); 
 		check(NumBatches == SourceImage->DataStorage.GetNumBatchesLODRange(NumBatchElems, BytesPerElem, LODBegin, LODEnd));
 
-		auto ProcessBatch = [DestImage, SourceImage, MaskImage, Palette, NumBatchElems, BytesPerElem, LODBegin, LODEnd](int32 BatchId)
+		ParallelExecutionUtils::InvokeBatchParallelFor(NumBatches,
+		[
+			DestImage, SourceImage, MaskImage, Palette, NumBatchElems, BytesPerElem, LODBegin, LODEnd
+		](int32 BatchId)
 		{
 			TArrayView<uint8> DestView = DestImage->DataStorage.GetBatchLODRange(BatchId, NumBatchElems, BytesPerElem, LODBegin, LODEnd); 
 			TArrayView<const uint8> SourceView = SourceImage->DataStorage.GetBatchLODRange(BatchId, NumBatchElems, BytesPerElem, LODBegin, LODEnd); 
@@ -140,16 +143,7 @@ namespace mu
 			default:
 				check(false);
 			}
-		};
-
-		if (NumBatches == 1)
-		{
-			ProcessBatch(0);
-		}
-		else if (NumBatches > 1)
-		{
-			ParallelFor(NumBatches, ProcessBatch);
-		}
+		});
 	}
 
 }
diff --git a/Engine/Plugins/Mutable/Source/MutableRuntime/Private/MuR/OpImageInterpolate.h b/Engine/Plugins/Mutable/Source/MutableRuntime/Private/MuR/OpImageInterpolate.h
--- a/Engine/Plugins/Mutable/Source/MutableRuntime/Private/MuR/OpImageInterpolate.h
+++ b/Engine/Plugins/Mutable/Source/MutableRuntime/Private/MuR/OpImageInterpolate.h
@@ -5,7 +5,7 @@
 #include "MuR/ImagePrivate.h"
 #include "MuR/Platform.h"
 
-#include "Async/ParallelFor.h"
+#include "MuR/ParallelExecutionUtils.h"
 
 namespace mu
 {
@@ -31,7 +31,10 @@ namespace mu
 		check(BImage->DataStorage.GetNumBatches(NumBatchElems, BytesPerElem) == NumBatches);
 
 		// Generic implementation	
-		auto ProcessBatch = [DestImage, BImage, NumBatchElems, BytesPerElem, FactorUNorm](int32 BatchId)
+		ParallelExecutionUtils::InvokeBatchParallelFor(NumBatches,
+		[
+			DestImage, BImage, NumBatchElems, BytesPerElem, FactorUNorm
+		](int32 BatchId)
 		{
 			TArrayView<uint8> DestView = DestImage->DataStorage.GetBatch(BatchId, NumBatchElems, BytesPerElem);
 			TArrayView<const uint8> BView = BImage->DataStorage.GetBatch(BatchId, NumBatchElems, BytesPerElem);
@@ -60,17 +63,7 @@ namespace mu
 			default:
 				check(false);
 			}
-		};
-
-		if (NumBatches == 1)
-		{
-			ProcessBatch(0);
-		}
-		else if (NumBatches > 1)
-		{
-			ParallelFor(NumBatches, ProcessBatch);
-		}
-
+		});
 	}
 
 }
diff --git a/Engine/Plugins/Mutable/Source/MutableRuntime/Private/MuR/OpImageInvert.h b/Engine/Plugins/Mutable/Source/MutableRuntime/Private/MuR/OpImageInvert.h
--- a/Engine/Plugins/Mutable/Source/MutableRuntime/Private/MuR/OpImageInvert.h
+++ b/Engine/Plugins/Mutable/Source/MutableRuntime/Private/MuR/OpImageInvert.h
@@ -4,34 +4,12 @@
 
 #include "MuR/ImagePrivate.h"
 #include "MuR/Platform.h"
+#include "MuR/ParallelExecutionUtils.h"
 
-#include "Async/ParallelFor.h"
 #include "Math/VectorRegister.h"
 
 namespace mu
 {
-
-namespace OpImageInvertInternal
-{
-	// TODO: This is used in multiple places, make it an utility so it can be shared.  
-	FORCENOINLINE void InvokeParallelFor(int32 NumBatches, TFunctionRef<void(int32)> BatchFunc)
-	{
-		if (NumBatches == 0)
-		{
-			return;
-		}
-
-		if (NumBatches == 1)
-		{
-			BatchFunc(0);
-		}
-		else
-		{
-			ParallelFor(NumBatches, BatchFunc);
-		}
-	}
-} // namespace mu::OpImageInvertInternal
-
 	/** Invert the colour (RGB or L) components of an image. Leave alpha untouched. */
 	template<bool bUseVectorImpl = false>
 	inline void ImageInvert(FImage* AImage)
@@ -49,7 +27,7 @@ namespace OpImageInvertInternal
 		const int32 BytesPerElem = GetImageFormatData(AImage->GetFormat()).BytesPerBlock;
 		const int32 NumBatches = AImage->DataStorage.GetNumBatches(NumBatchElems, BytesPerElem);
 
-		OpImageInvertInternal::InvokeParallelFor(NumBatches,
+		ParallelExecutionUtils::InvokeBatchParallelFor(NumBatches,
 		[
 			AImage, NumBatchElems, BytesPerElem
 		](int32 BatchId)
diff --git a/Engine/Plugins/Mutable/Source/MutableRuntime/Private/MuR/OpImageLayer.cpp b/Engine/Plugins/Mutable/Source/MutableRuntime/Private/MuR/OpImageLayer.cpp
--- a/Engine/Plugins/Mutable/Source/MutableRuntime/Private/MuR/OpImageLayer.cpp
+++ b/Engine/Plugins/Mutable/Source/MutableRuntime/Private/MuR/OpImageLayer.cpp
@@ -4,7 +4,7 @@
 
 #include "MuR/ImagePrivate.h"
 #include "MuR/ImageRLE.h"
-#include "Async/ParallelFor.h"
+#include "MuR/ParallelExecutionUtils.h"
 #include "Templates/UnrealTemplate.h"
 
 #include "MuR/OpImageBlend.h"
@@ -241,24 +241,6 @@ namespace OpImageLayerInternal
 		return BatchViewsResult;
 	}
 
-	FORCENOINLINE void InvokeParallelFor(int32 NumBatches, TFunctionRef<void(int32)> BatchFunc)
-	{
-		if (NumBatches == 0)
-		{
-			return;
-		}
-
-		if (NumBatches == 1)
-		{
-			BatchFunc(0);
-		}
-		else
-		{
-			ParallelFor(NumBatches, BatchFunc);
-		}
-
-	}
-
 	FORCEINLINE bool IsAnyComponentLargerThan1(FVector4f Value)
 	{
 		return (Value[0] > 1) | (Value[1] > 1) | (Value[2] > 1) | (Value[3] > 1);
@@ -349,7 +331,7 @@ namespace OpImageLayerInternal
 		const int32 NumBatches = OpImageLayerInternal::GetOpLayerNumBatches(BatchArgs);
 
 		const FIntVector ColorValue = FIntVector(Color.X * 255.0f, Color.Y * 255.0f, Color.Z * 255.0f);
-		OpImageLayerInternal::InvokeParallelFor(NumBatches,
+		ParallelExecutionUtils::InvokeBatchParallelFor(NumBatches,
 		[
 			&BatchArgs, ColorValue, BaseFormat
 		](int32 BatchId)
@@ -437,7 +419,7 @@ namespace OpImageLayerInternal
 		const int32 NumBatches = OpImageLayerInternal::GetOpLayerNumBatches(BatchArgs);
 
 		const FIntVector ColorValue = FIntVector(Color.X * 255.0f, Color.Y * 255.0f, Color.Z * 255.0f);
-		OpImageLayerInternal::InvokeParallelFor(NumBatches,
+		ParallelExecutionUtils::InvokeBatchParallelFor(NumBatches,
 		[
 			&BatchArgs, ColorValue, BaseFormat
 		](int32 BatchId)
@@ -546,7 +528,7 @@ namespace OpImageLayerInternal
 
 			const int32 NumBatches = OpImageLayerInternal::GetOpLayerNumBatchesLODRange(BatchArgs);
 
-			OpImageLayerInternal::InvokeParallelFor(NumBatches,
+			ParallelExecutionUtils::InvokeBatchParallelFor(NumBatches,
 			[
 				&BatchArgs, BaseChannelOffset, TopColor, ColorChannelOffset, NumColorChannels
 			](int32 BatchId)
@@ -827,7 +809,7 @@ namespace OpImageLayerInternal
 
 		const int32 NumBatches = OpImageLayerInternal::GetOpLayerNumBatchesLODRange(BatchArgs);
 
-		OpImageLayerInternal::InvokeParallelFor(NumBatches,
+		ParallelExecutionUtils::InvokeBatchParallelFor(NumBatches,
 		[
 			&BatchArgs, TopColor, BaseChannelOffset, ColorChannelOffset
 		] (uint32 BatchId)
@@ -1027,7 +1009,7 @@ namespace OpImageLayerInternal
 
 		const int32 NumBatches = OpImageLayerInternal::GetOpLayerNumBatchesLODRange(BatchArgs);
 
-		OpImageLayerInternal::InvokeParallelFor(NumBatches,
+		ParallelExecutionUtils::InvokeBatchParallelFor(NumBatches,
 		[
 			&BatchArgs, BaseChannelOffset, BlendedChannelOffset
 		] (int32 BatchId)
@@ -1102,7 +1084,7 @@ namespace OpImageLayerInternal
 
 		const int32 NumBatches = OpImageLayerInternal::GetOpLayerNumBatchesLODRange(BatchArgs);
 
-		OpImageLayerInternal::InvokeParallelFor(NumBatches,
+		ParallelExecutionUtils::InvokeBatchParallelFor(NumBatches,
 		[
 			&BatchArgs, UnblendedChannels
 		] (int32 BatchId)
@@ -1315,7 +1297,7 @@ namespace OpImageLayerInternal
 
 			const int32 NumBatches = OpImageLayerInternal::GetOpLayerNumBatchesLODRange(BatchArgs);
 
-			OpImageLayerInternal::InvokeParallelFor(NumBatches,
+			ParallelExecutionUtils::InvokeBatchParallelFor(NumBatches,
 			[
 				&BatchArgs, UnblendedChannels
 			] (uint32 BatchId)
@@ -1564,7 +1546,7 @@ namespace OpImageLayerInternal
 
 		const int32 NumBatches = OpImageLayerInternal::GetOpLayerNumBatchesLODRange(BatchArgs);
 
-		OpImageLayerInternal::InvokeParallelFor(NumBatches, 
+		ParallelExecutionUtils::InvokeBatchParallelFor(NumBatches, 
 		[
 			&BatchArgs, UnblendedChannels, BlendImageFormat = BlendImage->GetFormat()
 		] (int32 BatchId)
@@ -1889,7 +1871,7 @@ namespace OpImageLayerInternal
 				? FMath::DivideAndRoundUp(NumRelevantElems, BatchNumElems)
 				: NumBatches;
 
-		OpImageLayerInternal::InvokeParallelFor(NumBatches, 
+		ParallelExecutionUtils::InvokeBatchParallelFor(NumBatches, 
 		[
 			&BatchArgs, BlendAlphaSourceChannel
 		] (int32 BatchId)
@@ -1989,7 +1971,7 @@ namespace OpImageLayerInternal
 				? FMath::DivideAndRoundUp(NumRelevantElems, BatchNumElems)
 				: NumBatches;
 
-		OpImageLayerInternal::InvokeParallelFor(NumRelevantBatches,
+		ParallelExecutionUtils::InvokeBatchParallelFor(NumRelevantBatches,
 		[
 			&BatchArgs
 		] (uint32 BatchId)
@@ -2051,7 +2033,7 @@ namespace OpImageLayerInternal
 
 		const int32 NumBatches = OpImageLayerInternal::GetOpLayerNumBatchesLODRange(BatchArgs);
 
-		OpImageLayerInternal::InvokeParallelFor(NumBatches, 
+		ParallelExecutionUtils::InvokeBatchParallelFor(NumBatches, 
 		[
 			&BatchArgs, BlendAlphaSourceChannel
 		] (uint32 BatchId)
@@ -2226,7 +2208,7 @@ namespace OpImageLayerInternal
 
 		const int32 NumBatches = OpImageLayerInternal::GetOpLayerNumBatchesLODRange(BatchArgs);
 
-		OpImageLayerInternal::InvokeParallelFor(NumBatches,
+		ParallelExecutionUtils::InvokeBatchParallelFor(NumBatches,
 		[
 			&BatchArgs, TopColor
 		](int32 BatchId)
@@ -2276,7 +2258,7 @@ namespace OpImageLayerInternal
 
 		const int32 NumBatches = OpImageLayerInternal::GetOpLayerNumBatchesLODRange(BatchArgs);
 
-		OpImageLayerInternal::InvokeParallelFor(NumBatches,
+		ParallelExecutionUtils::InvokeBatchParallelFor(NumBatches,
 		[
 			&BatchArgs
 		](int32 BatchId)
@@ -2330,7 +2312,7 @@ namespace OpImageLayerInternal
 
 		const int32 NumBatches = OpImageLayerInternal::GetOpLayerNumBatchesLODRange(BatchArgs);
 		
-		OpImageLayerInternal::InvokeParallelFor(NumBatches,
+		ParallelExecutionUtils::InvokeBatchParallelFor(NumBatches,
 		[
 			&BatchArgs
 		](int32 BatchId)
@@ -2389,7 +2371,7 @@ namespace OpImageLayerInternal
 		
 		const int32 NumBatches = OpImageLayerInternal::GetOpLayerNumBatches(BatchArgs);
 
-		OpImageLayerInternal::InvokeParallelFor(NumBatches,
+		ParallelExecutionUtils::InvokeBatchParallelFor(NumBatches,
 		[
 			&BatchArgs, TopColor
 		](int32 BatchId)
@@ -2597,7 +2579,7 @@ namespace OpImageLayerInternal
 		
 		const int32 NumBatches = OpImageLayerInternal::GetOpLayerNumBatches(BatchArgs);
 
-		OpImageLayerInternal::InvokeParallelFor(NumBatches, 
+		ParallelExecutionUtils::InvokeBatchParallelFor(NumBatches, 
 		[
 			&BatchArgs, ImageCombine
 		](int32 BatchId)
diff --git a/Engine/Plugins/Mutable/Source/MutableRuntime/Private/MuR/OpImageLuminance.h b/Engine/Plugins/Mutable/Source/MutableRuntime/Private/MuR/OpImageLuminance.h
--- a/Engine/Plugins/Mutable/Source/MutableRuntime/Private/MuR/OpImageLuminance.h
+++ b/Engine/Plugins/Mutable/Source/MutableRuntime/Private/MuR/OpImageLuminance.h
@@ -4,8 +4,7 @@
 
 #include "MuR/ImagePrivate.h"
 #include "MuR/Platform.h"
-
-#include "Async/ParallelFor.h"
+#include "MuR/ParallelExecutionUtils.h"
 
 namespace mu
 {
@@ -20,7 +19,10 @@ namespace mu
 		const int32 NumBatches = DestImage->DataStorage.GetNumBatches(NumBatchElems, 1); 
 		check(NumBatches == AImage->DataStorage.GetNumBatches(NumBatchElems, BytesPerElem)); 
 
-        auto ProcessBatch = [DestImage, AImage, NumBatchElems, BytesPerElem](int32 BatchId)
+		ParallelExecutionUtils::InvokeBatchParallelFor(NumBatches,
+        [
+            DestImage, AImage, NumBatchElems, BytesPerElem
+        ](int32 BatchId)
         {
 			TArrayView<uint8> DestView = DestImage->DataStorage.GetBatch(BatchId, NumBatchElems, 1);
 			TArrayView<const uint8> AView = AImage->DataStorage.GetBatch(BatchId, NumBatchElems, BytesPerElem);
@@ -66,16 +68,7 @@ namespace mu
                 check(false);
                 break;
             }
-        };
-
-        if (NumBatches == 1)
-        {
-            ProcessBatch(0);
-        }
-        else
-        {
-            ParallelFor(NumBatches, ProcessBatch);
-        }
+        });
 	}
 
 }
diff --git a/Engine/Plugins/Mutable/Source/MutableRuntime/Private/MuR/OpImageResize.cpp b/Engine/Plugins/Mutable/Source/MutableRuntime/Private/MuR/OpImageResize.cpp
--- a/Engine/Plugins/Mutable/Source/MutableRuntime/Private/MuR/OpImageResize.cpp
+++ b/Engine/Plugins/Mutable/Source/MutableRuntime/Private/MuR/OpImageResize.cpp
@@ -2,7 +2,10 @@
 
 #include "MuR/ImagePrivate.h"
 #include "MuR/MutableTrace.h"
+#include "MuR/ParallelExecutionUtils.h"
+
 #include "Async/ParallelFor.h"
+#include "HAL/Platform.h"
 
 namespace
 {
@@ -1003,7 +1006,7 @@ namespace
 					uint16((BatchSizeInBlocksY + 1)*PixelsPerBlock.Y));
 		
 			// Add some extra padding so different threads do not share cache lines.
-			const int32 PerJobStagingBytes = StagingSize.X*StagingSize.Y*DestNumChannels + 64;
+			const int32 PerJobStagingBytes = StagingSize.X*StagingSize.Y*DestNumChannels + PLATFORM_CACHE_LINE_SIZE;
 
 			TArray<uint8, FDefaultMemoryTrackingAllocator<MemoryCounters::FImageMemoryCounter>> StagingMemory;
 			StagingMemory.SetNum(PerJobStagingBytes * NumParallelJobs);
@@ -1011,13 +1014,13 @@ namespace
 
 			miro::SubImageDecompression::FuncRefType DecompressionFunc = SelectDecompressionFunction(DestFormat, SrcFormat);
 
-			auto ProcessJob = 
-				[
-					NumParallelJobs, NumRowBatchesPerJob, 
-					StagingMemoryData, PerJobStagingBytes, 
-					NumBatches, NumBlocks, PixelsPerBlock, BlockSizeInBytes, DestNumChannels, DecompressionFunc,
-					SrcData, SrcSize, DestData, DestSize, OneOverScalingFactor
-				](int32 JobId)
+			ParallelExecutionUtils::InvokeBatchParallelFor(NumParallelJobs, 
+			[
+				NumParallelJobs, NumRowBatchesPerJob, 
+				StagingMemoryData, PerJobStagingBytes, 
+				NumBatches, NumBlocks, PixelsPerBlock, BlockSizeInBytes, DestNumChannels, DecompressionFunc,
+				SrcData, SrcSize, DestData, DestSize, OneOverScalingFactor
+			](int32 JobId)
 			{
 				const int32 JobRowBegin = JobId*NumRowBatchesPerJob;
 				const int32 JobRowEnd   = FMath::Min(JobRowBegin + NumRowBatchesPerJob, NumBatches.Y);
@@ -1118,16 +1121,7 @@ namespace
 						}
 					}
 				}
-			};
-
-			if (NumParallelJobs == 1)
-			{
-				ProcessJob(0);
-			}
-			else if (NumParallelJobs > 1)
-			{	
-				ParallelFor(NumParallelJobs, ProcessJob);
-			}
+			});
 		}
 		else
 		{
diff --git a/Engine/Plugins/Mutable/Source/MutableRuntime/Private/MuR/OpImageSaturate.h b/Engine/Plugins/Mutable/Source/MutableRuntime/Private/MuR/OpImageSaturate.h
--- a/Engine/Plugins/Mutable/Source/MutableRuntime/Private/MuR/OpImageSaturate.h
+++ b/Engine/Plugins/Mutable/Source/MutableRuntime/Private/MuR/OpImageSaturate.h
@@ -4,34 +4,13 @@
 
 #include "MuR/ImagePrivate.h"
 #include "MuR/Platform.h"
+#include "MuR/ParallelExecutionUtils.h" 
 
-#include "Async/ParallelFor.h" 
 #include "Math/VectorRegister.h"
 
 namespace mu
 {
 
-namespace OpImageSaturateInternal
-{
-	// TODO: This is used in multiple places, make it an utility so it can be shared.  
-	FORCENOINLINE void InvokeParallelFor(int32 NumBatches, TFunctionRef<void(int32)> BatchFunc)
-	{
-		if (NumBatches == 0)
-		{
-			return;
-		}
-
-		if (NumBatches == 1)
-		{
-			BatchFunc(0);
-		}
-		else
-		{
-			ParallelFor(NumBatches, BatchFunc);
-		}
-	}
-} // namespace mu::OpImageStaurateInternal
-
 	template<bool bUseVectorImpl = false>
 	inline void ImageSaturate(FImage* AImage, float Factor)
 	{
@@ -54,7 +33,7 @@ namespace OpImageSaturateInternal
 		constexpr int32 NumBatchElems = 1 << 14;
 		const int32 NumBatches = AImage->DataStorage.GetNumBatches(NumBatchElems, BytesPerElem); 
 
-		OpImageSaturateInternal::InvokeParallelFor(NumBatches,
+		ParallelExecutionUtils::InvokeBatchParallelFor(NumBatches,
 		[
 			AImage, NumBatchElems, BytesPerElem, FactorUNorm
 		](int32 BatchId)
diff --git a/Engine/Plugins/Mutable/Source/MutableRuntime/Private/MuR/OpImageSwizzle.cpp b/Engine/Plugins/Mutable/Source/MutableRuntime/Private/MuR/OpImageSwizzle.cpp
--- a/Engine/Plugins/Mutable/Source/MutableRuntime/Private/MuR/OpImageSwizzle.cpp
+++ b/Engine/Plugins/Mutable/Source/MutableRuntime/Private/MuR/OpImageSwizzle.cpp
@@ -2,7 +2,7 @@
 
 #include "MuR/ImagePrivate.h"
 #include "MuR/Operations.h"
-#include "Async/ParallelFor.h"
+#include "MuR/ParallelExecutionUtils.h"
 
 namespace mu
 {
@@ -68,7 +68,10 @@ namespace mu
 		
 						int32 SrcChannel = Channels[Channel];
 
-						auto ProcessBatch = [Result, Src, Sources, NumDestChannels, DestChannel, SrcChannel, NumBatchElems, NumLODs](int32 BatchId)
+						ParallelExecutionUtils::InvokeBatchParallelFor(NumBatches,
+						[
+							Result, Src, Sources, NumDestChannels, DestChannel, SrcChannel, NumBatchElems, NumLODs
+						](int32 BatchId)
 						{
 							TArrayView<const uint8> SrcView = Src->DataStorage.GetBatchLODRange(BatchId, NumBatchElems, 1, 0, NumLODs);
 							TArrayView<uint8> ResultView = Result->DataStorage.GetBatchLODRange(BatchId, NumBatchElems, NumDestChannels, 0, NumLODs);
@@ -83,16 +86,7 @@ namespace mu
 							{
 								DestBuf[I*NumDestChannels] = SrcBuf[I];
 							}
-						};
-
-						if (NumBatches == 1)
-						{
-							ProcessBatch(0);
-						}
-						else if (NumBatches > 1)
-						{
-							ParallelFor(NumBatches, ProcessBatch);
-						}
+						});
 
 						bFilled = true;
 					}
@@ -107,7 +101,10 @@ namespace mu
 		
 						const int32 SrcChannel = Channels[Channel];
 
-						auto ProcessBatch = [Result, Src, Sources, NumDestChannels, DestChannel, SrcChannel, NumBatchElems, NumLODs](int32 BatchId)
+						ParallelExecutionUtils::InvokeBatchParallelFor(NumBatches,
+						[
+							Result, Src, Sources, NumDestChannels, DestChannel, SrcChannel, NumBatchElems, NumLODs
+						](int32 BatchId)
 						{
 							TArrayView<const uint8> SrcView = Src->DataStorage.GetBatchLODRange(BatchId, NumBatchElems, 3, 0, NumLODs);
 							TArrayView<uint8> ResultView = Result->DataStorage.GetBatchLODRange(BatchId, NumBatchElems, NumDestChannels, 0, NumLODs);
@@ -122,16 +119,7 @@ namespace mu
 							{
 								DestBuf[I*NumDestChannels] = SrcBuf[I*3];
 							}
-						};
-
-						if (NumBatches == 1)
-						{
-							ProcessBatch(0);
-						}
-						else if (NumBatches > 1)
-						{
-							ParallelFor(NumBatches, ProcessBatch);
-						}
+						});
 
 						bFilled = true;
 					}
@@ -149,7 +137,10 @@ namespace mu
 						const int32 NumBatches = Src->DataStorage.GetNumBatchesLODRange(NumBatchElems, 4, 0, NumLODs);
 						check(NumBatches == Result->DataStorage.GetNumBatchesLODRange(NumBatchElems, NumDestChannels, 0, NumLODs));
 			
-						auto ProcessBatch = [Result, Src, Sources, NumDestChannels, DestChannel, SrcChannel, NumBatchElems, NumLODs](int32 BatchId)
+						ParallelExecutionUtils::InvokeBatchParallelFor(NumBatches,
+						[
+							Result, Src, Sources, NumDestChannels, DestChannel, SrcChannel, NumBatchElems, NumLODs
+						](int32 BatchId)
 						{
 							TArrayView<const uint8> SrcView = Src->DataStorage.GetBatchLODRange(BatchId, NumBatchElems, 4, 0, NumLODs);  
 							TArrayView<uint8> ResultView = Result->DataStorage.GetBatchLODRange(BatchId, NumBatchElems, NumDestChannels, 0, NumLODs);  
@@ -164,16 +155,7 @@ namespace mu
 							{
 								DestBuf[I*NumDestChannels] = SrcBuf[I*4];
 							}
-						};
-
-						if (NumBatches == 1)
-						{
-							ProcessBatch(0);
-						}
-						else if (NumBatches > 1)
-						{
-							ParallelFor(NumBatches, ProcessBatch);
-						}
+						});
 
 						bFilled = true;
 					}
@@ -193,7 +175,10 @@ namespace mu
 				// Alpha is expected to be filled with 1.
 				const uint8 FillValue = DestChannel < 3 ? 0 : 255;
 
-				auto ProcessBatch = [Result, NumDestChannels, NumBatchElems, NumLODs, DestChannel, FillValue](int32 BatchId)
+				ParallelExecutionUtils::InvokeBatchParallelFor(NumBatches,
+				[
+					Result, NumDestChannels, NumBatchElems, NumLODs, DestChannel, FillValue
+				](int32 BatchId)
 				{
 					TArrayView<uint8> ResultView = Result->DataStorage.GetBatchLODRange(BatchId, NumBatchElems, NumDestChannels, 0, NumLODs);
 					int32 NumElems = ResultView.Num() / NumDestChannels;
@@ -204,16 +189,7 @@ namespace mu
 					{
 						DestBuf[I * NumDestChannels] = FillValue;
 					}
-				};
-
-				if (NumBatches == 1)
-				{
-					ProcessBatch(0);
-				}
-				else if (NumBatches > 1)
-				{
-					ParallelFor(NumBatches, ProcessBatch);
-				}
+				});
 			}
 		}
 	}
diff --git a/Engine/Plugins/Mutable/Source/MutableRuntime/Private/MuR/OpImageTransformPrivate.cpp b/Engine/Plugins/Mutable/Source/MutableRuntime/Private/MuR/OpImageTransformPrivate.cpp
--- a/Engine/Plugins/Mutable/Source/MutableRuntime/Private/MuR/OpImageTransformPrivate.cpp
+++ b/Engine/Plugins/Mutable/Source/MutableRuntime/Private/MuR/OpImageTransformPrivate.cpp
@@ -2,7 +2,7 @@
 
 #include "OpImageTransformPrivate.h"
 
-#include "Async/ParallelFor.h"
+#include "MuR/ParallelExecutionUtils.h"
 
 #include "Math/Vector.h"
 #include "Math/IntVector.h"
@@ -42,7 +42,8 @@ namespace mu
 		const int32 NumBatchRows = FMath::DivideAndRoundUp(BatchSize, DestCropRect.Width());
 		const int32 NumBatches   = FMath::DivideAndRoundUp(DestCropRect.Height(), NumBatchRows);
 
-		auto ProcessBatch = [=](int32 BatchId)
+		ParallelExecutionUtils::InvokeBatchParallelFor(NumBatches,
+		[=](int32 BatchId)
 		{
 			const int32 RowOffset = NumBatchRows * BatchId;
 			const FIntRect BatchRect = FIntRect(
@@ -207,16 +208,7 @@ namespace mu
 					}
 				}
 			}
-		};
-
-		if (NumBatches == 1)
-		{
-			ProcessBatch(0);
-		}
-		else if (NumBatches > 1)
-		{
-			ParallelFor(NumBatches, ProcessBatch);
-		}
+		});
     }
 
     template<int32 NumChannels, EAddressMode AddressMode>
@@ -252,7 +244,8 @@ namespace mu
 		const int32 NumBatchRows = FMath::DivideAndRoundUp(BatchSize, DestCropRect.Width());
 		const int32 NumBatches   = FMath::DivideAndRoundUp(DestCropRect.Height(), NumBatchRows);
 
-		auto ProcessBatch = [=](int32 BatchId)
+		ParallelExecutionUtils::InvokeBatchParallelFor(NumBatches,
+		[=](int32 BatchId)
 		{
 			const int32 RowOffset = NumBatchRows * BatchId;
 			const FIntRect BatchRect = FIntRect(
@@ -399,16 +392,7 @@ namespace mu
 					}
 				}
 			}
-		};
-
-		if (NumBatches == 1)
-		{
-			ProcessBatch(0);
-		}
-		else if (NumBatches > 1)
-		{
-			ParallelFor(NumBatches, ProcessBatch);
-		}
+		});
     }
 
     template<int32 NumChannels, EAddressMode AddressMode, bool bUseVectorImpl>
diff --git a/Engine/Plugins/Mutable/Source/MutableRuntime/Private/MuR/OpMeshApplyPose.h b/Engine/Plugins/Mutable/Source/MutableRuntime/Private/MuR/OpMeshApplyPose.h
--- a/Engine/Plugins/Mutable/Source/MutableRuntime/Private/MuR/OpMeshApplyPose.h
+++ b/Engine/Plugins/Mutable/Source/MutableRuntime/Private/MuR/OpMeshApplyPose.h
@@ -7,7 +7,7 @@
 #include "MuR/MutableTrace.h"
 #include "MuR/Platform.h"
 
-#include "Async/ParallelFor.h"
+#include "MuR/ParallelExecutionUtils.h"
 
 namespace mu
 {
@@ -161,23 +161,25 @@ namespace mu
         check(WeightCount <= MAX_BONES_PER_VERTEX);
 
 		constexpr int32 NumVertsPerBatch = 1 << 11;
-		auto ProcessVertexBatch = 
-			[
-				SkinTransforms = MakeArrayView<const FTransform3f>(SkinTransforms.GetData(), SkinTransforms.Num()),
-				SourcePositionIterBegin,
-				SourceNormalIterBegin,
-				SourceTangentIterBegin,
-				SourceBiNormalIterBegin,
-				BoneIndicesIterBegin,
-				BoneWeightsIterBegin,
-				TargetPositionIterBegin,
-				TargetNormalIterBegin,
-				TargetTangentIterBegin,
-				TargetBiNormalIterBegin,
-				VertexCount,
-				WeightCount,
-				NumVertsPerBatch
-			](int32 BatchId)
+		const int32 NumBatches = FMath::DivideAndRoundUp<int32>(VertexCount, NumVertsPerBatch);
+		
+		ParallelExecutionUtils::InvokeBatchParallelFor(NumBatches,
+		[
+			SkinTransforms = MakeArrayView<const FTransform3f>(SkinTransforms.GetData(), SkinTransforms.Num()),
+			SourcePositionIterBegin,
+			SourceNormalIterBegin,
+			SourceTangentIterBegin,
+			SourceBiNormalIterBegin,
+			BoneIndicesIterBegin,
+			BoneWeightsIterBegin,
+			TargetPositionIterBegin,
+			TargetNormalIterBegin,
+			TargetTangentIterBegin,
+			TargetBiNormalIterBegin,
+			VertexCount,
+			WeightCount,
+			NumVertsPerBatch
+		](int32 BatchId)
 		{
 			const int32 BatchBeginVertIndex = BatchId * NumVertsPerBatch;
 			const int32 BatchEndVertIndex = FMath::Min(BatchBeginVertIndex + NumVertsPerBatch, VertexCount);
@@ -263,17 +265,6 @@ namespace mu
 					(TargetBiNormalIterBegin + VertexIndex).SetFromVec3f(BiNormal.GetSafeNormal());
 				}
 			}
-		};
-
-		const int32 NumBatches = FMath::DivideAndRoundUp<int32>(VertexCount, NumVertsPerBatch);
-
-		if (NumBatches == 1)
-		{
-			ProcessVertexBatch(0);
-		}
-		else
-		{
-			ParallelFor(NumBatches, ProcessVertexBatch);
-		}
+		});
     }
 }
diff --git a/Engine/Plugins/Mutable/Source/MutableRuntime/Private/MuR/ParallelExecutionUtils.cpp b/Engine/Plugins/Mutable/Source/MutableRuntime/Private/MuR/ParallelExecutionUtils.cpp
new file mode 100644
--- /dev/null
+++ b/Engine/Plugins/Mutable/Source/MutableRuntime/Private/MuR/ParallelExecutionUtils.cpp
@@ -0,0 +1,27 @@
+// Copyright Epic Games, Inc. All Rights Reserved.
+
+#include <MuR/ParallelExecutionUtils.h>
+#include <Async/ParallelFor.h>
+
+namespace mu::ParallelExecutionUtils
+{
+	void InvokeBatchParallelFor(int32 Num, TFunctionRef<void(int32)> Body)
+	{
+		if (Num == 1)
+		{
+			Body(0);
+		}
+		else if (Num > 1)
+		{
+			ParallelFor(Num, Body);
+		}
+	}
+
+	void InvokeBatchParallelForSingleThread(int32 Num, TFunctionRef<void(int32)> Body)
+	{
+		for (int32 BatchId = 0; BatchId < Num; ++BatchId)
+		{
+			Body(BatchId);
+		}
+	}
+}
