[Mutable] Now table nodes can have more than one compilation filter.

#rnx

diff --git a/Engine/Plugins/Mutable/Source/CustomizableObject/Internal/MuCO/CustomizableObjectCustomVersion.h b/Engine/Plugins/Mutable/Source/CustomizableObject/Internal/MuCO/CustomizableObjectCustomVersion.h
--- a/Engine/Plugins/Mutable/Source/CustomizableObject/Internal/MuCO/CustomizableObjectCustomVersion.h
+++ b/Engine/Plugins/Mutable/Source/CustomizableObject/Internal/MuCO/CustomizableObjectCustomVersion.h
@@ -188,6 +188,10 @@ struct FCustomizableObjectCustomVersion
 
 		ChangedSwitchNodesInputPinsFriendlyNames,
 
+		TableNodeCompilationFilterArray,
+
+		EnableMutableMacros,
+
 		// -----<new versions can be added above this line>-------------------------------------------------
 		VersionPlusOne,
 		LatestVersion = VersionPlusOne - 1
diff --git a/Engine/Plugins/Mutable/Source/CustomizableObjectEditor/Private/MuCOE/CustomizableObjectEditorModule.cpp b/Engine/Plugins/Mutable/Source/CustomizableObjectEditor/Private/MuCOE/CustomizableObjectEditorModule.cpp
--- a/Engine/Plugins/Mutable/Source/CustomizableObjectEditor/Private/MuCOE/CustomizableObjectEditorModule.cpp
+++ b/Engine/Plugins/Mutable/Source/CustomizableObjectEditor/Private/MuCOE/CustomizableObjectEditorModule.cpp
@@ -85,6 +85,7 @@
 #include "MuCOE/Nodes/CustomizableObjectNodeCurve.h"
 #include "MuCOE/Widgets/CustomizableObjectVariationCustomization.h"
 #include "MuCOE/Widgets/CustomizableObjectLODReductionSettings.h"
+#include "MuCOE/Widgets/CustomizableObjectNodeTableCompilationFilterEditor.h"
 #include "PropertyEditorModule.h"
 #include "AssetRegistry/AssetRegistryModule.h"
 #include "MuCO/CustomizableObjectInstance.h"
@@ -289,6 +290,7 @@ void FCustomizableObjectEditorModule::StartupModule()
 	PropertyModule.RegisterCustomPropertyTypeLayout(NAME_StrProperty, FOnGetPropertyTypeCustomizationInstance::CreateStatic(&FCustomizableObjectStateParameterSelector::MakeInstance), MakeShared<FStatePropertyTypeIdentifier>());
 	PropertyModule.RegisterCustomPropertyTypeLayout(FCustomizableObjectVariation::StaticStruct()->GetFName(), FOnGetPropertyTypeCustomizationInstance::CreateStatic(&FCustomizableObjectVariationCustomization::MakeInstance));
 	PropertyModule.RegisterCustomPropertyTypeLayout(FCustomizableObjectTextureVariation::StaticStruct()->GetFName(), FOnGetPropertyTypeCustomizationInstance::CreateStatic(&FCustomizableObjectVariationCustomization::MakeInstance));
+	PropertyModule.RegisterCustomPropertyTypeLayout(FTableNodeCompilationFilter::StaticStruct()->GetFName(), FOnGetPropertyTypeCustomizationInstance::CreateStatic(&FCustomizableObjectNodeTableCompilationFilterEditor::MakeInstance));
 
 	PropertyModule.NotifyCustomizationModuleChanged();
 
diff --git a/Engine/Plugins/Mutable/Source/CustomizableObjectEditor/Private/MuCOE/GenerateMutableSource/GenerateMutableSource.cpp b/Engine/Plugins/Mutable/Source/CustomizableObjectEditor/Private/MuCOE/GenerateMutableSource/GenerateMutableSource.cpp
--- a/Engine/Plugins/Mutable/Source/CustomizableObjectEditor/Private/MuCOE/GenerateMutableSource/GenerateMutableSource.cpp
+++ b/Engine/Plugins/Mutable/Source/CustomizableObjectEditor/Private/MuCOE/GenerateMutableSource/GenerateMutableSource.cpp
@@ -1969,8 +1969,8 @@ bool FMutableGraphGenerationContext::FSharedSurface::operator==(const FSharedSur
 }
 
 
-FGeneratedMutableDataTableKey::FGeneratedMutableDataTableKey(FString TableName, FName VersionColumn, FName CompilationFilterColumn, TArray<FString> CompilationFilters, ETableCompilationFilterOperationType FilterOperationType)
-	: TableName(TableName), VersionColumn(VersionColumn), CompilationFilterColumn(CompilationFilterColumn), CompilationFilters(CompilationFilters), FilterOperationType(FilterOperationType)
+FGeneratedMutableDataTableKey::FGeneratedMutableDataTableKey(FString TableName, FName VersionColumn, const TArray<FTableNodeCompilationFilter>& CompilationFilterOptions)
+	: TableName(TableName), VersionColumn(VersionColumn), CompilationFilterOptions(CompilationFilterOptions)
 {}
 
 
@@ -1978,14 +1978,18 @@ uint32 GetTypeHash(const FGeneratedMutableDataTableKey& Key)
 {
 	uint32 Hash = GetTypeHash(Key.TableName);
 	Hash = HashCombine(Hash, GetTypeHash(Key.VersionColumn));
-	Hash = HashCombine(Hash, GetTypeHash(Key.CompilationFilterColumn));
 
-	for (int32 NameIndex = 0; NameIndex < Key.CompilationFilters.Num(); ++NameIndex)
+	for (int32 FilterIndex = 0; FilterIndex < Key.CompilationFilterOptions.Num(); ++FilterIndex)
 	{
-		Hash = HashCombine(Hash, GetTypeHash(Key.CompilationFilters[NameIndex]));
-	}
+		Hash = HashCombine(Hash, GetTypeHash(Key.CompilationFilterOptions[FilterIndex].FilterColumn));
 
-	Hash = HashCombine(Hash, GetTypeHash(Key.FilterOperationType));
+		for (int32 NameIndex = 0; NameIndex < Key.CompilationFilterOptions[FilterIndex].Filters.Num(); ++NameIndex)
+		{
+			Hash = HashCombine(Hash, GetTypeHash(Key.CompilationFilterOptions[FilterIndex].Filters[NameIndex]));
+		}
+
+		Hash = HashCombine(Hash, GetTypeHash(Key.CompilationFilterOptions[FilterIndex].OperationType));
+	}
 
 	return Hash;
 }
diff --git a/Engine/Plugins/Mutable/Source/CustomizableObjectEditor/Private/MuCOE/GenerateMutableSource/GenerateMutableSource.h b/Engine/Plugins/Mutable/Source/CustomizableObjectEditor/Private/MuCOE/GenerateMutableSource/GenerateMutableSource.h
--- a/Engine/Plugins/Mutable/Source/CustomizableObjectEditor/Private/MuCOE/GenerateMutableSource/GenerateMutableSource.h
+++ b/Engine/Plugins/Mutable/Source/CustomizableObjectEditor/Private/MuCOE/GenerateMutableSource/GenerateMutableSource.h
@@ -534,7 +534,7 @@ public:
 
 struct FGeneratedMutableDataTableKey
 {
-	FGeneratedMutableDataTableKey(FString TableName, FName VersionColumn, FName CompilationFilterColumn, TArray<FString> CompilationFilters, ETableCompilationFilterOperationType FilterOperationType);
+	FGeneratedMutableDataTableKey(FString TableName, FName VersionColumn, const TArray<FTableNodeCompilationFilter>& CompilationFilterOptions);
 
 	// Name of the Data Table Asset
 	FString TableName;
@@ -543,14 +543,8 @@ struct FGeneratedMutableDataTableKey
 	// Name of the column that determines de version control
 	FName VersionColumn;
 
-	// Name of the column that determines the compilation filter
-	FName CompilationFilterColumn;
-
-	// Compilation Filter
-	TArray<FString> CompilationFilters;
-
-	// Restriction option
-	ETableCompilationFilterOperationType FilterOperationType;
+	// Compilation Filters
+	TArray<FTableNodeCompilationFilter> CompilationFilterOptions;
 
 	bool operator==(const FGeneratedMutableDataTableKey&) const = default;
 };
diff --git a/Engine/Plugins/Mutable/Source/CustomizableObjectEditor/Private/MuCOE/GenerateMutableSource/GenerateMutableSourceMesh.cpp b/Engine/Plugins/Mutable/Source/CustomizableObjectEditor/Private/MuCOE/GenerateMutableSource/GenerateMutableSourceMesh.cpp
--- a/Engine/Plugins/Mutable/Source/CustomizableObjectEditor/Private/MuCOE/GenerateMutableSource/GenerateMutableSourceMesh.cpp
+++ b/Engine/Plugins/Mutable/Source/CustomizableObjectEditor/Private/MuCOE/GenerateMutableSource/GenerateMutableSourceMesh.cpp
@@ -3517,8 +3517,7 @@ mu::NodeMeshPtr GenerateMorphMesh(const UEdGraphPin* Pin,
 		int32 NumRows = RowNames.Num();
 
 		// Should exist
-		FGeneratedMutableDataTableKey MutableTableKey(DataTable->GetName(), TypedNodeTable->VersionColumn,
-			TypedNodeTable->CompilationFilterColumn, TypedNodeTable->CompilationFilters, TypedNodeTable->FilterOperationType);
+		FGeneratedMutableDataTableKey MutableTableKey(DataTable->GetName(), TypedNodeTable->VersionColumn, TypedNodeTable->CompilationFilterOptions);
 		mu::Ptr<mu::FTable> GeneratedTable = GenerationContext.GeneratedTables[MutableTableKey].GeneratedTable;
 		check(GeneratedTable);
 
diff --git a/Engine/Plugins/Mutable/Source/CustomizableObjectEditor/Private/MuCOE/GenerateMutableSource/GenerateMutableSourceTable.cpp b/Engine/Plugins/Mutable/Source/CustomizableObjectEditor/Private/MuCOE/GenerateMutableSource/GenerateMutableSourceTable.cpp
--- a/Engine/Plugins/Mutable/Source/CustomizableObjectEditor/Private/MuCOE/GenerateMutableSource/GenerateMutableSourceTable.cpp
+++ b/Engine/Plugins/Mutable/Source/CustomizableObjectEditor/Private/MuCOE/GenerateMutableSource/GenerateMutableSourceTable.cpp
@@ -800,8 +800,7 @@ void GenerateUniqueRowIds(const TArray<FName>& RowNames, TArray<uint32>& OutRowI
 
 TArray<FName> GetRowsToCompile(const UDataTable& DataTable, const UCustomizableObjectNodeTable& TableNode, FMutableGraphGenerationContext& GenerationContext, TArray<uint32>& OutRowIds)
 {
-	FGeneratedMutableDataTableKey MutableTableKey(DataTable.GetName(), TableNode.VersionColumn,
-		TableNode.CompilationFilterColumn, TableNode.CompilationFilters, TableNode.FilterOperationType);
+	FGeneratedMutableDataTableKey MutableTableKey(DataTable.GetName(), TableNode.VersionColumn, TableNode.CompilationFilterOptions);
 	if (FMutableGraphGenerationContext::FGeneratedDataTablesData* Result = GenerationContext.GeneratedTables.Find(MutableTableKey))
 	{
 		OutRowIds = Result->RowIds;
@@ -1047,8 +1046,7 @@ mu::Ptr<mu::FTable> GenerateMutableSourceTable(const UDataTable* DataTable, cons
 	// Checking if the table is in the cache
 	const FString TableName = DataTable->GetName();
 
-	FGeneratedMutableDataTableKey MutableTableKey(TableName, TableNode->VersionColumn,
-		TableNode->CompilationFilterColumn, TableNode->CompilationFilters, TableNode->FilterOperationType);
+	FGeneratedMutableDataTableKey MutableTableKey(TableName, TableNode->VersionColumn, TableNode->CompilationFilterOptions);
 	if (FMutableGraphGenerationContext::FGeneratedDataTablesData* CachedTable = GenerationContext.GeneratedTables.Find(MutableTableKey))
 	{
 		// Generating Parameter Metadata for parameters that reuse a Table
@@ -1315,7 +1313,9 @@ bool GenerateNoneRow(const UCustomizableObjectNodeTable& TableNode, const UEdGra
 			ColumnId = MutableTable->AddColumn(ColumnName, mu::ETableColumnType::Color);
 		}
 
-		// Encoding an invalid value. QuietNaN does not trigger errors nor checks. 
+		// HACK: Encoding an invalid value (Nan) for table option "None".
+		// Using Nan avoids that color operations modify this encoded number since any operation returns Nan.
+		// Also, QuietNaNs do not trigger errors nor checks.
 		// It's checked in the moment that the material of the COI is generated.
 		FLinearColor InvalidColor(std::numeric_limits<float>::quiet_NaN(), 0.0f, 0.0f);
 		MutableTable->SetCell(ColumnId, RowId, TableNode.bUseMaterialColor ? InvalidColor : FLinearColor::Black);
diff --git a/Engine/Plugins/Mutable/Source/CustomizableObjectEditor/Private/MuCOE/Nodes/CustomizableObjectNodeTable.cpp b/Engine/Plugins/Mutable/Source/CustomizableObjectEditor/Private/MuCOE/Nodes/CustomizableObjectNodeTable.cpp
--- a/Engine/Plugins/Mutable/Source/CustomizableObjectEditor/Private/MuCOE/Nodes/CustomizableObjectNodeTable.cpp
+++ b/Engine/Plugins/Mutable/Source/CustomizableObjectEditor/Private/MuCOE/Nodes/CustomizableObjectNodeTable.cpp
@@ -311,8 +311,8 @@ void UCustomizableObjectNodeTable::BackwardsCompatibleFixup(int32 CustomizableOb
 
 				if (BoolProperty)
 				{
-					CompilationFilterColumn = FName(*BoolProperty->GetDisplayNameText().ToString());
-					CompilationFilters.Add("false");
+					CompilationFilterColumn_DEPRECATED = FName(*BoolProperty->GetDisplayNameText().ToString());
+					CompilationFilters_DEPRECATED.Add("false");
 
 					// There should be only one Bool column
 					break;
@@ -321,6 +321,14 @@ void UCustomizableObjectNodeTable::BackwardsCompatibleFixup(int32 CustomizableOb
 
 		}
 	}
+
+	if (CustomizableObjectCustomVersion == FCustomizableObjectCustomVersion::TableNodeCompilationFilterArray)
+	{
+		if (!CompilationFilterColumn_DEPRECATED.IsNone())
+		{
+			CompilationFilterOptions.Add({ CompilationFilterColumn_DEPRECATED, CompilationFilters_DEPRECATED, FilterOperationType_DEPRECATED });
+		}
+	}
 }
 
 
@@ -1590,7 +1598,7 @@ TArray<FName> UCustomizableObjectNodeTable::GetEnabledRows(const UDataTable& Dat
 
 	if (!TableStruct)
 	{
-		return RowNames;
+		return TArray<FName>();
 	}
 
 	TArray<FName> TableRowNames = DataTable.GetRowNames();
@@ -1600,81 +1608,88 @@ TArray<FName> UCustomizableObjectNodeTable::GetEnabledRows(const UDataTable& Dat
 	TableRowNames.Sort([](const FName& A, const FName& B) { return A.ToString() < B.ToString(); });
 
 	// Check if there is any filter condition
-	if (CompilationFilterColumn.IsNone() || !CompilationFilters.Num())
+	if (CompilationFilterOptions.IsEmpty())
 	{
 		return TableRowNames;
 	}
 
-	FProperty* ColumnProperty = FindColumnProperty(CompilationFilterColumn);
+	RowNames.Reserve(TableRowNames.Num());
 
-	if (!ColumnProperty)
+	for (int32 FilterIndex = 0; FilterIndex < CompilationFilterOptions.Num(); ++FilterIndex)
 	{
-		return TableRowNames;
-	}
+		FProperty* ColumnProperty = FindColumnProperty(CompilationFilterOptions[FilterIndex].FilterColumn);
 
-	RowNames.Reserve(TableRowNames.Num());
+		if (!ColumnProperty)
+		{
+			continue;
+		}
 
-	if (UCustomizableObjectNodeTable::SupportedFilterTypes.Contains(ColumnProperty->GetClass()))
-	{
-		for (int32 RowIndex = 0; RowIndex < TableRowNames.Num(); ++RowIndex)
+		if (UCustomizableObjectNodeTable::SupportedFilterTypes.Contains(ColumnProperty->GetClass()))
 		{
-			if (uint8* CellData = UCustomizableObjectNodeTable::GetCellData(TableRowNames[RowIndex], DataTable, *ColumnProperty))
+			for (int32 RowIndex = 0; RowIndex < TableRowNames.Num(); ++RowIndex)
 			{
-				// By getting values as Texts we enable the support to FText columns
-				FString Value = DataTableUtils::GetPropertyValueAsTextDirect(ColumnProperty, CellData).ToString();
-
-				if (CompilationFilters.Contains(Value))
+				if (uint8* CellData = UCustomizableObjectNodeTable::GetCellData(TableRowNames[RowIndex], DataTable, *ColumnProperty))
 				{
-					RowNames.Add(TableRowNames[RowIndex]);
+					// By getting values as Texts we enable the support to FText columns
+					FString Value = DataTableUtils::GetPropertyValueAsTextDirect(ColumnProperty, CellData).ToString();
+
+					if (CompilationFilterOptions[FilterIndex].Filters.Contains(Value))
+					{
+						RowNames.Add(TableRowNames[RowIndex]);
+					}
 				}
 			}
 		}
-	}
-
-	else if (FArrayProperty* ArrayProperty = CastField<FArrayProperty>(ColumnProperty))
-	{
-		FProperty* InnerProperty = ArrayProperty->Inner;
-		check(InnerProperty);
 
-		if (UCustomizableObjectNodeTable::SupportedFilterTypes.Contains(InnerProperty->GetClass()))
+		else if (FArrayProperty* ArrayProperty = CastField<FArrayProperty>(ColumnProperty))
 		{
-			for (int32 RowIndex = 0; RowIndex < TableRowNames.Num(); ++RowIndex)
+			FProperty* InnerProperty = ArrayProperty->Inner;
+			check(InnerProperty);
+
+			if (UCustomizableObjectNodeTable::SupportedFilterTypes.Contains(InnerProperty->GetClass()))
 			{
-				if (uint8* CellData = UCustomizableObjectNodeTable::GetCellData(TableRowNames[RowIndex], DataTable, *ArrayProperty))
+				for (int32 RowIndex = 0; RowIndex < TableRowNames.Num(); ++RowIndex)
 				{
-					FScriptArrayHelper ArrayHelper(ArrayProperty, CellData);
-					int32 NumMatchingFilters = 0;
-
-					for (int32 NameIndex = 0; NameIndex < ArrayHelper.Num(); ++NameIndex)
+					if (uint8* CellData = UCustomizableObjectNodeTable::GetCellData(TableRowNames[RowIndex], DataTable, *ArrayProperty))
 					{
-						// By getting values as Texts we enable the support to FText columns
-						FString Value = DataTableUtils::GetPropertyValueAsTextDirect(InnerProperty, ArrayHelper.GetRawPtr(NameIndex)).ToString();
+						FScriptArrayHelper ArrayHelper(ArrayProperty, CellData);
+						int32 NumMatchingFilters = 0;
 
-						if (CompilationFilters.Contains(Value))
+						for (int32 NameIndex = 0; NameIndex < ArrayHelper.Num(); ++NameIndex)
 						{
-							if (FilterOperationType == ETableCompilationFilterOperationType::TCFOT_OR)
-							{
-								RowNames.Add(TableRowNames[RowIndex]);
-								break;
-							}
-							else
-							{
-								NumMatchingFilters++;
+							// By getting values as Texts we enable the support to FText columns
+							FString Value = DataTableUtils::GetPropertyValueAsTextDirect(InnerProperty, ArrayHelper.GetRawPtr(NameIndex)).ToString();
 
-								if (NumMatchingFilters == CompilationFilters.Num())
+							if (CompilationFilterOptions[FilterIndex].Filters.Contains(Value))
+							{
+								if (CompilationFilterOptions[FilterIndex].OperationType == ETableCompilationFilterOperationType::TCFOT_OR)
 								{
 									RowNames.Add(TableRowNames[RowIndex]);
 									break;
 								}
+								else
+								{
+									NumMatchingFilters++;
+
+									if (NumMatchingFilters == CompilationFilterOptions[FilterIndex].Filters.Num())
+									{
+										RowNames.Add(TableRowNames[RowIndex]);
+										break;
+									}
+								}
 							}
 						}
 					}
 				}
 			}
 		}
+
+		// We will iterate through mutliple filters
+		TableRowNames = RowNames;
+		RowNames.Reset();
 	}
 
-	return RowNames;
+	return TableRowNames;
 }
 
 
diff --git a/Engine/Plugins/Mutable/Source/CustomizableObjectEditor/Private/MuCOE/Nodes/CustomizableObjectNodeTable.h b/Engine/Plugins/Mutable/Source/CustomizableObjectEditor/Private/MuCOE/Nodes/CustomizableObjectNodeTable.h
--- a/Engine/Plugins/Mutable/Source/CustomizableObjectEditor/Private/MuCOE/Nodes/CustomizableObjectNodeTable.h
+++ b/Engine/Plugins/Mutable/Source/CustomizableObjectEditor/Private/MuCOE/Nodes/CustomizableObjectNodeTable.h
@@ -188,6 +188,29 @@ public:
 };
 
 
+USTRUCT()
+struct FTableNodeCompilationFilter
+{
+	GENERATED_USTRUCT_BODY()
+
+	/** Column of the Data Table that contains the compilation filters of each row */
+	UPROPERTY(EditAnywhere, Category = CompilationRestrictions)
+	FName FilterColumn;
+
+	/** Compilation filters required for this Table node represented in a string
+		Supported types: Names, Strings, int, bools (true, false). */
+	UPROPERTY(EditAnywhere, Category = CompilationRestrictions)
+	TArray<FString> Filters;
+
+	/** Determines how Filters interact with the values of the Compilation Filter Column.
+		Check the tooltip of each operation for more information. */
+	UPROPERTY(EditAnywhere, Category = CompilationRestrictions)
+	ETableCompilationFilterOperationType OperationType = ETableCompilationFilterOperationType::TCFOT_OR;
+
+	bool operator==(const FTableNodeCompilationFilter&) const = default;
+};
+
+
 UCLASS()
 class CUSTOMIZABLEOBJECTEDITOR_API UCustomizableObjectNodeTable : public UCustomizableObjectNode
 {
@@ -203,20 +226,6 @@ public:
 	UPROPERTY(EditAnywhere, Category = TableProperties)
 	bool bAddNoneOption;
 
-	/** Column of the Data Table that contains the compilation filters of each row */
-	UPROPERTY(EditAnywhere, Category = CompilationRestrictions)
-	FName CompilationFilterColumn;
-
-	/** Compilation filters required for this Table node represented in a string
-		Supported types: Names, Strings, int, bools (true, false). */
-	UPROPERTY(EditAnywhere, Category = CompilationRestrictions)
-	TArray<FString> CompilationFilters;
-
-	/** Determines how Filters interact with the values of the Compilation Filter Column.
-		Check the tooltip of each operation for more information. */
-	UPROPERTY(EditAnywhere, Category = CompilationRestrictions)
-	ETableCompilationFilterOperationType FilterOperationType = ETableCompilationFilterOperationType::TCFOT_OR;
-
 	/** Source where table gathers the data */
 	UPROPERTY(EditAnywhere, Category = TableProperties)
 	ETableDataGatheringSource TableDataGatheringMode = ETableDataGatheringSource::ETDGM_DataTable;
@@ -270,6 +279,9 @@ public:
 	UPROPERTY(EditAnywhere, Category = TableProperties, meta = (DisplayAfter = "bAddNoneOption", EditCondition = "bAddNoneOption==true", EditConditionHides))
 	bool bUseMaterialColor = false;
 
+	UPROPERTY(EditAnywhere, Category = CompilationRestrictions)
+	TArray<FTableNodeCompilationFilter> CompilationFilterOptions;
+
 	// Array with all the classes supported to be used as row filters
 	static const TArray<FFieldClass*> SupportedFilterTypes;
 
@@ -443,4 +455,18 @@ private:
 	UPROPERTY()
 	bool bDisableCheckedRows_DEPRECATED = true;
 
+	/** Column of the Data Table that contains the compilation filters of each row */
+	UPROPERTY()
+	FName CompilationFilterColumn_DEPRECATED;
+
+	/** Compilation filters required for this Table node represented in a string
+		Supported types: Names, Strings, int, bools (true, false). */
+	UPROPERTY()
+	TArray<FString> CompilationFilters_DEPRECATED;
+
+	/** Determines how Filters interact with the values of the Compilation Filter Column.
+		Check the tooltip of each operation for more information. */
+	UPROPERTY()
+	ETableCompilationFilterOperationType FilterOperationType_DEPRECATED = ETableCompilationFilterOperationType::TCFOT_OR;
+
 };
diff --git a/Engine/Plugins/Mutable/Source/CustomizableObjectEditor/Private/MuCOE/Nodes/CustomizableObjectNodeTableDetails.cpp b/Engine/Plugins/Mutable/Source/CustomizableObjectEditor/Private/MuCOE/Nodes/CustomizableObjectNodeTableDetails.cpp
--- a/Engine/Plugins/Mutable/Source/CustomizableObjectEditor/Private/MuCOE/Nodes/CustomizableObjectNodeTableDetails.cpp
+++ b/Engine/Plugins/Mutable/Source/CustomizableObjectEditor/Private/MuCOE/Nodes/CustomizableObjectNodeTableDetails.cpp
@@ -66,7 +66,6 @@ void FCustomizableObjectNodeTableDetails::CustomizeDetails(const TSharedPtr<IDet
 		TSharedPtr<FString> CurrentMutableMetadataColumn = GenerateMutableMetaDataColumnComboBoxOptions();
 		TSharedPtr<FString> CurrentVersionColumn = GenerateVersionColumnComboBoxOptions();
 		TSharedPtr<FString> CurrentThumbnailColumn = GenerateThumbnailColumnComboBoxOptions();
-		TSharedPtr<FString> CurrentCompilationFilterColumn = GenerateCompilationFilterColumnComboBoxOptions();
 
 		CustomizableObjectCategory.AddProperty("ParameterName");
 		TSharedRef<IPropertyHandle> AddNoneOptionProperty = DetailBuilder->GetProperty("bAddNoneOption");
@@ -96,34 +95,6 @@ void FCustomizableObjectNodeTableDetails::CustomizeDetails(const TSharedPtr<IDet
 			]
 			.OverrideResetToDefault(FResetToDefaultOverride::Create(FSimpleDelegate::CreateSP(this, &FCustomizableObjectNodeTableDetails::OnVersionColumnComboBoxSelectionReset)));
 
-		IDetailGroup& CompilationFilterGroup = CompilationRestrictionsCategory.AddGroup(TEXT("CompilationFilterGroupo"), LOCTEXT("CompilationFilterColumn_Selector", "Compilation Filter Column"), false, true);
-		CompilationFilterGroup.HeaderRow()
-			.NameContent()
-			[
-				SNew(STextBlock)
-					.Text(LOCTEXT("CompilationFilterColumn_SelectorText", "Compilation Filter Column"))
-					.ToolTipText(LOCTEXT("CompilationFilterColumn_SelectorTooltip", "Select the column that contains the compilation filter of each row."))
-					.Font(IDetailLayoutBuilder::GetDetailFont())
-			]
-			.ValueContent()
-			[
-				SAssignNew(CompilationFilterColumnComboBox, STextComboBox)
-					.InitiallySelectedItem(CurrentCompilationFilterColumn)
-					.OptionsSource(&CompilationFilterColumnOptionNames)
-					.OnComboBoxOpening(this, &FCustomizableObjectNodeTableDetails::OnOpenCompilationFilterColumnComboBox)
-					.OnSelectionChanged(this, &FCustomizableObjectNodeTableDetails::OnCompilationFilterColumnComboBoxSelectionChanged)
-					.Font(IDetailLayoutBuilder::GetDetailFont())
-					.ColorAndOpacity(this, &FCustomizableObjectNodeTableDetails::GetCompilationFilterColumnComboBoxTextColor, &CompilationFilterColumnOptionNames)
-			]
-			.OverrideResetToDefault(FResetToDefaultOverride::Create(FSimpleDelegate::CreateSP(this, &FCustomizableObjectNodeTableDetails::OnCompilationFilterColumnComboBoxSelectionReset)));
-		
-		TSharedRef<IPropertyHandle> CompilationFilterPropertyHandle = DetailBuilder->GetProperty("CompilationFilters");
-		TSharedRef<IPropertyHandle> FilterRequiredPropertyHandle = DetailBuilder->GetProperty("FilterOperationType");
-
-		CompilationFilterGroup.AddPropertyRow(FilterRequiredPropertyHandle);
-		CompilationFilterGroup.AddPropertyRow(CompilationFilterPropertyHandle);
-
-
 		UICategory.AddCustomRow(LOCTEXT("MutableUIMetadataColumn_Selector", "MutableUIMetadataColumn"))
 			.NameContent()
 			[
@@ -869,120 +840,4 @@ void FCustomizableObjectNodeTableDetails::OnVersionColumnComboBoxSelectionReset(
 	}
 }
 
-
-TSharedPtr<FString> FCustomizableObjectNodeTableDetails::GenerateCompilationFilterColumnComboBoxOptions()
-{
-	const UScriptStruct* TableStruct = Node->GetTableNodeStruct();
-	TSharedPtr<FString> CurrentSelection;
-	CompilationFilterColumnOptionNames.Reset();
-
-	FText NothingSelectedText = LOCTEXT("NothingSelectedText", "- Nothing Selected -");
-	CompilationFilterColumnOptionNames.Add(MakeShareable(new FString(NothingSelectedText.ToString())));
-
-	if (!TableStruct)
-	{
-		return CurrentSelection;
-	}
-
-	// Iterating struct Options
-	for (TFieldIterator<FProperty> It(TableStruct); It; ++It)
-	{
-		FProperty* ColumnProperty = *It;
-
-		if (!ColumnProperty)
-		{
-			continue;
-		}
-
-		if (const FArrayProperty* ArrayProperty = CastField<FArrayProperty>(ColumnProperty))
-		{
-			ColumnProperty = ArrayProperty->Inner;
-		}
-
-		if (UCustomizableObjectNodeTable::SupportedFilterTypes.Contains(ColumnProperty->GetClass()))
-		{
-			TSharedPtr<FString> Option = MakeShareable(new FString(ColumnProperty->GetDisplayNameText().ToString()));
-			CompilationFilterColumnOptionNames.Add(Option);
-
-			if (*Option == Node->CompilationFilterColumn)
-			{
-				CurrentSelection = CompilationFilterColumnOptionNames.Last();
-			}
-		}
-	}
-
-	if (!CurrentSelection)
-	{
-		if (!Node->CompilationFilterColumn.IsNone())
-		{
-			CompilationFilterColumnOptionNames.Add(MakeShareable(new FString(Node->CompilationFilterColumn.ToString())));
-			CurrentSelection = CompilationFilterColumnOptionNames.Last();
-		}
-		else
-		{
-			CurrentSelection = CompilationFilterColumnOptionNames[0];
-		}
-	}
-
-	return CurrentSelection;
-}
-
-
-void FCustomizableObjectNodeTableDetails::OnOpenCompilationFilterColumnComboBox()
-{
-	TSharedPtr<FString> CurrentSelection = GenerateCompilationFilterColumnComboBoxOptions();
-
-	if (CompilationFilterColumnComboBox.IsValid())
-	{
-		CompilationFilterColumnComboBox->ClearSelection();
-		CompilationFilterColumnComboBox->RefreshOptions();
-		CompilationFilterColumnComboBox->SetSelectedItem(CurrentSelection);
-	}
-}
-
-
-void FCustomizableObjectNodeTableDetails::OnCompilationFilterColumnComboBoxSelectionChanged(TSharedPtr<FString> Selection, ESelectInfo::Type SelectInfo)
-{
-	if (Selection && Node->CompilationFilterColumn != FName(*Selection)
-		&& (SelectInfo == ESelectInfo::OnKeyPress || SelectInfo == ESelectInfo::OnMouseClick))
-	{
-		if (Selection == CompilationFilterColumnOptionNames[0])
-		{
-			Node->CompilationFilterColumn = NAME_None;
-		}
-		else
-		{
-			Node->CompilationFilterColumn = FName(*Selection);
-		}
-
-		Node->MarkPackageDirty();
-	}
-}
-
-
-FSlateColor FCustomizableObjectNodeTableDetails::GetCompilationFilterColumnComboBoxTextColor(TArray<TSharedPtr<FString>>* CurrentOptions) const
-{
-	if (Node->FindColumnProperty(Node->CompilationFilterColumn) || Node->CompilationFilterColumn.IsNone())
-	{
-		return FSlateColor::UseForeground();
-	}
-
-	// Table Struct null or does not contain the selected property anymore
-	return FSlateColor(FLinearColor(0.9f, 0.05f, 0.05f, 1.0f));
-}
-
-
-void FCustomizableObjectNodeTableDetails::OnCompilationFilterColumnComboBoxSelectionReset()
-{
-	Node->CompilationFilterColumn = NAME_None;
-
-	if (CompilationFilterColumnComboBox.IsValid())
-	{
-		GenerateCompilationFilterColumnComboBoxOptions();
-		CompilationFilterColumnComboBox->SetSelectedItem(CompilationFilterColumnOptionNames[0]);
-		CompilationFilterColumnComboBox->RefreshOptions();
-	}
-}
-
-
 #undef LOCTEXT_NAMESPACE
diff --git a/Engine/Plugins/Mutable/Source/CustomizableObjectEditor/Private/MuCOE/Nodes/CustomizableObjectNodeTableDetails.h b/Engine/Plugins/Mutable/Source/CustomizableObjectEditor/Private/MuCOE/Nodes/CustomizableObjectNodeTableDetails.h
--- a/Engine/Plugins/Mutable/Source/CustomizableObjectEditor/Private/MuCOE/Nodes/CustomizableObjectNodeTableDetails.h
+++ b/Engine/Plugins/Mutable/Source/CustomizableObjectEditor/Private/MuCOE/Nodes/CustomizableObjectNodeTableDetails.h
@@ -124,23 +124,6 @@ private:
 	// OnComboBoxSelectionChanged Callback for Layout ComboBox
 	void OnVersionColumnComboBoxSelectionReset();
 
-	// Generates MutableMetadata columns combobox options
-	// Returns the current selected option or a null pointer
-	TSharedPtr<FString> GenerateCompilationFilterColumnComboBoxOptions();
-
-	// Callback to regenerate the combobox options
-	void OnOpenCompilationFilterColumnComboBox();
-
-	// OnComboBoxSelectionChanged Callback for Layout ComboBox
-	void OnCompilationFilterColumnComboBoxSelectionChanged(TSharedPtr<FString> Selection, ESelectInfo::Type SelectInfo);
-
-	// Sets the combo box selection color
-	FSlateColor GetCompilationFilterColumnComboBoxTextColor(TArray<TSharedPtr<FString>>* CurrentOptions) const;
-
-	// OnComboBoxSelectionChanged Callback for Layout ComboBox
-	void OnCompilationFilterColumnComboBoxSelectionReset();
-
-
 private:
 
 	// Details -------------
@@ -196,10 +179,4 @@ private:
 	// ComboBox widget to select a VersionColumn from the NodeTable
 	TSharedPtr<STextComboBox> VersionColumnsComboBox;
 
-	// Array with the name of the Version columns
-	TArray<TSharedPtr<FString>> CompilationFilterColumnOptionNames;
-	
-	// ComboBox widget to select a VersionColumn from the NodeTable
-	TSharedPtr<STextComboBox> CompilationFilterColumnComboBox;
-
 };
diff --git a/Engine/Plugins/Mutable/Source/CustomizableObjectEditor/Private/MuCOE/Widgets/CustomizableObjectNodeTableCompilationFilterEditor.cpp b/Engine/Plugins/Mutable/Source/CustomizableObjectEditor/Private/MuCOE/Widgets/CustomizableObjectNodeTableCompilationFilterEditor.cpp
new file mode 100644
--- /dev/null
+++ b/Engine/Plugins/Mutable/Source/CustomizableObjectEditor/Private/MuCOE/Widgets/CustomizableObjectNodeTableCompilationFilterEditor.cpp
@@ -0,0 +1,211 @@
+// Copyright Epic Games, Inc. All Rights Reserved.
+
+#include "MuCOE/Widgets/CustomizableObjectNodeTableCompilationFilterEditor.h"
+
+#include "DetailLayoutBuilder.h"
+#include "IDetailChildrenBuilder.h"
+#include "IDetailGroup.h"
+#include "IDetailsView.h"
+#include "MuCOE/Nodes/CustomizableObjectNodeTable.h"
+#include "PropertyCustomizationHelpers.h"
+#include "Widgets/Input/STextComboBox.h"
+#include "Widgets/Text/STextBlock.h"
+
+
+#define LOCTEXT_NAMESPACE "CustomizableObjectDetails"
+
+
+TSharedRef<IPropertyTypeCustomization> FCustomizableObjectNodeTableCompilationFilterEditor::MakeInstance()
+{
+	return MakeShared<FCustomizableObjectNodeTableCompilationFilterEditor>();
+}
+
+
+void FCustomizableObjectNodeTableCompilationFilterEditor::CustomizeHeader(TSharedRef<IPropertyHandle> StructPropertyHandle, class FDetailWidgetRow& HeaderRow, IPropertyTypeCustomizationUtils& StructCustomizationUtils)
+{
+	HeaderRow
+	.NameContent()
+	[
+		StructPropertyHandle->CreatePropertyNameWidget()
+	];
+}
+
+
+void FCustomizableObjectNodeTableCompilationFilterEditor::CustomizeChildren(TSharedRef<IPropertyHandle> StructPropertyHandle, class IDetailChildrenBuilder& ChildBuilder, IPropertyTypeCustomizationUtils& StructCustomizationUtils)
+{
+	TArray<UObject*> OuterObjects;
+	StructPropertyHandle->GetOuterObjects(OuterObjects);
+
+	// Get table node
+	if (OuterObjects.Num())
+	{
+		Node = Cast<UCustomizableObjectNodeTable>(OuterObjects[0]);
+
+		if (!Node.IsValid())
+		{
+			return;
+		}
+	}
+
+	ColumnPropertyHandle = StructPropertyHandle->GetChildHandle(GET_MEMBER_NAME_CHECKED(FTableNodeCompilationFilter, FilterColumn));
+	TSharedPtr<IPropertyHandle> CompilationFilterPropertyHandle = StructPropertyHandle->GetChildHandle(GET_MEMBER_NAME_CHECKED(FTableNodeCompilationFilter, Filters));
+	TSharedPtr<IPropertyHandle> OperationTypePropertyHandle = StructPropertyHandle->GetChildHandle(GET_MEMBER_NAME_CHECKED(FTableNodeCompilationFilter, OperationType));
+
+	if (ColumnPropertyHandle && CompilationFilterPropertyHandle && OperationTypePropertyHandle)
+	{
+		TSharedPtr<FString> CurrentCompilationFilterColumn = GenerateCompilationFilterColumnComboBoxOptions();
+
+		ChildBuilder.AddProperty(ColumnPropertyHandle.ToSharedRef()).CustomWidget()
+			.NameContent()
+			[
+				ColumnPropertyHandle->CreatePropertyNameWidget()
+			]
+			.ValueContent()
+			[
+				SAssignNew(CompilationFilterColumnComboBox, STextComboBox)
+					.InitiallySelectedItem(CurrentCompilationFilterColumn)
+					.OptionsSource(&CompilationFilterColumnOptionNames)
+					.OnComboBoxOpening(this, &FCustomizableObjectNodeTableCompilationFilterEditor::OnOpenCompilationFilterColumnComboBox)
+					.OnSelectionChanged(this, &FCustomizableObjectNodeTableCompilationFilterEditor::OnCompilationFilterColumnComboBoxSelectionChanged)
+					.Font(IDetailLayoutBuilder::GetDetailFont())
+					.ColorAndOpacity(this, &FCustomizableObjectNodeTableCompilationFilterEditor::GetCompilationFilterColumnComboBoxTextColor)
+			];
+
+		ChildBuilder.AddProperty(CompilationFilterPropertyHandle.ToSharedRef());
+		ChildBuilder.AddProperty(OperationTypePropertyHandle.ToSharedRef());
+	}
+}
+
+
+TSharedPtr<FString> FCustomizableObjectNodeTableCompilationFilterEditor::GenerateCompilationFilterColumnComboBoxOptions()
+{
+	TSharedPtr<FString> CurrentSelection;
+	FName CurrentOption;
+	ColumnPropertyHandle->GetValue(CurrentOption);
+
+	if (!Node.IsValid())
+	{
+		return CurrentSelection;
+	}
+
+	const UScriptStruct* TableStruct = Node->GetTableNodeStruct();
+	
+	CompilationFilterColumnOptionNames.Reset();
+
+	FText NothingSelectedText = LOCTEXT("NothingSelectedText", "- Nothing Selected -");
+	CompilationFilterColumnOptionNames.Add(MakeShareable(new FString(NothingSelectedText.ToString())));
+
+	if (!TableStruct)
+	{
+		return CurrentSelection;
+	}
+
+	// Iterating struct Options
+	for (TFieldIterator<FProperty> It(TableStruct); It; ++It)
+	{
+		FProperty* ColumnProperty = *It;
+
+		if (!ColumnProperty)
+		{
+			continue;
+		}
+
+		if (const FArrayProperty* ArrayProperty = CastField<FArrayProperty>(ColumnProperty))
+		{
+			ColumnProperty = ArrayProperty->Inner;
+		}
+
+		if (UCustomizableObjectNodeTable::SupportedFilterTypes.Contains(ColumnProperty->GetClass()))
+		{
+			TSharedRef<FString> Option = MakeShareable(new FString(ColumnProperty->GetDisplayNameText().ToString()));
+			CompilationFilterColumnOptionNames.Add(Option);
+
+			if (*Option == CurrentOption.ToString())
+			{
+				CurrentSelection = CompilationFilterColumnOptionNames.Last();
+			}
+		}
+	}
+
+	if (!CurrentSelection)
+	{
+		if (!CurrentOption.IsNone())
+		{
+			CompilationFilterColumnOptionNames.Add(MakeShareable(new FString(CurrentOption.ToString())));
+			CurrentSelection = CompilationFilterColumnOptionNames.Last();
+		}
+		else
+		{
+			CurrentSelection = CompilationFilterColumnOptionNames[0];
+		}
+	}
+
+	return CurrentSelection;
+}
+
+
+void FCustomizableObjectNodeTableCompilationFilterEditor::OnOpenCompilationFilterColumnComboBox()
+{
+	TSharedPtr<FString> CurrentSelection = GenerateCompilationFilterColumnComboBoxOptions();
+
+	if (CompilationFilterColumnComboBox.IsValid())
+	{
+		CompilationFilterColumnComboBox->ClearSelection();
+		CompilationFilterColumnComboBox->RefreshOptions();
+		CompilationFilterColumnComboBox->SetSelectedItem(CurrentSelection);
+	}
+}
+
+
+void FCustomizableObjectNodeTableCompilationFilterEditor::OnCompilationFilterColumnComboBoxSelectionChanged(TSharedPtr<FString> Selection, ESelectInfo::Type SelectInfo)
+{
+	FName CurrentOption;
+	ColumnPropertyHandle->GetValue(CurrentOption);
+
+	if (Selection && CurrentOption != FName(*Selection)
+		&& (SelectInfo == ESelectInfo::OnKeyPress || SelectInfo == ESelectInfo::OnMouseClick))
+	{
+		if (Selection == CompilationFilterColumnOptionNames[0])
+		{
+			ColumnPropertyHandle->SetValue(NAME_None);
+		}
+		else
+		{
+			ColumnPropertyHandle->SetValue(FName(*Selection));
+		}
+	}
+}
+
+
+FSlateColor FCustomizableObjectNodeTableCompilationFilterEditor::GetCompilationFilterColumnComboBoxTextColor() const
+{
+	FName CurrentOption;
+	ColumnPropertyHandle->GetValue(CurrentOption);
+
+	if (Node.IsValid() && Node->FindColumnProperty(CurrentOption) || CurrentOption.IsNone())
+	{
+		return FSlateColor::UseForeground();
+	}
+
+	// Table Struct null or does not contain the selected property anymore
+	return FSlateColor(FLinearColor(0.9f, 0.05f, 0.05f, 1.0f));
+}
+
+
+void FCustomizableObjectNodeTableCompilationFilterEditor::OnCompilationFilterColumnComboBoxSelectionReset()
+{
+	if (Node.IsValid())
+	{
+		ColumnPropertyHandle->SetValue(NAME_None);
+
+		if (CompilationFilterColumnComboBox.IsValid())
+		{
+			GenerateCompilationFilterColumnComboBoxOptions();
+			CompilationFilterColumnComboBox->SetSelectedItem(CompilationFilterColumnOptionNames[0]);
+			CompilationFilterColumnComboBox->RefreshOptions();
+		}
+	}
+}
+
+
+#undef LOCTEXT_NAMESPACE
diff --git a/Engine/Plugins/Mutable/Source/CustomizableObjectEditor/Private/MuCOE/Widgets/CustomizableObjectNodeTableCompilationFilterEditor.h b/Engine/Plugins/Mutable/Source/CustomizableObjectEditor/Private/MuCOE/Widgets/CustomizableObjectNodeTableCompilationFilterEditor.h
new file mode 100644
--- /dev/null
+++ b/Engine/Plugins/Mutable/Source/CustomizableObjectEditor/Private/MuCOE/Widgets/CustomizableObjectNodeTableCompilationFilterEditor.h
@@ -0,0 +1,51 @@
+// Copyright Epic Games, Inc. All Rights Reserved.
+
+#pragma once
+
+#include "IPropertyTypeCustomization.h"
+
+class IPropertyHandle;
+class STextComboBox;
+class UCustomizableObjectNodeTable;
+
+class FCustomizableObjectNodeTableCompilationFilterEditor : public IPropertyTypeCustomization
+{
+public:
+	static TSharedRef<IPropertyTypeCustomization> MakeInstance();
+
+	// IPropertyTypeCustomization interface
+	void CustomizeHeader(TSharedRef<IPropertyHandle> StructPropertyHandle, class FDetailWidgetRow& HeaderRow, IPropertyTypeCustomizationUtils& StructCustomizationUtils) override;
+	void CustomizeChildren(TSharedRef<IPropertyHandle> StructPropertyHandle, class IDetailChildrenBuilder& ChildBuilder, IPropertyTypeCustomizationUtils& StructCustomizationUtils) override;
+
+private:
+
+	// Pointer to the node represented in this details
+	TWeakObjectPtr<UCustomizableObjectNodeTable> Node;
+
+	TSharedPtr<IPropertyHandle> ColumnPropertyHandle;
+
+	// Pinter to the structure instance, useful to get the array index of the property
+	TSharedPtr<IPropertyHandle> StructPropertyHandlePtr;
+	
+	// Array with the name of the Version columns
+	TArray<TSharedPtr<FString>> CompilationFilterColumnOptionNames;
+
+	// ComboBox widget to select a VersionColumn from the NodeTable
+	TSharedPtr<STextComboBox> CompilationFilterColumnComboBox;
+
+	// Generates MutableMetadata columns combobox options
+	// Returns the current selected option or a null pointer
+	TSharedPtr<FString> GenerateCompilationFilterColumnComboBoxOptions();
+
+	// Callback to regenerate the combobox options
+	void OnOpenCompilationFilterColumnComboBox();
+
+	// OnComboBoxSelectionChanged Callback for Layout ComboBox
+	void OnCompilationFilterColumnComboBoxSelectionChanged(TSharedPtr<FString> Selection, ESelectInfo::Type SelectInfo);
+
+	// Sets the combo box selection color
+	FSlateColor GetCompilationFilterColumnComboBoxTextColor() const;
+
+	// OnComboBoxSelectionChanged Callback for Layout ComboBox
+	void OnCompilationFilterColumnComboBoxSelectionReset();
+};
