[Mutable] Fix deadlock when saving an asset + compiling a CO while asset registry was loading.

diff --git a/Engine/Plugins/Mutable/Source/CustomizableObjectEditor/Private/MuCOE/CustomizableObjectCompiler.cpp b/Engine/Plugins/Mutable/Source/CustomizableObjectEditor/Private/MuCOE/CustomizableObjectCompiler.cpp
--- a/Engine/Plugins/Mutable/Source/CustomizableObjectEditor/Private/MuCOE/CustomizableObjectCompiler.cpp
+++ b/Engine/Plugins/Mutable/Source/CustomizableObjectEditor/Private/MuCOE/CustomizableObjectCompiler.cpp
@@ -48,6 +48,19 @@ bool FCustomizableObjectCompiler::Tick(bool bBlocking)
 {
 	MUTABLE_CPUPROFILER_SCOPE(FCustomizableObjectCompiler::Tick);
 
+	if (bBlocking)
+	{
+		// Compilations require the asset registry to have finished.
+		// If blocking, the asset registry will not finish by itself and we will deadlock. Force it to finish synchronously.
+		const FAssetRegistryModule& AssetRegistryModule = FModuleManager::LoadModuleChecked<FAssetRegistryModule>("AssetRegistry");
+		IAssetRegistry& AssetRegistry = AssetRegistryModule.Get();
+
+		if (AssetRegistry.IsGathering())
+		{
+			AssetRegistry.SearchAllAssets(true);
+		}
+	}
+	
 	ProcessCompileTasks();
 	
 	bool bFinished = true;
