Fixed clang related preflight error for Incremental UnrealEditor Linux PresubmitEngine build.

diff --git a/Engine/Plugins/Mutable/Source/CustomizableObjectEditor/Private/MuCOE/Widgets/CustomizableObjectNodeTableCompilationFilterEditor.cpp b/Engine/Plugins/Mutable/Source/CustomizableObjectEditor/Private/MuCOE/Widgets/CustomizableObjectNodeTableCompilationFilterEditor.cpp
--- a/Engine/Plugins/Mutable/Source/CustomizableObjectEditor/Private/MuCOE/Widgets/CustomizableObjectNodeTableCompilationFilterEditor.cpp
+++ b/Engine/Plugins/Mutable/Source/CustomizableObjectEditor/Private/MuCOE/Widgets/CustomizableObjectNodeTableCompilationFilterEditor.cpp
@@ -182,7 +182,7 @@ FSlateColor FCustomizableObjectNodeTableCompilationFilterEditor::GetCompilationF
 	FName CurrentOption;
 	ColumnPropertyHandle->GetValue(CurrentOption);
 
-	if (Node.IsValid() && Node->FindColumnProperty(CurrentOption) || CurrentOption.IsNone())
+	if ((Node.IsValid() && Node->FindColumnProperty(CurrentOption)) || CurrentOption.IsNone())
 	{
 		return FSlateColor::UseForeground();
 	}
