[Mutable] Enabled mutable macros.
+ Removed the Cvar that used to enable macros.
+ Forced the creation of Static String pins
+ Enabled the creation of Macro instance and Static String nodes.

#rnx

diff --git a/Engine/Plugins/Mutable/Source/CustomizableObjectEditor/Private/MuCOE/CustomizableObjectMacroLibrary/CustomizableObjectMacroLibrary.cpp b/Engine/Plugins/Mutable/Source/CustomizableObjectEditor/Private/MuCOE/CustomizableObjectMacroLibrary/CustomizableObjectMacroLibrary.cpp
--- a/Engine/Plugins/Mutable/Source/CustomizableObjectEditor/Private/MuCOE/CustomizableObjectMacroLibrary/CustomizableObjectMacroLibrary.cpp
+++ b/Engine/Plugins/Mutable/Source/CustomizableObjectEditor/Private/MuCOE/CustomizableObjectMacroLibrary/CustomizableObjectMacroLibrary.cpp
@@ -8,12 +8,6 @@
 
 #define LOCTEXT_NAMESPACE "CustomizableObjectMacroLibrary"
 
-TAutoConsoleVariable<bool> CVarEnableMutableMacros(
-	TEXT("mutable.EnableMutableMacros"),
-	false,
-	TEXT("Enables or disables the new nodes and assets to be able to create and use Mutable Macros."),
-	ECVF_Scalability);
-
 
 UCustomizableObjectMacro* UCustomizableObjectMacroLibrary::AddMacro()
 {
diff --git a/Engine/Plugins/Mutable/Source/CustomizableObjectEditor/Private/MuCOE/CustomizableObjectMacroLibrary/CustomizableObjectMacroLibrary.h b/Engine/Plugins/Mutable/Source/CustomizableObjectEditor/Private/MuCOE/CustomizableObjectMacroLibrary/CustomizableObjectMacroLibrary.h
--- a/Engine/Plugins/Mutable/Source/CustomizableObjectEditor/Private/MuCOE/CustomizableObjectMacroLibrary/CustomizableObjectMacroLibrary.h
+++ b/Engine/Plugins/Mutable/Source/CustomizableObjectEditor/Private/MuCOE/CustomizableObjectMacroLibrary/CustomizableObjectMacroLibrary.h
@@ -6,7 +6,6 @@
 
 #include "CustomizableObjectMacroLibrary.generated.h"
 
-extern TAutoConsoleVariable<bool> CVarEnableMutableMacros;
 
 class UEdGraph;
 class UCustomizableObjectNodeTunnel;
diff --git a/Engine/Plugins/Mutable/Source/CustomizableObjectEditor/Private/MuCOE/CustomizableObjectMacroLibrary/CustomizableObjectMacroLibraryFactory.cpp b/Engine/Plugins/Mutable/Source/CustomizableObjectEditor/Private/MuCOE/CustomizableObjectMacroLibrary/CustomizableObjectMacroLibraryFactory.cpp
--- a/Engine/Plugins/Mutable/Source/CustomizableObjectEditor/Private/MuCOE/CustomizableObjectMacroLibrary/CustomizableObjectMacroLibraryFactory.cpp
+++ b/Engine/Plugins/Mutable/Source/CustomizableObjectEditor/Private/MuCOE/CustomizableObjectMacroLibrary/CustomizableObjectMacroLibraryFactory.cpp
@@ -22,9 +22,3 @@ UObject* UCustomizableObjectMacroLibraryFactory::FactoryCreateNew(UClass* Class,
 
 	return NewMacroLibrary;
 }
-
-
-bool UCustomizableObjectMacroLibraryFactory::CanCreateNew() const
-{
-	return CVarEnableMutableMacros.GetValueOnAnyThread();
-}
diff --git a/Engine/Plugins/Mutable/Source/CustomizableObjectEditor/Private/MuCOE/CustomizableObjectMacroLibrary/CustomizableObjectMacroLibraryFactory.h b/Engine/Plugins/Mutable/Source/CustomizableObjectEditor/Private/MuCOE/CustomizableObjectMacroLibrary/CustomizableObjectMacroLibraryFactory.h
--- a/Engine/Plugins/Mutable/Source/CustomizableObjectEditor/Private/MuCOE/CustomizableObjectMacroLibrary/CustomizableObjectMacroLibraryFactory.h
+++ b/Engine/Plugins/Mutable/Source/CustomizableObjectEditor/Private/MuCOE/CustomizableObjectMacroLibrary/CustomizableObjectMacroLibraryFactory.h
@@ -17,7 +17,6 @@ public:
 
 	// UFactory Interface
 	UObject* FactoryCreateNew(UClass* Class, UObject* InParent, FName Name, EObjectFlags Flags, UObject* Context, FFeedbackContext* Warn) override;
-	virtual bool CanCreateNew() const override;
 
 };
 
diff --git a/Engine/Plugins/Mutable/Source/CustomizableObjectEditor/Private/MuCOE/CustomizableObjectNodeObjectGroupDetails.cpp b/Engine/Plugins/Mutable/Source/CustomizableObjectEditor/Private/MuCOE/CustomizableObjectNodeObjectGroupDetails.cpp
--- a/Engine/Plugins/Mutable/Source/CustomizableObjectEditor/Private/MuCOE/CustomizableObjectNodeObjectGroupDetails.cpp
+++ b/Engine/Plugins/Mutable/Source/CustomizableObjectEditor/Private/MuCOE/CustomizableObjectNodeObjectGroupDetails.cpp
@@ -46,15 +46,7 @@ void FCustomizableObjectNodeObjectGroupDetails::CustomizeDetails(IDetailLayoutBu
 
 		// Forcing property order
 		{
-			if (CVarEnableMutableMacros.GetValueOnAnyThread())
-			{
-				DetailBuilder.HideProperty("GroupName");
-			}
-			else
-			{
-				GroupInfoCategory.AddProperty("GroupName");
-			}
-
+			DetailBuilder.HideProperty("GroupName");
 			GroupInfoCategory.AddProperty("GroupType");
 		}
 
diff --git a/Engine/Plugins/Mutable/Source/CustomizableObjectEditor/Private/MuCOE/EdGraphSchema_CustomizableObject.cpp b/Engine/Plugins/Mutable/Source/CustomizableObjectEditor/Private/MuCOE/EdGraphSchema_CustomizableObject.cpp
--- a/Engine/Plugins/Mutable/Source/CustomizableObjectEditor/Private/MuCOE/EdGraphSchema_CustomizableObject.cpp
+++ b/Engine/Plugins/Mutable/Source/CustomizableObjectEditor/Private/MuCOE/EdGraphSchema_CustomizableObject.cpp
@@ -444,21 +444,12 @@ void UEdGraphSchema_CustomizableObject::GetGraphContextActions(FGraphContextMenu
 			ContextMenuBuilder.CreateTemplateNode<UCustomizableObjectNodeMaterialVariation>(),
 			ContextMenuBuilder.CreateTemplateNode<UCustomizableObjectNodeMaterialSwitch>(),
 			ContextMenuBuilder.CreateTemplateNode<UCustomizableObjectNodeCopyMaterial>(),
-			ContextMenuBuilder.CreateTemplateNode<UCustomizableObjectNodeTable>()
+			ContextMenuBuilder.CreateTemplateNode<UCustomizableObjectNodeTable>(),
+			ContextMenuBuilder.CreateTemplateNode<UCustomizableObjectNodeMacroInstance>(),
+			ContextMenuBuilder.CreateTemplateNode<UCustomizableObjectNodeStaticString>()
 		};
 
 		AddNewNodeCategoryActionsFiltered(ObjectTemplateNodes, ContextMenuBuilder, TEXT("Object"), GeneralGrouping, Filter);
-
-		if (CVarEnableMutableMacros.GetValueOnAnyThread())
-		{
-			UCustomizableObjectNode* MacroTemplateNodes[]
-			{
-				ContextMenuBuilder.CreateTemplateNode<UCustomizableObjectNodeMacroInstance>(),
-				ContextMenuBuilder.CreateTemplateNode<UCustomizableObjectNodeStaticString>()
-			};
-
-			AddNewNodeCategoryActionsFiltered(MacroTemplateNodes, ContextMenuBuilder, TEXT("Object"), GeneralGrouping, Filter);
-		}
 	}
 
 	{
diff --git a/Engine/Plugins/Mutable/Source/CustomizableObjectEditor/Private/MuCOE/Nodes/CustomizableObjectNodeComponent.cpp b/Engine/Plugins/Mutable/Source/CustomizableObjectEditor/Private/MuCOE/Nodes/CustomizableObjectNodeComponent.cpp
--- a/Engine/Plugins/Mutable/Source/CustomizableObjectEditor/Private/MuCOE/Nodes/CustomizableObjectNodeComponent.cpp
+++ b/Engine/Plugins/Mutable/Source/CustomizableObjectEditor/Private/MuCOE/Nodes/CustomizableObjectNodeComponent.cpp
@@ -2,6 +2,7 @@
 
 #include "MuCOE/Nodes/CustomizableObjectNodeComponent.h"
 
+#include "MuCO/CustomizableObjectCustomVersion.h"
 #include "MuCOE/EdGraphSchema_CustomizableObject.h"
 #include "MuCOE/GraphTraversal.h"
 #include "MuCOE/Nodes/CustomizableObjectNodeMacroInstance.h"
@@ -11,11 +12,7 @@
 void UCustomizableObjectNodeComponent::AllocateDefaultPins(UCustomizableObjectNodeRemapPins* RemapPins)
 {
 	OutputPin = CustomCreatePin(EGPD_Output, UEdGraphSchema_CustomizableObject::PC_Component, FName(TEXT("Component")));
-
-	if (CVarEnableMutableMacros.GetValueOnAnyThread())
-	{
-		ComponentNamePin = CustomCreatePin(EGPD_Input, UEdGraphSchema_CustomizableObject::PC_String, FName("Name"));
-	}
+	ComponentNamePin = CustomCreatePin(EGPD_Input, UEdGraphSchema_CustomizableObject::PC_String, FName("Name"));
 }
 
 
@@ -31,6 +28,20 @@ bool UCustomizableObjectNodeComponent::IsAffectedByLOD() const
 }
 
 
+void UCustomizableObjectNodeComponent::BackwardsCompatibleFixup(int32 CustomizableObjectCustomVersion)
+{
+	Super::BackwardsCompatibleFixup(CustomizableObjectCustomVersion);
+
+	if (CustomizableObjectCustomVersion == FCustomizableObjectCustomVersion::EnableMutableMacros)
+	{
+		if (!ComponentNamePin.Get())
+		{
+			ComponentNamePin = CustomCreatePin(EGPD_Input, UEdGraphSchema_CustomizableObject::PC_String, FName("Name"));
+		}
+	}
+}
+
+
 bool UCustomizableObjectNodeComponent::GetCanRenameNode() const
 {
 	const UEdGraphPin* NamePin = GetComponentNamePin();
@@ -91,10 +102,5 @@ void UCustomizableObjectNodeComponent::SetComponentName(const FName& InComponent
 
 UEdGraphPin* UCustomizableObjectNodeComponent::GetComponentNamePin() const
 {
-	if (CVarEnableMutableMacros.GetValueOnAnyThread())
-	{
-		return ComponentNamePin.Get();
-	}
-
-	return nullptr;
+	return ComponentNamePin.Get();
 }
diff --git a/Engine/Plugins/Mutable/Source/CustomizableObjectEditor/Private/MuCOE/Nodes/CustomizableObjectNodeComponent.h b/Engine/Plugins/Mutable/Source/CustomizableObjectEditor/Private/MuCOE/Nodes/CustomizableObjectNodeComponent.h
--- a/Engine/Plugins/Mutable/Source/CustomizableObjectEditor/Private/MuCOE/Nodes/CustomizableObjectNodeComponent.h
+++ b/Engine/Plugins/Mutable/Source/CustomizableObjectEditor/Private/MuCOE/Nodes/CustomizableObjectNodeComponent.h
@@ -25,6 +25,7 @@ public:
 
 	// UCustomizableObjectNode interface
 	virtual bool IsAffectedByLOD() const override;
+	virtual void BackwardsCompatibleFixup(int32 CustomizableObjectCustomVersion) override;
 
 	//Own Interface
 	FName GetComponentName(TArray<const UCustomizableObjectNodeMacroInstance*>* MacroContext = nullptr) const;
diff --git a/Engine/Plugins/Mutable/Source/CustomizableObjectEditor/Private/MuCOE/Nodes/CustomizableObjectNodeComponentMeshAddTo.cpp b/Engine/Plugins/Mutable/Source/CustomizableObjectEditor/Private/MuCOE/Nodes/CustomizableObjectNodeComponentMeshAddTo.cpp
--- a/Engine/Plugins/Mutable/Source/CustomizableObjectEditor/Private/MuCOE/Nodes/CustomizableObjectNodeComponentMeshAddTo.cpp
+++ b/Engine/Plugins/Mutable/Source/CustomizableObjectEditor/Private/MuCOE/Nodes/CustomizableObjectNodeComponentMeshAddTo.cpp
@@ -25,11 +25,7 @@ void UCustomizableObjectNodeComponentMeshAddTo::AllocateDefaultPins(UCustomizabl
 		LODPins.Add(Pin);
 	}
 
-	if (CVarEnableMutableMacros.GetValueOnAnyThread())
-	{
-		ParentComponentNamePin = CustomCreatePin(EGPD_Input, UEdGraphSchema_CustomizableObject::PC_String, FName("Component Name"));
-	}
-
+	ParentComponentNamePin = CustomCreatePin(EGPD_Input, UEdGraphSchema_CustomizableObject::PC_String, FName("Component Name"));
 	OutputPin = CustomCreatePin(EGPD_Output, UEdGraphSchema_CustomizableObject::PC_Component, FName(TEXT("Component")));
 }
 
@@ -65,6 +61,20 @@ bool UCustomizableObjectNodeComponentMeshAddTo::IsSingleOutputNode() const
 }
 
 
+void UCustomizableObjectNodeComponentMeshAddTo::BackwardsCompatibleFixup(int32 CustomizableObjectCustomVersion)
+{
+	Super::BackwardsCompatibleFixup(CustomizableObjectCustomVersion);
+
+	if (CustomizableObjectCustomVersion == FCustomizableObjectCustomVersion::EnableMutableMacros)
+	{
+		if (!ParentComponentNamePin.Get())
+		{
+			ParentComponentNamePin = CustomCreatePin(EGPD_Input, UEdGraphSchema_CustomizableObject::PC_String, FName("Component Name"));
+		}
+	}
+}
+
+
 FText UCustomizableObjectNodeComponentMeshAddTo::GetNodeTitle(ENodeTitleType::Type TitleType) const
 {
 	const UEdGraphPin* ParentNamePin = GetParentComponentNamePin();
@@ -177,12 +187,7 @@ void UCustomizableObjectNodeComponentMeshAddTo::SetParentComponentName(const FNa
 
 UEdGraphPin* UCustomizableObjectNodeComponentMeshAddTo::GetParentComponentNamePin() const
 {
-	if (CVarEnableMutableMacros.GetValueOnAnyThread())
-	{
-		return ParentComponentNamePin.Get();
-	}
-
-	return nullptr;
+	return ParentComponentNamePin.Get();
 }
 
 #undef LOCTEXT_NAMESPACE
diff --git a/Engine/Plugins/Mutable/Source/CustomizableObjectEditor/Private/MuCOE/Nodes/CustomizableObjectNodeComponentMeshAddTo.h b/Engine/Plugins/Mutable/Source/CustomizableObjectEditor/Private/MuCOE/Nodes/CustomizableObjectNodeComponentMeshAddTo.h
--- a/Engine/Plugins/Mutable/Source/CustomizableObjectEditor/Private/MuCOE/Nodes/CustomizableObjectNodeComponentMeshAddTo.h
+++ b/Engine/Plugins/Mutable/Source/CustomizableObjectEditor/Private/MuCOE/Nodes/CustomizableObjectNodeComponentMeshAddTo.h
@@ -29,6 +29,7 @@ public:
 	//UCustomizableObjectNode interface
 	virtual bool IsAffectedByLOD() const override;
 	bool IsSingleOutputNode() const;
+	virtual void BackwardsCompatibleFixup(int32 CustomizableObjectCustomVersion) override;
 
 	// ComponentMesh Interface
 	virtual int32 GetNumLODs() override;
diff --git a/Engine/Plugins/Mutable/Source/CustomizableObjectEditor/Private/MuCOE/Nodes/CustomizableObjectNodeMaterial.cpp b/Engine/Plugins/Mutable/Source/CustomizableObjectEditor/Private/MuCOE/Nodes/CustomizableObjectNodeMaterial.cpp
--- a/Engine/Plugins/Mutable/Source/CustomizableObjectEditor/Private/MuCOE/Nodes/CustomizableObjectNodeMaterial.cpp
+++ b/Engine/Plugins/Mutable/Source/CustomizableObjectEditor/Private/MuCOE/Nodes/CustomizableObjectNodeMaterial.cpp
@@ -337,7 +337,6 @@ void UCustomizableObjectNodeMaterial::AllocateDefaultPins(UCustomizableObjectNod
 		TableMaterialPin->PinToolTip = "Pin for a Material from a Table Node";
 	}
 
-	if (CVarEnableMutableMacros.GetValueOnAnyThread())
 	{
 		const FString PinFriendlyName = "Enable Tags";
 		const FString PinName = PinFriendlyName + FString(TEXT("_Input_Pin"));
@@ -778,6 +777,19 @@ void UCustomizableObjectNodeMaterial::BackwardsCompatibleFixup(int32 Customizabl
 			FoundPin->PinName = FName(TargetPinName);
 		}
 	}
+
+	if (CustomizableObjectCustomVersion == FCustomizableObjectCustomVersion::EnableMutableMacros)
+	{
+		if (!EnableTagsPinRef.Get())
+		{
+			const FString PinFriendlyName = "Enable Tags";
+			const FString PinName = PinFriendlyName + FString(TEXT("_Input_Pin"));
+			UEdGraphPin* TagsPin = CustomCreatePin(EGPD_Input, UEdGraphSchema_CustomizableObject::PC_String, FName(PinName), true);
+			TagsPin->PinFriendlyName = FText::FromString(PinFriendlyName);
+			TagsPin->PinToolTip = "List of Tags that this node will Enable";
+			EnableTagsPinRef = TagsPin;
+		}
+	}
 }
 
 
@@ -890,12 +902,7 @@ UEdGraphPin* UCustomizableObjectNodeMaterial::GetMaterialAssetPin() const
 
 UEdGraphPin* UCustomizableObjectNodeMaterial::GetEnableTagsPin() const
 {
-	if (CVarEnableMutableMacros.GetValueOnAnyThread())
-	{
-		return EnableTagsPinRef.Get();
-	}
-
-	return nullptr;
+	return EnableTagsPinRef.Get();
 }
 
 
diff --git a/Engine/Plugins/Mutable/Source/CustomizableObjectEditor/Private/MuCOE/Nodes/CustomizableObjectNodeMeshMorph.cpp b/Engine/Plugins/Mutable/Source/CustomizableObjectEditor/Private/MuCOE/Nodes/CustomizableObjectNodeMeshMorph.cpp
--- a/Engine/Plugins/Mutable/Source/CustomizableObjectEditor/Private/MuCOE/Nodes/CustomizableObjectNodeMeshMorph.cpp
+++ b/Engine/Plugins/Mutable/Source/CustomizableObjectEditor/Private/MuCOE/Nodes/CustomizableObjectNodeMeshMorph.cpp
@@ -40,12 +40,8 @@ void UCustomizableObjectNodeMeshMorph::AllocateDefaultPins(UCustomizableObjectNo
 	UEdGraphPin* FactorPin = CustomCreatePin(EGPD_Input, Schema->PC_Float, FName(*PinName));
 	FactorPin->bDefaultValueIsIgnored = true;
 
-	// Create this pin only if we are in a Macro graph
-	if (CVarEnableMutableMacros.GetValueOnAnyThread())
-	{
-		PinName = TEXT("Morph Target Name");
-		MorphTargetNamePinRef = CustomCreatePin(EGPD_Input, Schema->PC_String, FName(*PinName));
-	}
+	PinName = TEXT("Morph Target Name");
+	MorphTargetNamePinRef = CustomCreatePin(EGPD_Input, Schema->PC_String, FName(*PinName));
 }
 
 
@@ -146,12 +142,7 @@ FText UCustomizableObjectNodeMeshMorph::GetTooltipText() const
 
 UEdGraphPin* UCustomizableObjectNodeMeshMorph::MorphTargetNamePin() const
 {
-	if (CVarEnableMutableMacros.GetValueOnAnyThread())
-	{
-		return MorphTargetNamePinRef.Get();
-	}
-
-	return nullptr;
+	return MorphTargetNamePinRef.Get();
 }
 
 
@@ -178,5 +169,20 @@ bool UCustomizableObjectNodeMeshMorph::CreateStaticStringPinWidget() const
 }
 
 
+void UCustomizableObjectNodeMeshMorph::BackwardsCompatibleFixup(int32 CustomizableObjectCustomVersion)
+{
+	Super::BackwardsCompatibleFixup(CustomizableObjectCustomVersion);
+
+	if (CustomizableObjectCustomVersion == FCustomizableObjectCustomVersion::EnableMutableMacros)
+	{
+		if (!MorphTargetNamePinRef.Get())
+		{
+			FString PinName = TEXT("Morph Target Name");
+			MorphTargetNamePinRef = CustomCreatePin(EGPD_Input, UEdGraphSchema_CustomizableObject::PC_String, FName(*PinName));
+		}
+	}
+}
+
+
 
 #undef LOCTEXT_NAMESPACE
diff --git a/Engine/Plugins/Mutable/Source/CustomizableObjectEditor/Private/MuCOE/Nodes/CustomizableObjectNodeMeshMorph.h b/Engine/Plugins/Mutable/Source/CustomizableObjectEditor/Private/MuCOE/Nodes/CustomizableObjectNodeMeshMorph.h
--- a/Engine/Plugins/Mutable/Source/CustomizableObjectEditor/Private/MuCOE/Nodes/CustomizableObjectNodeMeshMorph.h
+++ b/Engine/Plugins/Mutable/Source/CustomizableObjectEditor/Private/MuCOE/Nodes/CustomizableObjectNodeMeshMorph.h
@@ -32,6 +32,7 @@ public:
 	FString GetRefreshMessage() const override;
 	void PinConnectionListChanged(UEdGraphPin* Pin) override;
 	virtual bool CreateStaticStringPinWidget() const override;
+	virtual void BackwardsCompatibleFixup(int32 CustomizableObjectCustomVersion) override;
 
 
 	// Own interface
diff --git a/Engine/Plugins/Mutable/Source/CustomizableObjectEditor/Private/MuCOE/Nodes/CustomizableObjectNodeMeshParameter.cpp b/Engine/Plugins/Mutable/Source/CustomizableObjectEditor/Private/MuCOE/Nodes/CustomizableObjectNodeMeshParameter.cpp
--- a/Engine/Plugins/Mutable/Source/CustomizableObjectEditor/Private/MuCOE/Nodes/CustomizableObjectNodeMeshParameter.cpp
+++ b/Engine/Plugins/Mutable/Source/CustomizableObjectEditor/Private/MuCOE/Nodes/CustomizableObjectNodeMeshParameter.cpp
@@ -22,10 +22,7 @@ void UCustomizableObjectNodeMeshParameter::AllocateDefaultPins(UCustomizableObje
 		return;
 	}
 
-	if (CVarEnableMutableMacros.GetValueOnAnyThread())
-	{
-		NamePin = CustomCreatePin(EGPD_Input, UEdGraphSchema_CustomizableObject::PC_String, FName("Name"));
-	}
+	NamePin = CustomCreatePin(EGPD_Input, UEdGraphSchema_CustomizableObject::PC_String, FName("Name"));
 
 	const UEdGraphSchema_CustomizableObject* Schema = GetDefault<UEdGraphSchema_CustomizableObject>();
 
@@ -85,6 +82,20 @@ bool UCustomizableObjectNodeMeshParameter::IsExperimental() const
 }
 
 
+void UCustomizableObjectNodeMeshParameter::BackwardsCompatibleFixup(int32 CustomizableObjectCustomVersion)
+{
+	Super::BackwardsCompatibleFixup(CustomizableObjectCustomVersion);
+
+	if (CustomizableObjectCustomVersion == FCustomizableObjectCustomVersion::EnableMutableMacros)
+	{
+		if (!NamePin.Get())
+		{
+			NamePin = CustomCreatePin(EGPD_Input, UEdGraphSchema_CustomizableObject::PC_String, FName("Name"));
+		}
+	}
+}
+
+
 FText UCustomizableObjectNodeMeshParameter::GetNodeTitle(ENodeTitleType::Type TitleType) const
 {
 	const UEdGraphPin* ParamNamePin = NamePin.Get();
@@ -130,7 +141,7 @@ bool UCustomizableObjectNodeMeshParameter::GetCanRenameNode() const
 {
 	const UEdGraphPin* ParamNamePin = NamePin.Get();
 
-	if (CVarEnableMutableMacros.GetValueOnAnyThread() && ParamNamePin && ParamNamePin->LinkedTo.Num())
+	if (ParamNamePin && ParamNamePin->LinkedTo.Num())
 	{
 		return false;
 	}
diff --git a/Engine/Plugins/Mutable/Source/CustomizableObjectEditor/Private/MuCOE/Nodes/CustomizableObjectNodeMeshParameter.h b/Engine/Plugins/Mutable/Source/CustomizableObjectEditor/Private/MuCOE/Nodes/CustomizableObjectNodeMeshParameter.h
--- a/Engine/Plugins/Mutable/Source/CustomizableObjectEditor/Private/MuCOE/Nodes/CustomizableObjectNodeMeshParameter.h
+++ b/Engine/Plugins/Mutable/Source/CustomizableObjectEditor/Private/MuCOE/Nodes/CustomizableObjectNodeMeshParameter.h
@@ -76,6 +76,7 @@ public:
 	virtual void AllocateDefaultPins(UCustomizableObjectNodeRemapPins* RemapPins) override;
 	virtual bool IsAffectedByLOD() const override { return false; }
 	virtual bool IsExperimental() const override;
+	virtual void BackwardsCompatibleFixup(int32 CustomizableObjectCustomVersion) override;
 
 	// UCustomizableObjectNodeMesh interface
 	virtual UTexture2D* FindTextureForPin(const UEdGraphPin* Pin) const override;
diff --git a/Engine/Plugins/Mutable/Source/CustomizableObjectEditor/Private/MuCOE/Nodes/CustomizableObjectNodeModifierBase.cpp b/Engine/Plugins/Mutable/Source/CustomizableObjectEditor/Private/MuCOE/Nodes/CustomizableObjectNodeModifierBase.cpp
--- a/Engine/Plugins/Mutable/Source/CustomizableObjectEditor/Private/MuCOE/Nodes/CustomizableObjectNodeModifierBase.cpp
+++ b/Engine/Plugins/Mutable/Source/CustomizableObjectEditor/Private/MuCOE/Nodes/CustomizableObjectNodeModifierBase.cpp
@@ -27,12 +27,7 @@ UEdGraphPin* UCustomizableObjectNodeModifierBase::OutputPin() const
 
 UEdGraphPin* UCustomizableObjectNodeModifierBase::RequiredTagsPin() const
 {
-	if (CVarEnableMutableMacros.GetValueOnAnyThread())
-	{
-		return RequiredTagsPinRef.Get();
-	}
-
-	return nullptr;
+	return RequiredTagsPinRef.Get();
 }
 
 
@@ -261,6 +256,14 @@ void UCustomizableObjectNodeModifierBase::BackwardsCompatibleFixup(int32 Customi
 			}
 		}
 	}
+
+	if (CustomizableObjectCustomVersion == FCustomizableObjectCustomVersion::EnableMutableMacros)
+	{
+		if (!RequiredTagsPinRef.Get())
+		{
+			RequiredTagsPinRef = CustomCreatePin(EGPD_Input, UEdGraphSchema_CustomizableObject::PC_String, "Required Tags", true);
+		}
+	}
 }
 
 void UCustomizableObjectNodeModifierBase::PostBackwardsCompatibleFixup()
@@ -335,8 +338,5 @@ FString UCustomizableObjectNodeModifierBase::MakeNodeAutoTag( UEdGraphNode* Node
 
 void UCustomizableObjectNodeModifierBase::AllocateDefaultPins(UCustomizableObjectNodeRemapPins* RemapPins)
 {
-	if (CVarEnableMutableMacros.GetValueOnAnyThread())
-	{
-		RequiredTagsPinRef = CustomCreatePin(EGPD_Input, UEdGraphSchema_CustomizableObject::PC_String, "Required Tags", true);
-	}
+	RequiredTagsPinRef = CustomCreatePin(EGPD_Input, UEdGraphSchema_CustomizableObject::PC_String, "Required Tags", true);
 }
diff --git a/Engine/Plugins/Mutable/Source/CustomizableObjectEditor/Private/MuCOE/Nodes/CustomizableObjectNodeModifierExtendMeshSection.cpp b/Engine/Plugins/Mutable/Source/CustomizableObjectEditor/Private/MuCOE/Nodes/CustomizableObjectNodeModifierExtendMeshSection.cpp
--- a/Engine/Plugins/Mutable/Source/CustomizableObjectEditor/Private/MuCOE/Nodes/CustomizableObjectNodeModifierExtendMeshSection.cpp
+++ b/Engine/Plugins/Mutable/Source/CustomizableObjectEditor/Private/MuCOE/Nodes/CustomizableObjectNodeModifierExtendMeshSection.cpp
@@ -108,6 +108,16 @@ void UCustomizableObjectNodeModifierExtendMeshSection::BackwardsCompatibleFixup(
 
 		ReconstructNode();
 	}
+
+	if (CustomizableObjectCustomVersion == FCustomizableObjectCustomVersion::EnableMutableMacros)
+	{
+		if (!EnableTagsPinRef.Get())
+		{
+			UEdGraphPin* EnableTagsPin = CustomCreatePin(EGPD_Input, UEdGraphSchema_CustomizableObject::PC_String, FName("Enable Tags"), true);
+			EnableTagsPin->PinToolTip = "List of Tags that this node will Enable";
+			EnableTagsPinRef = EnableTagsPin;
+		}
+	}
 }
 
 
@@ -120,7 +130,7 @@ void UCustomizableObjectNodeModifierExtendMeshSection::AllocateDefaultPins(UCust
 	//Create Node Modifier Common Pins
 	Super::AllocateDefaultPins(RemapPins);
 
-	if (CVarEnableMutableMacros.GetValueOnAnyThread())
+	// Create enable tags pin
 	{
 		UEdGraphPin* EnableTagsPin = CustomCreatePin(EGPD_Input, Schema->PC_String, FName("Enable Tags"), true);
 		EnableTagsPin->PinToolTip = "List of Tags that this node will Enable";
@@ -167,12 +177,7 @@ UEdGraphPin* UCustomizableObjectNodeModifierExtendMeshSection::AddMeshPin() cons
 
 UEdGraphPin* UCustomizableObjectNodeModifierExtendMeshSection::GetEnableTagsPin() const
 {
-	if (CVarEnableMutableMacros.GetValueOnAnyThread())
-	{
-		return EnableTagsPinRef.Get();
-	}
-
-	return nullptr;
+	return EnableTagsPinRef.Get();
 }
 
 
diff --git a/Engine/Plugins/Mutable/Source/CustomizableObjectEditor/Private/MuCOE/Nodes/CustomizableObjectNodeModifierMorphMeshSection.cpp b/Engine/Plugins/Mutable/Source/CustomizableObjectEditor/Private/MuCOE/Nodes/CustomizableObjectNodeModifierMorphMeshSection.cpp
--- a/Engine/Plugins/Mutable/Source/CustomizableObjectEditor/Private/MuCOE/Nodes/CustomizableObjectNodeModifierMorphMeshSection.cpp
+++ b/Engine/Plugins/Mutable/Source/CustomizableObjectEditor/Private/MuCOE/Nodes/CustomizableObjectNodeModifierMorphMeshSection.cpp
@@ -28,11 +28,8 @@ void UCustomizableObjectNodeModifierMorphMeshSection::AllocateDefaultPins(UCusto
 	UEdGraphPin* FactorPin = CustomCreatePin(EGPD_Input, Schema->PC_Float, FName(*PinName));
 	FactorPin->bDefaultValueIsIgnored = true;
 
-	if (CVarEnableMutableMacros.GetValueOnAnyThread())
-	{
-		PinName = TEXT("Morph Target Name");
-		MorphTargetNamePinRef = CustomCreatePin(EGPD_Input, Schema->PC_String, FName(*PinName));
-	}
+	PinName = TEXT("Morph Target Name");
+	MorphTargetNamePinRef = CustomCreatePin(EGPD_Input, Schema->PC_String, FName(*PinName));
 
 	//Create Node Modifier Common Pins
 	Super::AllocateDefaultPins(RemapPins);
@@ -65,12 +62,7 @@ bool UCustomizableObjectNodeModifierMorphMeshSection::IsSingleOutputNode() const
 
 UEdGraphPin* UCustomizableObjectNodeModifierMorphMeshSection::MorphTargetNamePin() const
 {
-	if (CVarEnableMutableMacros.GetValueOnAnyThread())
-	{
-		return MorphTargetNamePinRef.Get();
-	}
-
-	return nullptr;
+	return MorphTargetNamePinRef.Get();
 }
 
 
@@ -80,4 +72,18 @@ bool UCustomizableObjectNodeModifierMorphMeshSection::CreateStaticStringPinWidge
 }
 
 
+void UCustomizableObjectNodeModifierMorphMeshSection::BackwardsCompatibleFixup(int32 CustomizableObjectCustomVersion)
+{
+	Super::BackwardsCompatibleFixup(CustomizableObjectCustomVersion);
+
+	if (CustomizableObjectCustomVersion == FCustomizableObjectCustomVersion::EnableMutableMacros)
+	{
+		if (!MorphTargetNamePinRef.Get())
+		{
+			MorphTargetNamePinRef = CustomCreatePin(EGPD_Input, UEdGraphSchema_CustomizableObject::PC_String, FName("Morph Target Name"));
+		}
+	}
+}
+
+
 #undef LOCTEXT_NAMESPACE
diff --git a/Engine/Plugins/Mutable/Source/CustomizableObjectEditor/Private/MuCOE/Nodes/CustomizableObjectNodeModifierMorphMeshSection.h b/Engine/Plugins/Mutable/Source/CustomizableObjectEditor/Private/MuCOE/Nodes/CustomizableObjectNodeModifierMorphMeshSection.h
--- a/Engine/Plugins/Mutable/Source/CustomizableObjectEditor/Private/MuCOE/Nodes/CustomizableObjectNodeModifierMorphMeshSection.h
+++ b/Engine/Plugins/Mutable/Source/CustomizableObjectEditor/Private/MuCOE/Nodes/CustomizableObjectNodeModifierMorphMeshSection.h
@@ -34,6 +34,7 @@ public:
 	virtual FString GetRefreshMessage() const override;
 	virtual bool IsSingleOutputNode() const override;
 	virtual bool CreateStaticStringPinWidget() const override;
+	virtual void BackwardsCompatibleFixup(int32 CustomizableObjectCustomVersion) override;
 
 	UEdGraphPin* FactorPin() const
 	{
diff --git a/Engine/Plugins/Mutable/Source/CustomizableObjectEditor/Private/MuCOE/Nodes/CustomizableObjectNodeObjectChild.cpp b/Engine/Plugins/Mutable/Source/CustomizableObjectEditor/Private/MuCOE/Nodes/CustomizableObjectNodeObjectChild.cpp
--- a/Engine/Plugins/Mutable/Source/CustomizableObjectEditor/Private/MuCOE/Nodes/CustomizableObjectNodeObjectChild.cpp
+++ b/Engine/Plugins/Mutable/Source/CustomizableObjectEditor/Private/MuCOE/Nodes/CustomizableObjectNodeObjectChild.cpp
@@ -2,6 +2,7 @@
 
 #include "MuCOE/Nodes/CustomizableObjectNodeObjectChild.h"
 
+#include "MuCO/CustomizableObjectCustomVersion.h"
 #include "MuCOE/CustomizableObjectMacroLibrary/CustomizableObjectMacroLibrary.h"
 #include "MuCOE/EdGraphSchema_CustomizableObject.h"
 #include "MuCOE/GraphTraversal.h"
@@ -69,7 +70,7 @@ bool UCustomizableObjectNodeObjectChild::GetCanRenameNode() const
 {
 	const UEdGraphPin* ChildNamePin = NamePin.Get();
 
-	if (CVarEnableMutableMacros.GetValueOnAnyThread() && ChildNamePin && ChildNamePin->LinkedTo.Num())
+	if (ChildNamePin && ChildNamePin->LinkedTo.Num())
 	{
 		return false;
 	}
@@ -90,9 +91,20 @@ void UCustomizableObjectNodeObjectChild::AllocateDefaultPins(UCustomizableObject
 {
 	Super::AllocateDefaultPins(RemapPins);
 
-	if (CVarEnableMutableMacros.GetValueOnAnyThread())
+	NamePin = CustomCreatePin(EGPD_Input, UEdGraphSchema_CustomizableObject::PC_String, "Name");
+}
+
+
+void UCustomizableObjectNodeObjectChild::BackwardsCompatibleFixup(int32 CustomizableObjectCustomVersion)
+{
+	Super::BackwardsCompatibleFixup(CustomizableObjectCustomVersion);
+
+	if (CustomizableObjectCustomVersion == FCustomizableObjectCustomVersion::EnableMutableMacros)
 	{
-		NamePin = CustomCreatePin(EGPD_Input, UEdGraphSchema_CustomizableObject::PC_String, "Name");
+		if (!NamePin.Get())
+		{
+			NamePin = CustomCreatePin(EGPD_Input, UEdGraphSchema_CustomizableObject::PC_String, "Name");
+		}
 	}
 }
 
@@ -105,7 +117,7 @@ bool UCustomizableObjectNodeObjectChild::IsNodeSupportedInMacros() const
 
 FString UCustomizableObjectNodeObjectChild::GetObjectName(TArray<const UCustomizableObjectNodeMacroInstance*>* MacroContext) const
 {
-	if (CVarEnableMutableMacros.GetValueOnAnyThread() && NamePin.Get())
+	if (NamePin.Get())
 	{
 		if (const UEdGraphPin* LinkedPin = FollowInputPin(*NamePin.Get()))
 		{
diff --git a/Engine/Plugins/Mutable/Source/CustomizableObjectEditor/Private/MuCOE/Nodes/CustomizableObjectNodeObjectChild.h b/Engine/Plugins/Mutable/Source/CustomizableObjectEditor/Private/MuCOE/Nodes/CustomizableObjectNodeObjectChild.h
--- a/Engine/Plugins/Mutable/Source/CustomizableObjectEditor/Private/MuCOE/Nodes/CustomizableObjectNodeObjectChild.h
+++ b/Engine/Plugins/Mutable/Source/CustomizableObjectEditor/Private/MuCOE/Nodes/CustomizableObjectNodeObjectChild.h
@@ -34,6 +34,7 @@ public:
 	// UCustomizableObjectNode interface
 	virtual bool IsNodeSupportedInMacros() const override;
 	virtual void AllocateDefaultPins(UCustomizableObjectNodeRemapPins* RemapPins) override;
+	virtual void BackwardsCompatibleFixup(int32 CustomizableObjectCustomVersion) override;
 
 	// UCustomizableObjectNodeObject interface
 	virtual FString GetObjectName(TArray<const UCustomizableObjectNodeMacroInstance*>* MacroContext = nullptr) const override;
diff --git a/Engine/Plugins/Mutable/Source/CustomizableObjectEditor/Private/MuCOE/Nodes/CustomizableObjectNodeObjectDetails.cpp b/Engine/Plugins/Mutable/Source/CustomizableObjectEditor/Private/MuCOE/Nodes/CustomizableObjectNodeObjectDetails.cpp
--- a/Engine/Plugins/Mutable/Source/CustomizableObjectEditor/Private/MuCOE/Nodes/CustomizableObjectNodeObjectDetails.cpp
+++ b/Engine/Plugins/Mutable/Source/CustomizableObjectEditor/Private/MuCOE/Nodes/CustomizableObjectNodeObjectDetails.cpp
@@ -253,11 +253,7 @@ void FCustomizableObjectNodeObjectDetails::CustomizeDetails( IDetailLayoutBuilde
 		else
 		{
 			DetailBuilder.HideProperty("NumMeshComponents");
-
-			if (CVarEnableMutableMacros.GetValueOnAnyThread())
-			{
-				DetailBuilder.HideProperty("ObjectName");
-			}
+			DetailBuilder.HideProperty("ObjectName");
 		}
 
 		
diff --git a/Engine/Plugins/Mutable/Source/CustomizableObjectEditor/Private/MuCOE/Nodes/CustomizableObjectNodeObjectGroup.cpp b/Engine/Plugins/Mutable/Source/CustomizableObjectEditor/Private/MuCOE/Nodes/CustomizableObjectNodeObjectGroup.cpp
--- a/Engine/Plugins/Mutable/Source/CustomizableObjectEditor/Private/MuCOE/Nodes/CustomizableObjectNodeObjectGroup.cpp
+++ b/Engine/Plugins/Mutable/Source/CustomizableObjectEditor/Private/MuCOE/Nodes/CustomizableObjectNodeObjectGroup.cpp
@@ -136,7 +136,7 @@ bool UCustomizableObjectNodeObjectGroup::GetCanRenameNode() const
 {
 	const UEdGraphPin* GroupNamePin = NamePin.Get();
 
-	if (CVarEnableMutableMacros.GetValueOnAnyThread() && GroupNamePin && GroupNamePin->LinkedTo.Num())
+	if (GroupNamePin && GroupNamePin->LinkedTo.Num())
 	{
 		return false;
 	}
@@ -166,10 +166,7 @@ void UCustomizableObjectNodeObjectGroup::AllocateDefaultPins(UCustomizableObject
 
 	CustomCreatePin(EGPD_Output, Schema->PC_Object, FName("Group"));
 
-	if (CVarEnableMutableMacros.GetValueOnAnyThread())
-	{
-		NamePin = CustomCreatePin(EGPD_Input, UEdGraphSchema_CustomizableObject::PC_String, FName("Name"));
-	}
+	NamePin = CustomCreatePin(EGPD_Input, UEdGraphSchema_CustomizableObject::PC_String, FName("Name"));
 }
 
 
@@ -179,6 +176,20 @@ bool UCustomizableObjectNodeObjectGroup::IsSingleOutputNode() const
 }
 
 
+void UCustomizableObjectNodeObjectGroup::BackwardsCompatibleFixup(int32 CustomizableObjectCustomVersion)
+{
+	Super::BackwardsCompatibleFixup(CustomizableObjectCustomVersion);
+
+	if (CustomizableObjectCustomVersion == FCustomizableObjectCustomVersion::EnableMutableMacros)
+	{
+		if (!NamePin.Get())
+		{
+			NamePin = CustomCreatePin(EGPD_Input, UEdGraphSchema_CustomizableObject::PC_String, FName("Name"));
+		}
+	}
+}
+
+
 UEdGraphPin* UCustomizableObjectNodeObjectGroup::ObjectsPin() const
 {
 	return FindPin(TEXT("Objects"));
@@ -199,7 +210,7 @@ UEdGraphPin* UCustomizableObjectNodeObjectGroup::GroupPin() const
 
 FString UCustomizableObjectNodeObjectGroup::GetGroupName(TArray<const UCustomizableObjectNodeMacroInstance*>* MacroContext) const
 {
-	if (CVarEnableMutableMacros.GetValueOnAnyThread() && NamePin.Get())
+	if (NamePin.Get())
 	{
 		if (const UEdGraphPin* LinkedPin = FollowInputPin(*NamePin.Get()))
 		{
diff --git a/Engine/Plugins/Mutable/Source/CustomizableObjectEditor/Private/MuCOE/Nodes/CustomizableObjectNodeObjectGroup.h b/Engine/Plugins/Mutable/Source/CustomizableObjectEditor/Private/MuCOE/Nodes/CustomizableObjectNodeObjectGroup.h
--- a/Engine/Plugins/Mutable/Source/CustomizableObjectEditor/Private/MuCOE/Nodes/CustomizableObjectNodeObjectGroup.h
+++ b/Engine/Plugins/Mutable/Source/CustomizableObjectEditor/Private/MuCOE/Nodes/CustomizableObjectNodeObjectGroup.h
@@ -56,6 +56,7 @@ public:
 	// UCustomizableObjectNode interface
 	virtual void AllocateDefaultPins(UCustomizableObjectNodeRemapPins* RemapPins) override;
 	virtual bool IsSingleOutputNode() const override;
+	virtual void BackwardsCompatibleFixup(int32 CustomizableObjectCustomVersion) override;
 
 	// Own interface
 	UEdGraphPin* ObjectsPin() const;
diff --git a/Engine/Plugins/Mutable/Source/CustomizableObjectEditor/Private/MuCOE/Nodes/CustomizableObjectNodeParameter.cpp b/Engine/Plugins/Mutable/Source/CustomizableObjectEditor/Private/MuCOE/Nodes/CustomizableObjectNodeParameter.cpp
--- a/Engine/Plugins/Mutable/Source/CustomizableObjectEditor/Private/MuCOE/Nodes/CustomizableObjectNodeParameter.cpp
+++ b/Engine/Plugins/Mutable/Source/CustomizableObjectEditor/Private/MuCOE/Nodes/CustomizableObjectNodeParameter.cpp
@@ -2,6 +2,7 @@
 
 #include "MuCOE/Nodes/CustomizableObjectNodeParameter.h"
 
+#include "MuCO/CustomizableObjectCustomVersion.h"
 #include "MuCOE/EdGraphSchema_CustomizableObject.h"
 #include "MuCOE/CustomizableObjectMacroLibrary/CustomizableObjectMacroLibrary.h"
 #include "MuCOE/GraphTraversal.h"
@@ -23,9 +24,20 @@ void UCustomizableObjectNodeParameter::AllocateDefaultPins(UCustomizableObjectNo
 	ValuePin->PinFriendlyName = PinFriendlyName;
 	ValuePin->bDefaultValueIsIgnored = true;
 
-	if (CVarEnableMutableMacros.GetValueOnAnyThread())
+	NamePin = CustomCreatePin(EGPD_Input, UEdGraphSchema_CustomizableObject::PC_String, FName("Name"));
+}
+
+
+void UCustomizableObjectNodeParameter::BackwardsCompatibleFixup(int32 CustomizableObjectCustomVersion)
+{
+	Super::BackwardsCompatibleFixup(CustomizableObjectCustomVersion);
+
+	if (CustomizableObjectCustomVersion == FCustomizableObjectCustomVersion::EnableMutableMacros)
 	{
-		NamePin = CustomCreatePin(EGPD_Input, UEdGraphSchema_CustomizableObject::PC_String, FName("Name"));
+		if (!NamePin.Get())
+		{
+			NamePin = CustomCreatePin(EGPD_Input, UEdGraphSchema_CustomizableObject::PC_String, FName("Name"));
+		}
 	}
 }
 
@@ -80,7 +92,7 @@ bool UCustomizableObjectNodeParameter::GetCanRenameNode() const
 {
 	const UEdGraphPin* ParamNamePin = NamePin.Get();
 
-	if (CVarEnableMutableMacros.GetValueOnAnyThread() && ParamNamePin && ParamNamePin->LinkedTo.Num())
+	if (ParamNamePin && ParamNamePin->LinkedTo.Num())
 	{
 		return false;
 	}
@@ -100,7 +112,7 @@ void UCustomizableObjectNodeParameter::PinConnectionListChanged(UEdGraphPin* Pin
 
 FString UCustomizableObjectNodeParameter::GetParameterName(TArray<const UCustomizableObjectNodeMacroInstance*>* MacroContext) const
 {
-	if (CVarEnableMutableMacros.GetValueOnAnyThread() && NamePin.Get())
+	if (NamePin.Get())
 	{
 		if (const UEdGraphPin* LinkedPin = FollowInputPin(*NamePin.Get()))
 		{
diff --git a/Engine/Plugins/Mutable/Source/CustomizableObjectEditor/Private/MuCOE/Nodes/CustomizableObjectNodeParameter.h b/Engine/Plugins/Mutable/Source/CustomizableObjectEditor/Private/MuCOE/Nodes/CustomizableObjectNodeParameter.h
--- a/Engine/Plugins/Mutable/Source/CustomizableObjectEditor/Private/MuCOE/Nodes/CustomizableObjectNodeParameter.h
+++ b/Engine/Plugins/Mutable/Source/CustomizableObjectEditor/Private/MuCOE/Nodes/CustomizableObjectNodeParameter.h
@@ -42,6 +42,7 @@ public:
 	// UCustomizableObjectNode interface
 	virtual void AllocateDefaultPins(UCustomizableObjectNodeRemapPins* RemapPins) override;
 	virtual bool IsAffectedByLOD() const override { return false; }
+	virtual void BackwardsCompatibleFixup(int32 CustomizableObjectCustomVersion) override;
 
 	// Own interface
 	/** Return the pin category of this node. */
diff --git a/Engine/Plugins/Mutable/Source/CustomizableObjectEditor/Private/MuCOE/Nodes/CustomizableObjectNodeParameterDetails.cpp b/Engine/Plugins/Mutable/Source/CustomizableObjectEditor/Private/MuCOE/Nodes/CustomizableObjectNodeParameterDetails.cpp
--- a/Engine/Plugins/Mutable/Source/CustomizableObjectEditor/Private/MuCOE/Nodes/CustomizableObjectNodeParameterDetails.cpp
+++ b/Engine/Plugins/Mutable/Source/CustomizableObjectEditor/Private/MuCOE/Nodes/CustomizableObjectNodeParameterDetails.cpp
@@ -15,8 +15,5 @@ void FCustomizableObjectNodeParameterDetails::CustomizeDetails(IDetailLayoutBuil
 {
 	FCustomizableObjectNodeDetails::CustomizeDetails(DetailBuilder);
 	
-	if (CVarEnableMutableMacros.GetValueOnAnyThread())
-	{
-		DetailBuilder.HideProperty("ParameterName");
-	}
+	DetailBuilder.HideProperty("ParameterName");
 }
\ No newline at end of file
diff --git a/Engine/Plugins/Mutable/Source/CustomizableObjectEditor/Private/MuCOE/Nodes/CustomizableObjectNodeVariation.cpp b/Engine/Plugins/Mutable/Source/CustomizableObjectEditor/Private/MuCOE/Nodes/CustomizableObjectNodeVariation.cpp
--- a/Engine/Plugins/Mutable/Source/CustomizableObjectEditor/Private/MuCOE/Nodes/CustomizableObjectNodeVariation.cpp
+++ b/Engine/Plugins/Mutable/Source/CustomizableObjectEditor/Private/MuCOE/Nodes/CustomizableObjectNodeVariation.cpp
@@ -37,11 +37,7 @@ void UCustomizableObjectNodeVariation::AllocateDefaultPins(UCustomizableObjectNo
 	}
 	
 	VariationsPins.SetNum(VariationsData.Num());
-
-	if (CVarEnableMutableMacros.GetValueOnAnyThread())
-	{
-		VariationTagPins.SetNum(VariationsData.Num());
-	}
+	VariationTagPins.SetNum(VariationsData.Num());
 
 	for (int32 VariationIndex = VariationsData.Num() - 1; VariationIndex >= 0; --VariationIndex)
 	{
@@ -53,14 +49,11 @@ void UCustomizableObjectNodeVariation::AllocateDefaultPins(UCustomizableObjectNo
 		
 		VariationsPins[VariationIndex] = VariationPin;
 
-		if (CVarEnableMutableMacros.GetValueOnAnyThread())
-		{
-			FText VariationTagSufix = LOCTEXT("TagSufix", " - Tag");
-			const FName TagPinName = FName(PinName.ToString() + VariationTagSufix.ToString());
-			UEdGraphPin* VariationTagPin = CustomCreatePin(EGPD_Input, Schema->PC_String, TagPinName);
-			
-			VariationTagPins[VariationIndex] = VariationTagPin;
-		}
+		FText VariationTagSufix = LOCTEXT("TagSufix", " - Tag");
+		const FName TagPinName = FName(PinName.ToString() + VariationTagSufix.ToString());
+		UEdGraphPin* VariationTagPin = CustomCreatePin(EGPD_Input, Schema->PC_String, TagPinName);
+
+		VariationTagPins[VariationIndex] = VariationTagPin;
 	}
 
 	CustomCreatePin(EGPD_Input, Category, FName(TEXT("Default")), bIsInputPinArray);
@@ -126,7 +119,7 @@ UEdGraphPin* UCustomizableObjectNodeVariation::VariationPin(int32 Index) const
 
 UEdGraphPin* UCustomizableObjectNodeVariation::VariationTagPin(int32 Index) const
 {
-	if (VariationTagPins.IsValidIndex(Index) && CVarEnableMutableMacros.GetValueOnAnyThread())
+	if (VariationTagPins.IsValidIndex(Index))
 	{
 		return VariationTagPins[Index].Get();
 	}
@@ -163,6 +156,23 @@ void UCustomizableObjectNodeVariation::BackwardsCompatibleFixup(int32 Customizab
 		ReconstructNode();
 	}
 
+	if (CustomizableObjectCustomVersion == FCustomizableObjectCustomVersion::EnableMutableMacros)
+	{
+		if (VariationTagPins.IsEmpty())
+		{
+			VariationTagPins.SetNum(VariationsData.Num());
+
+			for (int32 VariationIndex = VariationsData.Num() - 1; VariationIndex >= 0; --VariationIndex)
+			{
+				const FName PinName = FName(FString::Printf(TEXT("Variation %d"), VariationIndex));
+				FText VariationTagSufix = LOCTEXT("TagSufix", " - Tag");
+				const FName TagPinName = FName(PinName.ToString() + VariationTagSufix.ToString());
+				UEdGraphPin* VariationTagPin = CustomCreatePin(EGPD_Input, UEdGraphSchema_CustomizableObject::PC_String, TagPinName);
+				VariationTagPins[VariationIndex] = VariationTagPin;
+			}
+		}
+	}
+
 	// TODO(Max UE-268615): Remove the tag name from the variation pins since now there is a pin with a textbox for it.
 }
 
