[Mutable] Fix missing mesh sections in higher LODs when mesh streaming is enabled.

#rnx

diff --git a/Engine/Plugins/Mutable/Source/CustomizableObject/Private/MuCO/CustomizableObjectSkeletalMesh.cpp b/Engine/Plugins/Mutable/Source/CustomizableObject/Private/MuCO/CustomizableObjectSkeletalMesh.cpp
--- a/Engine/Plugins/Mutable/Source/CustomizableObject/Private/MuCO/CustomizableObjectSkeletalMesh.cpp
+++ b/Engine/Plugins/Mutable/Source/CustomizableObject/Private/MuCO/CustomizableObjectSkeletalMesh.cpp
@@ -42,9 +42,11 @@ void UCustomizableObjectSkeletalMesh::InitMutableStreamingData(const TSharedRef<
 		MeshIDs[LODIndex] = LOD.MeshID;
 
 		SurfaceIDs[LODIndex].SetNum(LOD.SurfaceCount);
+		check(LOD.Mesh && LOD.Mesh->GetSurfaceCount() == LOD.SurfaceCount);
+
 		for (int32 SurfaceIndex = 0; SurfaceIndex < LOD.SurfaceCount; ++SurfaceIndex)
 		{
-			SurfaceIDs[LODIndex][SurfaceIndex] = InOperationData->InstanceUpdateData.Surfaces[LOD.FirstSurface + SurfaceIndex].SurfaceId;
+			SurfaceIDs[LODIndex][SurfaceIndex] = LOD.Mesh->GetSurfaceId(SurfaceIndex);
 		}
 	}
 
