[mutable] Fixed transform controls not updating the mesh in the COI viewport.
- Replaced all ObjectPtr with StrongObjectPtr in CustomizableObjectEditorViewportClient (cleaner than using FGCObject)
- Added delegates called when the transform values of the "Clip Mesh with Mesh" and "Transform Mesh in Mesh (Experimental)" get updated.
- Made the changes in the transform properties of the mentioned nodes trigger the update of the transform of the "ClipMesh" shown in the preview window.

#rnx

diff --git a/Engine/Plugins/Mutable/Source/CustomizableObjectEditor/Private/MuCOE/CustomizableObjectEditor.cpp b/Engine/Plugins/Mutable/Source/CustomizableObjectEditor/Private/MuCOE/CustomizableObjectEditor.cpp
--- a/Engine/Plugins/Mutable/Source/CustomizableObjectEditor/Private/MuCOE/CustomizableObjectEditor.cpp
+++ b/Engine/Plugins/Mutable/Source/CustomizableObjectEditor/Private/MuCOE/CustomizableObjectEditor.cpp
@@ -1440,14 +1440,14 @@ void FCustomizableObjectEditor::OnSelectedGraphNodesChanged(const FGraphPanelSel
 		}
 		else if (UCustomizableObjectNodeModifierClipWithMesh* NodeModifierClipWithMesh = Cast<UCustomizableObjectNodeModifierClipWithMesh>(Objects[0]))
 		{
-			if (const UEdGraphPin* Pin = NodeModifierClipWithMesh->ClipMeshPin())
+			if (const UEdGraphPin* Pin = NodeModifierClipWithMesh->GetClipMeshPin())
 			{
 				ShowGizmoClipMesh(*NodeModifierClipWithMesh, &NodeModifierClipWithMesh->Transform, *Pin);
 			}
 		}
 		else if (UCustomizableObjectNodeModifierTransformInMesh* NodeModifierTransformInMesh = Cast<UCustomizableObjectNodeModifierTransformInMesh>(Objects[0]))
 		{
-			if (const UEdGraphPin* Pin = NodeModifierTransformInMesh->BoundingMeshPin())
+			if (const UEdGraphPin* Pin = NodeModifierTransformInMesh->GetBoundingMeshPin())
 			{
 				ShowGizmoClipMesh(*NodeModifierTransformInMesh, &NodeModifierTransformInMesh->BoundingMeshTransform, *Pin);
 			}
diff --git a/Engine/Plugins/Mutable/Source/CustomizableObjectEditor/Private/MuCOE/CustomizableObjectEditorViewportClient.cpp b/Engine/Plugins/Mutable/Source/CustomizableObjectEditor/Private/MuCOE/CustomizableObjectEditorViewportClient.cpp
--- a/Engine/Plugins/Mutable/Source/CustomizableObjectEditor/Private/MuCOE/CustomizableObjectEditorViewportClient.cpp
+++ b/Engine/Plugins/Mutable/Source/CustomizableObjectEditor/Private/MuCOE/CustomizableObjectEditorViewportClient.cpp
@@ -5,9 +5,7 @@
 #include "AnimPreviewInstance.h"
 #include "MuCOE/CustomizableObjectInstanceBakingUtils.h"
 #include "Animation/DebugSkelMeshComponent.h"
-#include "Animation/PoseAsset.h"
 #include "Animation/AnimSingleNodeInstance.h"
-#include "Animation/Skeleton.h"
 #include "AssetViewerSettings.h"
 #include "CanvasTypes.h"
 #include "Components/SphereReflectionCaptureComponent.h"
@@ -20,13 +18,9 @@
 #include "Editor/UnrealEdTypes.h"
 #include "EditorModeManager.h"
 #include "GameFramework/WorldSettings.h"
-#include "FileHelpers.h"
 #include "IContentBrowserSingleton.h"
 #include "InputKeyEventArgs.h"
-#include "ObjectTools.h"
 #include "MaterialDomain.h"
-#include "Materials/MaterialInstanceDynamic.h"
-#include "Materials/MaterialInstanceConstant.h"
 #include "Misc/MessageDialog.h"
 #include "Misc/PackageName.h"
 #include "Misc/ConfigCacheIni.h"
@@ -35,22 +29,19 @@
 #include "MuCO/CustomizableObjectSystem.h"
 #include "MuCO/CustomizableSkeletalMeshActor.h"
 #include "MuCO/CustomizableSkeletalComponent.h"
-#include "MuCO/CustomizableSkeletalComponentPrivate.h"
 #include "MuCOE/CustomizableObjectPreviewScene.h"
 #include "MuCOE/ICustomizableObjectInstanceEditor.h"
 #include "MuCOE/Nodes/CustomizableObjectNodeModifierClipMorph.h"
 #include "MuCOE/Nodes/CustomizableObjectNodeModifierClipWithMesh.h"
+#include "MuCOE/Nodes/CustomizableObjectNodeModifierTransformInMesh.h"
 #include "MuCOE/Nodes/CustomizableObjectNodeProjectorConstant.h"
 #include "MuCOE/Nodes/CustomizableObjectNodeProjectorParameter.h"
 #include "MuCOE/UnrealEditorPortabilityHelpers.h"
-#include "MuT/UnrealPixelFormatOverride.h"
 #include "Preferences/PersonaOptions.h"
 #include "SceneView.h"
-#include "ScopedTransaction.h"
 #include "SkeletalDebugRendering.h"
 #include "UnrealWidget.h"
 #include "MuCO/CustomizableObjectMipDataProvider.h"
-#include "MuCOE/CustomizableObjectEditorLogger.h"
 #include "MuCOE/UnrealBakeHelpers.h"
 #include "MuCO/CustomizableObjectPrivate.h"
 #include "MuCO/LoadUtils.h"
@@ -143,7 +134,7 @@ FCustomizableObjectEditorViewportClient::FCustomizableObjectEditorViewportClient
 		PluginConfig->GetString(TEXT("EditorDefaults"), TEXT("ClipMorphMaterialName"), ClipMorphMaterialName);
 		if (!ClipMorphMaterialName.IsEmpty())
 		{
-			ClipMorphMaterial = MutablePrivate::LoadObject<UMaterial>(nullptr, *ClipMorphMaterialName);
+			ClipMorphMaterial = TStrongObjectPtr<UMaterial>(MutablePrivate::LoadObject<UMaterial>(nullptr, *ClipMorphMaterialName));
 			ensure(ClipMorphMaterial);
 		}
 
@@ -152,19 +143,19 @@ FCustomizableObjectEditorViewportClient::FCustomizableObjectEditorViewportClient
 		PluginConfig->GetString(TEXT("EditorDefaults"), TEXT("ClipMeshMaterialName"), ClipMeshMaterialName);
 		if (!ClipMeshMaterialName.IsEmpty())
 		{
-			ClipMeshMaterial = MutablePrivate::LoadObject<UMaterial>(nullptr, *ClipMeshMaterialName);
+			ClipMeshMaterial = TStrongObjectPtr<UMaterial>(MutablePrivate::LoadObject<UMaterial>(nullptr, *ClipMeshMaterialName));
 			ensure(ClipMeshMaterial);
 		}
 	}
 
 	if (!ClipMorphMaterial)
 	{
-		ClipMorphMaterial = UMaterial::GetDefaultMaterial(MD_Surface);
+		ClipMorphMaterial = TStrongObjectPtr<UMaterial>(UMaterial::GetDefaultMaterial(MD_Surface));
 	}
 
 	if (!ClipMeshMaterial)
 	{
-		ClipMeshMaterial = UMaterial::GetDefaultMaterial(MD_Surface);
+		ClipMeshMaterial = TStrongObjectPtr<UMaterial>(UMaterial::GetDefaultMaterial(MD_Surface));
 	}
 	
 	const float FOVMin = 5.f;
@@ -195,7 +186,7 @@ FCustomizableObjectEditorViewportClient::FCustomizableObjectEditorViewportClient
 	FAdvancedPreviewScene* PreviewSceneCasted = static_cast<FAdvancedPreviewScene*>(PreviewScene);
 	PreviewSceneCasted->SetProfileIndex(ProfileIndex);
 
-	TransparentPlaneMaterialXY = (UMaterial*)StaticLoadObject(UMaterial::StaticClass(), NULL, TEXT("/Engine/EditorMaterials/WidgetVertexColorMaterial.WidgetVertexColorMaterial"), NULL, LOAD_None, NULL);
+	TransparentPlaneMaterialXY = TStrongObjectPtr<UMaterial>((UMaterial*)StaticLoadObject(UMaterial::StaticClass(), NULL, TEXT("/Engine/EditorMaterials/WidgetVertexColorMaterial.WidgetVertexColorMaterial"), NULL, LOAD_None, NULL));
 }
 
 
@@ -633,34 +624,66 @@ void FCustomizableObjectEditorViewportClient::HideGizmoClipMorph()
 }
 
 
-void FCustomizableObjectEditorViewportClient::ShowGizmoClipMesh(UCustomizableObjectNode& InClipMeshNode, FTransform* InClipMeshTransform, UObject& ClipMesh, int32 LODIndex, int32 SectionIndex, int32 MaterialSlotIndex)
+void FCustomizableObjectEditorViewportClient::ShowGizmoClipMesh(UCustomizableObjectNode& InClipMeshNode, FTransform* InClipMeshTransform, UObject& InClipMesh, int32 LODIndex, int32 SectionIndex, int32 MaterialSlotIndex)
 {
-	HideGizmoClipMesh();
+	HideGizmoClipMesh(); 
 
 	SetWidgetType(EWidgetType::ClipMesh);
 
-	ClipMeshNode = &InClipMeshNode;
+	ClipMesh = TStrongObjectPtr<UObject>(&InClipMesh);
+	ClipMeshNode = TStrongObjectPtr<UObject>(&InClipMeshNode);
 	ClipMeshTransform = InClipMeshTransform;
 
-	if (UStaticMesh* StaticMesh = Cast<UStaticMesh>(&ClipMesh))
+	check(!TransformExternallyChangedDelegateHandle.IsValid())
+	if (UCustomizableObjectNodeModifierClipWithMesh* NodeModifierClipWithMesh = Cast<UCustomizableObjectNodeModifierClipWithMesh>(ClipMeshNode.Get()))
+	{
+		TransformExternallyChangedDelegateHandle = NodeModifierClipWithMesh->TransformChangedDelegate.AddSP(this, &FCustomizableObjectEditorViewportClient::UpdateGizmoClipMeshTransform);
+	}
+	else if (UCustomizableObjectNodeModifierTransformInMesh* NodeModifierTransformInMesh = Cast<UCustomizableObjectNodeModifierTransformInMesh>(ClipMeshNode.Get()))
+	{
+		TransformExternallyChangedDelegateHandle = NodeModifierTransformInMesh->TransformChangedDelegate.AddSP(this, &FCustomizableObjectEditorViewportClient::UpdateGizmoClipMeshTransform);
+	}
+	else
+	{
+		unimplemented();
+	}
+
+	if (UStaticMesh* StaticMesh = Cast<UStaticMesh>(&InClipMesh))
 	{
 		ClipMeshStaticMeshComp->SetStaticMesh(StaticMesh);
 		ClipMeshStaticMeshComp->SetVisibility(true);
 		ClipMeshStaticMeshComp->SetWorldTransform(*ClipMeshTransform);
 		ClipMeshStaticMeshComp->EmptyOverrideMaterials();
-		ClipMeshStaticMeshComp->SetMaterial(MaterialSlotIndex, ClipMeshMaterial);
+		ClipMeshStaticMeshComp->SetMaterial(MaterialSlotIndex, ClipMeshMaterial.Get());
 		ClipMeshStaticMeshComp->SetSectionPreview(SectionIndex);
 	}
-	else if (USkeletalMesh* SkeletalMesh = Cast<USkeletalMesh>(&ClipMesh))
+	else if (USkeletalMesh* SkeletalMesh = Cast<USkeletalMesh>(&InClipMesh))
 	{
 		ClipMeshSkeletalMeshComp->SetSkeletalMesh(SkeletalMesh);
 		ClipMeshSkeletalMeshComp->SetVisibility(true);
 		ClipMeshSkeletalMeshComp->SetWorldTransform(*ClipMeshTransform);
 		ClipMeshSkeletalMeshComp->SetForcedLOD(LODIndex);
 		ClipMeshSkeletalMeshComp->EmptyOverrideMaterials();
-		ClipMeshSkeletalMeshComp->SetMaterial(MaterialSlotIndex, ClipMeshMaterial);
+		ClipMeshSkeletalMeshComp->SetMaterial(MaterialSlotIndex, ClipMeshMaterial.Get());
 		ClipMeshSkeletalMeshComp->SetSectionPreview(SectionIndex);
 	}
+	else
+	{
+		unimplemented();
+	}
+}
+
+
+void FCustomizableObjectEditorViewportClient::UpdateGizmoClipMeshTransform(const FTransform& InTransform)
+{
+	if (ClipMesh->IsA(UStaticMesh::StaticClass()))
+	{
+		ClipMeshStaticMeshComp->SetWorldTransform(InTransform);
+	}
+	else if (ClipMesh->IsA(USkeletalMesh::StaticClass()))
+	{
+		ClipMeshSkeletalMeshComp->SetWorldTransform(InTransform);
+	}
 }
 
 
@@ -671,6 +694,25 @@ void FCustomizableObjectEditorViewportClient::HideGizmoClipMesh()
 		ClipMeshStaticMeshComp->SetVisibility(false);
 		ClipMeshSkeletalMeshComp->SetVisibility(false);
 
+		// Unbound the control that allows the user to move the widget from the details view
+		if (ensure(TransformExternallyChangedDelegateHandle.IsValid()))
+		{
+			if (UCustomizableObjectNodeModifierClipWithMesh* NodeModifierClipWithMesh = Cast<UCustomizableObjectNodeModifierClipWithMesh>(ClipMeshNode.Get()))
+			{
+				NodeModifierClipWithMesh->TransformChangedDelegate.Remove(TransformExternallyChangedDelegateHandle);
+			}
+			else if (UCustomizableObjectNodeModifierTransformInMesh* NodeModifierTransformInMesh = Cast<UCustomizableObjectNodeModifierTransformInMesh>(ClipMeshNode.Get()))
+			{
+				NodeModifierTransformInMesh->TransformChangedDelegate.Remove(TransformExternallyChangedDelegateHandle);
+			}
+			else
+			{
+				unimplemented();
+			}
+
+			TransformExternallyChangedDelegateHandle.Reset();
+		}
+
 		SetWidgetType(EWidgetType::Hidden);
 	}
 }
@@ -1564,7 +1606,7 @@ void FCustomizableObjectEditorViewportClient::BakeInstance()
 		return;
 	}
 	
-	BakeTempInstance = Instance->Clone();
+	BakeTempInstance = TStrongObjectPtr<UCustomizableObjectInstance>(Instance->Clone());
 
 	// Call the instance update async method
 	FInstanceUpdateNativeDelegate UpdateDelegate;
@@ -1710,19 +1752,6 @@ void FCustomizableObjectEditorViewportClient::DrawShadowedString(FCanvas* Canvas
 }
 
 
-void FCustomizableObjectEditorViewportClient::AddReferencedObjects(FReferenceCollector& Collector)
-{
-	Collector.AddReferencedObject(ClipMorphMaterial);
-	Collector.AddReferencedObject(ClipMeshMaterial);
-	Collector.AddReferencedObject(TransparentPlaneMaterialXY);
-
-	if (BakeTempInstance)
-	{
-		Collector.AddReferencedObject(BakeTempInstance);
-	}
-}
-
-
 void FCustomizableObjectEditorViewportClient::SetAdvancedShowFlagsForScene(const bool bAdvancedShowFlags)
 {
 	if (bAdvancedShowFlags)
diff --git a/Engine/Plugins/Mutable/Source/CustomizableObjectEditor/Private/MuCOE/CustomizableObjectEditorViewportClient.h b/Engine/Plugins/Mutable/Source/CustomizableObjectEditor/Private/MuCOE/CustomizableObjectEditorViewportClient.h
--- a/Engine/Plugins/Mutable/Source/CustomizableObjectEditor/Private/MuCOE/CustomizableObjectEditorViewportClient.h
+++ b/Engine/Plugins/Mutable/Source/CustomizableObjectEditor/Private/MuCOE/CustomizableObjectEditorViewportClient.h
@@ -189,8 +189,13 @@ public:
 	void HideGizmoClipMorph();
 	
 	/** Do not call directly. Use ICustomizableObjectEditor functions instead. */
-	void ShowGizmoClipMesh(UCustomizableObjectNode& ClipMeshNode, FTransform* ClipMeshTransform, UObject& ClipMesh, int32 LODIndex, int32 SectionIndex, int32 MaterialSlotIndex);
+	void ShowGizmoClipMesh(UCustomizableObjectNode& ClipMeshNode, FTransform* ClipMeshTransform, UObject& InClipMesh, int32 LODIndex, int32 SectionIndex, int32 MaterialSlotIndex);
 
+private:
+	/** Do not call directly. Method invoked each time the transform value of the clipping mesh is modified from the node. */
+	void UpdateGizmoClipMeshTransform(const FTransform& InTransform);
+
+public:
 	/** Do not call directly. Use ICustomizableObjectEditor functions instead. */
 	void HideGizmoClipMesh();
 	
@@ -232,14 +237,6 @@ public:
 	/** Helper method to draw shadowed string on viewport */
 	void DrawShadowedString(FCanvas* Canvas, float StartX, float StartY, const FLinearColor& Color, float TextScale, FString String);
 
-	// FSerializableObject interface
-	virtual void AddReferencedObjects(FReferenceCollector& Collector) override;
-	virtual FString GetReferencerName() const override
-	{
-		return TEXT("FCustomizableObjectEditorViewportClient");
-	}
-	// End of FSerializableObject interface
-
 	/** Show pero LOD geometric information of the instance */
 	void ShowInstanceGeometryInformation(FCanvas* InCanvas);
 
@@ -349,7 +346,7 @@ private:
 	int32 UVDrawUVIndex = 0;
 	
 	UCustomizableObjectNodeModifierClipMorph* ClipMorphNode;
-	TObjectPtr<UMaterial> ClipMorphMaterial;
+	TStrongObjectPtr<UMaterial> ClipMorphMaterial;
 	bool bClipMorphLocalStartOffset;
 	FVector ClipMorphOrigin;
 	FVector ClipMorphOffset;
@@ -359,9 +356,11 @@ private:
 	FVector ClipMorphYAxis;
 	float MorphLength;
 	FBoxSphereBounds MorphBounds;
-	TObjectPtr<UObject> ClipMeshNode;
+	TStrongObjectPtr<UObject> ClipMeshNode;
+	TStrongObjectPtr<UObject> ClipMesh;
+	FDelegateHandle TransformExternallyChangedDelegateHandle;
 	FTransform* ClipMeshTransform = nullptr;
-	TObjectPtr<UMaterial> ClipMeshMaterial;
+	TStrongObjectPtr<UMaterial> ClipMeshMaterial;
 	UStaticMeshComponent* ClipMeshStaticMeshComp;
 	USkeletalMeshComponent* ClipMeshSkeletalMeshComp;
 	float Radius1;
@@ -387,7 +386,7 @@ private:
 	bool StateChangeShowGeometryDataFlag;
 	
 	/** Material for cylinder arc solid render */
-	TObjectPtr<UMaterialInterface> TransparentPlaneMaterialXY;
+	TStrongObjectPtr<UMaterialInterface> TransparentPlaneMaterialXY;
 
 	/** bool to return the Camera mode to Orbital when changing the Camera view to Perspective */
 	bool bSetOrbitalOnPerspectiveMode;
@@ -400,7 +399,7 @@ private:
 	
 	// Temp Instance used in the bake process if a new instance is needed because mutable texture streaming is enabled so the viewport 
 	// instance does not have the high quality mips in the texture's platform data
-	TObjectPtr<UCustomizableObjectInstance> BakeTempInstance = nullptr;
+	TStrongObjectPtr<UCustomizableObjectInstance> BakeTempInstance = nullptr;
 
 	// Cache System configuration cached before performing the mandatory instance update for baking so we can restore it after the bake operation
 	bool bIsProgressiveMipStreamingEnabled = false;
diff --git a/Engine/Plugins/Mutable/Source/CustomizableObjectEditor/Private/MuCOE/GenerateMutableSource/GenerateMutableSourceModifier.cpp b/Engine/Plugins/Mutable/Source/CustomizableObjectEditor/Private/MuCOE/GenerateMutableSource/GenerateMutableSourceModifier.cpp
--- a/Engine/Plugins/Mutable/Source/CustomizableObjectEditor/Private/MuCOE/GenerateMutableSource/GenerateMutableSourceModifier.cpp
+++ b/Engine/Plugins/Mutable/Source/CustomizableObjectEditor/Private/MuCOE/GenerateMutableSource/GenerateMutableSourceModifier.cpp
@@ -173,7 +173,7 @@ mu::Ptr<mu::NodeModifier> GenerateMutableSourceModifier(const UEdGraphPin * Pin,
 
 		ClipNode->FaceCullStrategy = TypedNodeClipMesh->FaceCullStrategy;
 
-		if (const UEdGraphPin* ConnectedPin = FollowInputPin(*TypedNodeClipMesh->ClipMeshPin()))
+		if (const UEdGraphPin* ConnectedPin = FollowInputPin(*TypedNodeClipMesh->GetClipMeshPin()))
 		{
 			mu::NodeMeshPtr ClipMesh = GenerateMutableSourceMesh(ConnectedPin, GenerationContext, FMutableSourceMeshData(), false, true);
 
@@ -623,13 +623,13 @@ mu::Ptr<mu::NodeModifier> GenerateMutableSourceModifier(const UEdGraphPin * Pin,
 		mu::Ptr<mu::NodeModifierMeshTransformInMesh> TransformNode = new mu::NodeModifierMeshTransformInMesh();
 		Result = TransformNode;
 
-		if (const UEdGraphPin* ConnectedPin = FollowInputPin(*TypedNodeTransformMesh->TransformPin()))
+		if (const UEdGraphPin* ConnectedPin = FollowInputPin(*TypedNodeTransformMesh->GetTransformPin()))
 		{
 			TransformNode->MatrixNode = GenerateMutableSourceTransform(ConnectedPin, GenerationContext);
 		}
 
 		// If no bounding mesh is provided, we transform the entire mesh.
-		if (const UEdGraphPin* ConnectedPin = FollowInputPin(*TypedNodeTransformMesh->BoundingMeshPin()))
+		if (const UEdGraphPin* ConnectedPin = FollowInputPin(*TypedNodeTransformMesh->GetBoundingMeshPin()))
 		{
 			mu::Ptr<mu::NodeMesh> BoundingMesh = GenerateMutableSourceMesh(ConnectedPin, GenerationContext, FMutableSourceMeshData(), false, true);
 
diff --git a/Engine/Plugins/Mutable/Source/CustomizableObjectEditor/Private/MuCOE/Nodes/CustomizableObjectNodeMeshClipMorph.cpp b/Engine/Plugins/Mutable/Source/CustomizableObjectEditor/Private/MuCOE/Nodes/CustomizableObjectNodeMeshClipMorph.cpp
--- a/Engine/Plugins/Mutable/Source/CustomizableObjectEditor/Private/MuCOE/Nodes/CustomizableObjectNodeMeshClipMorph.cpp
+++ b/Engine/Plugins/Mutable/Source/CustomizableObjectEditor/Private/MuCOE/Nodes/CustomizableObjectNodeMeshClipMorph.cpp
@@ -87,7 +87,7 @@ void UCustomizableObjectNodeModifierClipMorph::ChangeStartOffsetTransform()
 }
 
 
-UEdGraphPin* UCustomizableObjectNodeModifierClipMorph::OutputPin() const
+UEdGraphPin* UCustomizableObjectNodeModifierClipMorph::GetOutputPin() const
 {
 	return FindPin(TEXT("Modifier"));
 }
@@ -157,7 +157,7 @@ void UCustomizableObjectNodeModifierClipMorph::PinConnectionListChanged(UEdGraph
 {
 	Super::PinConnectionListChanged(Pin);
 
-	if (Pin == OutputPin())
+	if (Pin == GetOutputPin())
 	{
 		TSharedPtr<FCustomizableObjectGraphEditorToolkit> Editor = GetGraphEditor();
 
diff --git a/Engine/Plugins/Mutable/Source/CustomizableObjectEditor/Private/MuCOE/Nodes/CustomizableObjectNodeMeshClipWithMesh.cpp b/Engine/Plugins/Mutable/Source/CustomizableObjectEditor/Private/MuCOE/Nodes/CustomizableObjectNodeMeshClipWithMesh.cpp
--- a/Engine/Plugins/Mutable/Source/CustomizableObjectEditor/Private/MuCOE/Nodes/CustomizableObjectNodeMeshClipWithMesh.cpp
+++ b/Engine/Plugins/Mutable/Source/CustomizableObjectEditor/Private/MuCOE/Nodes/CustomizableObjectNodeMeshClipWithMesh.cpp
@@ -49,16 +49,14 @@ void UCustomizableObjectNodeModifierClipWithMesh::BackwardsCompatibleFixup(int32
 		for ( const FGuid& MeshSectionNodeID : ArrayMaterialNodeToClipWithID_DEPRECATED )
 		{
 			// Look for the parent material and set it as the modifier reference material 
-			UCustomizableObjectNode* ParentNode = GetCustomizableObjectExternalNode<UCustomizableObjectNode>(CustomizableObjectToClipWith_DEPRECATED.Get(), MeshSectionNodeID);
 
-			if (ParentNode)
+			if (UCustomizableObjectNode* ParentNode = GetCustomizableObjectExternalNode<UCustomizableObjectNode>(CustomizableObjectToClipWith_DEPRECATED.Get(), MeshSectionNodeID))
 			{
 				// Add an autogenerated tag to the legacy parent, so that it can be referenced in this modifier.
 				FString NewLegacyTag = MakeNodeAutoTag(ParentNode);
 				RequiredTags.Add(NewLegacyTag);
 
-				TArray<FString>* NodeEnableTags = ParentNode->GetEnableTagsArray();
-				if (NodeEnableTags)
+				if (TArray<FString>* NodeEnableTags = ParentNode->GetEnableTagsArray())
 				{
 					NodeEnableTags->AddUnique(NewLegacyTag);
 				}
@@ -92,19 +90,38 @@ void UCustomizableObjectNodeModifierClipWithMesh::BackwardsCompatibleFixup(int32
 }
 
 
+void UCustomizableObjectNodeModifierClipWithMesh::PostEditChangeProperty(struct FPropertyChangedEvent& PropertyChangedEvent)
+{
+	Super::PostEditChangeProperty(PropertyChangedEvent);
+
+	if (PropertyChangedEvent.Property == nullptr)
+	{
+		return;
+	}
+	
+	if (PropertyChangedEvent.Property->GetFName() == GET_MEMBER_NAME_CHECKED(UCustomizableObjectNodeModifierClipWithMesh, Transform))
+	{
+		if (TransformChangedDelegate.IsBound())
+		{
+			TransformChangedDelegate.Broadcast(Transform);
+		}
+	}
+}
+
+
 FText UCustomizableObjectNodeModifierClipWithMesh::GetNodeTitle(ENodeTitleType::Type TitleType) const
 {
 	return LOCTEXT("Clip_Mesh_With_Mesh", "Clip Mesh With Mesh");
 }
 
 
-UEdGraphPin* UCustomizableObjectNodeModifierClipWithMesh::OutputPin() const
+UEdGraphPin* UCustomizableObjectNodeModifierClipWithMesh::GetOutputPin() const
 {
 	return FindPin(TEXT("Modifier"));
 }
 
 
-UEdGraphPin* UCustomizableObjectNodeModifierClipWithMesh::ClipMeshPin() const
+UEdGraphPin* UCustomizableObjectNodeModifierClipWithMesh::GetClipMeshPin() const
 {
 	return FindPin(TEXT("Clip Mesh"));
 }
diff --git a/Engine/Plugins/Mutable/Source/CustomizableObjectEditor/Private/MuCOE/Nodes/CustomizableObjectNodeModifierBase.cpp b/Engine/Plugins/Mutable/Source/CustomizableObjectEditor/Private/MuCOE/Nodes/CustomizableObjectNodeModifierBase.cpp
--- a/Engine/Plugins/Mutable/Source/CustomizableObjectEditor/Private/MuCOE/Nodes/CustomizableObjectNodeModifierBase.cpp
+++ b/Engine/Plugins/Mutable/Source/CustomizableObjectEditor/Private/MuCOE/Nodes/CustomizableObjectNodeModifierBase.cpp
@@ -19,7 +19,7 @@ FLinearColor UCustomizableObjectNodeModifierBase::GetNodeTitleColor() const
 }
 
 
-UEdGraphPin* UCustomizableObjectNodeModifierBase::OutputPin() const
+UEdGraphPin* UCustomizableObjectNodeModifierBase::GetOutputPin() const
 {
 	return FindPin(TEXT("Modifier"));
 }
diff --git a/Engine/Plugins/Mutable/Source/CustomizableObjectEditor/Private/MuCOE/Nodes/CustomizableObjectNodeModifierBase.h b/Engine/Plugins/Mutable/Source/CustomizableObjectEditor/Private/MuCOE/Nodes/CustomizableObjectNodeModifierBase.h
--- a/Engine/Plugins/Mutable/Source/CustomizableObjectEditor/Private/MuCOE/Nodes/CustomizableObjectNodeModifierBase.h
+++ b/Engine/Plugins/Mutable/Source/CustomizableObjectEditor/Private/MuCOE/Nodes/CustomizableObjectNodeModifierBase.h
@@ -55,7 +55,7 @@ public:
 	virtual void AllocateDefaultPins(UCustomizableObjectNodeRemapPins* RemapPins) override;
 
 	// Own interface
-	virtual UEdGraphPin* OutputPin() const;
+	virtual UEdGraphPin* GetOutputPin() const;
 	virtual UEdGraphPin* RequiredTagsPin() const;
 	TArray<FString> GetNodeRequiredTags(TArray<const UCustomizableObjectNodeMacroInstance*>* MacroContext = nullptr) const;
 
diff --git a/Engine/Plugins/Mutable/Source/CustomizableObjectEditor/Private/MuCOE/Nodes/CustomizableObjectNodeModifierClipMorph.h b/Engine/Plugins/Mutable/Source/CustomizableObjectEditor/Private/MuCOE/Nodes/CustomizableObjectNodeModifierClipMorph.h
--- a/Engine/Plugins/Mutable/Source/CustomizableObjectEditor/Private/MuCOE/Nodes/CustomizableObjectNodeModifierClipMorph.h
+++ b/Engine/Plugins/Mutable/Source/CustomizableObjectEditor/Private/MuCOE/Nodes/CustomizableObjectNodeModifierClipMorph.h
@@ -94,7 +94,7 @@ public:
 	virtual void BackwardsCompatibleFixup(int32 CustomizableObjectCustomVersion) override;
 
 	// Own interface
-	UEdGraphPin* OutputPin() const;
+	UEdGraphPin* GetOutputPin() const;
 
 	FVector GetOriginWithOffset() const;
 	void FindLocalAxes(FVector& Right, FVector& Up, FVector& Forward) const;
diff --git a/Engine/Plugins/Mutable/Source/CustomizableObjectEditor/Private/MuCOE/Nodes/CustomizableObjectNodeModifierClipWithMesh.h b/Engine/Plugins/Mutable/Source/CustomizableObjectEditor/Private/MuCOE/Nodes/CustomizableObjectNodeModifierClipWithMesh.h
--- a/Engine/Plugins/Mutable/Source/CustomizableObjectEditor/Private/MuCOE/Nodes/CustomizableObjectNodeModifierClipWithMesh.h
+++ b/Engine/Plugins/Mutable/Source/CustomizableObjectEditor/Private/MuCOE/Nodes/CustomizableObjectNodeModifierClipWithMesh.h
@@ -3,7 +3,6 @@
 #pragma once
 
 #include "MuCOE/Nodes/CustomizableObjectNodeModifierBase.h"
-#include "MuT/NodeModifier.h"
 
 #include "CustomizableObjectNodeModifierClipWithMesh.generated.h"
 
@@ -15,6 +14,8 @@ class UEdGraphPin;
 class UObject;
 struct FGuid;
 
+DECLARE_MULTICAST_DELEGATE_OneParam(OnTransformChanged, const FTransform&);
+
 UCLASS()
 class CUSTOMIZABLEOBJECTEDITOR_API UCustomizableObjectNodeModifierClipWithMesh : public UCustomizableObjectNodeModifierBase
 {
@@ -41,6 +42,11 @@ public:
 	UPROPERTY(EditAnywhere, Category = ClipMesh)
 	EFaceCullStrategy FaceCullStrategy = EFaceCullStrategy::AllVerticesCulled;
 
+	/**
+	 * Delegate invoked each time the BoundingMeshTransform value gets modified
+	 */
+	OnTransformChanged TransformChangedDelegate;
+	
 public:
 
 	UCustomizableObjectNodeModifierClipWithMesh();
@@ -52,10 +58,12 @@ public:
 	// UCustomizableObjectNode interface
 	virtual void AllocateDefaultPins(UCustomizableObjectNodeRemapPins* RemapPins) override;
 	virtual void BackwardsCompatibleFixup(int32 CustomizableObjectCustomVersion) override;
+	virtual void PostEditChangeProperty( struct FPropertyChangedEvent& PropertyChangedEvent) override;
 
 	// Own interface
-	UEdGraphPin* OutputPin() const;
+	virtual UEdGraphPin* GetOutputPin() const override;
+
+	UEdGraphPin* GetClipMeshPin() const;
 
-	UEdGraphPin* ClipMeshPin() const;
 };
 
diff --git a/Engine/Plugins/Mutable/Source/CustomizableObjectEditor/Private/MuCOE/Nodes/CustomizableObjectNodeModifierMorphMeshSectionDetails.cpp b/Engine/Plugins/Mutable/Source/CustomizableObjectEditor/Private/MuCOE/Nodes/CustomizableObjectNodeModifierMorphMeshSectionDetails.cpp
--- a/Engine/Plugins/Mutable/Source/CustomizableObjectEditor/Private/MuCOE/Nodes/CustomizableObjectNodeModifierMorphMeshSectionDetails.cpp
+++ b/Engine/Plugins/Mutable/Source/CustomizableObjectEditor/Private/MuCOE/Nodes/CustomizableObjectNodeModifierMorphMeshSectionDetails.cpp
@@ -288,9 +288,9 @@ void FCustomizableObjectNodeModifierMorphMeshSectionDetails::AddMorphsFromNode(
 
 	else if (UCustomizableObjectNodeModifierExtendMeshSection* ExtendNode = Cast<UCustomizableObjectNodeModifierExtendMeshSection>(Candidate))
 	{
-		if (ExtendNode->OutputPin())
+		if (ExtendNode->GetOutputPin())
 		{
-			const UEdGraphPin* SourceMeshPin = FindMeshBaseSource(*ExtendNode->OutputPin(), false);
+			const UEdGraphPin* SourceMeshPin = FindMeshBaseSource(*ExtendNode->GetOutputPin(), false);
 			if (SourceMeshPin)
 			{
 				UCustomizableObjectNodeSkeletalMesh* SkeletalNode = Cast<UCustomizableObjectNodeSkeletalMesh>(SourceMeshPin->GetOwningNode());
diff --git a/Engine/Plugins/Mutable/Source/CustomizableObjectEditor/Private/MuCOE/Nodes/CustomizableObjectNodeModifierRemoveMesh.cpp b/Engine/Plugins/Mutable/Source/CustomizableObjectEditor/Private/MuCOE/Nodes/CustomizableObjectNodeModifierRemoveMesh.cpp
--- a/Engine/Plugins/Mutable/Source/CustomizableObjectEditor/Private/MuCOE/Nodes/CustomizableObjectNodeModifierRemoveMesh.cpp
+++ b/Engine/Plugins/Mutable/Source/CustomizableObjectEditor/Private/MuCOE/Nodes/CustomizableObjectNodeModifierRemoveMesh.cpp
@@ -35,7 +35,7 @@ void UCustomizableObjectNodeModifierRemoveMesh::PinConnectionListChanged(UEdGrap
 {
 	Super::PinConnectionListChanged(Pin);
 
-	if (Pin == OutputPin())
+	if (Pin == GetOutputPin())
 	{
 		TSharedPtr<FCustomizableObjectGraphEditorToolkit> Editor = GetGraphEditor();
 
diff --git a/Engine/Plugins/Mutable/Source/CustomizableObjectEditor/Private/MuCOE/Nodes/CustomizableObjectNodeModifierTransformInMesh.cpp b/Engine/Plugins/Mutable/Source/CustomizableObjectEditor/Private/MuCOE/Nodes/CustomizableObjectNodeModifierTransformInMesh.cpp
--- a/Engine/Plugins/Mutable/Source/CustomizableObjectEditor/Private/MuCOE/Nodes/CustomizableObjectNodeModifierTransformInMesh.cpp
+++ b/Engine/Plugins/Mutable/Source/CustomizableObjectEditor/Private/MuCOE/Nodes/CustomizableObjectNodeModifierTransformInMesh.cpp
@@ -38,17 +38,37 @@ void UCustomizableObjectNodeModifierTransformInMesh::AllocateDefaultPins(UCustom
 	Super::AllocateDefaultPins(RemapPins);
 }
 
-UEdGraphPin* UCustomizableObjectNodeModifierTransformInMesh::OutputPin() const
+
+void UCustomizableObjectNodeModifierTransformInMesh::PostEditChangeProperty(struct FPropertyChangedEvent& PropertyChangedEvent)
+{
+	Super::PostEditChangeProperty(PropertyChangedEvent);
+
+	if (PropertyChangedEvent.Property == nullptr)
+	{
+		return;
+	}
+	
+	if (PropertyChangedEvent.Property->GetFName() == GET_MEMBER_NAME_CHECKED(UCustomizableObjectNodeModifierTransformInMesh, BoundingMeshTransform))
+	{
+		if (TransformChangedDelegate.IsBound())
+		{
+			TransformChangedDelegate.Broadcast(BoundingMeshTransform);
+		}
+	}
+}
+
+
+UEdGraphPin* UCustomizableObjectNodeModifierTransformInMesh::GetOutputPin() const
 {
 	return FindPin(OutputPinName);
 }
 
-UEdGraphPin* UCustomizableObjectNodeModifierTransformInMesh::BoundingMeshPin() const
+UEdGraphPin* UCustomizableObjectNodeModifierTransformInMesh::GetBoundingMeshPin() const
 {
 	return FindPin(BoundingMeshPinName);
 }
 
-UEdGraphPin* UCustomizableObjectNodeModifierTransformInMesh::TransformPin() const
+UEdGraphPin* UCustomizableObjectNodeModifierTransformInMesh::GetTransformPin() const
 {
 	return FindPin(TransformPinName);
 }
diff --git a/Engine/Plugins/Mutable/Source/CustomizableObjectEditor/Private/MuCOE/Nodes/CustomizableObjectNodeModifierTransformInMesh.h b/Engine/Plugins/Mutable/Source/CustomizableObjectEditor/Private/MuCOE/Nodes/CustomizableObjectNodeModifierTransformInMesh.h
--- a/Engine/Plugins/Mutable/Source/CustomizableObjectEditor/Private/MuCOE/Nodes/CustomizableObjectNodeModifierTransformInMesh.h
+++ b/Engine/Plugins/Mutable/Source/CustomizableObjectEditor/Private/MuCOE/Nodes/CustomizableObjectNodeModifierTransformInMesh.h
@@ -13,6 +13,8 @@ class UCustomizableObject;
 class UCustomizableObjectNodeRemapPins;
 class UEdGraphPin;
 
+DECLARE_MULTICAST_DELEGATE_OneParam(OnTransformChanged, const FTransform&);
+
 UCLASS()
 class CUSTOMIZABLEOBJECTEDITOR_API UCustomizableObjectNodeModifierTransformInMesh : public UCustomizableObjectNodeModifierBase
 {
@@ -23,6 +25,11 @@ public:
 	UPROPERTY(EditAnywhere, Category = BoundingMesh)
 	FTransform BoundingMeshTransform = FTransform::Identity;
 
+	/**
+	 * Delegate invoked each time the BoundingMeshTransform value gets modified
+	 */
+	OnTransformChanged TransformChangedDelegate;
+	
 public:
 	// EdGraphNode interface
 	virtual FText GetNodeTitle(ENodeTitleType::Type TitleType) const override;
@@ -31,14 +38,15 @@ public:
 	// UCustomizableObjectNode interface
 	virtual void AllocateDefaultPins(UCustomizableObjectNodeRemapPins* RemapPins) override;
 	virtual bool IsExperimental() const override { return true; }
+	virtual void PostEditChangeProperty( struct FPropertyChangedEvent& PropertyChangedEvent) override;
 
 	// UCustomizableObjectNodeModifierBase interface
-	UEdGraphPin* OutputPin() const override;
+	virtual UEdGraphPin* GetOutputPin() const override;
 
 	// Own interface
-	UEdGraphPin* BoundingMeshPin() const;
-	UEdGraphPin* TransformPin() const;
-	
+	UEdGraphPin* GetBoundingMeshPin() const;
+	UEdGraphPin* GetTransformPin() const;
+
 private:
 	static const TCHAR* OutputPinName;
 	static const TCHAR* BoundingMeshPinName;
