[Mutable] Fixed infinite loop in macro graphs.

#rnx

diff --git a/Engine/Plugins/Mutable/Source/CustomizableObjectEditor/Private/MuCOE/GraphTraversal.cpp b/Engine/Plugins/Mutable/Source/CustomizableObjectEditor/Private/MuCOE/GraphTraversal.cpp
--- a/Engine/Plugins/Mutable/Source/CustomizableObjectEditor/Private/MuCOE/GraphTraversal.cpp
+++ b/Engine/Plugins/Mutable/Source/CustomizableObjectEditor/Private/MuCOE/GraphTraversal.cpp
@@ -562,6 +562,10 @@ const UEdGraphPin* GraphTraversal::FindIOPinSourceThroughMacroContext(const UEdG
 				{
 					MacroContext = &NewMacroContext;
 				}
+				else if (MacroContext->Contains(NodeMacro))
+				{
+					return nullptr;
+				}
 
 				MacroContext->Push(NodeMacro);
 				ReturnPin = GraphTraversal::FindIOPinSourceThroughMacroContext(*ConnectedPin, MacroContext);
@@ -781,51 +785,12 @@ const UEdGraphPin* FindMeshBaseSource(const UEdGraphPin& Pin, const bool bOnlyLo
 
 	else if (const UCustomizableObjectNodeMacroInstance* NodeMacro = Cast<UCustomizableObjectNodeMacroInstance>(Node))
 	{
-		if (const UEdGraphPin* OutputPin = NodeMacro->GetMacroIOPin(ECOMacroIOType::COMVT_Output, Pin.PinName))
-		{
-			if (const UEdGraphPin* ConnectedPin = FollowInputPin(*OutputPin))
-			{
-				TArray<const UCustomizableObjectNodeMacroInstance*> NewMacroContext;
-
-				if (!MacroContext)
-				{
-					MacroContext = &NewMacroContext;
-				}
-
-				MacroContext->Push(NodeMacro);
-				const UEdGraphPin* ReturnPin = FindMeshBaseSource(*ConnectedPin, bOnlyLookForStaticMesh, MacroContext);
-				MacroContext->Pop();
-
-				return ReturnPin;
-			}
-		}
-	}
-
-	else if (const UCustomizableObjectNodeTunnel* NodeTunnel = Cast<UCustomizableObjectNodeTunnel>(Node))
-	{
-		check(NodeTunnel->bIsInputNode);
-
-		const UEdGraphPin* ReturnPin = nullptr;
-
-		if (MacroContext && MacroContext->Num())
+		const UEdGraphPin* ConnectedPin = GraphTraversal::FindIOPinSourceThroughMacroContext(Pin, MacroContext);
+		
+		if (ConnectedPin)
 		{
-			const UCustomizableObjectNodeMacroInstance* MacroInstanceNode = MacroContext->Pop();
-
-			if (MacroInstanceNode)
-			{
-				if (const UEdGraphPin* InputPin = MacroInstanceNode->FindPin(Pin.PinName, EEdGraphPinDirection::EGPD_Input))
-				{
-					if (const UEdGraphPin* ConnectedPin = FollowInputPin(*InputPin))
-					{
-						ReturnPin = FindMeshBaseSource(*ConnectedPin, bOnlyLookForStaticMesh, MacroContext);
-					}
-				}
-			}
-
-			MacroContext->Push(MacroInstanceNode);
+			return FindMeshBaseSource(*ConnectedPin, bOnlyLookForStaticMesh, MacroContext);
 		}
-
-		return ReturnPin;
 	}
 
 	else
diff --git a/Engine/Plugins/Mutable/Source/CustomizableObjectEditor/Private/MuCOE/Nodes/CustomizableObjectNodeMacroInstanceDetails.cpp b/Engine/Plugins/Mutable/Source/CustomizableObjectEditor/Private/MuCOE/Nodes/CustomizableObjectNodeMacroInstanceDetails.cpp
--- a/Engine/Plugins/Mutable/Source/CustomizableObjectEditor/Private/MuCOE/Nodes/CustomizableObjectNodeMacroInstanceDetails.cpp
+++ b/Engine/Plugins/Mutable/Source/CustomizableObjectEditor/Private/MuCOE/Nodes/CustomizableObjectNodeMacroInstanceDetails.cpp
@@ -47,7 +47,6 @@ void FCustomizableObjectNodeMacroInstanceDetails::CustomizeDetails(IDetailLayout
 
 	// Macro selector widget
 	DetailBuilder.EditDefaultProperty(PropertyHandle)->CustomWidget()
-	.OverrideResetToDefault(FResetToDefaultOverride::Hide())
 	.IsValueEnabled(TAttribute<bool>::CreateSPLambda(this, [this]() { return Node->ParentMacroLibrary != nullptr; }))
 	.NameContent()
 	[
@@ -79,7 +78,8 @@ void FCustomizableObjectNodeMacroInstanceDetails::CustomizeDetails(IDetailLayout
 			.Font(IDetailLayoutBuilder::GetDetailFont())
 			.Text(this, &FCustomizableObjectNodeMacroInstanceDetails::GetSelectedMacroName)
 		]
-	];
+	]
+	.OverrideResetToDefault(FResetToDefaultOverride::Create(FSimpleDelegate::CreateSP(this, &FCustomizableObjectNodeMacroInstanceDetails::ResetSelectedParentMacro)));
 }
 
 
@@ -90,9 +90,13 @@ TSharedPtr<FMacroSelectorItem> FCustomizableObjectNodeMacroInstanceDetails::Gene
 
 	if (Node && Node->ParentMacroLibrary)
 	{
+		UEdGraph* NodeGraph = Node->GetGraph();
+		check(NodeGraph);
+
 		for (UCustomizableObjectMacro* Macro : Node->ParentMacroLibrary->Macros)
 		{
-			if (Macro)
+			// Do not allow to instantiate a macro in its own graph
+			if (Macro && Macro->Graph != NodeGraph)
 			{
 				TSharedPtr<FMacroSelectorItem> MacroSelector = MakeShareable(new FMacroSelectorItem());
 				MacroSelector->Macro = Macro;
@@ -123,4 +127,15 @@ FText FCustomizableObjectNodeMacroInstanceDetails::GetSelectedMacroName() const
 	return Text;
 }
 
+
+void FCustomizableObjectNodeMacroInstanceDetails::ResetSelectedParentMacro()
+{
+	if (Node)
+	{
+		Node->ParentMacro = nullptr;
+		Node->ReconstructNode();
+	}
+}
+
+
 #undef LOCTEXT_NAMESPACE
\ No newline at end of file
diff --git a/Engine/Plugins/Mutable/Source/CustomizableObjectEditor/Private/MuCOE/Nodes/CustomizableObjectNodeMacroInstanceDetails.h b/Engine/Plugins/Mutable/Source/CustomizableObjectEditor/Private/MuCOE/Nodes/CustomizableObjectNodeMacroInstanceDetails.h
--- a/Engine/Plugins/Mutable/Source/CustomizableObjectEditor/Private/MuCOE/Nodes/CustomizableObjectNodeMacroInstanceDetails.h
+++ b/Engine/Plugins/Mutable/Source/CustomizableObjectEditor/Private/MuCOE/Nodes/CustomizableObjectNodeMacroInstanceDetails.h
@@ -32,6 +32,9 @@ private:
 
 	FText GetSelectedMacroName() const;
 
+	/** Set the selected macro property to nullptr */
+	void ResetSelectedParentMacro();
+
 private:
 
 	/** Pointer to the Macro Instance Node. */
diff --git a/Engine/Plugins/Mutable/Source/CustomizableObjectEditor/Private/MuCOE/Nodes/CustomizableObjectNodeMaterial.cpp b/Engine/Plugins/Mutable/Source/CustomizableObjectEditor/Private/MuCOE/Nodes/CustomizableObjectNodeMaterial.cpp
--- a/Engine/Plugins/Mutable/Source/CustomizableObjectEditor/Private/MuCOE/Nodes/CustomizableObjectNodeMaterial.cpp
+++ b/Engine/Plugins/Mutable/Source/CustomizableObjectEditor/Private/MuCOE/Nodes/CustomizableObjectNodeMaterial.cpp
@@ -1267,12 +1267,13 @@ TArray<FString> UCustomizableObjectNodeMaterial::GetEnableTags(TArray<const UCus
 {
 	// Getting tags from linked pin
 	const UEdGraphPin* EnableTagsPin = GetEnableTagsPin();
+	const TArray<UEdGraphPin*> ConnectedPins = FollowInputPinArray(*EnableTagsPin);
 
-	if (EnableTagsPin && FollowInputPinArray(*EnableTagsPin).Num())
+	if (EnableTagsPin && ConnectedPins.Num())
 	{
 		TArray<FString> OutTags;
 
-		for (const UEdGraphPin* StringPin : FollowInputPinArray(*EnableTagsPin))
+		for (const UEdGraphPin* StringPin : ConnectedPins)
 		{
 			const UEdGraphPin* SourceStringPin = GraphTraversal::FindIOPinSourceThroughMacroContext(*StringPin, MacroContext);
 
diff --git a/Engine/Plugins/Mutable/Source/CustomizableObjectEditor/Private/MuCOE/Nodes/CustomizableObjectNodeModifierBase.cpp b/Engine/Plugins/Mutable/Source/CustomizableObjectEditor/Private/MuCOE/Nodes/CustomizableObjectNodeModifierBase.cpp
--- a/Engine/Plugins/Mutable/Source/CustomizableObjectEditor/Private/MuCOE/Nodes/CustomizableObjectNodeModifierBase.cpp
+++ b/Engine/Plugins/Mutable/Source/CustomizableObjectEditor/Private/MuCOE/Nodes/CustomizableObjectNodeModifierBase.cpp
@@ -34,12 +34,13 @@ UEdGraphPin* UCustomizableObjectNodeModifierBase::RequiredTagsPin() const
 TArray<FString> UCustomizableObjectNodeModifierBase::GetNodeRequiredTags(TArray<const UCustomizableObjectNodeMacroInstance*>* MacroContext) const
 {
 	const UEdGraphPin* ReqTagsPin = RequiredTagsPin();
+	const TArray<UEdGraphPin*> ConnectedPins = FollowInputPinArray(*ReqTagsPin);
 
-	if (ReqTagsPin && FollowInputPinArray(*ReqTagsPin).Num())
+	if (ReqTagsPin && ConnectedPins.Num())
 	{
 		TArray<FString> OutTags;
 
-		for (const UEdGraphPin* StringPin : FollowInputPinArray(*ReqTagsPin))
+		for (const UEdGraphPin* StringPin : ConnectedPins)
 		{
 			const UEdGraphPin* SourceStringPin = GraphTraversal::FindIOPinSourceThroughMacroContext(*StringPin, MacroContext);
 
diff --git a/Engine/Plugins/Mutable/Source/CustomizableObjectEditor/Private/MuCOE/Nodes/CustomizableObjectNodeModifierExtendMeshSection.cpp b/Engine/Plugins/Mutable/Source/CustomizableObjectEditor/Private/MuCOE/Nodes/CustomizableObjectNodeModifierExtendMeshSection.cpp
--- a/Engine/Plugins/Mutable/Source/CustomizableObjectEditor/Private/MuCOE/Nodes/CustomizableObjectNodeModifierExtendMeshSection.cpp
+++ b/Engine/Plugins/Mutable/Source/CustomizableObjectEditor/Private/MuCOE/Nodes/CustomizableObjectNodeModifierExtendMeshSection.cpp
@@ -237,13 +237,14 @@ bool UCustomizableObjectNodeModifierExtendMeshSection::CustomRemovePin(UEdGraphP
 TArray<FString> UCustomizableObjectNodeModifierExtendMeshSection::GetEnableTags(TArray<const UCustomizableObjectNodeMacroInstance*>* MacroContext)
 { 
 	const UEdGraphPin* TagsPin = GetEnableTagsPin();
+	const TArray<UEdGraphPin*> ConnectedPins = FollowInputPinArray(*TagsPin);
 
 	// Getting tags from linked pin
-	if (TagsPin && FollowInputPinArray(*TagsPin).Num())
+	if (TagsPin && ConnectedPins.Num())
 	{
 		TArray<FString> OutTags;
 
-		for (const UEdGraphPin* StringPin : FollowInputPinArray(*TagsPin))
+		for (const UEdGraphPin* StringPin : ConnectedPins)
 		{
 			const UEdGraphPin* SourceStringPin = GraphTraversal::FindIOPinSourceThroughMacroContext(*StringPin, MacroContext);
 
