[Mutable] Fixed pointer check

#rnx

diff --git a/Engine/Plugins/Mutable/Source/CustomizableObjectEditor/Private/MuCOE/Nodes/CustomizableObjectNodeMaterial.cpp b/Engine/Plugins/Mutable/Source/CustomizableObjectEditor/Private/MuCOE/Nodes/CustomizableObjectNodeMaterial.cpp
--- a/Engine/Plugins/Mutable/Source/CustomizableObjectEditor/Private/MuCOE/Nodes/CustomizableObjectNodeMaterial.cpp
+++ b/Engine/Plugins/Mutable/Source/CustomizableObjectEditor/Private/MuCOE/Nodes/CustomizableObjectNodeMaterial.cpp
@@ -1267,9 +1267,15 @@ TArray<FString> UCustomizableObjectNodeMaterial::GetEnableTags(TArray<const UCus
 {
 	// Getting tags from linked pin
 	const UEdGraphPin* EnableTagsPin = GetEnableTagsPin();
+
+	if (!EnableTagsPin)
+	{
+		return Tags;
+	}
+
 	const TArray<UEdGraphPin*> ConnectedPins = FollowInputPinArray(*EnableTagsPin);
 
-	if (EnableTagsPin && ConnectedPins.Num())
+	if (ConnectedPins.Num())
 	{
 		TArray<FString> OutTags;
 
@@ -1288,10 +1294,8 @@ TArray<FString> UCustomizableObjectNodeMaterial::GetEnableTags(TArray<const UCus
 
 		return OutTags;
 	}
-	else
-	{
-		return Tags;
-	}
+
+	return Tags;
 }
 
 
diff --git a/Engine/Plugins/Mutable/Source/CustomizableObjectEditor/Private/MuCOE/Nodes/CustomizableObjectNodeModifierBase.cpp b/Engine/Plugins/Mutable/Source/CustomizableObjectEditor/Private/MuCOE/Nodes/CustomizableObjectNodeModifierBase.cpp
--- a/Engine/Plugins/Mutable/Source/CustomizableObjectEditor/Private/MuCOE/Nodes/CustomizableObjectNodeModifierBase.cpp
+++ b/Engine/Plugins/Mutable/Source/CustomizableObjectEditor/Private/MuCOE/Nodes/CustomizableObjectNodeModifierBase.cpp
@@ -34,9 +34,15 @@ UEdGraphPin* UCustomizableObjectNodeModifierBase::RequiredTagsPin() const
 TArray<FString> UCustomizableObjectNodeModifierBase::GetNodeRequiredTags(TArray<const UCustomizableObjectNodeMacroInstance*>* MacroContext) const
 {
 	const UEdGraphPin* ReqTagsPin = RequiredTagsPin();
+
+	if (!ReqTagsPin)
+	{
+		return RequiredTags;
+	}
+
 	const TArray<UEdGraphPin*> ConnectedPins = FollowInputPinArray(*ReqTagsPin);
 
-	if (ReqTagsPin && ConnectedPins.Num())
+	if (ConnectedPins.Num())
 	{
 		TArray<FString> OutTags;
 
diff --git a/Engine/Plugins/Mutable/Source/CustomizableObjectEditor/Private/MuCOE/Nodes/CustomizableObjectNodeModifierExtendMeshSection.cpp b/Engine/Plugins/Mutable/Source/CustomizableObjectEditor/Private/MuCOE/Nodes/CustomizableObjectNodeModifierExtendMeshSection.cpp
--- a/Engine/Plugins/Mutable/Source/CustomizableObjectEditor/Private/MuCOE/Nodes/CustomizableObjectNodeModifierExtendMeshSection.cpp
+++ b/Engine/Plugins/Mutable/Source/CustomizableObjectEditor/Private/MuCOE/Nodes/CustomizableObjectNodeModifierExtendMeshSection.cpp
@@ -237,10 +237,16 @@ bool UCustomizableObjectNodeModifierExtendMeshSection::CustomRemovePin(UEdGraphP
 TArray<FString> UCustomizableObjectNodeModifierExtendMeshSection::GetEnableTags(TArray<const UCustomizableObjectNodeMacroInstance*>* MacroContext)
 { 
 	const UEdGraphPin* TagsPin = GetEnableTagsPin();
+
+	if (!TagsPin)
+	{
+		return Tags;
+	}
+
 	const TArray<UEdGraphPin*> ConnectedPins = FollowInputPinArray(*TagsPin);
 
 	// Getting tags from linked pin
-	if (TagsPin && ConnectedPins.Num())
+	if (ConnectedPins.Num())
 	{
 		TArray<FString> OutTags;
 
@@ -260,10 +266,8 @@ TArray<FString> UCustomizableObjectNodeModifierExtendMeshSection::GetEnableTags(
 
 		return OutTags;
 	}
-	else
-	{
-		return Tags;
-	}
+	
+	return Tags;
 }
 
 
