[Mutable] Fixed Variable for mutable macros

#rnx

diff --git a/Engine/Plugins/Mutable/Source/CustomizableObject/Internal/MuCO/CustomizableObjectCustomVersion.h b/Engine/Plugins/Mutable/Source/CustomizableObject/Internal/MuCO/CustomizableObjectCustomVersion.h
--- a/Engine/Plugins/Mutable/Source/CustomizableObject/Internal/MuCO/CustomizableObjectCustomVersion.h
+++ b/Engine/Plugins/Mutable/Source/CustomizableObject/Internal/MuCO/CustomizableObjectCustomVersion.h
@@ -191,6 +191,9 @@ struct FCustomizableObjectCustomVersion
 		TableNodeCompilationFilterArray,
 
 		EnableMutableMacros,
+			
+		EnableMutableMacrosNewVersion,
+
 
 		// -----<new versions can be added above this line>-------------------------------------------------
 		VersionPlusOne,
diff --git a/Engine/Plugins/Mutable/Source/CustomizableObjectEditor/Private/MuCOE/Nodes/CustomizableObjectNodeComponent.cpp b/Engine/Plugins/Mutable/Source/CustomizableObjectEditor/Private/MuCOE/Nodes/CustomizableObjectNodeComponent.cpp
--- a/Engine/Plugins/Mutable/Source/CustomizableObjectEditor/Private/MuCOE/Nodes/CustomizableObjectNodeComponent.cpp
+++ b/Engine/Plugins/Mutable/Source/CustomizableObjectEditor/Private/MuCOE/Nodes/CustomizableObjectNodeComponent.cpp
@@ -32,7 +32,7 @@ void UCustomizableObjectNodeComponent::BackwardsCompatibleFixup(int32 Customizab
 {
 	Super::BackwardsCompatibleFixup(CustomizableObjectCustomVersion);
 
-	if (CustomizableObjectCustomVersion == FCustomizableObjectCustomVersion::EnableMutableMacros)
+	if (CustomizableObjectCustomVersion == FCustomizableObjectCustomVersion::EnableMutableMacrosNewVersion)
 	{
 		if (!ComponentNamePin.Get())
 		{
diff --git a/Engine/Plugins/Mutable/Source/CustomizableObjectEditor/Private/MuCOE/Nodes/CustomizableObjectNodeComponentMeshAddTo.cpp b/Engine/Plugins/Mutable/Source/CustomizableObjectEditor/Private/MuCOE/Nodes/CustomizableObjectNodeComponentMeshAddTo.cpp
--- a/Engine/Plugins/Mutable/Source/CustomizableObjectEditor/Private/MuCOE/Nodes/CustomizableObjectNodeComponentMeshAddTo.cpp
+++ b/Engine/Plugins/Mutable/Source/CustomizableObjectEditor/Private/MuCOE/Nodes/CustomizableObjectNodeComponentMeshAddTo.cpp
@@ -65,7 +65,7 @@ void UCustomizableObjectNodeComponentMeshAddTo::BackwardsCompatibleFixup(int32 C
 {
 	Super::BackwardsCompatibleFixup(CustomizableObjectCustomVersion);
 
-	if (CustomizableObjectCustomVersion == FCustomizableObjectCustomVersion::EnableMutableMacros)
+	if (CustomizableObjectCustomVersion == FCustomizableObjectCustomVersion::EnableMutableMacrosNewVersion)
 	{
 		if (!ParentComponentNamePin.Get())
 		{
diff --git a/Engine/Plugins/Mutable/Source/CustomizableObjectEditor/Private/MuCOE/Nodes/CustomizableObjectNodeMaterial.cpp b/Engine/Plugins/Mutable/Source/CustomizableObjectEditor/Private/MuCOE/Nodes/CustomizableObjectNodeMaterial.cpp
--- a/Engine/Plugins/Mutable/Source/CustomizableObjectEditor/Private/MuCOE/Nodes/CustomizableObjectNodeMaterial.cpp
+++ b/Engine/Plugins/Mutable/Source/CustomizableObjectEditor/Private/MuCOE/Nodes/CustomizableObjectNodeMaterial.cpp
@@ -778,7 +778,7 @@ void UCustomizableObjectNodeMaterial::BackwardsCompatibleFixup(int32 Customizabl
 		}
 	}
 
-	if (CustomizableObjectCustomVersion == FCustomizableObjectCustomVersion::EnableMutableMacros)
+	if (CustomizableObjectCustomVersion == FCustomizableObjectCustomVersion::EnableMutableMacrosNewVersion)
 	{
 		if (!EnableTagsPinRef.Get())
 		{
diff --git a/Engine/Plugins/Mutable/Source/CustomizableObjectEditor/Private/MuCOE/Nodes/CustomizableObjectNodeMeshMorph.cpp b/Engine/Plugins/Mutable/Source/CustomizableObjectEditor/Private/MuCOE/Nodes/CustomizableObjectNodeMeshMorph.cpp
--- a/Engine/Plugins/Mutable/Source/CustomizableObjectEditor/Private/MuCOE/Nodes/CustomizableObjectNodeMeshMorph.cpp
+++ b/Engine/Plugins/Mutable/Source/CustomizableObjectEditor/Private/MuCOE/Nodes/CustomizableObjectNodeMeshMorph.cpp
@@ -173,7 +173,7 @@ void UCustomizableObjectNodeMeshMorph::BackwardsCompatibleFixup(int32 Customizab
 {
 	Super::BackwardsCompatibleFixup(CustomizableObjectCustomVersion);
 
-	if (CustomizableObjectCustomVersion == FCustomizableObjectCustomVersion::EnableMutableMacros)
+	if (CustomizableObjectCustomVersion == FCustomizableObjectCustomVersion::EnableMutableMacrosNewVersion)
 	{
 		if (!MorphTargetNamePinRef.Get())
 		{
diff --git a/Engine/Plugins/Mutable/Source/CustomizableObjectEditor/Private/MuCOE/Nodes/CustomizableObjectNodeMeshParameter.cpp b/Engine/Plugins/Mutable/Source/CustomizableObjectEditor/Private/MuCOE/Nodes/CustomizableObjectNodeMeshParameter.cpp
--- a/Engine/Plugins/Mutable/Source/CustomizableObjectEditor/Private/MuCOE/Nodes/CustomizableObjectNodeMeshParameter.cpp
+++ b/Engine/Plugins/Mutable/Source/CustomizableObjectEditor/Private/MuCOE/Nodes/CustomizableObjectNodeMeshParameter.cpp
@@ -86,7 +86,7 @@ void UCustomizableObjectNodeMeshParameter::BackwardsCompatibleFixup(int32 Custom
 {
 	Super::BackwardsCompatibleFixup(CustomizableObjectCustomVersion);
 
-	if (CustomizableObjectCustomVersion == FCustomizableObjectCustomVersion::EnableMutableMacros)
+	if (CustomizableObjectCustomVersion == FCustomizableObjectCustomVersion::EnableMutableMacrosNewVersion)
 	{
 		if (!NamePin.Get())
 		{
diff --git a/Engine/Plugins/Mutable/Source/CustomizableObjectEditor/Private/MuCOE/Nodes/CustomizableObjectNodeModifierBase.cpp b/Engine/Plugins/Mutable/Source/CustomizableObjectEditor/Private/MuCOE/Nodes/CustomizableObjectNodeModifierBase.cpp
--- a/Engine/Plugins/Mutable/Source/CustomizableObjectEditor/Private/MuCOE/Nodes/CustomizableObjectNodeModifierBase.cpp
+++ b/Engine/Plugins/Mutable/Source/CustomizableObjectEditor/Private/MuCOE/Nodes/CustomizableObjectNodeModifierBase.cpp
@@ -264,7 +264,7 @@ void UCustomizableObjectNodeModifierBase::BackwardsCompatibleFixup(int32 Customi
 		}
 	}
 
-	if (CustomizableObjectCustomVersion == FCustomizableObjectCustomVersion::EnableMutableMacros)
+	if (CustomizableObjectCustomVersion == FCustomizableObjectCustomVersion::EnableMutableMacrosNewVersion)
 	{
 		if (!RequiredTagsPinRef.Get())
 		{
diff --git a/Engine/Plugins/Mutable/Source/CustomizableObjectEditor/Private/MuCOE/Nodes/CustomizableObjectNodeModifierExtendMeshSection.cpp b/Engine/Plugins/Mutable/Source/CustomizableObjectEditor/Private/MuCOE/Nodes/CustomizableObjectNodeModifierExtendMeshSection.cpp
--- a/Engine/Plugins/Mutable/Source/CustomizableObjectEditor/Private/MuCOE/Nodes/CustomizableObjectNodeModifierExtendMeshSection.cpp
+++ b/Engine/Plugins/Mutable/Source/CustomizableObjectEditor/Private/MuCOE/Nodes/CustomizableObjectNodeModifierExtendMeshSection.cpp
@@ -109,7 +109,7 @@ void UCustomizableObjectNodeModifierExtendMeshSection::BackwardsCompatibleFixup(
 		ReconstructNode();
 	}
 
-	if (CustomizableObjectCustomVersion == FCustomizableObjectCustomVersion::EnableMutableMacros)
+	if (CustomizableObjectCustomVersion == FCustomizableObjectCustomVersion::EnableMutableMacrosNewVersion)
 	{
 		if (!EnableTagsPinRef.Get())
 		{
diff --git a/Engine/Plugins/Mutable/Source/CustomizableObjectEditor/Private/MuCOE/Nodes/CustomizableObjectNodeModifierMorphMeshSection.cpp b/Engine/Plugins/Mutable/Source/CustomizableObjectEditor/Private/MuCOE/Nodes/CustomizableObjectNodeModifierMorphMeshSection.cpp
--- a/Engine/Plugins/Mutable/Source/CustomizableObjectEditor/Private/MuCOE/Nodes/CustomizableObjectNodeModifierMorphMeshSection.cpp
+++ b/Engine/Plugins/Mutable/Source/CustomizableObjectEditor/Private/MuCOE/Nodes/CustomizableObjectNodeModifierMorphMeshSection.cpp
@@ -76,7 +76,7 @@ void UCustomizableObjectNodeModifierMorphMeshSection::BackwardsCompatibleFixup(i
 {
 	Super::BackwardsCompatibleFixup(CustomizableObjectCustomVersion);
 
-	if (CustomizableObjectCustomVersion == FCustomizableObjectCustomVersion::EnableMutableMacros)
+	if (CustomizableObjectCustomVersion == FCustomizableObjectCustomVersion::EnableMutableMacrosNewVersion)
 	{
 		if (!MorphTargetNamePinRef.Get())
 		{
diff --git a/Engine/Plugins/Mutable/Source/CustomizableObjectEditor/Private/MuCOE/Nodes/CustomizableObjectNodeObjectChild.cpp b/Engine/Plugins/Mutable/Source/CustomizableObjectEditor/Private/MuCOE/Nodes/CustomizableObjectNodeObjectChild.cpp
--- a/Engine/Plugins/Mutable/Source/CustomizableObjectEditor/Private/MuCOE/Nodes/CustomizableObjectNodeObjectChild.cpp
+++ b/Engine/Plugins/Mutable/Source/CustomizableObjectEditor/Private/MuCOE/Nodes/CustomizableObjectNodeObjectChild.cpp
@@ -99,7 +99,7 @@ void UCustomizableObjectNodeObjectChild::BackwardsCompatibleFixup(int32 Customiz
 {
 	Super::BackwardsCompatibleFixup(CustomizableObjectCustomVersion);
 
-	if (CustomizableObjectCustomVersion == FCustomizableObjectCustomVersion::EnableMutableMacros)
+	if (CustomizableObjectCustomVersion == FCustomizableObjectCustomVersion::EnableMutableMacrosNewVersion)
 	{
 		if (!NamePin.Get())
 		{
diff --git a/Engine/Plugins/Mutable/Source/CustomizableObjectEditor/Private/MuCOE/Nodes/CustomizableObjectNodeObjectGroup.cpp b/Engine/Plugins/Mutable/Source/CustomizableObjectEditor/Private/MuCOE/Nodes/CustomizableObjectNodeObjectGroup.cpp
--- a/Engine/Plugins/Mutable/Source/CustomizableObjectEditor/Private/MuCOE/Nodes/CustomizableObjectNodeObjectGroup.cpp
+++ b/Engine/Plugins/Mutable/Source/CustomizableObjectEditor/Private/MuCOE/Nodes/CustomizableObjectNodeObjectGroup.cpp
@@ -180,7 +180,7 @@ void UCustomizableObjectNodeObjectGroup::BackwardsCompatibleFixup(int32 Customiz
 {
 	Super::BackwardsCompatibleFixup(CustomizableObjectCustomVersion);
 
-	if (CustomizableObjectCustomVersion == FCustomizableObjectCustomVersion::EnableMutableMacros)
+	if (CustomizableObjectCustomVersion == FCustomizableObjectCustomVersion::EnableMutableMacrosNewVersion)
 	{
 		if (!NamePin.Get())
 		{
diff --git a/Engine/Plugins/Mutable/Source/CustomizableObjectEditor/Private/MuCOE/Nodes/CustomizableObjectNodeParameter.cpp b/Engine/Plugins/Mutable/Source/CustomizableObjectEditor/Private/MuCOE/Nodes/CustomizableObjectNodeParameter.cpp
--- a/Engine/Plugins/Mutable/Source/CustomizableObjectEditor/Private/MuCOE/Nodes/CustomizableObjectNodeParameter.cpp
+++ b/Engine/Plugins/Mutable/Source/CustomizableObjectEditor/Private/MuCOE/Nodes/CustomizableObjectNodeParameter.cpp
@@ -32,7 +32,7 @@ void UCustomizableObjectNodeParameter::BackwardsCompatibleFixup(int32 Customizab
 {
 	Super::BackwardsCompatibleFixup(CustomizableObjectCustomVersion);
 
-	if (CustomizableObjectCustomVersion == FCustomizableObjectCustomVersion::EnableMutableMacros)
+	if (CustomizableObjectCustomVersion == FCustomizableObjectCustomVersion::EnableMutableMacrosNewVersion)
 	{
 		if (!NamePin.Get())
 		{
diff --git a/Engine/Plugins/Mutable/Source/CustomizableObjectEditor/Private/MuCOE/Nodes/CustomizableObjectNodeVariation.cpp b/Engine/Plugins/Mutable/Source/CustomizableObjectEditor/Private/MuCOE/Nodes/CustomizableObjectNodeVariation.cpp
--- a/Engine/Plugins/Mutable/Source/CustomizableObjectEditor/Private/MuCOE/Nodes/CustomizableObjectNodeVariation.cpp
+++ b/Engine/Plugins/Mutable/Source/CustomizableObjectEditor/Private/MuCOE/Nodes/CustomizableObjectNodeVariation.cpp
@@ -156,7 +156,7 @@ void UCustomizableObjectNodeVariation::BackwardsCompatibleFixup(int32 Customizab
 		ReconstructNode();
 	}
 
-	if (CustomizableObjectCustomVersion == FCustomizableObjectCustomVersion::EnableMutableMacros)
+	if (CustomizableObjectCustomVersion == FCustomizableObjectCustomVersion::EnableMutableMacrosNewVersion)
 	{
 		if (VariationTagPins.IsEmpty())
 		{
