[Mutable] Add Mesh Parameters interface.

- New API functions.
- New Editor widgets.
- Renamed Mesh Parameter to Skeletal Mesh Parameter.
- Changed the type of Texture Parameter from UTexture2D to UTexture.
- Removed Texture Provider.
- Added TStroingObjectPtr to Core.

diff --git a/Engine/Plugins/Mutable/Config/BaseMutable.ini b/Engine/Plugins/Mutable/Config/BaseMutable.ini
--- a/Engine/Plugins/Mutable/Config/BaseMutable.ini
+++ b/Engine/Plugins/Mutable/Config/BaseMutable.ini
@@ -2,7 +2,7 @@
 
 +StructRedirects=(OldName="/Script/CustomizableObject.ParameterUIData",NewName="/Script/CustomizableObject.MutableParameterData")
 +StructRedirects=(OldName="/Script/CustomizableObject.CompilationOptions",NewName="/Script/CustomizableObject.CompilationOptions_DEPRECATED")
-+StructRedirects=(OldName="/Script/CustomizableObject.FCustomizableObjectTextureParameterValue",NewName="/Script/CustomizableObject.FCustomizableObjectAssetParameterValue")
++StructRedirects=(OldName="/Script/CustomizableObject.FCustomizableObjectAssetParameterValue",NewName="/Script/CustomizableObject.FCustomizableObjectTextureParameterValue")
 
 +PropertyRedirects=(OldName="GroupProjectorParameterImage.OptionImage", NewName="OptionTexture")
 +PropertyRedirects=(OldName="CustomizableObjectNodeGroupProjectorParameter.OptionImages", NewName="OptionTextures")
diff --git a/Engine/Plugins/Mutable/Source/CustomizableObject/Internal/MuCO/CustomizableObjectMipDataProvider.h b/Engine/Plugins/Mutable/Source/CustomizableObject/Internal/MuCO/CustomizableObjectMipDataProvider.h
--- a/Engine/Plugins/Mutable/Source/CustomizableObject/Internal/MuCO/CustomizableObjectMipDataProvider.h
+++ b/Engine/Plugins/Mutable/Source/CustomizableObject/Internal/MuCO/CustomizableObjectMipDataProvider.h
@@ -9,7 +9,6 @@
 #include "MuR/Image.h"
 #include "Streaming/TextureMipDataProvider.h"
 
-#include "MuCO/DefaultImageProvider.h"
 #include "MuR/System.h"
 
 #include "CustomizableObjectMipDataProvider.generated.h"
@@ -35,9 +34,7 @@ public:
 	
 	FMutableUpdateContext(const FString& InCustomizableObjectPathName, const FString& InInstancePathName, TSharedPtr<mu::FSystem> InSystem,
 		TSharedPtr<mu::FModel, ESPMode::ThreadSafe> InModel, TSharedPtr<FModelStreamableBulkData> InModelStreamable, TSharedPtr<const mu::FParameters> InParameters, int32 InState);
-
-	~FMutableUpdateContext();
-
+	
 	const FString& GetCustomizableObjectPathName() const;
 
 	const FString& GetInstancePathName() const;
diff --git a/Engine/Plugins/Mutable/Source/CustomizableObject/Internal/MuCO/CustomizableObjectPrivate.h b/Engine/Plugins/Mutable/Source/CustomizableObject/Internal/MuCO/CustomizableObjectPrivate.h
--- a/Engine/Plugins/Mutable/Source/CustomizableObject/Internal/MuCO/CustomizableObjectPrivate.h
+++ b/Engine/Plugins/Mutable/Source/CustomizableObject/Internal/MuCO/CustomizableObjectPrivate.h
@@ -1189,7 +1189,7 @@ public:
 #if WITH_EDITORONLY_DATA
 	/** Runtime referenced textures used by the mu::FModel. */
 	UPROPERTY()
-	TArray<TSoftObjectPtr<const UTexture>> RuntimeReferencedTextures;
+	TArray<TSoftObjectPtr<UTexture2D>> RuntimeReferencedTextures;
 	
 	/** Runtime referenced meshes used by the mu::FModel.
 	* TODO: Move to FMutableSourceMeshData when runtime is really implemented.
@@ -1329,6 +1329,12 @@ public:
 	/** Minimum LOD to render per Quality level. */
 	UPROPERTY()
 	TMap<FName, FPerQualityLevelInt> MinQualityLevelLODPerComponent;
+
+	UPROPERTY()
+	TMap<FName, TObjectPtr<UTexture>> TextureParameterDefaultValues;
+
+	UPROPERTY()
+	TMap<FName, TObjectPtr<USkeletalMesh>> SkeletalMeshParameterDefaultValues;
 	
 	UPROPERTY()
 	FString ReleaseVersion;
@@ -1611,6 +1617,8 @@ enum class ECustomizableObjectVersions : int32
 	FixMeshReusalDueToLayouts,
 
 	IncorrectBonePoseMerging,
+
+	CoreParameterUObjects,
 	
 	// -----<new versions can be added above this line>--------
 	LastCustomizableObjectVersion
diff --git a/Engine/Plugins/Mutable/Source/CustomizableObject/Internal/MuCO/CustomizableObjectSystemPrivate.h b/Engine/Plugins/Mutable/Source/CustomizableObject/Internal/MuCO/CustomizableObjectSystemPrivate.h
--- a/Engine/Plugins/Mutable/Source/CustomizableObject/Internal/MuCO/CustomizableObjectSystemPrivate.h
+++ b/Engine/Plugins/Mutable/Source/CustomizableObject/Internal/MuCO/CustomizableObjectSystemPrivate.h
@@ -580,9 +580,7 @@ public:
 	FUpdateContextPrivate(UCustomizableObjectInstance& InInstance, const FCustomizableObjectInstanceDescriptor& Descriptor);
 
 	FUpdateContextPrivate(UCustomizableObjectInstance& InInstance);
-
-	~FUpdateContextPrivate();
-
+	
 	bool IsContextValid() const;
 	
 	void SetMinLOD(const TMap<FName, uint8>& MinLOD);
@@ -969,10 +967,6 @@ public:
 	 * Used for stats that are costly to update. */
 	void UpdateStats();
 
-	void CacheTextureParameters(const TArray<FCustomizableObjectAssetParameterValue>& TextureParameters) const;
-
-	void UnCacheTextureParameters(const TArray<FCustomizableObjectAssetParameterValue>& TextureParameters) const;
-
 #if WITH_EDITOR
 	/** Add a CO to the pending - load list.COs need to wait until all related objects are fully loaded before being able
 	 * to do things like check-if-up-to-date or compile.
@@ -992,14 +986,9 @@ public:
 	UE::Tasks::FTask LastUpdateMutableTask = UE::Tasks::MakeCompletedTask<void>();
 	
 #if WITH_EDITORONLY_DATA
-	/** Mutable default image provider. Used by the COIEditor and Instance/Descriptor APIs. */
-	UPROPERTY(Transient)
-	TObjectPtr<UEditorImageProvider> EditorImageProvider = nullptr;
-
 	/** List of CustomizableObjects pending to complete loading. */
 	UPROPERTY(Transient)
 	TArray<TObjectPtr<UCustomizableObject>> ObjectsPendingLoad;
-
 #endif
 
 	FLogBenchmarkUtil LogBenchmarkUtil;
@@ -1089,4 +1078,17 @@ CUSTOMIZABLEOBJECT_API FTexturePlatformData* MutableCreateImagePlatformData(TSha
 /** Return true if Streaming is enabled for the given Object. */
 bool IsStreamingEnabled(const UCustomizableObject& Object);
 
+
+template<typename Type>
+Type* ToObject(const FString& Parameter)
+{
+	return Cast<Type>(FSoftObjectPath::ConstructFromStringPath(Parameter).TryLoad());
+}
+template<typename Type>
+Type* ToObject(const FName& Parameter)
+{
+	return ToObject<Type>(Parameter.ToString());
+}
+
+
 #undef UE_API
diff --git a/Engine/Plugins/Mutable/Source/CustomizableObject/Internal/MuCO/EditorImageProvider.h b/Engine/Plugins/Mutable/Source/CustomizableObject/Internal/MuCO/EditorImageProvider.h
deleted file mode 100644
diff --git a/Engine/Plugins/Mutable/Source/CustomizableObject/Internal/MuCO/UnrealToMutableTextureConversionUtils.h b/Engine/Plugins/Mutable/Source/CustomizableObject/Internal/MuCO/UnrealToMutableTextureConversionUtils.h
--- a/Engine/Plugins/Mutable/Source/CustomizableObject/Internal/MuCO/UnrealToMutableTextureConversionUtils.h
+++ b/Engine/Plugins/Mutable/Source/CustomizableObject/Internal/MuCO/UnrealToMutableTextureConversionUtils.h
@@ -22,7 +22,7 @@ struct FMutableSourceTextureData
 {
 	FMutableSourceTextureData() = default;
 	
-	UE_API FMutableSourceTextureData(const UTexture& Texture);
+	UE_API FMutableSourceTextureData(const UTexture2D& Texture);
 	
 	UE_API FTextureSource& GetSource();
 	
diff --git a/Engine/Plugins/Mutable/Source/CustomizableObject/Private/MuCO/CustomizableObject.cpp b/Engine/Plugins/Mutable/Source/CustomizableObject/Private/MuCO/CustomizableObject.cpp
--- a/Engine/Plugins/Mutable/Source/CustomizableObject/Private/MuCO/CustomizableObject.cpp
+++ b/Engine/Plugins/Mutable/Source/CustomizableObject/Private/MuCO/CustomizableObject.cpp
@@ -1471,7 +1471,7 @@ void UCustomizableObjectPrivate::UpdateParameterPropertiesFromModel(const TShare
 
 			case mu::EParameterType::Mesh:
 			{
-				Data.Type = EMutableParameterType::Mesh;
+				Data.Type = EMutableParameterType::SkeletalMesh;
 				break;
 			}
 
@@ -1895,23 +1895,17 @@ FCustomizableObjectProjector UCustomizableObject::GetProjectorParameterDefaultVa
 }
 
 
-FName UCustomizableObject::GetTextureParameterDefaultValue(const FString& InParameterName) const
+UTexture* UCustomizableObject::GetTextureParameterDefaultValue(const FString& ParameterName) const
 {
-	const int32 ParameterIndex = GetPrivate()->FindParameter(InParameterName);
-	if (ParameterIndex == INDEX_NONE)
-	{
-		UE_LOG(LogMutable, Error, TEXT("Tried to access the default value of the nonexistent texture parameter [%s] in the CustomizableObject [%s]."), *InParameterName, *GetName());
-		return FCustomizableObjectAssetParameterValue::DEFAULT_PARAMETER_VALUE;
-	}
-
-	const TSharedPtr<const mu::FModel>& Model = GetPrivate()->GetModel();
-	if (!Model)
-	{
-		checkNoEntry();
-		return FCustomizableObjectAssetParameterValue::DEFAULT_PARAMETER_VALUE;
-	}
-	
-	return Model->GetImageDefaultValue(ParameterIndex);
+	const TObjectPtr<UTexture>* Result = GetPrivate()->GetModelResources()->TextureParameterDefaultValues.Find(FName(ParameterName));
+	return Result ? *Result : nullptr;
+}
+ 
+ 
+USkeletalMesh* UCustomizableObject::GetSkeletalMeshParameterDefaultValue(const FString& ParameterName) const
+{
+	const TObjectPtr<USkeletalMesh>* Result = GetPrivate()->GetModelResources()->SkeletalMeshParameterDefaultValues.Find(FName(ParameterName));
+	return Result ? *Result : nullptr;
 }
 
 
diff --git a/Engine/Plugins/Mutable/Source/CustomizableObject/Private/MuCO/CustomizableObjectInstance.cpp b/Engine/Plugins/Mutable/Source/CustomizableObject/Private/MuCO/CustomizableObjectInstance.cpp
--- a/Engine/Plugins/Mutable/Source/CustomizableObject/Private/MuCO/CustomizableObjectInstance.cpp
+++ b/Engine/Plugins/Mutable/Source/CustomizableObject/Private/MuCO/CustomizableObjectInstance.cpp
@@ -385,19 +385,6 @@ void UCustomizableInstancePrivate::BindObjectDelegates(UCustomizableObject*  Cur
 }
 
 
-void UCustomizableObjectInstance::PostEditChangeProperty(FPropertyChangedEvent& PropertyChangedEvent)
-{
-	Super::PostEditChangeProperty(PropertyChangedEvent);
-
-	const FName PropertyName = PropertyChangedEvent.MemberProperty ? PropertyChangedEvent.MemberProperty->GetFName() : NAME_None;
-
-	if (PropertyName == GET_MEMBER_NAME_CHECKED(UCustomizableObjectInstance, TextureParameterDeclarations))
-	{
-		UpdateSkeletalMeshAsync(true, true);
-	}
-}
-
-
 bool UCustomizableObjectInstance::CanEditChange(const FProperty* InProperty) const
 {
 	bool bIsMutable = Super::CanEditChange(InProperty);
@@ -511,12 +498,6 @@ void UCustomizableInstancePrivate::ReleaseMutableResources(bool bCalledFromBegin
 				}
 			}
 		}
-
-		// Remove all references to cached Texture Parameters. Only if we're destroying the COI.
-		if (bCalledFromBeginDestroy)
-		{
-			CustomizableObjectSystem->UnCacheTextureParameters(CommittedDescriptor.GetTextureParameters());
-		}
 	}
 
 	GeneratedTextures.Empty();
@@ -759,12 +740,12 @@ bool UCustomizableObjectInstance::IsParameterDirty(const FString& ParamName, con
 
 	case EMutableParameterType::Projector:
 		{
-			const FCustomizableObjectProjectorParameterValue* Result = Descriptor.GetProjectorParameters().FindByPredicate([&](const FCustomizableObjectProjectorParameterValue& Value)
+			const FCustomizableObjectProjectorParameterValue* Result = Descriptor.ProjectorParameters.FindByPredicate([&](const FCustomizableObjectProjectorParameterValue& Value)
 				{
 					return Value.ParameterName == ParamName;
 				});
 			
-			const FCustomizableObjectProjectorParameterValue* ResultCommited = GetPrivate()->CommittedDescriptor.GetProjectorParameters().FindByPredicate([&](const FCustomizableObjectProjectorParameterValue& Value)
+			const FCustomizableObjectProjectorParameterValue* ResultCommited = GetPrivate()->CommittedDescriptor.ProjectorParameters.FindByPredicate([&](const FCustomizableObjectProjectorParameterValue& Value)
 				{
 					return Value.ParameterName == ParamName;
 				});
@@ -794,12 +775,12 @@ bool UCustomizableObjectInstance::IsParameterDirty(const FString& ParamName, con
 		}		
 	case EMutableParameterType::Texture:
 		{
-			const FCustomizableObjectAssetParameterValue* Result = Descriptor.GetTextureParameters().FindByPredicate([&](const FCustomizableObjectAssetParameterValue& Value)
+			const FCustomizableObjectTextureParameterValue* Result = Descriptor.TextureParameters.FindByPredicate([&](const FCustomizableObjectTextureParameterValue& Value)
 				{
 					return Value.ParameterName == ParamName;
 				});
 			
-			const FCustomizableObjectAssetParameterValue* ResultCommited = GetPrivate()->CommittedDescriptor.GetTextureParameters().FindByPredicate([&](const FCustomizableObjectAssetParameterValue& Value)
+			const FCustomizableObjectTextureParameterValue* ResultCommited = GetPrivate()->CommittedDescriptor.TextureParameters.FindByPredicate([&](const FCustomizableObjectTextureParameterValue& Value)
 				{
 					return Value.ParameterName == ParamName;
 				});
@@ -830,12 +811,12 @@ bool UCustomizableObjectInstance::IsParameterDirty(const FString& ParamName, con
 
 	case EMutableParameterType::Bool:
 		{
-			const FCustomizableObjectBoolParameterValue* Result = Descriptor.GetBoolParameters().FindByPredicate([&](const FCustomizableObjectBoolParameterValue& Value)
+			const FCustomizableObjectBoolParameterValue* Result = Descriptor.BoolParameters.FindByPredicate([&](const FCustomizableObjectBoolParameterValue& Value)
 				{
 					return Value.ParameterName == ParamName;
 				});
 			
-			const FCustomizableObjectBoolParameterValue* ResultCommited = GetPrivate()->CommittedDescriptor.GetBoolParameters().FindByPredicate([&](const FCustomizableObjectBoolParameterValue& Value)
+			const FCustomizableObjectBoolParameterValue* ResultCommited = GetPrivate()->CommittedDescriptor.BoolParameters.FindByPredicate([&](const FCustomizableObjectBoolParameterValue& Value)
 				{
 					return Value.ParameterName == ParamName;
 				});
@@ -858,12 +839,12 @@ bool UCustomizableObjectInstance::IsParameterDirty(const FString& ParamName, con
 		}
 	case EMutableParameterType::Int:
 		{
-			const FCustomizableObjectIntParameterValue* Result = Descriptor.GetIntParameters().FindByPredicate([&](const FCustomizableObjectIntParameterValue& Value)
+			const FCustomizableObjectIntParameterValue* Result = Descriptor.IntParameters.FindByPredicate([&](const FCustomizableObjectIntParameterValue& Value)
 				{
 					return Value.ParameterName == ParamName;
 				});
 			
-			const FCustomizableObjectIntParameterValue* ResultCommited = GetPrivate()->CommittedDescriptor.GetIntParameters().FindByPredicate([&](const FCustomizableObjectIntParameterValue& Value)
+			const FCustomizableObjectIntParameterValue* ResultCommited = GetPrivate()->CommittedDescriptor.IntParameters.FindByPredicate([&](const FCustomizableObjectIntParameterValue& Value)
 				{
 					return Value.ParameterName == ParamName;
 				});
@@ -894,12 +875,12 @@ bool UCustomizableObjectInstance::IsParameterDirty(const FString& ParamName, con
 		
 	case EMutableParameterType::Float:
 		{
-			const FCustomizableObjectFloatParameterValue* Result = Descriptor.GetFloatParameters().FindByPredicate([&](const FCustomizableObjectFloatParameterValue& Value)
+			const FCustomizableObjectFloatParameterValue* Result = Descriptor.FloatParameters.FindByPredicate([&](const FCustomizableObjectFloatParameterValue& Value)
 				{
 					return Value.ParameterName == ParamName;
 				});
 			
-			const FCustomizableObjectFloatParameterValue* ResultCommited = GetPrivate()->CommittedDescriptor.GetFloatParameters().FindByPredicate([&](const FCustomizableObjectFloatParameterValue& Value)
+			const FCustomizableObjectFloatParameterValue* ResultCommited = GetPrivate()->CommittedDescriptor.FloatParameters.FindByPredicate([&](const FCustomizableObjectFloatParameterValue& Value)
 				{
 					return Value.ParameterName == ParamName;
 				});
@@ -930,12 +911,12 @@ bool UCustomizableObjectInstance::IsParameterDirty(const FString& ParamName, con
 		
 	case EMutableParameterType::Color:
 		{
-			const FCustomizableObjectVectorParameterValue* Result = Descriptor.GetVectorParameters().FindByPredicate([&](const FCustomizableObjectVectorParameterValue& Value)
+			const FCustomizableObjectVectorParameterValue* Result = Descriptor.VectorParameters.FindByPredicate([&](const FCustomizableObjectVectorParameterValue& Value)
 				{
 					return Value.ParameterName == ParamName;
 				});
 			
-			const FCustomizableObjectVectorParameterValue* ResultCommited = GetPrivate()->CommittedDescriptor.GetVectorParameters().FindByPredicate([&](const FCustomizableObjectVectorParameterValue& Value)
+			const FCustomizableObjectVectorParameterValue* ResultCommited = GetPrivate()->CommittedDescriptor.VectorParameters.FindByPredicate([&](const FCustomizableObjectVectorParameterValue& Value)
 				{
 					return Value.ParameterName == ParamName;
 				});
@@ -2718,18 +2699,30 @@ void UCustomizableObjectInstance::SetFloatParameterSelectedOption(const FString&
 }
 
 
-FName UCustomizableObjectInstance::GetTextureParameterSelectedOption(const FString& TextureParamName, const int32 RangeIndex) const
+UTexture* UCustomizableObjectInstance::GetTextureParameterSelectedOption(const FString& TextureParamName, const int32 RangeIndex) const
 {
 	return Descriptor.GetTextureParameterSelectedOption(TextureParamName, RangeIndex);
 }
 
 
-void UCustomizableObjectInstance::SetTextureParameterSelectedOption(const FString& TextureParamName, const FString& TextureValue, const int32 RangeIndex)
+void UCustomizableObjectInstance::SetTextureParameterSelectedOption(const FString& TextureParamName, UTexture* TextureValue, const int32 RangeIndex)
 {
 	Descriptor.SetTextureParameterSelectedOption(TextureParamName, TextureValue, RangeIndex);
 }
 
 
+USkeletalMesh* UCustomizableObjectInstance::GetSkeletalMeshParameterSelectedOption(const FString& SkeletalMeshParamName, int32 RangeIndex)
+{
+	return Descriptor.GetSkeletalMeshParameterSelectedOption(SkeletalMeshParamName, RangeIndex);
+}
+
+
+void UCustomizableObjectInstance::SetSkeletalMeshParameterSelectedOption(const FString& SkeletalMeshParamName, USkeletalMesh* SkeletalMeshValue, int32 RangeIndex)
+{
+	Descriptor.SetSkeletalMeshParameterSelectedOption(SkeletalMeshParamName, SkeletalMeshValue, RangeIndex);
+}
+
+
 FLinearColor UCustomizableObjectInstance::GetColorParameterSelectedOption(const FString& ColorParamName) const
 {
 	return Descriptor.GetColorParameterSelectedOption(ColorParamName);
@@ -3134,9 +3127,9 @@ bool UCustomizableInstancePrivate::MigrateProfileParametersToCurrentInstance(int
 		}
 	}
 
-	for (FCustomizableObjectAssetParameterValue& Parameter : TempProfile.TextureParameters)
+	for (FCustomizableObjectTextureParameterValue& Parameter : TempProfile.TextureParameters)
 	{
-		using ParamValType = FCustomizableObjectAssetParameterValue;
+		using ParamValType = FCustomizableObjectTextureParameterValue;
 		ParamValType* Found = Profile.TextureParameters.FindByPredicate(
 			[&Parameter](const ParamValType& P) { return P.ParameterName == Parameter.ParameterName; });
 
@@ -3146,10 +3139,10 @@ bool UCustomizableInstancePrivate::MigrateProfileParametersToCurrentInstance(int
 		}
 	}
 
-	for (FCustomizableObjectAssetParameterValue& Parameter : TempProfile.MeshParameters)
+	for (FCustomizableObjectSkeletalMeshParameterValue& Parameter : TempProfile.SkeletalMeshParameters)
 	{
-		using ParamValType = FCustomizableObjectAssetParameterValue;
-		ParamValType* Found = Profile.MeshParameters.FindByPredicate(
+		using ParamValType = FCustomizableObjectSkeletalMeshParameterValue;
+		ParamValType* Found = Profile.SkeletalMeshParameters.FindByPredicate(
 			[&Parameter](const ParamValType& P) { return P.ParameterName == Parameter.ParameterName; });
 
 		if (Found)
diff --git a/Engine/Plugins/Mutable/Source/CustomizableObject/Private/MuCO/CustomizableObjectInstanceDescriptor.cpp b/Engine/Plugins/Mutable/Source/CustomizableObject/Private/MuCO/CustomizableObjectInstanceDescriptor.cpp
--- a/Engine/Plugins/Mutable/Source/CustomizableObject/Private/MuCO/CustomizableObjectInstanceDescriptor.cpp
+++ b/Engine/Plugins/Mutable/Source/CustomizableObject/Private/MuCO/CustomizableObjectInstanceDescriptor.cpp
@@ -1,56 +1,57 @@
 // Copyright Epic Games, Inc. All Rights Reserved.
-
+ 
 #include "MuCO/CustomizableObjectInstanceDescriptor.h"
-
+ 
 #include "MuCO/CustomizableObjectInstancePrivate.h"
 #include "MuCO/CustomizableObjectSystemPrivate.h"
 #include "MuCO/UnrealMutableImageProvider.h"
 #include "MuCO/CustomizableObject.h"
 #include "MuCO/CustomizableObjectPrivate.h"
-#include "MuCO/DefaultImageProvider.h"
 #include "MuCO/MutableProjectorTypeUtils.h"
 #include "MuR/Model.h"
 #include "MuR/Parameters.h"
 #include "MuR/Ptr.h"
-
+#include "Engine/Texture2D.h"
+#include "Engine/SkeletalMesh.h"
+ 
 #if WITH_EDITOR
 #include "Editor.h"
 #endif
-
+ 
 #include UE_INLINE_GENERATED_CPP_BY_NAME(CustomizableObjectInstanceDescriptor)
-
-
+ 
+ 
 void CustomizableObjectNullErrorMessage()
 {
 	UE_LOG(LogMutable, Error,
 		TEXT("Tried to perform actions on a CustomizableObjectInstance with no CustomizableObject set. Please set the CustomizableObject of the Instance before doing anything with it."));
 }
-
+ 
 FString GetAvailableOptionsString(const UCustomizableObject& CustomizableObject, const int32 ParameterIndexInObject)
 {
 	FString OptionsString;
 	const int32 NumOptions = CustomizableObject.GetPrivate()->GetEnumParameterNumValues(ParameterIndexInObject);
-
+ 
 	for (int32 k = 0; k < NumOptions; ++k)
 	{
 		OptionsString += CustomizableObject.GetPrivate()->GetIntParameterAvailableOption(ParameterIndexInObject, k);
-
+ 
 		if (k < NumOptions - 1)
 		{
 			OptionsString += FString(", ");
 		}
 	}
-
+ 
 	return OptionsString;
 }
-
-
+ 
+ 
 FCustomizableObjectInstanceDescriptor::FCustomizableObjectInstanceDescriptor(UCustomizableObject& Object)
 {
 	SetCustomizableObject(&Object);
 }
-
-
+ 
+ 
 void FCustomizableObjectInstanceDescriptor::SaveDescriptor(FArchive& Ar, bool bUseCompactDescriptor)
 {
 	if (!CustomizableObject)
@@ -58,30 +59,30 @@ void FCustomizableObjectInstanceDescriptor::SaveDescriptor(FArchive& Ar, bool bU
 		CustomizableObjectNullErrorMessage();
 		return;
 	}
-
+ 
 	Ar << bUseCompactDescriptor;
-
+ 
 	// Not sure if this is needed, but it is small.
 	Ar << State;
-
+ 
 	int32 ModelParameterCount = CustomizableObject->GetParameterCount();
-
+ 
 	if (!bUseCompactDescriptor)
 	{
 		Ar << ModelParameterCount;
 	}
-
+ 
 	for (int32 ModelParameterIndex = 0; ModelParameterIndex < ModelParameterCount; ++ModelParameterIndex)
 	{
-		const FString & Name = CustomizableObject->GetParameterName(ModelParameterIndex);
+		const FString& Name = CustomizableObject->GetParameterName(ModelParameterIndex);
 		EMutableParameterType Type = CustomizableObject->GetParameterTypeByName(Name);
-
+ 
 		if (!bUseCompactDescriptor)
 		{
 			check(Ar.IsSaving());
 			Ar << const_cast<FString &>(Name);
 		}
-
+ 
 		switch (Type)
 		{
 		case EMutableParameterType::Bool:
@@ -98,7 +99,7 @@ void FCustomizableObjectInstanceDescriptor::SaveDescriptor(FArchive& Ar, bool bU
 			Ar << Value;
 			break;
 		}
-
+ 
 		case EMutableParameterType::Float:
 		{
 			float Value = 0.0f;
@@ -116,26 +117,26 @@ void FCustomizableObjectInstanceDescriptor::SaveDescriptor(FArchive& Ar, bool bU
 			Ar << RangeValues;
 			break;
 		}
-
+ 
 		case EMutableParameterType::Int:
 		{
 			int32 Value = 0;
 			FString ValueName;
-
+ 
 			TArray<int32> Values;
 			TArray<FString> ValueNames;
-
+ 
 			bool bIsParamMultidimensional = false;
-
+ 
 			const int32 IntParameterIndex = FindTypedParameterIndex(Name, EMutableParameterType::Int);
 			if (IntParameterIndex >= 0)
 			{					
 				const FCustomizableObjectIntParameterValue & P = IntParameters[IntParameterIndex];
 				Value = CustomizableObject->GetPrivate()->FindIntParameterValue(ModelParameterIndex, P.ParameterValueName);
-
+ 
 				int32 ParameterIndexInObject = CustomizableObject->GetPrivate()->FindParameter(IntParameters[IntParameterIndex].ParameterName);
 				bIsParamMultidimensional = CustomizableObject->GetPrivate()->IsParameterMultidimensional(ParameterIndexInObject);
-
+ 
 				if (bIsParamMultidimensional)
 				{
 					for (int32 i = 0; i < P.ParameterRangeValueNames.Num(); ++i)
@@ -144,17 +145,17 @@ void FCustomizableObjectInstanceDescriptor::SaveDescriptor(FArchive& Ar, bool bU
 						Values.Add(CustomizableObject->GetPrivate()->FindIntParameterValue(ModelParameterIndex, P.ParameterRangeValueNames[i]));
 					}
 				}
-
+ 
 				if (!bUseCompactDescriptor)
 				{
 					ValueName = P.ParameterValueName;
 				}
 			}
-
+ 
 			if (bUseCompactDescriptor)
 			{
 				Ar << Value;
-
+ 
 				if (bIsParamMultidimensional)
 				{
 					Ar << Values;
@@ -163,16 +164,16 @@ void FCustomizableObjectInstanceDescriptor::SaveDescriptor(FArchive& Ar, bool bU
 			else
 			{
 				Ar << ValueName;
-
+ 
 				if (bIsParamMultidimensional)
 				{
 					Ar << ValueNames;
 				}
 			}
-
+ 
 			break;
 		}
-
+ 
 		case EMutableParameterType::Color:
 		{
 			FLinearColor Value(FLinearColor::Black);
@@ -185,10 +186,10 @@ void FCustomizableObjectInstanceDescriptor::SaveDescriptor(FArchive& Ar, bool bU
 				}
 			}
 			Ar << Value;
-
+ 
 			break;
 		}
-
+ 
 		case EMutableParameterType::Transform:
 		{
 			FTransform Value(FTransform::Identity);
@@ -201,55 +202,68 @@ void FCustomizableObjectInstanceDescriptor::SaveDescriptor(FArchive& Ar, bool bU
 				}
 			}
 			Ar << Value;
-
+ 
 			break;
 		}
-
+ 
 		case EMutableParameterType::Texture:
 		{
-			FName Value;
-			TArray<FName> RangeValues;
-
-			for (const FCustomizableObjectAssetParameterValue& P : TextureParameters)
+			FString Value;
+			TArray<FString> RangeValues;
+ 
+			for (const FCustomizableObjectTextureParameterValue& P : TextureParameters)
 			{
 				if (P.ParameterName == Name)
 				{
-					Value = P.ParameterValue;
-					RangeValues = P.ParameterRangeValues;
+					Value = P.ParameterValue.GetPathName();
+ 
+					RangeValues.Empty(P.ParameterRangeValues.Num());	
+					Algo::Transform(P.ParameterRangeValues, RangeValues, [](const TObjectPtr<UTexture>& Element)
+					{
+						return Element->GetPathName();
+					});
+ 
 					break;
 				}
 			}
 			Ar << Value;
 			Ar << RangeValues;
-
+ 
 			break;
 		}
-
-		case EMutableParameterType::Mesh:
+ 
+		case EMutableParameterType::SkeletalMesh:
 		{
-			FName Value;
-			TArray<FName> RangeValues;
-
-			for (const FCustomizableObjectAssetParameterValue& P : MeshParameters)
+			FString Value;
+			TArray<FString> RangeValues;
+ 
+			for (const FCustomizableObjectSkeletalMeshParameterValue& P : SkeletalMeshParameters)
 			{
 				if (P.ParameterName == Name)
 				{
-					Value = P.ParameterValue;
-					RangeValues = P.ParameterRangeValues;
+					Value = P.ParameterValue.GetPathName();
+					
+					RangeValues.Empty(P.ParameterRangeValues.Num());	
+					Algo::Transform(P.ParameterRangeValues, RangeValues, [](const TObjectPtr<USkeletalMesh>& Element)
+					{
+						return Element->GetPathName();
+					});
+					
 					break;
 				}
 			}
+ 
 			Ar << Value;
 			Ar << RangeValues;
-
+				
 			break;
 		}
-
+ 
 		case EMutableParameterType::Projector:
 		{
 			FCustomizableObjectProjector Value;
 			TArray<FCustomizableObjectProjector> RangeValues;
-
+ 
 			for (const FCustomizableObjectProjectorParameterValue& P : ProjectorParameters)
 			{
 				if (P.ParameterName == Name)
@@ -261,18 +275,18 @@ void FCustomizableObjectInstanceDescriptor::SaveDescriptor(FArchive& Ar, bool bU
 			}
 			Ar << Value;
 			Ar << RangeValues;
-
+ 
 			break;
 		}
-
+ 
 		default:
 			// Parameter type replication not implemented.
 			check(false);
 		}
 	}
 }
-
-
+ 
+ 
 void FCustomizableObjectInstanceDescriptor::LoadDescriptor(FArchive& Ar)
 {
 	if (!CustomizableObject)
@@ -280,26 +294,26 @@ void FCustomizableObjectInstanceDescriptor::LoadDescriptor(FArchive& Ar)
 		CustomizableObjectNullErrorMessage();
 		return;
 	}
-
+ 
 	bool bUseCompactDescriptor;
 	Ar << bUseCompactDescriptor;
-
+ 
 	// Not sure if this is needed, but it is small.
 	Ar << State;
 	
 	int32 ModelParameterCount = CustomizableObject->GetParameterCount();
-
+ 
 	if (!bUseCompactDescriptor)
 	{
 		Ar << ModelParameterCount;
 	}
-
+ 
 	for (int32 ParameterIndex = 0; ParameterIndex < ModelParameterCount; ++ParameterIndex)
 	{
 		FString Name;
 		EMutableParameterType Type;
 		int32 ModelParameterIndex = -1;
-
+ 
 		if (bUseCompactDescriptor)
 		{
 			ModelParameterIndex = ParameterIndex;
@@ -311,7 +325,7 @@ void FCustomizableObjectInstanceDescriptor::LoadDescriptor(FArchive& Ar)
 			Ar << Name;
 			Type = CustomizableObject->GetParameterTypeByName(Name);
 		}
-
+ 
 		switch (Type)
 		{
 		case EMutableParameterType::Bool:
@@ -328,7 +342,7 @@ void FCustomizableObjectInstanceDescriptor::LoadDescriptor(FArchive& Ar)
 			}
 			break;
 		}
-
+ 
 		case EMutableParameterType::Float:
 		{
 			float Value = 0.0f;
@@ -346,23 +360,23 @@ void FCustomizableObjectInstanceDescriptor::LoadDescriptor(FArchive& Ar)
 			}
 			break;
 		}
-
+ 
 		case EMutableParameterType::Int:
 		{
 			int32 Value = 0;
 			FString ValueName;
-
+ 
 			TArray<int32> Values;
 			TArray<FString> ValueNames;
-
+ 
 			const int32 IntParameterIndex = FindTypedParameterIndex(Name, EMutableParameterType::Int);
 			const int32 ParameterIndexInObject = CustomizableObject->GetPrivate()->FindParameter(IntParameters[IntParameterIndex].ParameterName);
 			const bool bIsParamMultidimensional = CustomizableObject->GetPrivate()->IsParameterMultidimensional(ParameterIndexInObject);
-
+ 
 			if (bUseCompactDescriptor)
 			{
 				Ar << Value;
-
+ 
 				if (bIsParamMultidimensional)
 				{
 					Ar << Values;
@@ -371,13 +385,13 @@ void FCustomizableObjectInstanceDescriptor::LoadDescriptor(FArchive& Ar)
 			else
 			{
 				Ar << ValueName;
-
+ 
 				if (bIsParamMultidimensional)
 				{
 					Ar << ValueNames;
 				}
 			}
-
+ 
 			for (FCustomizableObjectIntParameterValue& P : IntParameters)
 			{
 				if (P.ParameterName == Name)
@@ -387,7 +401,7 @@ void FCustomizableObjectInstanceDescriptor::LoadDescriptor(FArchive& Ar)
 						P.ParameterValueName = CustomizableObject->GetPrivate()->FindIntParameterValueName(ModelParameterIndex, Value);
 						//check((P.ParameterValueName.IsEmpty() && ValueName.Equals(FString("None"))) || P.ParameterValueName.Equals(ValueName));
 						P.ParameterRangeValueNames.SetNum(Values.Num());
-
+ 
 						for (int32 ParamIndex = 0; ParamIndex < Values.Num(); ++ParamIndex)
 						{
 							P.ParameterRangeValueNames[ParamIndex] = CustomizableObject->GetPrivate()->FindIntParameterValueName(ModelParameterIndex, Values[ParamIndex]);
@@ -398,14 +412,14 @@ void FCustomizableObjectInstanceDescriptor::LoadDescriptor(FArchive& Ar)
 						P.ParameterValueName = ValueName;
 						P.ParameterRangeValueNames = ValueNames;
 					}
-
+ 
 					break;
 				}
 			}
-
+ 
 			break;
 		}
-
+ 
 		case EMutableParameterType::Color:
 		{
 			FLinearColor Value(FLinearColor::Black);
@@ -418,10 +432,10 @@ void FCustomizableObjectInstanceDescriptor::LoadDescriptor(FArchive& Ar)
 					break;
 				}
 			}
-
+ 
 			break;
 		}
-
+ 
 		case EMutableParameterType::Transform:
 		{
 			FTransform Value(FTransform::Identity);
@@ -434,57 +448,69 @@ void FCustomizableObjectInstanceDescriptor::LoadDescriptor(FArchive& Ar)
 					break;
 				}
 			}
-
+ 
 			break;
 		}
-
+ 
 		case EMutableParameterType::Texture:
 		{
-			FName Value;
-			TArray<FName> RangeValues;
+			FString Value;
+			TArray<FString> RangeValues;
 			Ar << Value;
 			Ar << RangeValues;
-
-			for (FCustomizableObjectAssetParameterValue& P : TextureParameters)
+ 
+			for (FCustomizableObjectTextureParameterValue& P : TextureParameters)
 			{
 				if (P.ParameterName == Name)
 				{
-					P.ParameterValue = Value;
-					P.ParameterRangeValues = RangeValues;
+					P.ParameterValue = ToObject<UTexture>(Value);
+ 
+					P.ParameterRangeValues.Empty(RangeValues.Num());
+					Algo::Transform(RangeValues, P.ParameterRangeValues, [](const FString& Element)
+					{
+						return ToObject<UTexture>(Element);
+					});
+					
 					break;
 				}
 			}
-
+ 
 			break;
 		}
-
-		case EMutableParameterType::Mesh:
+ 
+		case EMutableParameterType::SkeletalMesh:
 		{
-			FName Value;
-			TArray<FName> RangeValues;
+			FString Value;
+			TArray<FString> RangeValues;
 			Ar << Value;
 			Ar << RangeValues;
-
-			for (FCustomizableObjectAssetParameterValue& P : MeshParameters)
+ 
+			for (FCustomizableObjectSkeletalMeshParameterValue& P : SkeletalMeshParameters)
 			{
 				if (P.ParameterName == Name)
 				{
-					P.ParameterValue = Value;
-					P.ParameterRangeValues = RangeValues;
+					P.ParameterValue = ToObject<USkeletalMesh>(Value);
+ 
+					P.ParameterRangeValues.Empty(RangeValues.Num());
+					Algo::Transform( RangeValues, P.ParameterRangeValues, [](const FString& Element)
+					{
+						return ToObject<USkeletalMesh>(Element);
+					});
+					
 					break;
 				}
 			}
-
+ 
 			break;
 		}
-
+ 
 		case EMutableParameterType::Projector:
 		{
 			FCustomizableObjectProjector Value;
 			TArray<FCustomizableObjectProjector> RangeValues;
 			Ar << Value;
 			Ar << RangeValues;
-
+ 
 			for (FCustomizableObjectProjectorParameterValue& P : ProjectorParameters)
 			{
 				if (P.ParameterName == Name)
@@ -494,31 +520,31 @@ void FCustomizableObjectInstanceDescriptor::LoadDescriptor(FArchive& Ar)
 					break;
 				}
 			}
-
+ 
 			break;
 		}
-
+ 
 		default:
 			// Parameter type replication not implemented.
 			check(false);
 		}
 	}
 }
-
-
+ 
+ 
 UCustomizableObject* FCustomizableObjectInstanceDescriptor::GetCustomizableObject() const
 {
 	return CustomizableObject;
 }
-
-
+ 
+ 
 void FCustomizableObjectInstanceDescriptor::SetCustomizableObject(UCustomizableObject* InCustomizableObject)
 {
 	CustomizableObject = InCustomizableObject;
 	ReloadParameters();
 }
-
-
+ 
+ 
 bool FCustomizableObjectInstanceDescriptor::GetBuildParameterRelevancy() const
 {
 #if WITH_EDITOR
@@ -538,14 +564,14 @@ bool FCustomizableObjectInstanceDescriptor::GetBuildParameterRelevancy() const
 	return bBuildParameterRelevancy;
 #endif
 }
-
-
+ 
+ 
 void FCustomizableObjectInstanceDescriptor::SetBuildParameterRelevancy(const bool Value)
 {
 	bBuildParameterRelevancy = Value;
 }
-
-
+ 
+ 
 TSharedPtr<mu::FParameters> FCustomizableObjectInstanceDescriptor::GetParameters() const
 {
 	if (!CustomizableObject)
@@ -553,7 +579,7 @@ TSharedPtr<mu::FParameters> FCustomizableObjectInstanceDescriptor::GetParameters
 		CustomizableObjectNullErrorMessage();
 		return nullptr;
 	}
-
+ 
 	if (!CustomizableObject->IsCompiled())
 	{	
 		return nullptr;
@@ -563,16 +589,21 @@ TSharedPtr<mu::FParameters> FCustomizableObjectInstanceDescriptor::GetParameters
 	{
 		return nullptr;
 	}
-
-	TSharedPtr<mu::FParameters> MutableParameters = mu::FModel::NewParameters( CustomizableObject->GetPrivate()->GetModel() );
-
+ 
+	const UModelResources& ModelResources = CustomizableObject->GetPrivate()->GetModelResourcesChecked();
+	
+	TSharedPtr<mu::FParameters> MutableParameters = mu::FModel::NewParameters(CustomizableObject->GetPrivate()->GetModel(),
+		nullptr,
+		&ModelResources.TextureParameterDefaultValues,
+		&ModelResources.SkeletalMeshParameterDefaultValues);
+ 
 	const int32 ParamCount = MutableParameters->GetCount();
 	for (int32 ParamIndex = 0; ParamIndex < ParamCount; ++ParamIndex)
 	{
 		const FString& Name = MutableParameters->GetName(ParamIndex);
 		const FGuid& Uid = MutableParameters->GetUid(ParamIndex);
 		const mu::EParameterType MutableType = MutableParameters->GetType(ParamIndex);
-
+ 
 		switch (MutableType)
 		{
 		case mu::EParameterType::Bool:
@@ -585,10 +616,10 @@ TSharedPtr<mu::FParameters> FCustomizableObjectInstanceDescriptor::GetParameters
 					break;
 				}
 			}
-
+ 
 			break;
 		}
-
+ 
 		case mu::EParameterType::Int:
 		{
 			for (const FCustomizableObjectIntParameterValue& IntParameter : IntParameters)
@@ -600,7 +631,7 @@ TSharedPtr<mu::FParameters> FCustomizableObjectInstanceDescriptor::GetParameters
 						for (int32 RangeIndex = 0; RangeIndex < IntParameter.ParameterRangeValueNames.Num(); ++RangeIndex)
 						{
 							RangeIdxPtr->SetPosition(0, RangeIndex);
-
+ 
 							const FString& RangeValueName = IntParameter.ParameterRangeValueNames[RangeIndex];
 							const int32 Value = CustomizableObject->GetPrivate()->FindIntParameterValue(ParamIndex, RangeValueName);
 							MutableParameters->SetIntValue(ParamIndex, Value, RangeIdxPtr.Get());
@@ -611,14 +642,14 @@ TSharedPtr<mu::FParameters> FCustomizableObjectInstanceDescriptor::GetParameters
 						const int32 Value = CustomizableObject->GetPrivate()->FindIntParameterValue(ParamIndex, IntParameter.ParameterValueName);
 						MutableParameters->SetIntValue(ParamIndex, Value);
 					}
-
+ 
 					break;
 				}
 			}
-
+ 
 			break;
 		}
-
+ 
 		case mu::EParameterType::Float:
 		{
 			for (const FCustomizableObjectFloatParameterValue& FloatParameter : FloatParameters)
@@ -637,14 +668,14 @@ TSharedPtr<mu::FParameters> FCustomizableObjectInstanceDescriptor::GetParameters
 					{
 						MutableParameters->SetFloatValue(ParamIndex, FloatParameter.ParameterValue);
 					}
-
+ 
 					break;
 				}
 			}
-
+ 
 			break;
 		}
-
+ 
 		case mu::EParameterType::Color:
 		{
 			for (const FCustomizableObjectVectorParameterValue& VectorParameter : VectorParameters)
@@ -652,14 +683,14 @@ TSharedPtr<mu::FParameters> FCustomizableObjectInstanceDescriptor::GetParameters
 				if (VectorParameter.ParameterName == Name || (Uid.IsValid() && VectorParameter.Id == Uid))
 				{
 					MutableParameters->SetColourValue(ParamIndex, VectorParameter.ParameterValue);
-
+ 
 					break;
 				}
 			}
-
+ 
 			break;
 		}
-
+ 
 		case mu::EParameterType::Matrix:
 			{
 				for (const FCustomizableObjectTransformParameterValue& TransformParameter : TransformParameters)
@@ -667,11 +698,11 @@ TSharedPtr<mu::FParameters> FCustomizableObjectInstanceDescriptor::GetParameters
 					if (TransformParameter.ParameterName == Name || (Uid.IsValid() && TransformParameter.Id == Uid))
 					{
 						MutableParameters->SetMatrixValue(ParamIndex, FMatrix44f(TransformParameter.ParameterValue.ToMatrixWithScale()));
-
+ 
 						break;
 					}
 				}
-
+ 
 				break;
 			}
 			
@@ -695,7 +726,7 @@ TSharedPtr<mu::FParameters> FCustomizableObjectInstanceDescriptor::GetParameters
 								Value.Angle,
 								RangeIdxPtr);
 							break;
-
+ 
 						case ECustomizableObjectProjectorType::Cylindrical:
 						{
 							// Apply strange swizzle for scales
@@ -712,14 +743,14 @@ TSharedPtr<mu::FParameters> FCustomizableObjectInstanceDescriptor::GetParameters
 									RangeIdxPtr);
 							break;
 						}
-
+ 
 						default:
 							check(false); // Not implemented.
 						}
 					};
-
+ 
 					CopyProjector(ProjectorParameter.Value);
-
+ 
 					if (const TSharedPtr<mu::FRangeIndex> RangeIdxPtr = MutableParameters->NewRangeIndex(ParamIndex))
 					{
 						for (int32 RangeIndex = 0; RangeIndex < ProjectorParameter.RangeValues.Num(); ++RangeIndex)
@@ -730,13 +761,13 @@ TSharedPtr<mu::FParameters> FCustomizableObjectInstanceDescriptor::GetParameters
 					}
 				}
 			}
-
+ 
 			break;
 		}
-
+ 
 		case mu::EParameterType::Image:
 		{
-			for (const FCustomizableObjectAssetParameterValue& TextureParameter : TextureParameters)
+			for (const FCustomizableObjectTextureParameterValue& TextureParameter : TextureParameters)
 			{
 				if (TextureParameter.ParameterName == Name || (Uid.IsValid() && TextureParameter.Id == Uid))
 				{
@@ -752,17 +783,17 @@ TSharedPtr<mu::FParameters> FCustomizableObjectInstanceDescriptor::GetParameters
 					{
 						MutableParameters->SetImageValue(ParamIndex, TextureParameter.ParameterValue);
 					}
-
+ 
 					break;
 				}
 			}
-
+ 
 			break;
 		}
-
+ 
 		case mu::EParameterType::Mesh:
 		{
-			for (const FCustomizableObjectAssetParameterValue& MeshParameter : MeshParameters)
+			for (const FCustomizableObjectSkeletalMeshParameterValue& MeshParameter : SkeletalMeshParameters)
 			{
 				if (MeshParameter.ParameterName == Name || (Uid.IsValid() && MeshParameter.Id == Uid))
 				{
@@ -778,14 +809,14 @@ TSharedPtr<mu::FParameters> FCustomizableObjectInstanceDescriptor::GetParameters
 					{
 						MutableParameters->SetMeshValue(ParamIndex, MeshParameter.ParameterValue);
 					}
-
+ 
 					break;
 				}
 			}
-
+ 
 			break;
 		}
-
+ 
 		default:
 			check(false); // Missing case.
 			break;
@@ -794,55 +825,55 @@ TSharedPtr<mu::FParameters> FCustomizableObjectInstanceDescriptor::GetParameters
 	
 	return MutableParameters;
 }
-
-
+ 
+ 
 FString FCustomizableObjectInstanceDescriptor::ToString() const
 {
 	const UScriptStruct* Struct = StaticStruct();
 	FString ExportedText;
-
+ 
 	Struct->ExportText(ExportedText, this, nullptr, nullptr, (PPF_ExportsNotFullyQualified | PPF_Copy | PPF_Delimited | PPF_IncludeTransient), nullptr);
-
+ 
 	return ExportedText;
 }
-
-
+ 
+ 
 void FCustomizableObjectInstanceDescriptor::ReloadParameters()
 {
 	if (IsRunningCookCommandlet())
 	{
 		return;
 	}
-
+ 
 	if (!CustomizableObject)
 	{
 		CustomizableObjectNullErrorMessage();
 		return;
 	}
-
+ 
 	if (!CustomizableObject->IsCompiled())
 	{
 		return;
 	}
-
+ 
 	SetState(FMath::Clamp(GetState(), 0, CustomizableObject->GetStateCount() - 1));
-
+ 
 	FirstRequestedLOD.Empty();
-
+ 
 	TArray<FCustomizableObjectBoolParameterValue> OldBoolParameters = BoolParameters;
 	TArray<FCustomizableObjectIntParameterValue> OldIntParameters = IntParameters;
 	TArray<FCustomizableObjectFloatParameterValue> OldFloatParameters = FloatParameters;
-	TArray<FCustomizableObjectAssetParameterValue> OldTextureParameters = TextureParameters;
-	TArray<FCustomizableObjectAssetParameterValue> OldMeshParameters = MeshParameters;
+	TArray<FCustomizableObjectTextureParameterValue> OldTextureParameters = TextureParameters;
+	TArray<FCustomizableObjectSkeletalMeshParameterValue> OldSkeletalMeshParameters = SkeletalMeshParameters;
 	TArray<FCustomizableObjectVectorParameterValue> OldVectorParameters = VectorParameters;
 	TArray<FCustomizableObjectProjectorParameterValue> OldProjectorParameters = ProjectorParameters;
 	TArray<FCustomizableObjectTransformParameterValue> OldTransformParameters = TransformParameters;
-
+ 
 	BoolParameters.Reset();
 	IntParameters.Reset();
 	FloatParameters.Reset();
 	TextureParameters.Reset();
-	MeshParameters.Reset();
+	SkeletalMeshParameters.Reset();
 	VectorParameters.Reset();
 	ProjectorParameters.Reset();
 	TransformParameters.Reset();
@@ -852,16 +883,22 @@ void FCustomizableObjectInstanceDescriptor::ReloadParameters()
 		UE_LOG(LogMutable, Warning, TEXT("[ReloadParametersFromObject] No model in object [%s], generated empty parameters for [%s] "), *CustomizableObject->GetName(), *CustomizableObject->GetName());
 		return;
 	}
-
-	TSharedPtr<mu::FParameters> MutableParameters = mu::FModel::NewParameters(CustomizableObject->GetPrivate()->GetModel());
-
+ 
+	const UModelResources& ModelResources = CustomizableObject->GetPrivate()->GetModelResourcesChecked();
+	
+	TSharedPtr<mu::FParameters> MutableParameters = mu::FModel::NewParameters(
+		CustomizableObject->GetPrivate()->GetModel(),
+		nullptr,
+		&ModelResources.TextureParameterDefaultValues,
+		&ModelResources.SkeletalMeshParameterDefaultValues);
+ 
 	int32 ParamCount = MutableParameters->GetCount();
 	for (int32 ParamIndex = 0; ParamIndex < ParamCount; ++ParamIndex)
 	{
 		const FString& Name = MutableParameters->GetName(ParamIndex);
 		const FGuid& Uid = MutableParameters->GetUid(ParamIndex);
 		const mu::EParameterType MutableType = MutableParameters->GetType(ParamIndex);
-
+ 
 		switch (MutableType)
 		{
 		case mu::EParameterType::Bool:
@@ -869,12 +906,12 @@ void FCustomizableObjectInstanceDescriptor::ReloadParameters()
 			FCustomizableObjectBoolParameterValue Param;
 			Param.ParameterName = Name;
 			Param.Id = Uid;
-
+ 
 			auto FindByNameAndUid = [&](const FCustomizableObjectBoolParameterValue& P)
 			{
 				return P.ParameterName == Name || (Uid.IsValid() && P.Id == Uid);
 			};
-
+ 
 			if (FCustomizableObjectBoolParameterValue* Result = OldBoolParameters.FindByPredicate(FindByNameAndUid))
 			{	
 				Param.ParameterValue = Result->ParameterValue;
@@ -883,17 +920,17 @@ void FCustomizableObjectInstanceDescriptor::ReloadParameters()
 			{
 				Param.ParameterValue = MutableParameters->GetBoolValue(ParamIndex);
 			}
-
+ 
 			BoolParameters.Add(Param);
 			break;
 		}
-
+ 
 		case mu::EParameterType::Int:
 		{
 			FCustomizableObjectIntParameterValue Param;
 			Param.ParameterName = Name;
 			Param.Id = Uid;
-
+ 
 			auto FindByNameAndUid = [&](const FCustomizableObjectIntParameterValue& P)
 			{
 				return P.ParameterName == Name || (Uid.IsValid() && P.Id == Uid);
@@ -902,7 +939,7 @@ void FCustomizableObjectInstanceDescriptor::ReloadParameters()
 			if (FCustomizableObjectIntParameterValue* Result = OldIntParameters.FindByPredicate(FindByNameAndUid))
 			{
 				const int32 NumValueIndex = MutableParameters->GetIntPossibleValueCount(ParamIndex);
-
+ 
 				auto ValueExists = [&](const FString& ValueName)
 				{
 					for (int32 ValueIndex = 0; ValueIndex < NumValueIndex; ++ValueIndex)
@@ -912,7 +949,7 @@ void FCustomizableObjectInstanceDescriptor::ReloadParameters()
 							return true;
 						}
 					}
-
+ 
 					return false;
 				};
 				
@@ -921,7 +958,7 @@ void FCustomizableObjectInstanceDescriptor::ReloadParameters()
 					// Get num of ranges (layers) from the instance
 					int32 ValueCount = Result->ParameterRangeValueNames.Num();
 					Param.ParameterRangeValueNames.Reserve(ValueCount);
-
+ 
 					for (int32 RangeIndex = 0; RangeIndex < ValueCount; ++RangeIndex)
 					{
 						// Checking if the selected value still exists as option in the parameter
@@ -948,7 +985,7 @@ void FCustomizableObjectInstanceDescriptor::ReloadParameters()
 						const int32 ParamValue = MutableParameters->GetIntValue(ParamIndex);
 						Param.ParameterValueName = CustomizableObject->GetPrivate()->FindIntParameterValueName(ParamIndex, ParamValue);
 					}
-
+ 
 					// Multilayer ints with one option are not multidimensional parameters. However, we need to preserve the layer 
 					// information in case that we add a new option to the parameter, and it is converted to multidimensional.
 					for (int32 RangeIndex = 0; RangeIndex < Result->ParameterRangeValueNames.Num(); ++RangeIndex)
@@ -964,17 +1001,17 @@ void FCustomizableObjectInstanceDescriptor::ReloadParameters()
 				if (TSharedPtr<mu::FRangeIndex> RangeIdxPtr = MutableParameters->NewRangeIndex(ParamIndex))
 				{
 					const int32 ValueCount = MutableParameters->GetValueCount(ParamIndex);
-
+ 
 					for (int32 ValueIndex = 0; ValueIndex < ValueCount; ++ValueIndex)
 					{	
 						const TSharedPtr<mu::FRangeIndex> RangeValueIdxPtr = MutableParameters->GetValueIndex(ParamIndex, ValueIndex);
 						const int32 RangeIndex = RangeValueIdxPtr->GetPosition(0);
-
+ 
 						if (!Param.ParameterRangeValueNames.IsValidIndex(RangeIndex))
 						{
 							Param.ParameterRangeValueNames.AddDefaulted(RangeIndex + 1 - Param.ParameterRangeValueNames.Num());
 						}
-
+ 
 						const int32 Value = MutableParameters->GetIntValue(ParamIndex, RangeValueIdxPtr.Get());
 						const FString AuxParameterValueName = CustomizableObject->GetPrivate()->FindIntParameterValueName(ParamIndex, Value);
 						Param.ParameterRangeValueNames[RangeIndex] = AuxParameterValueName;
@@ -986,18 +1023,18 @@ void FCustomizableObjectInstanceDescriptor::ReloadParameters()
 					Param.ParameterValueName = CustomizableObject->GetPrivate()->FindIntParameterValueName(ParamIndex, ParamValue);
 				}
 			}
-
+ 
 			IntParameters.Add(Param);
-
+ 
 			break;
 		}
-
+ 
 		case mu::EParameterType::Float:
 		{
 			FCustomizableObjectFloatParameterValue Param;
 			Param.ParameterName = Name;
 			Param.Id = Uid;
-
+ 
 			auto FindByNameAndUid = [&](const FCustomizableObjectFloatParameterValue& P)
 			{
 				return P.ParameterName == Name || (Uid.IsValid() && P.Id == Uid);
@@ -1019,17 +1056,17 @@ void FCustomizableObjectInstanceDescriptor::ReloadParameters()
 				if (TSharedPtr<mu::FRangeIndex> RangeIdxPtr = MutableParameters->NewRangeIndex(ParamIndex))
 				{
 					const int32 ValueCount = MutableParameters->GetValueCount(ParamIndex);
-
+ 
 					for (int32 ValueIndex = 0; ValueIndex < ValueCount; ++ValueIndex)
 					{	
 						TSharedPtr<mu::FRangeIndex> RangeValueIdxPtr = MutableParameters->GetValueIndex(ParamIndex, ValueIndex);
 						int32 RangeIndex = RangeValueIdxPtr->GetPosition(0);
-
+ 
 						if (!Param.ParameterRangeValues.IsValidIndex(RangeIndex))
 						{
 							Param.ParameterRangeValues.AddDefaulted(RangeIndex + 1 - Param.ParameterRangeValues.Num());
 						}
-
+ 
 						Param.ParameterRangeValues[RangeIndex] = MutableParameters->GetFloatValue(ParamIndex, RangeValueIdxPtr.Get());
 					}
 				}
@@ -1038,22 +1075,22 @@ void FCustomizableObjectInstanceDescriptor::ReloadParameters()
 					Param.ParameterValue = MutableParameters->GetFloatValue(ParamIndex);
 				}
 			}
-
+ 
 			FloatParameters.Add(Param);
 			break;
 		}
-
+ 
 		case mu::EParameterType::Color:
 		{
 			FCustomizableObjectVectorParameterValue Param;
 			Param.ParameterName = Name;
 			Param.Id = Uid;
-
+ 
 			auto FindByNameAndUid = [&](const FCustomizableObjectVectorParameterValue& P)
 			{
 				return P.ParameterName == Name || (Uid.IsValid() && P.Id == Uid);
 			};
-
+ 
 			if (FCustomizableObjectVectorParameterValue* Result = OldVectorParameters.FindByPredicate(FindByNameAndUid))
 			{	
 				Param.ParameterValue = Result->ParameterValue;
@@ -1064,22 +1101,22 @@ void FCustomizableObjectInstanceDescriptor::ReloadParameters()
 				MutableParameters->GetColourValue(ParamIndex, V);
 				Param.ParameterValue = V;
 			}
-
+ 
 			VectorParameters.Add(Param);
 			break;
 		}
-
+ 
 		case mu::EParameterType::Matrix:
 			{
 				FCustomizableObjectTransformParameterValue Param;
 				Param.ParameterName = Name;
 				Param.Id = Uid;
-
+ 
 				auto FindByNameAndUid = [&](const FCustomizableObjectTransformParameterValue& P)
 				{
 					return P.ParameterName == Name || (Uid.IsValid() && P.Id == Uid);
 				};
-
+ 
 				if (FCustomizableObjectTransformParameterValue* Result = OldTransformParameters.FindByPredicate(FindByNameAndUid))
 				{	
 					Param.ParameterValue = Result->ParameterValue;
@@ -1090,7 +1127,7 @@ void FCustomizableObjectInstanceDescriptor::ReloadParameters()
 					MutableParameters->GetMatrixValue(ParamIndex, Matrix);
 					Param.ParameterValue = FTransform(FMatrix(Matrix));
 				}
-
+ 
 				TransformParameters.Add(Param);
 				break;
 			}
@@ -1100,10 +1137,10 @@ void FCustomizableObjectInstanceDescriptor::ReloadParameters()
 			FCustomizableObjectProjectorParameterValue Param;
 			Param.ParameterName = Name;
 			Param.Id = Uid;
-
+ 
 			// Projector to check if the porojector's type has changed
 			FCustomizableObjectProjector DefaultProjectorValue = CustomizableObject->GetProjectorParameterDefaultValue(Name);
-
+ 
 			auto FindByNameAndUid = [&](const FCustomizableObjectProjectorParameterValue& P)
 			{
 				return P.ParameterName == Name || (Uid.IsValid() && P.Id == Uid);
@@ -1115,7 +1152,7 @@ void FCustomizableObjectInstanceDescriptor::ReloadParameters()
 				{
 					Param.RangeValues = Result->RangeValues;
 					Param.Value.ProjectionType = DefaultProjectorValue.ProjectionType;
-
+ 
 					for (FCustomizableObjectProjector& Projector : Param.RangeValues)
 					{
 						Projector.ProjectionType = DefaultProjectorValue.ProjectionType;
@@ -1152,44 +1189,44 @@ void FCustomizableObjectInstanceDescriptor::ReloadParameters()
 						Value.Scale[0] = Value.Scale[1] = Value.Scale[1] * 2.0f;
 					}
 				};
-
+ 
 				GetProjector(Param.Value);
-
+ 
 				if (TSharedPtr<mu::FRangeIndex> RangeIdxPtr = MutableParameters->NewRangeIndex(ParamIndex))
 				{
 					const int32 ValueCount = MutableParameters->GetValueCount(ParamIndex);
-
+ 
 					for (int32 ValueIndex = 0; ValueIndex < ValueCount; ++ValueIndex)
 					{	
 						TSharedPtr<mu::FRangeIndex> RangeValueIdxPtr = MutableParameters->GetValueIndex(ParamIndex, ValueIndex);
 						int32 RangeIndex = RangeValueIdxPtr->GetPosition(0);
-
+ 
 						if (!Param.RangeValues.IsValidIndex(RangeIndex))
 						{
 							Param.RangeValues.AddDefaulted(RangeIndex + 1 - Param.RangeValues.Num());
 						}
-
+ 
 						GetProjector(Param.RangeValues[RangeIndex], RangeValueIdxPtr);
 					}
 				}
 			}
-
+ 
 			ProjectorParameters.Add(Param);
 			break;
 		}
-
+ 
 		case mu::EParameterType::Image:
 		{
-			FCustomizableObjectAssetParameterValue Param;
+			FCustomizableObjectTextureParameterValue Param;
 			Param.ParameterName = Name;
 			Param.Id = Uid;
-
-			auto FindByNameAndUid = [&](const FCustomizableObjectAssetParameterValue& P)
+ 
+			auto FindByNameAndUid = [&](const FCustomizableObjectTextureParameterValue& P)
 			{
 				return P.ParameterName == Name || (Uid.IsValid() && P.Id == Uid);
 			};
 			
-			if (FCustomizableObjectAssetParameterValue* Result = OldTextureParameters.FindByPredicate(FindByNameAndUid))
+			if (FCustomizableObjectTextureParameterValue* Result = OldTextureParameters.FindByPredicate(FindByNameAndUid))
 			{	
 				if (TSharedPtr<mu::FRangeIndex> RangeIdxPtr = MutableParameters->NewRangeIndex(ParamIndex))
 				{
@@ -1205,23 +1242,23 @@ void FCustomizableObjectInstanceDescriptor::ReloadParameters()
 				if (TSharedPtr<mu::FRangeIndex> RangeIdxPtr = MutableParameters->NewRangeIndex(ParamIndex))
 				{
 					const int32 ValueCount = MutableParameters->GetValueCount(ParamIndex);
-
+ 
 					for (int32 ValueIndex = 0; ValueIndex < ValueCount; ++ValueIndex)
 					{	
 						TSharedPtr<mu::FRangeIndex> RangeValueIdxPtr = MutableParameters->GetValueIndex(ParamIndex, ValueIndex);
 						int32 RangeIndex = RangeValueIdxPtr->GetPosition(0);
-
+ 
 						if (!Param.ParameterRangeValues.IsValidIndex(RangeIndex))
 						{
 							const int32 PreviousNum = Param.ParameterRangeValues.Num();
 							Param.ParameterRangeValues.AddUninitialized(RangeIndex + 1 - Param.ParameterRangeValues.Num());
-
+ 
 							for (int32 Index = PreviousNum; Index < Param.ParameterRangeValues.Num(); ++Index)
 							{
-								Param.ParameterRangeValues[Index] = FName();
+								Param.ParameterRangeValues[Index] = nullptr;
 							}
 						}
-
+ 
 						Param.ParameterRangeValues[RangeIndex] = MutableParameters->GetImageValue(ParamIndex, RangeValueIdxPtr.Get());
 					}
 				}
@@ -1230,24 +1267,24 @@ void FCustomizableObjectInstanceDescriptor::ReloadParameters()
 					Param.ParameterValue = MutableParameters->GetImageValue(ParamIndex);
 				}
 			}
-
+ 
 			TextureParameters.Add(Param);
 			break;
 		}
-
-
+ 
+ 
 		case mu::EParameterType::Mesh:
 		{
-			FCustomizableObjectAssetParameterValue Param;
+			FCustomizableObjectSkeletalMeshParameterValue Param;
 			Param.ParameterName = Name;
 			Param.Id = Uid;
-
-			auto FindByNameAndUid = [&](const FCustomizableObjectAssetParameterValue& P)
+ 
+			auto FindByNameAndUid = [&](const FCustomizableObjectSkeletalMeshParameterValue& P)
 				{
 					return P.ParameterName == Name || (Uid.IsValid() && P.Id == Uid);
 				};
-
-			if (FCustomizableObjectAssetParameterValue* Result = OldTextureParameters.FindByPredicate(FindByNameAndUid))
+ 
+			if (FCustomizableObjectSkeletalMeshParameterValue* Result = OldSkeletalMeshParameters.FindByPredicate(FindByNameAndUid))
 			{
 				if (TSharedPtr<mu::FRangeIndex> RangeIdxPtr = MutableParameters->NewRangeIndex(ParamIndex))
 				{
@@ -1263,23 +1300,23 @@ void FCustomizableObjectInstanceDescriptor::ReloadParameters()
 				if (TSharedPtr<mu::FRangeIndex> RangeIdxPtr = MutableParameters->NewRangeIndex(ParamIndex))
 				{
 					const int32 ValueCount = MutableParameters->GetValueCount(ParamIndex);
-
+ 
 					for (int32 ValueIndex = 0; ValueIndex < ValueCount; ++ValueIndex)
 					{
 						TSharedPtr<mu::FRangeIndex> RangeValueIdxPtr = MutableParameters->GetValueIndex(ParamIndex, ValueIndex);
 						int32 RangeIndex = RangeValueIdxPtr->GetPosition(0);
-
+ 
 						if (!Param.ParameterRangeValues.IsValidIndex(RangeIndex))
 						{
 							const int32 PreviousNum = Param.ParameterRangeValues.Num();
 							Param.ParameterRangeValues.AddUninitialized(RangeIndex + 1 - Param.ParameterRangeValues.Num());
-
+ 
 							for (int32 Index = PreviousNum; Index < Param.ParameterRangeValues.Num(); ++Index)
 							{
-								Param.ParameterRangeValues[Index] = FName();
+								Param.ParameterRangeValues[Index] = nullptr;
 							}
 						}
-
+ 
 						Param.ParameterRangeValues[RangeIndex] = MutableParameters->GetMeshValue(ParamIndex, RangeValueIdxPtr.Get());
 					}
 				}
@@ -1288,78 +1325,31 @@ void FCustomizableObjectInstanceDescriptor::ReloadParameters()
 					Param.ParameterValue = MutableParameters->GetMeshValue(ParamIndex);
 				}
 			}
-
-			MeshParameters.Add(Param);
+ 
+			SkeletalMeshParameters.Add(Param);
 			break;
 		}
-
+ 
 		default:
 			check(false); // Missing case.
 			break;
 		}
 	}
 }
-
-
+ 
+ 
 void FCustomizableObjectInstanceDescriptor::SetFirstRequestedLOD(const TMap<FName, uint8>& InFirstRequestedLOD)
 {
 	FirstRequestedLOD = InFirstRequestedLOD;
 }
-
-
+ 
+ 
 const TMap<FName, uint8>& FCustomizableObjectInstanceDescriptor::GetFirstRequestedLOD() const
 {
 	return FirstRequestedLOD;
 }
-
-
-const TArray<FCustomizableObjectBoolParameterValue>& FCustomizableObjectInstanceDescriptor::GetBoolParameters() const
-{
-	return BoolParameters;
-}
-
-
-const TArray<FCustomizableObjectIntParameterValue>& FCustomizableObjectInstanceDescriptor::GetIntParameters() const
-{
-	return IntParameters;	
-}
-
-
-const TArray<FCustomizableObjectFloatParameterValue>& FCustomizableObjectInstanceDescriptor::GetFloatParameters() const
-{
-	return FloatParameters;	
-}
-
-
-const TArray<FCustomizableObjectAssetParameterValue>& FCustomizableObjectInstanceDescriptor::GetTextureParameters() const
-{
-	return TextureParameters;
-}
-
-
-const TArray<FCustomizableObjectAssetParameterValue>& FCustomizableObjectInstanceDescriptor::GetMeshParameters() const
-{
-	return MeshParameters;
-}
-
-
-const TArray<FCustomizableObjectVectorParameterValue>& FCustomizableObjectInstanceDescriptor::GetVectorParameters() const
-{
-	return VectorParameters;
-}
-
-
-const TArray<FCustomizableObjectProjectorParameterValue>& FCustomizableObjectInstanceDescriptor::GetProjectorParameters() const
-{
-	return ProjectorParameters;	
-}
-
-const TArray<FCustomizableObjectTransformParameterValue>& FCustomizableObjectInstanceDescriptor::GetTransformParameters() const
-{
-	return TransformParameters;
-}
-
-
+ 
+ 
 bool FCustomizableObjectInstanceDescriptor::HasAnyParameters() const
 {
 	return !BoolParameters.IsEmpty() ||
@@ -1370,9 +1360,9 @@ bool FCustomizableObjectInstanceDescriptor::HasAnyParameters() const
 		!TransformParameters.IsEmpty() ||
 		!VectorParameters.IsEmpty();
 }
-
+ 
 #if WITH_EDITOR 
-
+ 
 #define RETURN_ON_UNCOMPILED_CO(CustomizableObject, ErrorMessage) \
 	if (!CustomizableObject->IsCompiled()) \
 	{ \
@@ -1380,9 +1370,9 @@ bool FCustomizableObjectInstanceDescriptor::HasAnyParameters() const
 		CustomizableObject->GetPrivate()->AddUncompiledCOWarning(AdditionalLoggingInfo);\
 		return; \
 	} \
-
+ 
 #else
-
+ 
 #define RETURN_ON_UNCOMPILED_CO(CustomizableObject, ErrorMessage) \
 	if (!ensureMsgf(CustomizableObject->IsCompiled(), TEXT("Customizable Object (%s) was not compiled."), *GetNameSafe(CustomizableObject))) \
 	{ \
@@ -1393,8 +1383,8 @@ bool FCustomizableObjectInstanceDescriptor::HasAnyParameters() const
  
 #endif
 	
-
-
+ 
+ 
 void LogParameterNotFoundWarning(const FString& ParameterName, const int32 ObjectParameterIndex, const int32 InstanceParameterIndex, const UCustomizableObject* CustomizableObject, char const* CallingFunction)
 {
 	UE_LOG(LogMutable, Error,
@@ -1402,8 +1392,8 @@ void LogParameterNotFoundWarning(const FString& ParameterName, const int32 Objec
 		CallingFunction, *ParameterName, *GetNameSafe(CustomizableObject), ObjectParameterIndex, InstanceParameterIndex
 	);
 }
-
-
+ 
+ 
 const FString& FCustomizableObjectInstanceDescriptor::GetIntParameterSelectedOption(const FString& ParamName, const int32 RangeIndex) const
 {
 	if (!CustomizableObject)
@@ -1414,7 +1404,7 @@ const FString& FCustomizableObjectInstanceDescriptor::GetIntParameterSelectedOpt
 	
 	const int32 ParameterIndexInObject = CustomizableObject->GetPrivate()->FindParameter(ParamName);
 	const int32 ParameterIndexInInstance = FindTypedParameterIndex(ParamName, EMutableParameterType::Int);
-
+ 
 	if (ParameterIndexInObject >= 0 && IntParameters.IsValidIndex(ParameterIndexInInstance))
 	{
 		// Due to optimizations a parameter may lose its multidimensionality (if it becomes constant). 
@@ -1432,13 +1422,13 @@ const FString& FCustomizableObjectInstanceDescriptor::GetIntParameterSelectedOpt
 			}
 		}
 	}
-
+ 
 	LogParameterNotFoundWarning(ParamName, ParameterIndexInObject, ParameterIndexInInstance, CustomizableObject, __FUNCTION__);
-
+ 
 	return FCustomizableObjectBoolParameterValue::DEFAULT_PARAMETER_VALUE_NAME;
 }
-
-
+ 
+ 
 void FCustomizableObjectInstanceDescriptor::SetIntParameterSelectedOption(const int32 ParameterIndexInInstance, const FString& SelectedOption, const int32 RangeIndex)
 {
 	if (!CustomizableObject)
@@ -1446,18 +1436,18 @@ void FCustomizableObjectInstanceDescriptor::SetIntParameterSelectedOption(const
 		CustomizableObjectNullErrorMessage();
 		return;
 	}
-
+ 
 	RETURN_ON_UNCOMPILED_CO(CustomizableObject, TEXT("Error: Cannot set Int parameter "));
-
+ 
 	const int32 ParameterIndexInObject = IntParameters.IsValidIndex(ParameterIndexInInstance) ? CustomizableObject->GetPrivate()->FindParameter(IntParameters[ParameterIndexInInstance].ParameterName) : INDEX_NONE; //-V781
-
+ 
 	if (ParameterIndexInObject < 0 || ParameterIndexInInstance < 0)
 	{
 		// Warn and early out since we could not find the parameter to set.
 		LogParameterNotFoundWarning(TEXT("Unknown Int Parmeter"), ParameterIndexInObject, ParameterIndexInInstance, CustomizableObject, __FUNCTION__);
 		return;
 	}
-
+ 
 	const bool bValid = SelectedOption == TEXT("None") || CustomizableObject->GetPrivate()->FindIntParameterValue(ParameterIndexInObject, SelectedOption) >= 0;
 	if (!bValid)
 	{
@@ -1473,7 +1463,7 @@ void FCustomizableObjectInstanceDescriptor::SetIntParameterSelectedOption(const
 #endif
 		return;
 	}
-
+ 
 	// Due to optimizations a parameter may lose its multidimensionality (if it becomes constant). 
 	// In that case it means there is only one possible option so it is ok to set it as if wasn't multidimensional.
 	bool bIsMultidimensional = CustomizableObject->GetPrivate()->IsParameterMultidimensional(ParameterIndexInObject);
@@ -1490,13 +1480,13 @@ void FCustomizableObjectInstanceDescriptor::SetIntParameterSelectedOption(const
 			const int32 NumInsertedElements = RangeIndex + 1 - InsertionIndex;
 			IntParameters[ParameterIndexInInstance].ParameterRangeValueNames.InsertDefaulted(InsertionIndex, NumInsertedElements);
 		}
-
+ 
 		check(IntParameters[ParameterIndexInInstance].ParameterRangeValueNames.IsValidIndex(RangeIndex));
 		IntParameters[ParameterIndexInInstance].ParameterRangeValueNames[RangeIndex] = SelectedOption;
 	}
 }
-
-
+ 
+ 
 void FCustomizableObjectInstanceDescriptor::SetIntParameterSelectedOption(const FString& ParamName, const FString& SelectedOptionName, const int32 RangeIndex)
 {
 	if (!CustomizableObject)
@@ -1504,9 +1494,9 @@ void FCustomizableObjectInstanceDescriptor::SetIntParameterSelectedOption(const
 		CustomizableObjectNullErrorMessage();
 		return;
 	}
-
+ 
 	RETURN_ON_UNCOMPILED_CO(CustomizableObject, TEXT("Error: Cannot set Int parameter "));
-
+ 
 	const int32 ParamIndexInInstance = FindTypedParameterIndex(ParamName, EMutableParameterType::Int);
 	if (ParamIndexInInstance == INDEX_NONE)
 	{
@@ -1514,11 +1504,11 @@ void FCustomizableObjectInstanceDescriptor::SetIntParameterSelectedOption(const
 		LogParameterNotFoundWarning(ParamName, ParamIndexInInstance, ParamIndexInInstance, CustomizableObject, __FUNCTION__);
 		return;
 	}
-
+ 
 	SetIntParameterSelectedOption(ParamIndexInInstance, SelectedOptionName, RangeIndex);
 }
-
-
+ 
+ 
 float FCustomizableObjectInstanceDescriptor::GetFloatParameterSelectedOption(const FString& FloatParamName, const int32 RangeIndex) const
 {
 	if (!CustomizableObject)
@@ -1526,10 +1516,10 @@ float FCustomizableObjectInstanceDescriptor::GetFloatParameterSelectedOption(con
 		CustomizableObjectNullErrorMessage();
 		return FCustomizableObjectFloatParameterValue::DEFAULT_PARAMETER_VALUE;
 	}
-
+ 
 	const int32 ParameterIndexInObject = CustomizableObject->GetPrivate()->FindParameter(FloatParamName);
 	const int32 FloatParamIndex = FindTypedParameterIndex(FloatParamName, EMutableParameterType::Float);
-
+ 
 	if (ParameterIndexInObject >= 0 && FloatParameters.IsValidIndex(FloatParamIndex))
 	{
 		// Due to optimizations a parameter may lose its multidimensionality (if it becomes constant). 
@@ -1547,13 +1537,13 @@ float FCustomizableObjectInstanceDescriptor::GetFloatParameterSelectedOption(con
 			}
 		}
 	}
-
+ 
 	LogParameterNotFoundWarning(FloatParamName, ParameterIndexInObject, FloatParamIndex, CustomizableObject, __FUNCTION__);
 	
 	return FCustomizableObjectFloatParameterValue::DEFAULT_PARAMETER_VALUE; 
 }
-
-
+ 
+ 
 void FCustomizableObjectInstanceDescriptor::SetFloatParameterSelectedOption(const FString& FloatParamName, const float FloatValue, const int32 RangeIndex)
 {
 	if (!CustomizableObject)
@@ -1561,19 +1551,19 @@ void FCustomizableObjectInstanceDescriptor::SetFloatParameterSelectedOption(cons
 		CustomizableObjectNullErrorMessage();
 		return;
 	}
-
+ 
 	RETURN_ON_UNCOMPILED_CO(CustomizableObject, TEXT("Error: Cannot set Int parameter "));
-
+ 
 	const int32 ParameterIndexInObject = CustomizableObject->GetPrivate()->FindParameter(FloatParamName);
 	const int32 ParameterIndexInInstance = FindTypedParameterIndex(FloatParamName, EMutableParameterType::Float);
-
+ 
 	if (ParameterIndexInObject < 0 || ParameterIndexInInstance < 0)
 	{
 		// Warn and early out since we could not find the parameter to set.
 		LogParameterNotFoundWarning(FloatParamName, ParameterIndexInObject, ParameterIndexInInstance, CustomizableObject, __FUNCTION__);
 		return;
 	}
-
+ 
 	// Due to optimizations a parameter may lose its multidimensionality (if it becomes constant). 
 	// In that case it means there is only one possible option so it is ok to set it as if wasn't multidimensional.
 	bool bIsMultidimensional = CustomizableObject->GetPrivate()->IsParameterMultidimensional(ParameterIndexInObject);
@@ -1589,25 +1579,25 @@ void FCustomizableObjectInstanceDescriptor::SetFloatParameterSelectedOption(cons
 			const int32 NumInsertedElements = RangeIndex + 1 - FloatParameters[ParameterIndexInInstance].ParameterRangeValues.Num();
 			FloatParameters[ParameterIndexInInstance].ParameterRangeValues.InsertDefaulted(InsertionIndex, NumInsertedElements);
 		}
-
+ 
 		check(FloatParameters[ParameterIndexInInstance].ParameterRangeValues.IsValidIndex(RangeIndex));
 		FloatParameters[ParameterIndexInInstance].ParameterRangeValues[RangeIndex] = FloatValue;
 	}
-
+ 
 }
-
-
-FName FCustomizableObjectInstanceDescriptor::GetTextureParameterSelectedOption(const FString& TextureParamName, const int32 RangeIndex) const
+ 
+ 
+UTexture* FCustomizableObjectInstanceDescriptor::GetTextureParameterSelectedOption(const FString& TextureParamName, const int32 RangeIndex) const
 {
 	if (!CustomizableObject)
 	{
 		CustomizableObjectNullErrorMessage();
-		return FName();
+		return nullptr;
 	}
-
+ 
 	const int32 ParameterIndexInObject = CustomizableObject->GetPrivate()->FindParameter(TextureParamName);
 	const int32 ParameterIndexInInstance = FindTypedParameterIndex(TextureParamName, EMutableParameterType::Texture);
-
+ 
 	if (ParameterIndexInObject >= 0 && ParameterIndexInInstance >= 0)
 	{
 		// Due to optimizations a parameter may lose its multidimensionality (if it becomes constant). 
@@ -1627,37 +1617,37 @@ FName FCustomizableObjectInstanceDescriptor::GetTextureParameterSelectedOption(c
 	}
 	
 	LogParameterNotFoundWarning(TextureParamName, ParameterIndexInObject, ParameterIndexInInstance, CustomizableObject, __FUNCTION__);
-
-	return FName(); 
+ 
+	return nullptr;
 }
-
-
-void FCustomizableObjectInstanceDescriptor::SetTextureParameterSelectedOption(const FString& TextureParamName, const FString& TextureValue, const int32 RangeIndex)
+ 
+ 
+void FCustomizableObjectInstanceDescriptor::SetTextureParameterSelectedOption(const FString& TextureParamName, UTexture* TextureValue, const int32 RangeIndex)
 {
 	if (!CustomizableObject)
 	{
 		CustomizableObjectNullErrorMessage();
 		return;
 	}
-
-	RETURN_ON_UNCOMPILED_CO(CustomizableObject, TEXT("Error: Cannot set Int parameter "));
-
+ 
+	RETURN_ON_UNCOMPILED_CO(CustomizableObject, TEXT("Error: Cannot set Texture parameter "));
+ 
 	const int32 ParameterIndexInObject = CustomizableObject->GetPrivate()->FindParameter(TextureParamName);
 	const int32 ParameterIndexInInstance = FindTypedParameterIndex(TextureParamName, EMutableParameterType::Texture);
-
+ 
 	if (ParameterIndexInObject < 0 || ParameterIndexInInstance < 0)
 	{
 		// Early out since we could not find the parameter to set.
 		LogParameterNotFoundWarning(TextureParamName, ParameterIndexInObject, ParameterIndexInInstance, CustomizableObject, __FUNCTION__);
 		return;
 	}
-
+ 
 	// Due to optimizations a parameter may lose its multidimensionality (if it becomes constant). 
 	// In that case it means there is only one possible option so it is ok to set it as if wasn't multidimensional.
 	bool bIsMultidimensional = CustomizableObject->GetPrivate()->IsParameterMultidimensional(ParameterIndexInObject);
 	if (!bIsMultidimensional || RangeIndex == INDEX_NONE)
 	{
-		TextureParameters[ParameterIndexInInstance].ParameterValue = FName(TextureValue);
+		TextureParameters[ParameterIndexInInstance].ParameterValue = TextureValue;
 	}
 	else
 	{
@@ -1667,13 +1657,90 @@ void FCustomizableObjectInstanceDescriptor::SetTextureParameterSelectedOption(co
 			const int32 NumInsertedElements = RangeIndex + 1 - TextureParameters[ParameterIndexInInstance].ParameterRangeValues.Num();
 			TextureParameters[ParameterIndexInInstance].ParameterRangeValues.InsertDefaulted(InsertionIndex, NumInsertedElements);
 		}
-
+ 
 		check(TextureParameters[ParameterIndexInInstance].ParameterRangeValues.IsValidIndex(RangeIndex));
-		TextureParameters[ParameterIndexInInstance].ParameterRangeValues[RangeIndex] = FName(TextureValue);
+		TextureParameters[ParameterIndexInInstance].ParameterRangeValues[RangeIndex] = TextureValue;
 	}
 }
-
-
+ 
+ 
+USkeletalMesh* FCustomizableObjectInstanceDescriptor::GetSkeletalMeshParameterSelectedOption(const FString& SkeletalMeshParamName, int32 RangeIndex) const
+{
+	if (!CustomizableObject)
+	{
+		CustomizableObjectNullErrorMessage();
+		return nullptr;
+	}
+ 
+	const int32 ParameterIndexInObject = CustomizableObject->GetPrivate()->FindParameter(SkeletalMeshParamName);
+	const int32 ParameterIndexInInstance = FindTypedParameterIndex(SkeletalMeshParamName, EMutableParameterType::SkeletalMesh);
+ 
+	if (ParameterIndexInObject >= 0 && ParameterIndexInInstance >= 0)
+	{
+		// Due to optimizations a parameter may lose its multidimensionality (if it becomes constant). 
+		// In that case it means there is only one possible option so it is ok to set it as if wasn't multidimensional.
+		bool bIsMultidimensional = CustomizableObject->GetPrivate()->IsParameterMultidimensional(ParameterIndexInObject);
+		if (!bIsMultidimensional || RangeIndex == INDEX_NONE)
+		{
+			return SkeletalMeshParameters[ParameterIndexInInstance].ParameterValue;
+		}
+		else
+		{
+			if (SkeletalMeshParameters[ParameterIndexInInstance].ParameterRangeValues.IsValidIndex(RangeIndex))
+			{
+				return SkeletalMeshParameters[ParameterIndexInInstance].ParameterRangeValues[RangeIndex];
+			}
+		}
+	}
+	
+	LogParameterNotFoundWarning(SkeletalMeshParamName, ParameterIndexInObject, ParameterIndexInInstance, CustomizableObject, __FUNCTION__);
+ 
+	return nullptr; 
+}
+ 
+ 
+void FCustomizableObjectInstanceDescriptor::SetSkeletalMeshParameterSelectedOption(const FString& SkeletalMeshParamName, USkeletalMesh* SkeletalMeshValue, int32 RangeIndex)
+{
+	if (!CustomizableObject)
+	{
+		CustomizableObjectNullErrorMessage();
+		return;
+	}
+ 
+	RETURN_ON_UNCOMPILED_CO(CustomizableObject, TEXT("Error: Cannot set Skeletal Mesh parameter "));
+	
+	const int32 ParameterIndexInObject = CustomizableObject->GetPrivate()->FindParameter(SkeletalMeshParamName);
+	const int32 ParameterIndexInInstance = FindTypedParameterIndex(SkeletalMeshParamName, EMutableParameterType::SkeletalMesh);
+ 
+	if (ParameterIndexInObject < 0 || ParameterIndexInInstance < 0)
+	{
+		// Early out since we could not find the parameter to set.
+		LogParameterNotFoundWarning(SkeletalMeshParamName, ParameterIndexInObject, ParameterIndexInInstance, CustomizableObject, __FUNCTION__);
+		return;
+	}
+ 
+	// Due to optimizations a parameter may lose its multidimensionality (if it becomes constant). 
+	// In that case it means there is only one possible option so it is ok to set it as if wasn't multidimensional.
+	bool bIsMultidimensional = CustomizableObject->GetPrivate()->IsParameterMultidimensional(ParameterIndexInObject);
+	if (!bIsMultidimensional || RangeIndex == INDEX_NONE)
+	{
+		SkeletalMeshParameters[ParameterIndexInInstance].ParameterValue = SkeletalMeshValue;
+	}
+	else
+	{
+		if (!SkeletalMeshParameters[ParameterIndexInInstance].ParameterRangeValues.IsValidIndex(RangeIndex))
+		{
+			const int32 InsertionIndex = SkeletalMeshParameters[ParameterIndexInInstance].ParameterRangeValues.Num();
+			const int32 NumInsertedElements = RangeIndex + 1 - SkeletalMeshParameters[ParameterIndexInInstance].ParameterRangeValues.Num();
+			SkeletalMeshParameters[ParameterIndexInInstance].ParameterRangeValues.InsertDefaulted(InsertionIndex, NumInsertedElements);
+		}
+ 
+		check(SkeletalMeshParameters[ParameterIndexInInstance].ParameterRangeValues.IsValidIndex(RangeIndex));
+		SkeletalMeshParameters[ParameterIndexInInstance].ParameterRangeValues[RangeIndex] = SkeletalMeshValue;
+	}
+}
+ 
+ 
 FLinearColor FCustomizableObjectInstanceDescriptor::GetColorParameterSelectedOption(const FString& ColorParamName) const
 {
 	if (!CustomizableObject)
@@ -1681,21 +1748,21 @@ FLinearColor FCustomizableObjectInstanceDescriptor::GetColorParameterSelectedOpt
 		CustomizableObjectNullErrorMessage();
 		return FCustomizableObjectVectorParameterValue::DEFAULT_PARAMETER_VALUE;
 	}
-
+ 
 	const int32 ParameterIndexInObject = CustomizableObject->GetPrivate()->FindParameter(ColorParamName);
 	const int32 ColorParamIndex = FindTypedParameterIndex(ColorParamName, EMutableParameterType::Color);
-
+ 
 	if (ColorParamIndex == INDEX_NONE)
 	{
 		LogParameterNotFoundWarning(ColorParamName, ParameterIndexInObject, ColorParamIndex, CustomizableObject, __FUNCTION__);
 		return FCustomizableObjectVectorParameterValue::DEFAULT_PARAMETER_VALUE;
 	}
-
+ 
 	// TODO: Multidimensional parameter support
 	return VectorParameters.IsValidIndex(ColorParamIndex) ? VectorParameters[ColorParamIndex].ParameterValue : FLinearColor();
 }
-
-
+ 
+ 
 void FCustomizableObjectInstanceDescriptor::SetColorParameterSelectedOption(const FString& ColorParamName, const FLinearColor& ColorValue)
 {
 	if (!CustomizableObject)
@@ -1703,13 +1770,13 @@ void FCustomizableObjectInstanceDescriptor::SetColorParameterSelectedOption(cons
 		CustomizableObjectNullErrorMessage();
 		return;
 	}
-
+ 
 	RETURN_ON_UNCOMPILED_CO(CustomizableObject, TEXT("Error: Cannot set Int parameter "));
-
+ 
 	// TODO: Multidimensional parameter support
 	SetVectorParameterSelectedOption(ColorParamName, ColorValue);
 }
-
+ 
 FTransform FCustomizableObjectInstanceDescriptor::GetTransformParameterSelectedOption(const FString& TransformParamName) const
 {
 	if (!CustomizableObject)
@@ -1717,20 +1784,20 @@ FTransform FCustomizableObjectInstanceDescriptor::GetTransformParameterSelectedO
 		CustomizableObjectNullErrorMessage();
 		return FCustomizableObjectTransformParameterValue::DEFAULT_PARAMETER_VALUE;
 	}
-
+ 
 	const int32 ParameterIndexInObject = CustomizableObject->GetPrivate()->FindParameter(TransformParamName);
 	const int32 TransformParamIndex = FindTypedParameterIndex(TransformParamName, EMutableParameterType::Transform);
-
+ 
 	if (TransformParamIndex == INDEX_NONE)
 	{
 		LogParameterNotFoundWarning(TransformParamName, ParameterIndexInObject, TransformParamIndex, CustomizableObject, __FUNCTION__);
 		return FCustomizableObjectTransformParameterValue::DEFAULT_PARAMETER_VALUE;
 	}
-
+ 
 	// TODO: Multidimensional parameter support
 	return TransformParameters.IsValidIndex(TransformParamIndex) ? TransformParameters[TransformParamIndex].ParameterValue : FTransform::Identity;
 }
-
+ 
 void FCustomizableObjectInstanceDescriptor::SetTransformParameterSelectedOption(const FString& TransformParamName, const FTransform& TransformValue)
 {
 	if (!CustomizableObject)
@@ -1738,24 +1805,24 @@ void FCustomizableObjectInstanceDescriptor::SetTransformParameterSelectedOption(
 		CustomizableObjectNullErrorMessage();
 		return;
 	}
-
+ 
 	RETURN_ON_UNCOMPILED_CO(CustomizableObject, TEXT("Error: Cannot set Transform parameter "));
-
+ 
 	const int32 ParameterIndexInObject = CustomizableObject->GetPrivate()->FindParameter(TransformParamName);
 	const int32 ParameterIndexInInstance = FindTypedParameterIndex(TransformParamName, EMutableParameterType::Transform);
-
+ 
 	if (ParameterIndexInObject < 0 || ParameterIndexInInstance < 0)
 	{
 		// Early out since we could not find the parameter to set.
 		LogParameterNotFoundWarning(TransformParamName, ParameterIndexInObject, ParameterIndexInInstance, CustomizableObject, __FUNCTION__);
 		return;
 	}
-
+ 
 	// TODO: Multidimensional parameter support
 	TransformParameters[ParameterIndexInInstance].ParameterValue = TransformValue;
 }
-
-
+ 
+ 
 bool FCustomizableObjectInstanceDescriptor::GetBoolParameterSelectedOption(const FString& BoolParamName) const
 {
 	if (!CustomizableObject)
@@ -1763,20 +1830,20 @@ bool FCustomizableObjectInstanceDescriptor::GetBoolParameterSelectedOption(const
 		CustomizableObjectNullErrorMessage();
 		return FCustomizableObjectBoolParameterValue::DEFAULT_PARAMETER_VALUE;
 	}
-
+ 
 	const int32 ParameterIndexInObject = CustomizableObject->GetPrivate()->FindParameter(BoolParamName);
 	const int32 BoolParamIndex = FindTypedParameterIndex(BoolParamName, EMutableParameterType::Bool);
-
+ 
 	if (BoolParamIndex == INDEX_NONE)
 	{
 		LogParameterNotFoundWarning(BoolParamName, ParameterIndexInObject, BoolParamIndex, CustomizableObject, __FUNCTION__);
 		return FCustomizableObjectBoolParameterValue::DEFAULT_PARAMETER_VALUE;
 	}
-
+ 
 	return BoolParameters[BoolParamIndex].ParameterValue;
 }
-
-
+ 
+ 
 void FCustomizableObjectInstanceDescriptor::SetBoolParameterSelectedOption(const FString& BoolParamName, const bool BoolValue)
 {
 	if (!CustomizableObject)
@@ -1784,23 +1851,23 @@ void FCustomizableObjectInstanceDescriptor::SetBoolParameterSelectedOption(const
 		CustomizableObjectNullErrorMessage();
 		return;
 	}
-
+ 
 	RETURN_ON_UNCOMPILED_CO(CustomizableObject, TEXT("Error: Cannot set Int parameter "));
-
+ 
 	const int32 ParameterIndexInObject = CustomizableObject->GetPrivate()->FindParameter(BoolParamName);
 	const int32 ParameterIndexInInstance = FindTypedParameterIndex(BoolParamName, EMutableParameterType::Bool);
-
+ 
 	if (ParameterIndexInObject < 0 || ParameterIndexInInstance < 0)
 	{
 		// Early out since we could not find the parameter to set.
 		LogParameterNotFoundWarning(BoolParamName, ParameterIndexInObject, ParameterIndexInInstance, CustomizableObject, __FUNCTION__);
 		return;
 	}
-
+ 
 	BoolParameters[ParameterIndexInInstance].ParameterValue = BoolValue;
 }
-
-
+ 
+ 
 void FCustomizableObjectInstanceDescriptor::SetVectorParameterSelectedOption(const FString& VectorParamName, const FLinearColor& VectorValue)
 {
 	if (!CustomizableObject)
@@ -1808,23 +1875,23 @@ void FCustomizableObjectInstanceDescriptor::SetVectorParameterSelectedOption(con
 		CustomizableObjectNullErrorMessage();
 		return;
 	}
-
+ 
 	RETURN_ON_UNCOMPILED_CO(CustomizableObject, TEXT("Error: Cannot set Int parameter "));
-
+ 
 	const int32 ParameterIndexInObject = CustomizableObject->GetPrivate()->FindParameter(VectorParamName);
 	const int32 ParameterIndexInInstance = FindTypedParameterIndex(VectorParamName, EMutableParameterType::Color);
-
+ 
 	if (ParameterIndexInObject < 0 || ParameterIndexInInstance < 0)
 	{
 		// Early out since we could not find the parameter to set.
 		LogParameterNotFoundWarning(VectorParamName, ParameterIndexInObject, ParameterIndexInInstance, CustomizableObject, __FUNCTION__);
 		return;
 	}
-
+ 
 	VectorParameters[ParameterIndexInInstance].ParameterValue = VectorValue;
 }
-
-
+ 
+ 
 void FCustomizableObjectInstanceDescriptor::SetProjectorValue(const FString& ProjectorParamName,
 	const FVector& Pos, const FVector& Direction, const FVector& Up, const FVector& Scale,
 	const float Angle,
@@ -1835,22 +1902,22 @@ void FCustomizableObjectInstanceDescriptor::SetProjectorValue(const FString& Pro
 		CustomizableObjectNullErrorMessage();
 		return;
 	}
-
+ 
 	RETURN_ON_UNCOMPILED_CO(CustomizableObject, TEXT("Error: Cannot set Projector parameter "))
-
+ 
 	const int32 ParameterIndexInObject = CustomizableObject->GetPrivate()->FindParameter(ProjectorParamName);
 	const int32 ParameterIndexInInstance = FindTypedParameterIndex(ProjectorParamName, EMutableParameterType::Projector);
-
+ 
 	if (ParameterIndexInObject < 0 || ParameterIndexInInstance < 0)
 	{
 		// Early out since we could not find the parameter to set.
 		LogParameterNotFoundWarning(ProjectorParamName, ParameterIndexInObject, ParameterIndexInInstance, CustomizableObject, __FUNCTION__);
 		return;
 	}
-
+ 
 	// Parameter to modify
 	FCustomizableObjectProjectorParameterValue& ProjectorParameter = ProjectorParameters[ParameterIndexInInstance];
-
+ 
 	// New Value
 	FCustomizableObjectProjector ProjectorData;
 	ProjectorData.Position = static_cast<FVector3f>(Pos);
@@ -1859,7 +1926,7 @@ void FCustomizableObjectInstanceDescriptor::SetProjectorValue(const FString& Pro
 	ProjectorData.Scale = static_cast<FVector3f>(Scale);
 	ProjectorData.Angle = Angle;
 	ProjectorData.ProjectionType = ProjectorParameter.Value.ProjectionType;
-
+ 
 	// Due to optimizations a parameter may lose its multidimensionality (if it becomes constant). 
 	// In that case it means there is only one possible option so it is ok to set it as if wasn't multidimensional.
 	bool bIsMultidimensional = CustomizableObject->GetPrivate()->IsParameterMultidimensional(ParameterIndexInObject);
@@ -1875,13 +1942,13 @@ void FCustomizableObjectInstanceDescriptor::SetProjectorValue(const FString& Pro
 			const int32 NumInsertedElements = RangeIndex + 1 - ProjectorParameter.RangeValues.Num();
 			ProjectorParameter.RangeValues.InsertDefaulted(InsertionIndex, NumInsertedElements);
 		}
-
+ 
 		check(ProjectorParameter.RangeValues.IsValidIndex(RangeIndex));
 		ProjectorParameter.RangeValues[RangeIndex] = ProjectorData;
 	}
 }
-
-
+ 
+ 
 void FCustomizableObjectInstanceDescriptor::SetProjectorPosition(const FString& ProjectorParamName, const FVector& Pos, const int32 RangeIndex)
 {
 	FVector DummyPos, Direction, Up, Scale;
@@ -1891,8 +1958,8 @@ void FCustomizableObjectInstanceDescriptor::SetProjectorPosition(const FString&
 	
 	SetProjectorValue(ProjectorParamName, static_cast<FVector>(Pos), Direction, Up, Scale, Angle, RangeIndex);
 }
-
-
+ 
+ 
 void FCustomizableObjectInstanceDescriptor::SetProjectorDirection(const FString& ProjectorParamName, const FVector& Direction, int32 RangeIndex)
 {
 	FVector Position, DummyDirection, Up, Scale;
@@ -1902,19 +1969,19 @@ void FCustomizableObjectInstanceDescriptor::SetProjectorDirection(const FString&
 		
 	SetProjectorValue(ProjectorParamName, Position, Direction, Up, Scale, Angle, RangeIndex);
 }
-
-
+ 
+ 
 void FCustomizableObjectInstanceDescriptor::SetProjectorUp(const FString& ProjectorParamName, const FVector& Up, int32 RangeIndex)
 {
 	FVector Position, Direction, DummyUp, Scale;
 	float Angle;
 	ECustomizableObjectProjectorType Type;
 	GetProjectorValue(ProjectorParamName, Position, Direction, DummyUp, Scale, Angle, Type, RangeIndex);
-
+ 
 	SetProjectorValue(ProjectorParamName, Position, Direction, Up, Scale, Angle, RangeIndex);
 }
-
-
+ 
+ 
 void FCustomizableObjectInstanceDescriptor::SetProjectorScale(const FString& ProjectorParamName, const FVector& Scale, int32 RangeIndex)
 {
 	FVector Position, Direction, Up, DummyScale;
@@ -1924,8 +1991,8 @@ void FCustomizableObjectInstanceDescriptor::SetProjectorScale(const FString& Pro
 	
 	SetProjectorValue(ProjectorParamName, Position, Direction, Up, Scale, Angle, RangeIndex);
 }
-
-
+ 
+ 
 void FCustomizableObjectInstanceDescriptor::SetProjectorAngle(const FString& ProjectorParamName, float Angle, int32 RangeIndex)
 {
 	FVector Position, Direction, Up, Scale;
@@ -1935,8 +2002,8 @@ void FCustomizableObjectInstanceDescriptor::SetProjectorAngle(const FString& Pro
 	
 	SetProjectorValue(ProjectorParamName, Position, Direction, Up, Scale, Angle, RangeIndex);
 }
-
-
+ 
+ 
 void FCustomizableObjectInstanceDescriptor::GetProjectorValue(const FString& ProjectorParamName,
                                                               FVector& OutPos, FVector& OutDirection, FVector& OutUp, FVector& OutScale,
                                                               float& OutAngle, ECustomizableObjectProjectorType& OutType,
@@ -1944,14 +2011,14 @@ void FCustomizableObjectInstanceDescriptor::GetProjectorValue(const FString& Pro
 {
 	FVector3f Pos, Direction, Up, Scale;
 	GetProjectorValueF(ProjectorParamName, Pos, Direction, Up, Scale, OutAngle, OutType, RangeIndex);
-
+ 
 	OutPos = static_cast<FVector>(Pos);
 	OutDirection = static_cast<FVector>(Direction);
 	OutUp = static_cast<FVector>(Up);
 	OutScale = static_cast<FVector>(Scale);
 }
-
-
+ 
+ 
 void FCustomizableObjectInstanceDescriptor::GetProjectorValueF(const FString& ProjectorParamName,
 	FVector3f& OutPos, FVector3f& OutDirection, FVector3f& OutUp, FVector3f& OutScale,
 	float& OutAngle, ECustomizableObjectProjectorType& OutType,
@@ -1962,15 +2029,15 @@ void FCustomizableObjectInstanceDescriptor::GetProjectorValueF(const FString& Pr
 		CustomizableObjectNullErrorMessage();
 		return;
 	}
-
+ 
 	const int32 ParameterIndexInObject = CustomizableObject->GetPrivate()->FindParameter(ProjectorParamName);
 	const int32 ParameterIndexInInstance = FindTypedParameterIndex(ProjectorParamName, EMutableParameterType::Projector);
-
+ 
 	if (ParameterIndexInObject >= 0 && ParameterIndexInInstance >= 0)
 	{
 		const FCustomizableObjectProjectorParameterValue& ProjectorParameter = ProjectorParameters[ParameterIndexInInstance];
 		const FCustomizableObjectProjector* ProjectorData;
-
+ 
 		// Due to optimizations a parameter may lose its multidimensionality (if it becomes constant). 
 		// In that case it means there is only one possible option so it is ok to set it as if wasn't multidimensional.
 		bool bIsMultidimensional = CustomizableObject->GetPrivate()->IsParameterMultidimensional(ParameterIndexInObject);
@@ -1981,10 +2048,10 @@ void FCustomizableObjectInstanceDescriptor::GetProjectorValueF(const FString& Pr
 		else
 		{
 			check(ProjectorParameter.RangeValues.IsValidIndex(RangeIndex));
-
+ 
 			ProjectorData = &ProjectorParameter.RangeValues[RangeIndex];
 		}
-
+ 
 		if (ProjectorData)
 		{
 			OutPos = ProjectorData->Position;
@@ -2000,8 +2067,8 @@ void FCustomizableObjectInstanceDescriptor::GetProjectorValueF(const FString& Pr
 		LogParameterNotFoundWarning(ProjectorParamName, ParameterIndexInObject, ParameterIndexInInstance, CustomizableObject, __FUNCTION__);
 	}
 }
-
-
+ 
+ 
 FVector FCustomizableObjectInstanceDescriptor::GetProjectorPosition(const FString& ParamName, const int32 RangeIndex) const
 {
 	if (!CustomizableObject)
@@ -2009,10 +2076,10 @@ FVector FCustomizableObjectInstanceDescriptor::GetProjectorPosition(const FStrin
 		CustomizableObjectNullErrorMessage();
 		return FVector(-0.0, -0.0, -0.0);
 	}
-
+ 
 	const int32 ParameterIndexInObject = CustomizableObject->GetPrivate()->FindParameter(ParamName);
 	const int32 ParameterIndexInInstance = FindTypedParameterIndex(ParamName, EMutableParameterType::Projector);
-
+ 
 	if (ParameterIndexInObject >= 0 && ParameterIndexInInstance >= 0)
 	{
 		// Due to optimizations a parameter may lose its multidimensionality (if it becomes constant). 
@@ -2027,13 +2094,13 @@ FVector FCustomizableObjectInstanceDescriptor::GetProjectorPosition(const FStrin
 			return static_cast<FVector>(ProjectorParameters[ParameterIndexInInstance].RangeValues[RangeIndex].Position);
 		}
 	}
-
+ 
 	LogParameterNotFoundWarning(ParamName, ParameterIndexInObject, ParameterIndexInInstance, CustomizableObject, __FUNCTION__);
-
+ 
 	return FVector(-0.0,-0.0,-0.0);
 }
-
-
+ 
+ 
 FVector FCustomizableObjectInstanceDescriptor::GetProjectorDirection(const FString& ParamName, const int32 RangeIndex) const
 {
 	if (!CustomizableObject)
@@ -2041,10 +2108,10 @@ FVector FCustomizableObjectInstanceDescriptor::GetProjectorDirection(const FStri
 		CustomizableObjectNullErrorMessage();
 		return FVector(-0.0, -0.0, -0.0);
 	}
-
+ 
 	const int32 ParameterIndexInObject = CustomizableObject->GetPrivate()->FindParameter(ParamName);
 	const int32 ParameterIndexInInstance = FindTypedParameterIndex(ParamName, EMutableParameterType::Projector);
-
+ 
 	if (ParameterIndexInObject >= 0 && ParameterIndexInInstance >= 0)
 	{
 		// Due to optimizations a parameter may lose its multidimensionality (if it becomes constant). 
@@ -2061,11 +2128,11 @@ FVector FCustomizableObjectInstanceDescriptor::GetProjectorDirection(const FStri
 	}
 	
 	LogParameterNotFoundWarning(ParamName, ParameterIndexInObject, ParameterIndexInInstance, CustomizableObject, __FUNCTION__);
-
+ 
 	return FVector(-0.0, -0.0, -0.0);
 }
-
-
+ 
+ 
 FVector FCustomizableObjectInstanceDescriptor::GetProjectorUp(const FString& ParamName, const int32 RangeIndex) const
 {
 	if (!CustomizableObject)
@@ -2073,10 +2140,10 @@ FVector FCustomizableObjectInstanceDescriptor::GetProjectorUp(const FString& Par
 		CustomizableObjectNullErrorMessage();
 		return FVector(-0.0, -0.0, -0.0);
 	}
-
+ 
 	const int32 ParameterIndexInObject = CustomizableObject->GetPrivate()->FindParameter(ParamName);
 	const int32 ParameterIndexInInstance = FindTypedParameterIndex(ParamName, EMutableParameterType::Projector);
-
+ 
 	if (ParameterIndexInObject >= 0 && ParameterIndexInInstance >= 0)
 	{
 		// Due to optimizations a parameter may lose its multidimensionality (if it becomes constant). 
@@ -2093,11 +2160,11 @@ FVector FCustomizableObjectInstanceDescriptor::GetProjectorUp(const FString& Par
 	}
 	
 	LogParameterNotFoundWarning(ParamName, ParameterIndexInObject, ParameterIndexInInstance, CustomizableObject, __FUNCTION__);
-
+ 
 	return FVector(-0.0, -0.0, -0.0);	
 }
-
-
+ 
+ 
 FVector FCustomizableObjectInstanceDescriptor::GetProjectorScale(const FString& ParamName, const int32 RangeIndex) const
 {
 	if (!CustomizableObject)
@@ -2105,10 +2172,10 @@ FVector FCustomizableObjectInstanceDescriptor::GetProjectorScale(const FString&
 		CustomizableObjectNullErrorMessage();
 		return FVector(-0.0, -0.0, -0.0);
 	}
-
+ 
 	const int32 ParameterIndexInObject = CustomizableObject->GetPrivate()->FindParameter(ParamName);
 	const int32 ParameterIndexInInstance = FindTypedParameterIndex(ParamName, EMutableParameterType::Projector);
-
+ 
 	if (ParameterIndexInObject >= 0 && ParameterIndexInInstance >= 0)
 	{
 		// Due to optimizations a parameter may lose its multidimensionality (if it becomes constant). 
@@ -2125,11 +2192,11 @@ FVector FCustomizableObjectInstanceDescriptor::GetProjectorScale(const FString&
 	}
 	
 	LogParameterNotFoundWarning(ParamName, ParameterIndexInObject, ParameterIndexInInstance, CustomizableObject, __FUNCTION__);
-
+ 
 	return FVector(-0.0, -0.0, -0.0);
 }
-
-
+ 
+ 
 float FCustomizableObjectInstanceDescriptor::GetProjectorAngle(const FString& ParamName, const int32 RangeIndex) const
 {
 	if (!CustomizableObject)
@@ -2137,10 +2204,10 @@ float FCustomizableObjectInstanceDescriptor::GetProjectorAngle(const FString& Pa
 		CustomizableObjectNullErrorMessage();
 		return 0.0;
 	}
-
+ 
 	const int32 ParameterIndexInObject = CustomizableObject->GetPrivate()->FindParameter(ParamName);
 	const int32 ParameterIndexInInstance = FindTypedParameterIndex(ParamName, EMutableParameterType::Projector);
-
+ 
 	if (ParameterIndexInObject >= 0 && ParameterIndexInInstance >= 0)
 	{
 		// Due to optimizations a parameter may lose its multidimensionality (if it becomes constant). 
@@ -2155,13 +2222,13 @@ float FCustomizableObjectInstanceDescriptor::GetProjectorAngle(const FString& Pa
 			return ProjectorParameters[ParameterIndexInInstance].RangeValues[RangeIndex].Angle;
 		}
 	}
-
+ 
 	LogParameterNotFoundWarning(ParamName, ParameterIndexInObject, ParameterIndexInInstance, CustomizableObject, __FUNCTION__);
-
+ 
 	return 0.0;
 }
-
-
+ 
+ 
 ECustomizableObjectProjectorType FCustomizableObjectInstanceDescriptor::GetProjectorParameterType(const FString& ParamName, const int32 RangeIndex) const
 {
 	if (!CustomizableObject)
@@ -2169,10 +2236,10 @@ ECustomizableObjectProjectorType FCustomizableObjectInstanceDescriptor::GetProje
 		CustomizableObjectNullErrorMessage();
 		return ECustomizableObjectProjectorType::Planar;
 	}
-
+ 
 	const int32 ParameterIndexInObject = CustomizableObject->GetPrivate()->FindParameter(ParamName);
 	const int32 ParameterIndexInInstance = FindTypedParameterIndex(ParamName, EMutableParameterType::Projector);
-
+ 
 	if (ParameterIndexInObject >= 0 && ParameterIndexInInstance >= 0)
 	{
 		// Due to optimizations a parameter may lose its multidimensionality (if it becomes constant). 
@@ -2187,18 +2254,18 @@ ECustomizableObjectProjectorType FCustomizableObjectInstanceDescriptor::GetProje
 			return ProjectorParameters[ParameterIndexInInstance].RangeValues[RangeIndex].ProjectionType;
 		}
 	}
-
+ 
 	LogParameterNotFoundWarning(ParamName, ParameterIndexInObject, ParameterIndexInInstance, CustomizableObject, __FUNCTION__);
-
+ 
 	return ECustomizableObjectProjectorType::Planar;
 }
-
-
+ 
+ 
 FCustomizableObjectProjector FCustomizableObjectInstanceDescriptor::GetProjector(const FString& ParamName, const int32 RangeIndex) const
 {
 	const int32 ParameterIndexInObject = CustomizableObject->GetPrivate()->FindParameter(ParamName);
 	const int32 ParameterIndexInInstance = FindTypedParameterIndex(ParamName, EMutableParameterType::Projector);
-
+ 
 	if (ParameterIndexInObject >= 0 && ParameterIndexInInstance >= 0)
 	{
 		// Due to optimizations a parameter may lose its multidimensionality (if it becomes constant). 
@@ -2213,19 +2280,19 @@ FCustomizableObjectProjector FCustomizableObjectInstanceDescriptor::GetProjector
 			return ProjectorParameters[ParameterIndexInInstance].RangeValues[RangeIndex];
 		}
 	}
-
+ 
 	LogParameterNotFoundWarning(ParamName, ParameterIndexInObject, ParameterIndexInInstance, CustomizableObject, __FUNCTION__);
-
+ 
 	return FCustomizableObjectProjectorParameterValue::DEFAULT_PARAMETER_VALUE;
 }
-
-
+ 
+ 
 int32 FCustomizableObjectInstanceDescriptor::FindTypedParameterIndex(const FString& ParamName, EMutableParameterType Type) const
 {
 	return CustomizableObject ? CustomizableObject->GetPrivate()->FindParameterTyped(ParamName, Type) : INDEX_NONE;
 }
-
-
+ 
+ 
 int32 FCustomizableObjectInstanceDescriptor::GetProjectorValueRange(const FString& ParamName) const
 {
 	if (!CustomizableObject)
@@ -2233,16 +2300,16 @@ int32 FCustomizableObjectInstanceDescriptor::GetProjectorValueRange(const FStrin
 		CustomizableObjectNullErrorMessage();
 		return -1;
 	}
-
+ 
 	const int32 ProjectorParamIndex = FindTypedParameterIndex(ParamName, EMutableParameterType::Projector);
 	if (ProjectorParamIndex < 0)
 	{
 		return -1;
 	}
-
+ 
 	return ProjectorParameters[ProjectorParamIndex].RangeValues.Num();
 }
-
+ 
 int32 FCustomizableObjectInstanceDescriptor::GetIntValueRange(const FString& ParamName) const
 {
 	if (!CustomizableObject)
@@ -2250,17 +2317,17 @@ int32 FCustomizableObjectInstanceDescriptor::GetIntValueRange(const FString& Par
 		CustomizableObjectNullErrorMessage();
 		return -1;
 	}
-
+ 
 	const int32 IntParamIndex = FindTypedParameterIndex(ParamName, EMutableParameterType::Int);
 	if (IntParamIndex < 0)
 	{
 		return -1;
 	}
-
+ 
 	return IntParameters[IntParamIndex].ParameterRangeValueNames.Num();
 }
-
-
+ 
+ 
 int32 FCustomizableObjectInstanceDescriptor::GetFloatValueRange(const FString& ParamName) const
 {
 	if (!CustomizableObject)
@@ -2268,17 +2335,17 @@ int32 FCustomizableObjectInstanceDescriptor::GetFloatValueRange(const FString& P
 		CustomizableObjectNullErrorMessage();
 		return -1;
 	}
-
+ 
 	const int32 FloatParamIndex = FindTypedParameterIndex(ParamName, EMutableParameterType::Float);
 	if (FloatParamIndex < 0)
 	{
 		return -1;
 	}
-
+ 
 	return FloatParameters[FloatParamIndex].ParameterRangeValues.Num();
 }
-
-
+ 
+ 
 int32 FCustomizableObjectInstanceDescriptor::GetTextureValueRange(const FString& ParamName) const
 {
 	if (!CustomizableObject)
@@ -2286,17 +2353,17 @@ int32 FCustomizableObjectInstanceDescriptor::GetTextureValueRange(const FString&
 		CustomizableObjectNullErrorMessage();
 		return -1;
 	}
-
+ 
 	const int32 TextureParamIndex = FindTypedParameterIndex(ParamName, EMutableParameterType::Texture);
 	if (TextureParamIndex < 0)
 	{
 		return -1;
 	}
-
+ 
 	return TextureParameters[TextureParamIndex].ParameterRangeValues.Num();
 }
-
-
+ 
+ 
 int32 FCustomizableObjectInstanceDescriptor::AddValueToIntRange(const FString& ParamName)
 {
 	if (!CustomizableObject)
@@ -2304,7 +2371,7 @@ int32 FCustomizableObjectInstanceDescriptor::AddValueToIntRange(const FString& P
 		CustomizableObjectNullErrorMessage();
 		return -1;
 	}
-
+ 
 	const int32 IntParameterIndex = FindTypedParameterIndex(ParamName, EMutableParameterType::Int);
 	if (IntParameterIndex != INDEX_NONE)
 	{
@@ -2316,8 +2383,8 @@ int32 FCustomizableObjectInstanceDescriptor::AddValueToIntRange(const FString& P
 	}
 	return -1;
 }
-
-
+ 
+ 
 int32 FCustomizableObjectInstanceDescriptor::AddValueToFloatRange(const FString& ParamName)
 {
 	const int32 FloatParameterIndex = FindTypedParameterIndex(ParamName, EMutableParameterType::Float);
@@ -2329,21 +2396,21 @@ int32 FCustomizableObjectInstanceDescriptor::AddValueToFloatRange(const FString&
 	}
 	return -1;
 }
-
-
+ 
+ 
 int32 FCustomizableObjectInstanceDescriptor::AddValueToTextureRange(const FString& ParamName)
 {
 	const int32 TextureParameterIndex = FindTypedParameterIndex(ParamName, EMutableParameterType::Texture);
 	if (TextureParameterIndex != INDEX_NONE)
 	{
-		FCustomizableObjectAssetParameterValue& TextureParameter = TextureParameters[TextureParameterIndex];
-		return TextureParameter.ParameterRangeValues.Add(FName());		
+		FCustomizableObjectTextureParameterValue& TextureParameter = TextureParameters[TextureParameterIndex];
+		return TextureParameter.ParameterRangeValues.Add(nullptr);
 	}
 	
 	return -1;
 }
-
-
+ 
+ 
 int32 FCustomizableObjectInstanceDescriptor::AddValueToProjectorRange(const FString& ParamName)
 {
 	if (!CustomizableObject)
@@ -2351,7 +2418,7 @@ int32 FCustomizableObjectInstanceDescriptor::AddValueToProjectorRange(const FStr
 		CustomizableObjectNullErrorMessage();
 		return -1;
 	}
-
+ 
 	const int32 ProjectorParameterIndex = FindTypedParameterIndex(ParamName, EMutableParameterType::Projector);
 	if (ProjectorParameterIndex != INDEX_NONE)
 	{
@@ -2359,11 +2426,11 @@ int32 FCustomizableObjectInstanceDescriptor::AddValueToProjectorRange(const FStr
 		const FCustomizableObjectProjector Projector = GetCustomizableObject()->GetProjectorParameterDefaultValue(ParamName);
 		return ProjectorParameter.RangeValues.Add(Projector);
 	}
-
+ 
 	return -1;
 }
-
-
+ 
+ 
 int32 FCustomizableObjectInstanceDescriptor::RemoveValueFromIntRange(const FString& ParamName, const int32 RangeIndex)
 {
 	const int32 IntParameterIndex = FindTypedParameterIndex(ParamName, EMutableParameterType::Int);
@@ -2379,8 +2446,8 @@ int32 FCustomizableObjectInstanceDescriptor::RemoveValueFromIntRange(const FStri
 	}
 	return -1;
 }
-
-
+ 
+ 
 int32 FCustomizableObjectInstanceDescriptor::RemoveValueFromFloatRange(const FString& ParamName, const int32 RangeIndex)
 {
 	const int32 FloatParameterIndex = FindTypedParameterIndex(ParamName, EMutableParameterType::Float);
@@ -2396,14 +2463,14 @@ int32 FCustomizableObjectInstanceDescriptor::RemoveValueFromFloatRange(const FSt
 	}
 	return -1;
 }
-
-
+ 
+ 
 int32 FCustomizableObjectInstanceDescriptor::RemoveValueFromTextureRange(const FString& ParamName)
 {
 	const int32 TextureParameterIndex = FindTypedParameterIndex(ParamName, EMutableParameterType::Texture);
 	if (TextureParameterIndex != INDEX_NONE)
 	{
-		FCustomizableObjectAssetParameterValue& TextureParameter = TextureParameters[TextureParameterIndex];
+		FCustomizableObjectTextureParameterValue& TextureParameter = TextureParameters[TextureParameterIndex];
 		if (TextureParameter.ParameterRangeValues.Num() > 0)
 		{
 			TextureParameter.ParameterRangeValues.Pop();
@@ -2412,14 +2479,14 @@ int32 FCustomizableObjectInstanceDescriptor::RemoveValueFromTextureRange(const F
 	}
 	return -1;
 }
-
-
+ 
+ 
 int32 FCustomizableObjectInstanceDescriptor::RemoveValueFromTextureRange(const FString& ParamName, const int32 RangeIndex)
 {
 	const int32 TextureParameterIndex = FindTypedParameterIndex(ParamName, EMutableParameterType::Texture);
 	if (TextureParameterIndex != INDEX_NONE)
 	{
-		FCustomizableObjectAssetParameterValue& TextureParameter = TextureParameters[TextureParameterIndex];
+		FCustomizableObjectTextureParameterValue& TextureParameter = TextureParameters[TextureParameterIndex];
 		if (TextureParameter.ParameterRangeValues.Num() > 0)
 		{
 			TextureParameter.ParameterRangeValues.RemoveAt(RangeIndex);
@@ -2428,15 +2495,15 @@ int32 FCustomizableObjectInstanceDescriptor::RemoveValueFromTextureRange(const F
 	}
 	return -1;
 }
-
-
+ 
+ 
 int32 FCustomizableObjectInstanceDescriptor::RemoveValueFromProjectorRange(const FString& ParamName, const int32 RangeIndex)
 {
 	const int32 ProjectorParameterIndex = FindTypedParameterIndex(ParamName, EMutableParameterType::Projector);
 	if (ProjectorParameterIndex != INDEX_NONE)
 	{
 		FCustomizableObjectProjectorParameterValue& ProjectorParameter = ProjectorParameters[ProjectorParameterIndex];
-
+ 
 		if (ProjectorParameter.RangeValues.Num() > 0)
 		{
 			int32 IndexToRemove = ProjectorParameter.RangeValues.IsValidIndex(RangeIndex) ? RangeIndex : ProjectorParameter.RangeValues.Num() - 1;
@@ -2447,14 +2514,14 @@ int32 FCustomizableObjectInstanceDescriptor::RemoveValueFromProjectorRange(const
 	}
 	return -1;
 }
-
-
+ 
+ 
 int32 FCustomizableObjectInstanceDescriptor::GetState() const
 {
 	return State;
 }
-
-
+ 
+ 
 FString FCustomizableObjectInstanceDescriptor::GetCurrentState() const
 {
 	if (!CustomizableObject)
@@ -2462,17 +2529,17 @@ FString FCustomizableObjectInstanceDescriptor::GetCurrentState() const
 		CustomizableObjectNullErrorMessage();
 		return FString();
 	}
-
+ 
 	return CustomizableObject->GetPrivate()->GetStateName(GetState());
 }
-
-
+ 
+ 
 void FCustomizableObjectInstanceDescriptor::SetState(const int32 InState)
 {
 	State = InState;
 }
-
-
+ 
+ 
 void FCustomizableObjectInstanceDescriptor::SetCurrentState(const FString& StateName)
 {
 	if (!CustomizableObject)
@@ -2480,9 +2547,9 @@ void FCustomizableObjectInstanceDescriptor::SetCurrentState(const FString& State
 		CustomizableObjectNullErrorMessage();
 		return;
 	}
-
+ 
 	RETURN_ON_UNCOMPILED_CO(CustomizableObject, TEXT("Error: Cannot set state"))
-
+ 
 	const int32 Result = CustomizableObject->GetPrivate()->FindState(StateName);
 #if WITH_EDITOR
 	if (Result != INDEX_NONE)
@@ -2497,16 +2564,16 @@ void FCustomizableObjectInstanceDescriptor::SetCurrentState(const FString& State
 		UE_LOG(LogMutable, Error, TEXT("%hs: Unknown %s state."), __FUNCTION__, *StateName);
 	}
 }
-
-
+ 
+ 
 void FCustomizableObjectInstanceDescriptor::SetRandomValues()
 {
 	const int32 RandomSeed = FMath::SRand() * TNumericLimits<int32>::Max();
 	FRandomStream RandomStream{RandomSeed};
 	SetRandomValuesFromStream(RandomSeed);
 }
-
-
+ 
+ 
 void FCustomizableObjectInstanceDescriptor::SetRandomValuesFromStream(const FRandomStream& InStream)
 {
 	if (!CustomizableObject)
@@ -2520,13 +2587,13 @@ void FCustomizableObjectInstanceDescriptor::SetRandomValuesFromStream(const FRan
 	for (FCustomizableObjectFloatParameterValue& FloatParameter : FloatParameters)
 	{
 		FloatParameter.ParameterValue = InStream.GetFraction();
-
+ 
 		for (float& RangeValue : FloatParameter.ParameterRangeValues)
 		{
 			RangeValue = InStream.GetFraction();
 		}
 	}
-
+ 
 	for (FCustomizableObjectBoolParameterValue& BoolParameter : BoolParameters)
 	{
 		BoolParameter.ParameterValue = static_cast<bool>(InStream.RandRange(0, 1));
@@ -2535,18 +2602,18 @@ void FCustomizableObjectInstanceDescriptor::SetRandomValuesFromStream(const FRan
 	for (FCustomizableObjectIntParameterValue& IntParameter : IntParameters)
 	{
 		const int32 NumValues = CustomizableObject->GetEnumParameterNumValues(IntParameter.ParameterName);
-
+ 
 		if (NumValues)
 		{
 			IntParameter.ParameterValueName = CustomizableObject->GetEnumParameterValue(IntParameter.ParameterName, NumValues * InStream.GetFraction());
-
+ 
 			for (FString& RangeValue : IntParameter.ParameterRangeValueNames)
 			{
 				RangeValue = CustomizableObject->GetEnumParameterValue(IntParameter.ParameterName, NumValues * InStream.GetFraction());
 			}
 		}		
 	}
-
+ 
 	for (FCustomizableObjectVectorParameterValue& VectorParameter : VectorParameters)
 	{
 		VectorParameter.ParameterValue.R = InStream.GetFraction();
@@ -2554,41 +2621,11 @@ void FCustomizableObjectInstanceDescriptor::SetRandomValuesFromStream(const FRan
 		VectorParameter.ParameterValue.B = InStream.GetFraction();
 		VectorParameter.ParameterValue.A = InStream.GetFraction();
 	}
-
-	const UCustomizableObjectSystemPrivate* SystemPrivate = UCustomizableObjectSystem::GetInstance()->GetPrivate();
-
-	TArray<FName> PossibleValues;
-
-	// Get all possible values
-	TArray<FCustomizableObjectExternalTexture> ProviderValues;
-	for (const TWeakObjectPtr<UCustomizableSystemImageProvider>& Provider : SystemPrivate->ResourceProvider->ImageProviders)
-	{
-		ProviderValues.Reset();
-		Provider->GetTextureParameterValues(ProviderValues);
-
-		for (const FCustomizableObjectExternalTexture& ProviderValue : ProviderValues)
-		{
-			PossibleValues.Add(ProviderValue.Value);
-		}
-	}
-
-	if (const int32 NumPossibleValues = PossibleValues.Num())
-	{
-		for (FCustomizableObjectAssetParameterValue& TextureParameter : TextureParameters)
-		{
-			TextureParameter.ParameterValue = PossibleValues[NumPossibleValues * InStream.GetFraction()];
-		
-			for (FName& RangeValue : TextureParameter.ParameterRangeValues)
-			{
-				RangeValue = PossibleValues[NumPossibleValues * InStream.GetFraction()];
-			}				
-		}		
-	}
-	
+ 
 	// Currently we are not randomizing the projectors since we do not know the valid range of values.
 }
-
-
+ 
+ 
 void FCustomizableObjectInstanceDescriptor::SetDefaultValue(int32 ParamIndex)
 {
 	if (!CustomizableObject)
@@ -2596,17 +2633,17 @@ void FCustomizableObjectInstanceDescriptor::SetDefaultValue(int32 ParamIndex)
 		CustomizableObjectNullErrorMessage();
 		return;
 	}
-
+ 
 	if (ParamIndex >= CustomizableObject->GetParameterCount())
 	{
 		return;
 	}
-
+ 
 	const FString ParamName = CustomizableObject->GetParameterName(ParamIndex);
 	const EMutableParameterType ParamType = CustomizableObject->GetPrivate()->GetParameterType(ParamIndex);
 	
 	const int32 TypedIndex = FindTypedParameterIndex(ParamName, ParamType);
-
+ 
 	switch (ParamType)
 	{
 	case EMutableParameterType::Bool:
@@ -2615,34 +2652,34 @@ void FCustomizableObjectInstanceDescriptor::SetDefaultValue(int32 ParamIndex)
 		BoolParameters[TypedIndex].ParameterValue = DefaultValue;
 		break;
 	}
-
+ 
 	case EMutableParameterType::Int:
 	{
 		const FString DefaultValue = CustomizableObject->GetPrivate()->FindIntParameterValueName(ParamIndex, CustomizableObject->GetEnumParameterDefaultValue(ParamName));
-
+ 
 		FCustomizableObjectIntParameterValue& IntParameter = IntParameters[TypedIndex];
 		IntParameter.ParameterValueName = DefaultValue;
 		IntParameter.ParameterRangeValueNames.Empty();
 		break;
 	}
-
+ 
 	case EMutableParameterType::Float:
 	{
 		const float DefaultValue = CustomizableObject->GetFloatParameterDefaultValue(ParamName);
-
+ 
 		FCustomizableObjectFloatParameterValue& FloatParameter = FloatParameters[TypedIndex];
 		FloatParameter.ParameterValue = DefaultValue;
 		FloatParameter.ParameterRangeValues.Empty();
 		break;
 	}
-
+ 
 	case EMutableParameterType::Color:
 	{
 		const FLinearColor DefaultValue = CustomizableObject->GetColorParameterDefaultValue(ParamName);
 		VectorParameters[TypedIndex].ParameterValue = DefaultValue;
 		break;
 	}
-
+ 
 	case EMutableParameterType::Transform:
 	{
 		const FTransform DefaultValue = CustomizableObject->GetTransformParameterDefaultValue(ParamName);
@@ -2653,33 +2690,46 @@ void FCustomizableObjectInstanceDescriptor::SetDefaultValue(int32 ParamIndex)
 	case EMutableParameterType::Projector:
 	{
 		const FCustomizableObjectProjector DefaultValue = CustomizableObject->GetProjectorParameterDefaultValue(ParamName);
-
+ 
 		FCustomizableObjectProjectorParameterValue& ProjectorParameter = ProjectorParameters[TypedIndex];
 		ProjectorParameter.Value = DefaultValue;
 		ProjectorParameter.RangeValues.Empty();
 		break;
 	}
-
+ 
 	case EMutableParameterType::Texture:
 	{
-		const FName	DefaultValue = CustomizableObject->GetTextureParameterDefaultValue(ParamName);
+		UTexture* DefaultValue = CustomizableObject->GetTextureParameterDefaultValue(ParamName);
 		
 		if (ensure(TextureParameters.IsValidIndex(TypedIndex)))
 		{
-			FCustomizableObjectAssetParameterValue& TextureParameter = TextureParameters[TypedIndex];
+			FCustomizableObjectTextureParameterValue& TextureParameter = TextureParameters[TypedIndex];
 			TextureParameter.ParameterValue = DefaultValue;
 			TextureParameter.ParameterRangeValues.Empty();
 		}
 		break;
 	}
-
+		
+	case EMutableParameterType::SkeletalMesh:
+	{
+		USkeletalMesh* DefaultValue = CustomizableObject->GetSkeletalMeshParameterDefaultValue(ParamName);
+	
+		if (ensure(SkeletalMeshParameters.IsValidIndex(TypedIndex)))
+		{
+			FCustomizableObjectSkeletalMeshParameterValue& SkeletalMeshParameter = SkeletalMeshParameters[TypedIndex];
+			SkeletalMeshParameter.ParameterValue = DefaultValue;
+			SkeletalMeshParameter.ParameterRangeValues.Empty();
+		}
+		break;
+	}
+ 
 	default:
 		// Parameter type replication not implemented.
 		unimplemented();
 	}
 }
-
-
+ 
+ 
 void FCustomizableObjectInstanceDescriptor::SetDefaultValues()
 {
 	if (!CustomizableObject)
@@ -2687,14 +2737,14 @@ void FCustomizableObjectInstanceDescriptor::SetDefaultValues()
 		CustomizableObjectNullErrorMessage();
 		return;
 	}
-
+ 
 	const int32 NumParameters = CustomizableObject->GetParameterCount();
 	for (int32 ParamIndex = 0; ParamIndex < NumParameters; ++ParamIndex)
 	{
 		SetDefaultValue(ParamIndex);
 	}
 }
-
+ 
 bool FCustomizableObjectInstanceDescriptor::IsMultilayerProjector(const FString& ParamName) const
 {
 	if (!CustomizableObject)
@@ -2702,7 +2752,7 @@ bool FCustomizableObjectInstanceDescriptor::IsMultilayerProjector(const FString&
 		CustomizableObjectNullErrorMessage();
 		return false;
 	}
-
+ 
 	// Projector.
 	if (const int32 ProjectorParameterIndex = FindTypedParameterIndex(ParamName, EMutableParameterType::Projector);
 		ProjectorParameterIndex == INDEX_NONE)
@@ -2710,7 +2760,7 @@ bool FCustomizableObjectInstanceDescriptor::IsMultilayerProjector(const FString&
 		ensureAlwaysMsgf(false, TEXT("%s"), *MULTILAYER_PROJECTOR_PARAMETERS_INVALID);
 		return false;
 	}
-
+ 
 	// Num layers.
 	if (const int32 FloatParameterIndex = FindTypedParameterIndex(ParamName + NUM_LAYERS_PARAMETER_POSTFIX, EMutableParameterType::Float);
 		FloatParameterIndex == INDEX_NONE)
@@ -2718,7 +2768,7 @@ bool FCustomizableObjectInstanceDescriptor::IsMultilayerProjector(const FString&
 		ensureAlwaysMsgf(false, TEXT("%s"), *MULTILAYER_PROJECTOR_PARAMETERS_INVALID);
 		return false;
 	}
-
+ 
 	// Selected Image.
 	if (const int32 IntParameterIndex = FindTypedParameterIndex(ParamName + IMAGE_PARAMETER_POSTFIX, EMutableParameterType::Int);
 		IntParameterIndex == INDEX_NONE)
@@ -2726,7 +2776,7 @@ bool FCustomizableObjectInstanceDescriptor::IsMultilayerProjector(const FString&
 		ensureAlwaysMsgf(false, TEXT("%s"), *MULTILAYER_PROJECTOR_PARAMETERS_INVALID);
 		return false;
 	}
-
+ 
 	// Opacity.	
 	if (const int32 FloatParameterIndex = FindTypedParameterIndex(ParamName + OPACITY_PARAMETER_POSTFIX, EMutableParameterType::Float);
 		FloatParameterIndex == INDEX_NONE)
@@ -2734,111 +2784,111 @@ bool FCustomizableObjectInstanceDescriptor::IsMultilayerProjector(const FString&
 		ensureAlwaysMsgf(false, TEXT("%s"), *MULTILAYER_PROJECTOR_PARAMETERS_INVALID);
 		return false;
 	}
-
+ 
 	return true;
 }
-
-
+ 
+ 
 int32 FCustomizableObjectInstanceDescriptor::NumProjectorLayers(const FName& ProjectorParamName) const
 {
 	const FString ParamName = ProjectorParamName.ToString();
-
+ 
 	if (IsMultilayerProjector(ParamName))
 	{
 		return GetFloatParameterSelectedOption(ParamName + NUM_LAYERS_PARAMETER_POSTFIX);
 	}
-
+ 
 	return INDEX_NONE;
 }
-
-
+ 
+ 
 void FCustomizableObjectInstanceDescriptor::CreateLayer(const FName& ProjectorParamName, int32 RangeIndex)
 {
 	const FString ParamName = ProjectorParamName.ToString();
-
+ 
 	if (!IsMultilayerProjector(ParamName))
 	{
 		return;
 	}
-
+ 
 	// Num Layers.
 	SetFloatParameterSelectedOption(ParamName + NUM_LAYERS_PARAMETER_POSTFIX, NumProjectorLayers(ProjectorParamName) + 1);
-
+ 
 	// Projector Range. New value is defaulted.
 	AddValueToProjectorRange(ParamName);
-
+ 
 	// Selected Image Range.
 	{
 		const FString ImageParamName = ParamName + IMAGE_PARAMETER_POSTFIX;
 		AddValueToTextureRange(ImageParamName);
-
+ 
 		const int32 DefaultValueIndex = CustomizableObject->GetEnumParameterDefaultValue(ImageParamName);
 		const FString DefaultValueName = CustomizableObject->GetEnumParameterValue(ImageParamName, DefaultValueIndex);
 		SetIntParameterSelectedOption(ImageParamName, DefaultValueName, RangeIndex);
 	}
-
+ 
 	// Opacity Range.
 	{
 		const FString OpacityParamName = ParamName + OPACITY_PARAMETER_POSTFIX;
 		AddValueToFloatRange(OpacityParamName);
-
+ 
 		const float DefaultValue = CustomizableObject->GetFloatParameterDefaultValue(OpacityParamName);
 		SetFloatParameterSelectedOption(OpacityParamName, DefaultValue, RangeIndex);
 	}
 }
-
-
+ 
+ 
 void FCustomizableObjectInstanceDescriptor::RemoveLayerAt(const FName& ProjectorParamName, int32 RangeIndex)
 {
 	const FString ParamName = ProjectorParamName.ToString();
 	const int32 NumLayers = NumProjectorLayers(ProjectorParamName);
 	check(RangeIndex >= 0 && RangeIndex < NumLayers); // Layer out of range.
-
+ 
 	if (!IsMultilayerProjector(ParamName))
 	{
 		return;
 	}
-
+ 
 	// Num Layers.
 	SetFloatParameterSelectedOption(ParamName + NUM_LAYERS_PARAMETER_POSTFIX, NumLayers - 1);
-
+ 
 	// Projector Range.
 	RemoveValueFromProjectorRange(ParamName, RangeIndex);
-
+ 
 	// Selected Image Range.
 	RemoveValueFromIntRange(ParamName + IMAGE_PARAMETER_POSTFIX, RangeIndex);
-
+ 
 	// Opacity Range.
 	RemoveValueFromFloatRange(ParamName + OPACITY_PARAMETER_POSTFIX, RangeIndex);
 }
-
-
+ 
+ 
 FMultilayerProjectorLayer FCustomizableObjectInstanceDescriptor::GetLayer(const FName& ProjectorParamName, int32 Index) const
 {
 	const FString ParamName = ProjectorParamName.ToString();
-
+ 
 	FMultilayerProjectorLayer ProjectorLayer;
-
+ 
 	if (IsMultilayerProjector(ParamName))
 	{
 		ProjectorLayer.Read(*this, ParamName, Index);
 	}
-
+ 
 	return ProjectorLayer;
 }
-
-
+ 
+ 
 void FCustomizableObjectInstanceDescriptor::UpdateLayer(const FName& ProjectorParamName, int32 Index, const FMultilayerProjectorLayer& Layer)
 {
 	const FString& ParamName = ProjectorParamName.ToString();
-
+ 
 	if (!IsMultilayerProjector(ParamName))
 	{
 		return;
 	}
-
+ 
 	Layer.Write(*this, ParamName, Index);
 }
-
-
-#undef RETURN_ON_CO_UNCOMPILED
+ 
+ 
+#undef RETURN_ON_CO_UNCOMPILED
\ No newline at end of file
diff --git a/Engine/Plugins/Mutable/Source/CustomizableObject/Private/MuCO/CustomizableObjectMipDataProvider.cpp b/Engine/Plugins/Mutable/Source/CustomizableObject/Private/MuCO/CustomizableObjectMipDataProvider.cpp
--- a/Engine/Plugins/Mutable/Source/CustomizableObject/Private/MuCO/CustomizableObjectMipDataProvider.cpp
+++ b/Engine/Plugins/Mutable/Source/CustomizableObject/Private/MuCO/CustomizableObjectMipDataProvider.cpp
@@ -53,22 +53,6 @@ FMutableUpdateContext::FMutableUpdateContext(
 	Parameters(InParameters),
 	State(InState)
 {
-	if (Parameters)
-	{
-		const UCustomizableObjectSystemPrivate* Private = UCustomizableObjectSystem::GetInstance()->GetPrivate();
-		Private->GetResourceProviderChecked()->CacheImages(*Parameters);
-	}
-}
-
-
-FMutableUpdateContext::~FMutableUpdateContext()
-{
-	if (Parameters &&
-		UCustomizableObjectSystem::IsCreated())
-	{
-		const UCustomizableObjectSystemPrivate* Private = UCustomizableObjectSystem::GetInstance()->GetPrivate();
-		Private->GetResourceProviderChecked()->UnCacheImages(*Parameters);
-	}
 }
 
 
@@ -90,7 +74,7 @@ TSharedPtr<mu::FSystem> FMutableUpdateContext::GetSystem() const
 }
 
 
-TSharedPtr<mu::FModel, ESPMode::ThreadSafe> FMutableUpdateContext::GetModel() const
+TSharedPtr<mu::FModel> FMutableUpdateContext::GetModel() const
 {
 	return Model;
 }
diff --git a/Engine/Plugins/Mutable/Source/CustomizableObject/Private/MuCO/CustomizableObjectSystem.cpp b/Engine/Plugins/Mutable/Source/CustomizableObject/Private/MuCO/CustomizableObjectSystem.cpp
--- a/Engine/Plugins/Mutable/Source/CustomizableObject/Private/MuCO/CustomizableObjectSystem.cpp
+++ b/Engine/Plugins/Mutable/Source/CustomizableObject/Private/MuCO/CustomizableObjectSystem.cpp
@@ -14,7 +14,6 @@
 #include "MuCO/CustomizableObjectPrivate.h"
 #include "MuCO/CustomizableObjectInstanceUsagePrivate.h"
 #include "MuCO/CustomizableObjectUIData.h"
-#include "MuCO/DefaultImageProvider.h"
 #include "MuCO/CustomizableObjectInstanceUsage.h"
 #include "MuCO/ICustomizableObjectModule.h"
 #include "MuCO/LogBenchmarkUtil.h"
@@ -36,7 +35,6 @@
 #include "Engine/World.h"
 #include "Interfaces/ITargetPlatform.h"
 #include "MuCO/ICustomizableObjectEditorModule.h"
-#include "MuCO/EditorImageProvider.h"
 #include "MuT/UnrealPixelFormatOverride.h"
 #else
 #include "Engine/Engine.h"
@@ -330,9 +328,6 @@ FUpdateContextPrivate::FUpdateContextPrivate(UCustomizableObjectInstance& InInst
 	MutableSystem = UCustomizableObjectSystem::GetInstance()->GetPrivate()->MutableSystem;	
 	check(MutableSystem);
 
-	UCustomizableObjectSystem* System = UCustomizableObjectSystem::GetInstance();
-	System->GetPrivate()->CacheTextureParameters(CapturedDescriptor.GetTextureParameters());
-
 	bValid = true;
 }
 
@@ -579,18 +574,6 @@ FUpdateContextPrivate::FUpdateContextPrivate(UCustomizableObjectInstance& InInst
 }
 
 
-FUpdateContextPrivate::~FUpdateContextPrivate()
-{
-	check(IsInGameThread());
-
-	if (bValid && UCustomizableObjectSystem::IsCreated())
-	{
-		UCustomizableObjectSystem* System = UCustomizableObjectSystem::GetInstance();
-		System->GetPrivate()->UnCacheTextureParameters(CapturedDescriptor.GetTextureParameters());
-	}
-}
-
-
 void FUpdateContextPrivate::SetMinLOD(const TMap<FName, uint8>& MinLOD)
 {
 	CapturedDescriptor.MinLOD = MinLOD;
@@ -930,12 +913,6 @@ void UCustomizableObjectSystem::InitSystem()
 	Private->ResourceProvider = Provider;
 	Private->MutableSystem->SetExternalResourceProvider(Provider);
 
-#if WITH_EDITORONLY_DATA
-	Private->EditorImageProvider = NewObject<UEditorImageProvider>();
-	check(Private->EditorImageProvider);
-	RegisterImageProvider(Private->EditorImageProvider);
-#endif
-	
 	GetPrivate()->DefaultInstanceLODManagement = NewObject<UCustomizableInstanceLODManagement>();
 	check(GetPrivate()->DefaultInstanceLODManagement != nullptr);
 	GetPrivate()->CurrentInstanceLODManagement = GetPrivate()->DefaultInstanceLODManagement;
@@ -1185,22 +1162,11 @@ void FinishUpdateGlobal(const TSharedRef<FUpdateContextPrivate>& Context)
 		case EUpdateResult::Success:
 		case EUpdateResult::Warning:
 			PrivateInstance->SkeletalMeshStatus = ESkeletalMeshStatus::Success;
-
-			if (SystemPrivate)
-			{
-				SystemPrivate->UnCacheTextureParameters(PrivateInstance->CommittedDescriptor.GetTextureParameters());				
-			}
-
+			
 			PrivateInstance->CommittedDescriptor = Context->MoveCommittedDescriptor();
 			PrivateInstance->CommittedDescriptorHash = Context->GetCapturedDescriptorHash();
 
 			PrivateInstance->bAutomaticUpdateRequired = false;
-
-			if (SystemPrivate)
-			{
-				// Cache new Texture FParameters
-				SystemPrivate->CacheTextureParameters(PrivateInstance->CommittedDescriptor.GetTextureParameters());
-			}
 			
 			// Delegates must be called only after updating the Instance flags.
 			Instance->UpdatedDelegate.Broadcast(Instance);
@@ -3696,7 +3662,7 @@ namespace impl
 
 #if WITH_EDITOR
 		// Async load all Runtime Referenced Textures.
-		const TArray<TSoftObjectPtr<const UTexture>>& RuntimeReferencedTextures = ModelResources.RuntimeReferencedTextures;
+		const TArray<TSoftObjectPtr<UTexture2D>>& RuntimeReferencedTextures = ModelResources.RuntimeReferencedTextures;
 		if (!RuntimeReferencedTextures.IsEmpty())
 		{
 			UE::Tasks::FTaskEvent Event = UE::Tasks::FTaskEvent(TEXT("Texture"));
@@ -3704,7 +3670,7 @@ namespace impl
 
 			TArray<FSoftObjectPath> Textures;
 			Textures.Reserve(RuntimeReferencedTextures.Num());
-			for (const TSoftObjectPtr<const UTexture>& Texture : RuntimeReferencedTextures)
+			for (const TSoftObjectPtr<UTexture2D>& Texture : RuntimeReferencedTextures)
 			{
 				Textures.Add(Texture.ToSoftObjectPath());
 			}
@@ -4384,62 +4350,6 @@ bool UCustomizableObjectSystem::IsUpdating(const UCustomizableObjectInstance* In
 }
 
 
-TArray<FCustomizableObjectExternalTexture> UCustomizableObjectSystem::GetTextureParameterValues()
-{
-	TArray<FCustomizableObjectExternalTexture> Result;
-
-	for (const TWeakObjectPtr<UCustomizableSystemImageProvider> Provider : GetPrivate()->GetResourceProviderChecked()->ImageProviders)
-	{
-		if (Provider.IsValid())
-		{
-			Provider->GetTextureParameterValues(Result);
-		}
-	}
-
-	return Result;
-}
-
-
-void UCustomizableObjectSystem::RegisterImageProvider(UCustomizableSystemImageProvider* Provider)
-{
-	GetPrivate()->GetResourceProviderChecked()->ImageProviders.Add(Provider);
-}
-
-
-void UCustomizableObjectSystem::UnregisterImageProvider(UCustomizableSystemImageProvider* Provider)
-{
-	GetPrivate()->GetResourceProviderChecked()->ImageProviders.Remove(Provider);
-}
-
-
-void UCustomizableObjectSystemPrivate::CacheTextureParameters(const TArray<FCustomizableObjectAssetParameterValue>& TextureParameters) const
-{
-	for (const FCustomizableObjectAssetParameterValue& TextureParameter : TextureParameters)
-	{
-		ResourceProvider->CacheImage(TextureParameter.ParameterValue, false);
-
-		for (const FName& RangeValue : TextureParameter.ParameterRangeValues)
-		{
-			ResourceProvider->CacheImage(RangeValue, false);
-		}
-	}
-}
-
-
-void UCustomizableObjectSystemPrivate::UnCacheTextureParameters(const TArray<FCustomizableObjectAssetParameterValue>& TextureParameters) const
-{
-	for (const FCustomizableObjectAssetParameterValue& TextureParameter : TextureParameters)
-	{
-		ResourceProvider->UnCacheImage(TextureParameter.ParameterValue, false);
-
-		for (const FName& RangeValue : TextureParameter.ParameterRangeValues)
-		{
-			ResourceProvider->UnCacheImage(RangeValue, false);
-		}
-	}
-}
-
-
 bool UCustomizableObjectSystemPrivate::IsUsingBenchmarkingSettings()
 {
 	return bUseBenchmarkingSettings;
@@ -4684,24 +4594,6 @@ uint64 UCustomizableObjectSystem::GetMaxChunkSizeForPlatform(const ITargetPlatfo
 #endif // WITH_EDITOR
 
 
-void UCustomizableObjectSystem::CacheImage(FName ImageId)
-{
-	GetPrivate()->GetResourceProviderChecked()->CacheImage(ImageId, true);
-}
-
-
-void UCustomizableObjectSystem::UnCacheImage(FName ImageId)
-{
-	GetPrivate()->GetResourceProviderChecked()->UnCacheImage(ImageId, true);
-}
-
-
-void UCustomizableObjectSystem::ClearImageCache()
-{
-	GetPrivate()->GetResourceProviderChecked()->ClearCache(true);
-}
-
-
 bool UCustomizableObjectSystemPrivate::IsMutableAnimInfoDebuggingEnabled() const
 { 
 #if WITH_EDITORONLY_DATA
diff --git a/Engine/Plugins/Mutable/Source/CustomizableObject/Private/MuCO/DefaultImageProvider.cpp b/Engine/Plugins/Mutable/Source/CustomizableObject/Private/MuCO/DefaultImageProvider.cpp
deleted file mode 100644
diff --git a/Engine/Plugins/Mutable/Source/CustomizableObject/Private/MuCO/DescriptorHash.cpp b/Engine/Plugins/Mutable/Source/CustomizableObject/Private/MuCO/DescriptorHash.cpp
--- a/Engine/Plugins/Mutable/Source/CustomizableObject/Private/MuCO/DescriptorHash.cpp
+++ b/Engine/Plugins/Mutable/Source/CustomizableObject/Private/MuCO/DescriptorHash.cpp
@@ -29,12 +29,12 @@ FDescriptorHash::FDescriptorHash(const FCustomizableObjectInstanceDescriptor& De
 		Hash = HashCombine(Hash, GetTypeHash(Value));
 	}
 
-	for (const FCustomizableObjectAssetParameterValue& Value : Descriptor.TextureParameters)
+	for (const FCustomizableObjectTextureParameterValue& Value : Descriptor.TextureParameters)
 	{
 		Hash = HashCombine(Hash, GetTypeHash(Value));
 	}
 
-	for (const FCustomizableObjectAssetParameterValue& Value : Descriptor.MeshParameters)
+	for (const FCustomizableObjectSkeletalMeshParameterValue& Value : Descriptor.SkeletalMeshParameters)
 	{
 		Hash = HashCombine(Hash, GetTypeHash(Value));
 	}
diff --git a/Engine/Plugins/Mutable/Source/CustomizableObject/Private/MuCO/EditorImageProvider.cpp b/Engine/Plugins/Mutable/Source/CustomizableObject/Private/MuCO/EditorImageProvider.cpp
deleted file mode 100644
diff --git a/Engine/Plugins/Mutable/Source/CustomizableObject/Private/MuCO/MultilayerProjector.cpp b/Engine/Plugins/Mutable/Source/CustomizableObject/Private/MuCO/MultilayerProjector.cpp
--- a/Engine/Plugins/Mutable/Source/CustomizableObject/Private/MuCO/MultilayerProjector.cpp
+++ b/Engine/Plugins/Mutable/Source/CustomizableObject/Private/MuCO/MultilayerProjector.cpp
@@ -26,12 +26,12 @@ void FMultilayerProjectorLayer::Read(const FCustomizableObjectInstanceDescriptor
 
 	{
 		const int32 ImageParamIndex = Descriptor.FindTypedParameterIndex(ParamName + IMAGE_PARAMETER_POSTFIX, EMutableParameterType::Int);
-		Image = Descriptor.GetIntParameters()[ImageParamIndex].ParameterRangeValueNames[LayerIndex];
+		Image = Descriptor.IntParameters[ImageParamIndex].ParameterRangeValueNames[LayerIndex];
 	}
 
 	{
 		const int32 OpacityParamIndex = Descriptor.FindTypedParameterIndex(ParamName + OPACITY_PARAMETER_POSTFIX, EMutableParameterType::Float);
-		Opacity = Descriptor.GetFloatParameters()[OpacityParamIndex].ParameterRangeValues[LayerIndex];
+		Opacity = Descriptor.FloatParameters[OpacityParamIndex].ParameterRangeValues[LayerIndex];
 	}
 }
 
diff --git a/Engine/Plugins/Mutable/Source/CustomizableObject/Private/MuCO/UnrealMutableImageProvider.cpp b/Engine/Plugins/Mutable/Source/CustomizableObject/Private/MuCO/UnrealMutableImageProvider.cpp
--- a/Engine/Plugins/Mutable/Source/CustomizableObject/Private/MuCO/UnrealMutableImageProvider.cpp
+++ b/Engine/Plugins/Mutable/Source/CustomizableObject/Private/MuCO/UnrealMutableImageProvider.cpp
@@ -20,6 +20,7 @@
 #include "ReferenceSkeleton.h"
 #include "ImageCoreUtils.h"
 #include "TextureResource.h"
+#include "Engine/Texture.h"
 
 
 namespace
@@ -88,7 +89,7 @@ mu::EImageFormat GetMutablePixelFormat(EPixelFormat InTextureFormat)
 
 
 //-------------------------------------------------------------------------------------------------
-TTuple<UE::Tasks::FTask, TFunction<void()>> FUnrealMutableResourceProvider::GetImageAsync(FName Id, uint8 MipmapsToSkip, TFunction<void(TSharedPtr<mu::FImage>)>& ResultCallback)
+TTuple<UE::Tasks::FTask, TFunction<void()>> FUnrealMutableResourceProvider::GetImageAsync(UTexture* Texture, uint8 MipmapsToSkip, TFunction<void(TSharedPtr<mu::FImage>)>& ResultCallback)
 {
 	// Thread: worker
 	MUTABLE_CPUPROFILER_SCOPE(FUnrealMutableImageProvider::GetImageAsync);
@@ -108,233 +109,214 @@ TTuple<UE::Tasks::FTask, TFunction<void()>> FUnrealMutableResourceProvider::GetI
 		return MakeTuple(UE::Tasks::MakeCompletedTask<void>(), []() -> void {});
 	};
 
+	if (!Texture)
 	{
-		FScopeLock Lock(&ExternalImagesLock);
-		// Inside this scope it's safe to access GlobalExternalImages
-
+		UE_LOG(LogMutable, Warning, TEXT("Invalid Image Parameter. Nullptr"));
+		ResultCallback(CreateDummy());
+		return Invoke(TrivialReturn);
+	}
+	
+	UTexture2D* Texture2D = Cast<UTexture2D>(Texture);
+	if (!Texture2D)
+	{
+		UE_LOG(LogMutable, Warning, TEXT("Invalid Image Parameter [%s]. Is not a UTexture2D."), *Texture->GetName());
+		ResultCallback(CreateDummy());
+		return Invoke(TrivialReturn);
+	}
+	
+#if WITH_EDITOR
+	FTextureSource Source = Texture->Source.CopyTornOff();
+	
+	const int32 MipIndex = FMath::Min(static_cast<int32>(MipmapsToSkip), Source.GetNumMips() - 1);
+	check(MipIndex >= 0);
+		
+	// In the editor the src data can be directly accessed
+	TSharedPtr<mu::FImage> Image = MakeShared<mu::FImage>();
 
-		if (!GlobalExternalImages.Contains(Id))
-		{
-			// Null case, no image was provided
-			UE_LOG(LogMutable, Warning, TEXT("Failed to get external image [%s]. GlobalExternalImage not found."), *Id.ToString());
+	FMutableSourceTextureData Tex(*Texture2D);
+	EUnrealToMutableConversionError Error = ConvertTextureUnrealSourceToMutable(Image.Get(), Tex, MipIndex);
+	if (Error != EUnrealToMutableConversionError::Success)
+	{
+		// This could happen in the editor, because some source textures may have changed while there was a background compilation.
+		// We just show a warning and move on. This cannot happen during cooks, so it is fine.
+		UE_LOG(LogMutable, Warning, TEXT("Failed to load some source texture data for image [%s]. Some materials may look corrupted."), *Texture->GetName());
+	}
 
-			ResultCallback(CreateDummy());
-			return Invoke(TrivialReturn);
-		}
+	ResultCallback(Image);
+	return Invoke(TrivialReturn);
+		
+#else
+	// It's safe to access TextureToLoad because ExternalImagesLock guarantees that the data in GlobalExternalImages is valid,
+	// not being modified by the game thread at the moment and the texture cannot be GCed because of the AddReferencedObjects
+	// in the FUnrealMutableImageProvider
 
-		const FUnrealMutableImageInfo& ImageInfo = GlobalExternalImages[Id];
+	int32 MipIndex = MipmapsToSkip < Texture2D->GetPlatformData()->Mips.Num() ? MipmapsToSkip : Texture2D->GetPlatformData()->Mips.Num() - 1;
+	check (MipIndex >= 0);
 
-		if (ImageInfo.Image)
+	// Mips in the mip tail are inlined and can't be streamed, find the smallest mip available.
+	for (; MipIndex > 0; --MipIndex)
+	{
+		if (Texture2D->GetPlatformData()->Mips[MipIndex].BulkData.CanLoadFromDisk())
 		{
-			// Easy case where the image was directly provided
-			ResultCallback(ImageInfo.Image);
-			return Invoke(TrivialReturn);
+			break;
 		}
+	}
+	
+	// Texture format and the equivalent mutable format
+	Format = Texture2D->GetPlatformData()->PixelFormat;
+	MutImageFormat = GetMutablePixelFormat(Format);
 
-#if WITH_EDITOR
-		if (const TSharedPtr<FMutableSourceTextureData>& SourceTextureData = ImageInfo.SourceTextureData)
-		{
-			const int32 MipIndex = FMath::Min(static_cast<int32>(MipmapsToSkip), SourceTextureData->GetSource().GetNumMips() - 1);
-			check(MipIndex >= 0);
-			
-			// In the editor the src data can be directly accessed
-			TSharedPtr<mu::FImage> Image = MakeShared<mu::FImage>();
+	// Check if it's a format we support
+	if (MutImageFormat == mu::EImageFormat::None)
+	{
+		UE_LOG(LogMutable, Warning, TEXT("Failed to get Image Parameter [%s]. Unexpected image format. EImageFormat [%s]."), *Texture2D->GetName(), GetPixelFormatString(Format));
+		ResultCallback(CreateDummy());
+		return Invoke(TrivialReturn);
+	}
 
-			EUnrealToMutableConversionError Error = ConvertTextureUnrealSourceToMutable(Image.Get(), *SourceTextureData, MipIndex);
-			if (Error != EUnrealToMutableConversionError::Success)
-			{
-				// This could happen in the editor, because some source textures may have changed while there was a background compilation.
-				// We just show a warning and move on. This cannot happen during cooks, so it is fine.
-				UE_LOG(LogMutable, Warning, TEXT("Failed to load some source texture data for image [%s]. Some materials may look corrupted."), *Id.ToString());
-			}
+	int32 SizeX = Texture2D->GetSizeX() >> MipIndex;
+	int32 SizeY = Texture2D->GetSizeY() >> MipIndex;
 
-			ResultCallback(Image);
-			return Invoke(TrivialReturn);
-		}
-#else
-		if (UTexture2D* TextureToLoad = ImageInfo.TextureToLoad)
-		{
-			// It's safe to access TextureToLoad because ExternalImagesLock guarantees that the data in GlobalExternalImages is valid,
-			// not being modified by the game thread at the moment and the texture cannot be GCed because of the AddReferencedObjects
-			// in the FUnrealMutableImageProvider
+	check(LODs == 1);
+	TSharedPtr<mu::FImage> Image = MakeShared<mu::FImage>(SizeX, SizeY, LODs, MutImageFormat, mu::EInitializationType::NotInitialized);
+	TArrayView<uint8> MutImageDataView = Image->DataStorage.GetLOD(0);
+
+	// In a packaged game the bulk data has to be loaded
+	// Get the actual file to read the mip 0 data, do not keep any reference to Texture2D because once outside of the lock
+	// it may be GCed or changed. Just keep the actual file handle and some sizes instead of the texture
+	FByteBulkData& BulkData = Texture2D->GetPlatformData()->Mips[MipIndex].BulkData;
+	BulkDataSize = BulkData.GetBulkDataSize();
+	check(BulkDataSize > 0);
 
-			int32 MipIndex = MipmapsToSkip < TextureToLoad->GetPlatformData()->Mips.Num() ? MipmapsToSkip : TextureToLoad->GetPlatformData()->Mips.Num() - 1;
-			check (MipIndex >= 0);
+	if (BulkDataSize != MutImageDataView.Num())
+	{
+		UE_LOG(LogMutable, Warning, TEXT("Failed to get Image Parameter [%s]. Bulk data size is different than the expected size. BulkData size [%d]. Mutable image data size [%d]."),
+			*Texture->GetName(), BulkDataSize, MutImageDataSize);
+
+		ResultCallback(CreateDummy());
+		return Invoke(TrivialReturn);
+	}
 
-			// Mips in the mip tail are inlined and can't be streamed, find the smallest mip available.
-			for (; MipIndex > 0; --MipIndex)
+	// Create a streaming request if the data is not loaded or copy the mip data
+	if (!BulkData.IsBulkDataLoaded())
+	{
+		UE::Tasks::FTaskEvent IORequestCompletionEvent(TEXT("Mutable_IORequestCompletionEvent"));
+
+		TFunction<void(bool, IBulkDataIORequest*)> IOCallback =
+			[
+				MutImageDataView,
+				MutImageFormat,
+				Format,
+				Image,
+				BulkDataSize,	
+				ResultCallback, // Notice ResultCallback is captured by copy
+				IORequestCompletionEvent
+			](bool bWasCancelled, IBulkDataIORequest* IORequest)
+		{
+			ON_SCOPE_EXIT
 			{
-				if (TextureToLoad->GetPlatformData()->Mips[MipIndex].BulkData.CanLoadFromDisk())
-				{
-					break;
-				}
-			}
+				UE::Tasks::FTaskEvent EventCopy = IORequestCompletionEvent;
+				EventCopy.Trigger();
+			};
 			
-			// Texture format and the equivalent mutable format
-			Format = TextureToLoad->GetPlatformData()->PixelFormat;
-			MutImageFormat = GetMutablePixelFormat(Format);
-
-			// Check if it's a format we support
-			if (MutImageFormat == mu::EImageFormat::None)
+			// Should we do someting different than returning a dummy image if cancelled?
+			if (bWasCancelled)
 			{
-				UE_LOG(LogMutable, Warning, TEXT("Failed to get external image [%s]. Unexpected image format. EImageFormat [%s]."), *Id.ToString(), GetPixelFormatString(Format));
+				UE_LOG(LogMutable, Warning, TEXT("Failed to get external image. Cancelled IO Request"));
 				ResultCallback(CreateDummy());
-				return Invoke(TrivialReturn);
+				return;
 			}
 
-			int32 SizeX = TextureToLoad->GetSizeX() >> MipIndex;
-			int32 SizeY = TextureToLoad->GetSizeY() >> MipIndex;
-
-			check(LODs == 1);
-			TSharedPtr<mu::FImage> Image = MakeShared<mu::FImage>(SizeX, SizeY, LODs, MutImageFormat, mu::EInitializationType::NotInitialized);
-			TArrayView<uint8> MutImageDataView = Image->DataStorage.GetLOD(0);
+			uint8* Results = IORequest->GetReadResults(); // required?
 
-			// In a packaged game the bulk data has to be loaded
-			// Get the actual file to read the mip 0 data, do not keep any reference to TextureToLoad because once outside of the lock
-			// it may be GCed or changed. Just keep the actual file handle and some sizes instead of the texture
-			FByteBulkData& BulkData = TextureToLoad->GetPlatformData()->Mips[MipIndex].BulkData;
-			BulkDataSize = BulkData.GetBulkDataSize();
-			check(BulkDataSize > 0);
-
-			if (BulkDataSize != MutImageDataView.Num())
+			if (Results && MutImageDataView.Num() == (int32)IORequest->GetSize())
 			{
-				UE_LOG(LogMutable, Warning, TEXT("Failed to get external image [%s]. Bulk data size is different than the expected size. BulkData size [%d]. Mutable image data size [%d]."),
-					*Id.ToString(),	BulkDataSize, MutImageDataSize);
+				check(BulkDataSize == (int32)IORequest->GetSize());
+				check(Results == MutImageDataView.GetData());
 
-				ResultCallback(CreateDummy());
-				return Invoke(TrivialReturn);
+				ResultCallback(Image);
+				return;
 			}
 
-			// Create a streaming request if the data is not loaded or copy the mip data
-			if (!BulkData.IsBulkDataLoaded())
+			if (!Results)
 			{
-				UE::Tasks::FTaskEvent IORequestCompletionEvent(TEXT("Mutable_IORequestCompletionEvent"));
-
-				TFunction<void(bool, IBulkDataIORequest*)> IOCallback =
-					[
-						MutImageDataView,
-						MutImageFormat,
-						Format,
-						Image,
-						BulkDataSize,	
-						ResultCallback, // Notice ResultCallback is captured by copy
-						IORequestCompletionEvent
-					](bool bWasCancelled, IBulkDataIORequest* IORequest)
-				{
-					ON_SCOPE_EXIT
-					{
-						UE::Tasks::FTaskEvent EventCopy = IORequestCompletionEvent;
-						EventCopy.Trigger();
-					};
-					
-					// Should we do someting different than returning a dummy image if cancelled?
-					if (bWasCancelled)
-					{
-						UE_LOG(LogMutable, Warning, TEXT("Failed to get external image. Cancelled IO Request"));
-						ResultCallback(CreateDummy());
-						return;
-					}
-
-					uint8* Results = IORequest->GetReadResults(); // required?
-
-					if (Results && MutImageDataView.Num() == (int32)IORequest->GetSize())
-					{
-						check(BulkDataSize == (int32)IORequest->GetSize());
-						check(Results == MutImageDataView.GetData());
-
-						ResultCallback(Image);
-						return;
-					}
-
-					if (!Results)
-					{
-						UE_LOG(LogMutable, Warning, TEXT("Failed to get external image. IO Request failed. Request results [%hhd]. Format: [%s]. MutableFormat: [%d]."),
-							(Results != nullptr),
-							GetPixelFormatString(Format),
-							(int32)MutImageFormat);
-					}
-					else if (MutImageDataView.Num() != (int32)IORequest->GetSize())
-					{
-						UE_LOG(LogMutable, Warning, TEXT("Failed to get external image. Requested size is different than the expected size. RequestSize: [%lld]. ExpectedSize: [%d]. Format: [%s]. MutableFormat: [%d]."),
-							IORequest->GetSize(),
-							MutImageDataView.Num(),
-							GetPixelFormatString(Format),
-							(int32)MutImageFormat);
-					}
-					else
-					{
-						UE_LOG(LogMutable, Warning, TEXT("Failed to get external image."));
-					}
-
-					// Something failed when loading the bulk data, just return a dummy
-					ResultCallback(CreateDummy());
-				};
-			
-				// Is the resposability of the CreateStreamingRequest caller to delete the IORequest. 
-				// This can *not* be done in the IOCallback because it would cause a deadlock so it is deferred to the returned
-				// cleanup function. Another solution could be to spwan a new task that depends on the 
-				// IORequestComplitionEvent which deletes it.
-				TRACE_IOSTORE_METADATA_SCOPE_TAG(Id);
-				IORequest = BulkData.CreateStreamingRequest(EAsyncIOPriorityAndFlags::AIOP_High, &IOCallback, MutImageDataView.GetData());
-
-				if (IORequest)
-				{
-					// Make the lambda mutable and set the IORequest pointer to null when deleted so it is safer 
-					// agains multiple calls.
-					const auto DeleteIORequest = [IORequest]() mutable -> void
-					{
-						if (IORequest)
-						{
-							delete IORequest;
-						}
-						
-						IORequest = nullptr;
-					};
-
-					return MakeTuple(IORequestCompletionEvent, DeleteIORequest);
-
-				}
-				else
-				{
-					UE_LOG(LogMutable, Warning, TEXT("Failed to create an IORequest for a UTexture2D BulkData for an application-specific image parameter."));
-
-					IORequestCompletionEvent.Trigger();
-					
-					ResultCallback(CreateDummy());
-					return Invoke(TrivialReturn);
-				}
+				UE_LOG(LogMutable, Warning, TEXT("Failed to get Image Parameter. IO Request failed. Request results [%hhd]. Format: [%s]. MutableFormat: [%d]."),
+					(Results != nullptr),
+					GetPixelFormatString(Format),
+					(int32)MutImageFormat);
+			}
+			else if (MutImageDataView.Num() != (int32)IORequest->GetSize())
+			{
+				UE_LOG(LogMutable, Warning, TEXT("Failed to get Image Parameter. Requested size is different than the expected size. RequestSize: [%lld]. ExpectedSize: [%d]. Format: [%s]. MutableFormat: [%d]."),
+					IORequest->GetSize(),
+					MutImageDataView.Num(),
+					GetPixelFormatString(Format),
+					(int32)MutImageFormat);
 			}
 			else
 			{
-				// Bulk data already loaded
-				const void* Data = (!BulkData.IsLocked()) ? BulkData.LockReadOnly() : nullptr; // TODO: Retry if it fails?
-				
-				if (Data)
-				{
-					FMemory::Memcpy(MutImageDataView.GetData(), Data, BulkDataSize);
+				UE_LOG(LogMutable, Warning, TEXT("Failed to get Image Parameter."));
+			}
 
-					BulkData.Unlock();
-					ResultCallback(Image);
-					return Invoke(TrivialReturn);
-				}
-				else
+			// Something failed when loading the bulk data, just return a dummy
+			ResultCallback(CreateDummy());
+		};
+	
+		// Is the resposability of the CreateStreamingRequest caller to delete the IORequest. 
+		// This can *not* be done in the IOCallback because it would cause a deadlock so it is deferred to the returned
+		// cleanup function. Another solution could be to spwan a new task that depends on the 
+		// IORequestComplitionEvent which deletes it.
+		TRACE_IOSTORE_METADATA_SCOPE_TAG(FName(Texture->GetPathName()));
+		IORequest = BulkData.CreateStreamingRequest(EAsyncIOPriorityAndFlags::AIOP_High, &IOCallback, MutImageDataView.GetData());
+
+		if (IORequest)
+		{
+			// Make the lambda mutable and set the IORequest pointer to null when deleted so it is safer 
+			// agains multiple calls.
+			const auto DeleteIORequest = [IORequest]() mutable -> void
+			{
+				if (IORequest)
 				{
-					UE_LOG(LogMutable, Warning, TEXT("Failed to get external image. Bulk data already locked or null."));
-					ResultCallback(CreateDummy());
-					return Invoke(TrivialReturn);
+					delete IORequest;
 				}
-			}
+				
+				IORequest = nullptr;
+			};
+
+			return MakeTuple(IORequestCompletionEvent, DeleteIORequest);
+
+		}
+		else
+		{
+			UE_LOG(LogMutable, Warning, TEXT("Failed to create an IORequest for a UTexture2D BulkData for an application-specific image parameter."));
+
+			IORequestCompletionEvent.Trigger();
+			
+			ResultCallback(CreateDummy());
+			return Invoke(TrivialReturn);
 		}
-#endif
-		
-		// No UTexture2D was provided, cannot do anything, just provide a dummy texture
-		UE_LOG(LogMutable, Warning, TEXT("No UTexture2D was provided for an application-specific image parameter."));
-		ResultCallback(CreateDummy());
-		return Invoke(TrivialReturn);
 	}
+	else
+	{
+		// Bulk data already loaded
+		const void* Data = (!BulkData.IsLocked()) ? BulkData.LockReadOnly() : nullptr; // TODO: Retry if it fails?
+		
+		if (Data)
+		{
+			FMemory::Memcpy(MutImageDataView.GetData(), Data, BulkDataSize);
 
-	// Make sure the returned event is dispatched at some point for all code paths, 
-	// in this case returning Invoke(TrivialReturn) or through the IORequest callback.
+			BulkData.Unlock();
+			ResultCallback(Image);
+			return Invoke(TrivialReturn);
+		}
+		else
+		{
+			UE_LOG(LogMutable, Warning, TEXT("Failed to get Image Parameter. Bulk data already locked or null."));
+			ResultCallback(CreateDummy());
+			return Invoke(TrivialReturn);
+		}
+	}
+#endif
 }
 
 
@@ -350,39 +332,52 @@ TTuple<UE::Tasks::FTask, TFunction<void()>> FUnrealMutableResourceProvider::GetR
 	};
 
 #if WITH_EDITOR
-
-	TSharedPtr<mu::FImage> Image = MakeShared<mu::FImage>();
-
 	FScopeLock Lock(&RuntimeReferencedLock);
 	if (!RuntimeReferencedImages.Contains(ModelPtr))
 	{
-		UE_LOG(LogMutable, Error, TEXT("Failed to load image [%i]. Model not registered in the provider."), Id);
+		UE_LOG(LogMutable, Error, TEXT("Failed to load Referenced Image [%i]. Model not registered in the provider."), Id);
 		ResultCallback(CreateDummy());
 		return Invoke(TrivialReturn);
 	}
 
-	TArray<FMutableSourceTextureData>& SourceTextures = RuntimeReferencedImages[ModelPtr].SourceTextures;
+	TArray<TStrongObjectPtr<UTexture>>& SourceTextures = RuntimeReferencedImages[ModelPtr].Images;
 	if (!SourceTextures.IsValidIndex(Id))
 	{
 		// This could happen in the editor, because some source textures may have changed while there was a background compilation.
 		// We just show a warning and move on. This cannot happen during cooks, so it is fine.
-		UE_LOG(LogMutable, Warning, TEXT("Failed to load image [%i]."), Id);
-
+		UE_LOG(LogMutable, Warning, TEXT("Failed to load Referenced Image [%i]. Invalid id."), Id);
 		ResultCallback(CreateDummy());
 		return Invoke(TrivialReturn);
 	}
 
-	FMutableSourceTextureData& SourceTextureData = SourceTextures[Id];
+	UTexture* Texture = SourceTextures[Id].Get();
+	if (!Texture)
+	{
+		UE_LOG(LogMutable, Warning, TEXT("Invalid Referenced Image [%i]. Nullptr."), Id);
+		ResultCallback(CreateDummy());
+		return Invoke(TrivialReturn);
+	}
+	
+	UTexture2D* Texture2D = Cast<UTexture2D>(Texture);
+	if (!Texture2D)
+	{
+		UE_LOG(LogMutable, Warning, TEXT("Invalid Referenced Image [%i, %s]. Is not a UTexture2D."), Id, *Texture->GetName());
+		ResultCallback(CreateDummy());
+		return Invoke(TrivialReturn);
+	}
 	
-	const int32 MipIndex = FMath::Min(static_cast<int32>(MipmapsToSkip), SourceTextureData.GetSource().GetNumMips() - 1);
+	const int32 MipIndex = FMath::Min(static_cast<int32>(MipmapsToSkip), Texture2D->Source.CopyTornOff().GetNumMips() - 1);
 	check(MipIndex >= 0);
+
+	TSharedPtr<mu::FImage> Image = MakeShared<mu::FImage>();
 	
-	EUnrealToMutableConversionError Error = ConvertTextureUnrealSourceToMutable(Image.Get(), SourceTextureData, MipIndex);
+	FMutableSourceTextureData Tex(*Texture2D);
+	EUnrealToMutableConversionError Error = ConvertTextureUnrealSourceToMutable(Image.Get(), Tex, MipIndex);
 	if (Error != EUnrealToMutableConversionError::Success)
 	{
 		// This could happen in the editor, because some source textures may have changed while updating.
 		// We just show a warning and move on. This cannot happen during cooks, so it is fine.
-		UE_LOG(LogMutable, Warning, TEXT("Failed to load some source texture data for image [%i]. Some textures may be corrupted."), Id);
+		UE_LOG(LogMutable, Warning, TEXT("Failed to load some source texture data for Referenced Image [%i, %s]. Some textures may be corrupted."), Id, *Texture->GetName());
 		
 		ResultCallback(CreateDummy());
 		return Invoke(TrivialReturn);
@@ -393,7 +388,7 @@ TTuple<UE::Tasks::FTask, TFunction<void()>> FUnrealMutableResourceProvider::GetR
 #else // WITH_EDITOR
 
 	// Not supported outside editor yet.
-	UE_LOG(LogMutable, Warning, TEXT("Failed to get reference image. Only supported in editor."));
+	UE_LOG(LogMutable, Warning, TEXT("Failed to get Reference Image. Only supported in editor."));
 
 	ResultCallback(CreateDummy());
 	return Invoke(TrivialReturn);
@@ -403,188 +398,116 @@ TTuple<UE::Tasks::FTask, TFunction<void()>> FUnrealMutableResourceProvider::GetR
 
 
 // This should mantain parity with the descriptor of the images generated by GetImageAsync 
-mu::FExtendedImageDesc FUnrealMutableResourceProvider::GetImageDesc(FName Id)
+mu::FExtendedImageDesc FUnrealMutableResourceProvider::GetImageDesc(UTexture* Texture)
 {
 	MUTABLE_CPUPROFILER_SCOPE(FUnrealMutableImageProvider::GetImageDesc);
 
-	FScopeLock Lock(&ExternalImagesLock);
-	// Inside this scope it's safe to access GlobalExternalImages
-
-	if (!GlobalExternalImages.Contains(Id))
+	if (!Texture)
 	{
-		// Null case, no image was provided
 		return CreateDummyDesc();
 	}
 
-	const FUnrealMutableImageInfo& ImageInfo = GlobalExternalImages[Id];
-
-	if (ImageInfo.Image)
+	UTexture2D* Texture2D = Cast<UTexture2D>(Texture);
+	if (!Texture2D)
 	{
-		// Easy case where the image was directly provided
-		
-		const mu::FImageDesc ImageDesc = mu::FImageDesc 
-				{ ImageInfo.Image->GetSize(), ImageInfo.Image->GetFormat(), (uint8)ImageInfo.Image->GetLODCount() };
-		return mu::FExtendedImageDesc { ImageDesc, 0 }; 
+		return CreateDummyDesc();
 	}
-
+	
 #if WITH_EDITOR
-	if (const TSharedPtr<FMutableSourceTextureData>& SourceTextureData = ImageInfo.SourceTextureData)
-	{
-		const FTextureSource& Source = SourceTextureData->GetSource();
+	const FTextureSource& Source = Texture->Source.CopyTornOff();
 
-		const mu::FImageSize ImageSize = mu::FImageSize(Source.GetSizeX(), Source.GetSizeY());
-		const uint8 LODs = 1;
+	const mu::FImageSize ImageSize = mu::FImageSize(Source.GetSizeX(), Source.GetSizeY());
+	const uint8 LODs = 1;
 
-		return mu::FExtendedImageDesc { mu::FImageDesc { ImageSize, mu::EImageFormat::None, LODs }, 0 };
-	}
+	return mu::FExtendedImageDesc { mu::FImageDesc { ImageSize, mu::EImageFormat::None, LODs }, 0 };
 #else
-	if (UTexture2D* TextureToLoad = ImageInfo.TextureToLoad)
-	{
-		mu::FExtendedImageDesc Result;	
+	mu::FExtendedImageDesc Result;	
 
-		// It's safe to access TextureToLoad because ExternalImagesLock guarantees that the data in GlobalExternalImages is valid,
-		// not being modified by the game thread at the moment and the texture cannot be GCed because of the AddReferencedObjects
-		// in the FUnrealMutableImageProvider
+	// It's safe to access TextureToLoad because ExternalImagesLock guarantees that the data in GlobalExternalImages is valid,
+	// not being modified by the game thread at the moment and the texture cannot be GCed because of the AddReferencedObjects
+	// in the FUnrealMutableImageProvider
 
-		const int32 TextureToLoadNumMips = TextureToLoad->GetPlatformData()->Mips.Num();
+	const int32 TextureToLoadNumMips = Texture2D->GetPlatformData()->Mips.Num();
 
-		int32 FirstLODAvailable = 0;
-		for (; FirstLODAvailable < TextureToLoadNumMips; ++FirstLODAvailable)
+	int32 FirstLODAvailable = 0;
+	for (; FirstLODAvailable < TextureToLoadNumMips; ++FirstLODAvailable)
+	{
+		if (Texture2D->GetPlatformData()->Mips[FirstLODAvailable].BulkData.DoesExist())
 		{
-			if (TextureToLoad->GetPlatformData()->Mips[FirstLODAvailable].BulkData.DoesExist())
-			{
-				break;
-			}
+			break;
 		}
+	}
 	
-		// Texture format and the equivalent mutable format
-		const EPixelFormat Format = TextureToLoad->GetPlatformData()->PixelFormat;
-		const mu::EImageFormat MutableFormat = GetMutablePixelFormat(Format);
+	// Texture format and the equivalent mutable format
+	const EPixelFormat Format = Texture2D->GetPlatformData()->PixelFormat;
+	const mu::EImageFormat MutableFormat = GetMutablePixelFormat(Format);
 
-		// Check if it's a format we support
-		if (MutableFormat == mu::EImageFormat::None)
-		{
-			UE_LOG(LogMutable, Warning, TEXT("Failed to get external image descriptor. Unexpected image format. EImageFormat [%s]."), GetPixelFormatString(Format));
-			return CreateDummyDesc();
-		}
+	// Check if it's a format we support
+	if (MutableFormat == mu::EImageFormat::None)
+	{
+		UE_LOG(LogMutable, Warning, TEXT("Failed to get Image Parameter descriptor. Unexpected image format. EImageFormat [%s]."), GetPixelFormatString(Format));
+		return CreateDummyDesc();
+	}
 
-		const mu::FImageDesc ImageDesc = mu::FImageDesc 
-			{ mu::FImageSize(TextureToLoad->GetSizeX(), TextureToLoad->GetSizeY()), MutableFormat, 1 }; 
+	const mu::FImageDesc ImageDesc = mu::FImageDesc 
+		{ mu::FImageSize(Texture2D->GetSizeX(), Texture2D->GetSizeY()), MutableFormat, 1 }; 
 
-		Result = mu::FExtendedImageDesc { ImageDesc, (uint8)FirstLODAvailable }; 
+	Result = mu::FExtendedImageDesc { ImageDesc, (uint8)FirstLODAvailable }; 
 
-		return Result;
-	}
+	return Result;
 #endif
-	
-	// No UTexture2D was provided, cannot do anything, just provide a dummy texture
-	UE_LOG(LogMutable, Warning, TEXT("No UTexture2D was provided for an application-specific image parameter descriptor."));
-	return CreateDummyDesc();
 }
 
 
-TTuple<UE::Tasks::FTask, TFunction<void()>> FUnrealMutableResourceProvider::GetMeshAsync(FName Id, int32 LODIndex, int32 SectionIndex, TFunction<void(TSharedPtr<mu::FMesh>)>& ResultCallback)
+TTuple<UE::Tasks::FTask, TFunction<void()>> FUnrealMutableResourceProvider::GetMeshAsync(USkeletalMesh* SkeletalMesh, int32 LODIndex, int32 SectionIndex, TFunction<void(TSharedPtr<mu::FMesh>)>& ResultCallback)
 {
 	// Thread: worker
 	MUTABLE_CPUPROFILER_SCOPE(FUnrealMutableImageProvider::GetMeshAsync);
 
+	auto TrivialReturn = []() -> TTuple<UE::Tasks::FTask, TFunction<void()>>
+	{
+		return MakeTuple(UE::Tasks::MakeCompletedTask<void>(), []() -> void {});
+	};
+
 	TSharedPtr<mu::FMesh> Result = MakeShared<mu::FMesh>();
 
+	if (!SkeletalMesh)
+	{
+		UE_LOG(LogMutable, Warning, TEXT("Invalid Mesh Parameter. Nullptr."));
+		ResultCallback(Result);
+		return Invoke(TrivialReturn);
+	}
+	
 	UE::Tasks::FTaskEvent Completion(TEXT("MutableMeshParameterLoadInGameThread"));
+	
+	ExecuteOnGameThread(TEXT("GetMeshAsync"), [SkeletalMesh, LODIndex, SectionIndex, Result, Completion, Provider = SharedThis(this)]() mutable
+	{
+		check(IsInGameThread());
 
-	auto MeshLoadCallback = [Id, LODIndex, SectionIndex, Result, Completion, Provider = this](const FSoftObjectPath&, UObject* LoadedObject) mutable
-		{
-			check(IsInGameThread());
-
-			MUTABLE_CPUPROFILER_SCOPE(ActualLoad);
-
-			USkeletalMesh* SkeletalMesh = Cast<USkeletalMesh>(LoadedObject);
-			if (!SkeletalMesh)
-			{
-				UE_LOG(LogMutable, Error, TEXT("Failed to load the skeletal mesh [%s] set for a mesh parameter. Please check that it was packaged."), *Id.ToString());
-				Completion.Trigger();
-				return;
-			}
-
-			// Add a reference to the skeletalmesh
-			{
-				bool bFound = false;
-				for (FReferencedSkeletalMesh& Entry : Provider->ReferencedSkeletalMeshes)
-				{
-					if (Entry.SkeletalMesh == SkeletalMesh)
-					{
-						bFound = true;
-						++Entry.ReferenceCount;
-						break;
-					}
-				}
-
-				if (!bFound)
-				{
-					Provider->ReferencedSkeletalMeshes.Add({ SkeletalMesh,1 });
-				}
-			}
-
-			UModelResources* ModelResources = nullptr;
-
-			TStrongObjectPtr<UCustomizableObject> CO = Provider->CurrentCustomizableObject.Pin();
-			if (CO)
-			{
-				ModelResources = CO->GetPrivate()->GetModelResources();
-			}
-
-			// It is valid not to have a CO or MutableStreamedResources. It is only used for skeleton data. 
-			// This may happen when updating mips, and for those operations we don't need skeleton data.
-			//ensure(ModelResources);
-
-			UE::Tasks::FTask ConversionTask = UnrealConversionUtils::ConvertSkeletalMeshFromRuntimeData(SkeletalMesh, LODIndex, SectionIndex, ModelResources, Result.Get());
-
-			// The rest of the conversion may happen in a worker thread.
-			UE::Tasks::Launch(TEXT("MeshParameterLoadFinalize"),
-				[Completion, Provider, SkeletalMesh]() mutable
-				{					
-					// Signal
-					Completion.Trigger();
-
-					// Remove a reference to the skeletalmesh
-					ExecuteOnGameThread(TEXT("MeshParameterLoadReleaseMesh"),
-						[Provider, SkeletalMesh]() mutable
-						{
-							bool bFound = false;
-							for (int32 RefIndex = 0; RefIndex < Provider->ReferencedSkeletalMeshes.Num(); ++RefIndex)
-							{
-								FReferencedSkeletalMesh& Entry = Provider->ReferencedSkeletalMeshes[RefIndex];
-								if (Entry.SkeletalMesh == SkeletalMesh)
-								{
-									bFound = true;
-									--Entry.ReferenceCount;
-									if (Entry.ReferenceCount == 0)
-									{
-										Provider->ReferencedSkeletalMeshes.RemoveAtSwap(RefIndex);
-									}
-									break;
-								}
-							}
-
-							check(bFound);
-						});
-				},
-				// Dependencies
-				ConversionTask
-			);
-		};
+		MUTABLE_CPUPROFILER_SCOPE(ActualLoad);
 
-	// LoadAsync is only thread-safe when using the zenloader
-	ExecuteOnGameThread(TEXT("MeshParameterLoad"),
-		[Id, MeshLoadCallback]() mutable
+		if (!SkeletalMesh)
 		{
-			check(IsInGameThread());
-
-			MUTABLE_CPUPROFILER_SCOPE(MutableMeshParameterLoadInGameThread);
+			UE_LOG(LogMutable, Warning, TEXT("No Skeletal Mesh was provided for the Skeletal Mesh parameter."));
+			Completion.Trigger();
+			return;
+		}
 		
-			FSoftObjectPath(Id.ToString()).LoadAsync( FLoadSoftObjectPathAsyncDelegate::CreateLambda( MeshLoadCallback ));
+		UModelResources* ModelResources = nullptr;
+
+		TStrongObjectPtr<UCustomizableObject> CO = Provider->CurrentCustomizableObject.Pin();
+		if (CO)
+		{
+			ModelResources = CO->GetPrivate()->GetModelResources();
 		}
-	);
+
+		// It is valid not to have a CO or MutableStreamedResources. It is only used for skeleton data. 
+		// This may happen when updating mips, and for those operations we don't need skeleton data.
+		//ensure(ModelResources);
+
+		UE::Tasks::FTask ConversionTask = UnrealConversionUtils::ConvertSkeletalMeshFromRuntimeData(SkeletalMesh, LODIndex, SectionIndex, ModelResources, Result.Get());
+		Completion.Trigger();
+	});
 
 	return MakeTuple(
 		// Some post-game conversion stuff can happen here in a worker thread
@@ -603,221 +526,8 @@ TTuple<UE::Tasks::FTask, TFunction<void()>> FUnrealMutableResourceProvider::GetM
 }
 
 
-void FUnrealMutableResourceProvider::CacheImage(FName Id, bool bUser)
-{
-	if (Id == NAME_None)
-	{
-		return;	
-	}	
-
-	{
-		FScopeLock Lock(&ExternalImagesLock);
-
-		if (FUnrealMutableImageInfo* Result = GlobalExternalImages.Find(Id))
-		{
-			if (bUser)
-			{
-				Result->ReferencesUser = true;			
-			}
-			else
-			{
-				Result->ReferencesSystem++;							
-			}
-		}
-		else
-		{
-			bool bFound = false;
-
-			FUnrealMutableImageInfo ImageInfo;
-
-			// See if any provider provides this id.
-			for (const TWeakObjectPtr<UCustomizableSystemImageProvider>& Provider : ImageProviders)
-			{
-				if (bFound)
-				{
-					break;
-				}
-				
-				if (!Provider.IsValid())
-				{
-					continue;
-				}
-
-				// \TODO: all these queries could probably be optimized into a single call.
-				switch (Provider->HasTextureParameterValue(Id))
-				{
-				case UCustomizableSystemImageProvider::ValueType::Raw:
-					{
-						FIntVector desc = Provider->GetTextureParameterValueSize(Id);
-						TSharedPtr<mu::FImage> pResult = MakeShared<mu::FImage>(desc[0], desc[1], 1, mu::EImageFormat::RGBA_UByte, mu::EInitializationType::Black);
-						Provider->GetTextureParameterValueData(Id, pResult->GetLODData(0));
-
-						ImageInfo = FUnrealMutableImageInfo(pResult);
-						
-						bFound = true;
-						break;
-					}
-
-				case UCustomizableSystemImageProvider::ValueType::Unreal:
-					{
-						UTexture2D* UnrealTexture = Provider->GetTextureParameterValue(Id);
-						TSharedPtr<mu::FImage> pResult = MakeShared<mu::FImage>();
-
 #if WITH_EDITOR
-						FMutableSourceTextureData Tex(*UnrealTexture);
-						EUnrealToMutableConversionError Error = ConvertTextureUnrealSourceToMutable(pResult.Get(), Tex, 0);
-						if (Error != EUnrealToMutableConversionError::Success)
-						{
-							// This could happen in the editor, because some source textures may have changed while there was a background compilation.
-							// We just show a warning and move on. This cannot happen during cooks, so it is fine.
-							UE_LOG(LogMutable, Warning, TEXT("Failed to load some source texture data for [%s]. Some textures may be corrupted."), *UnrealTexture->GetName());
-						}
-#else
-						ConvertTextureUnrealPlatformToMutable(pResult.Get(), UnrealTexture, 0);
-#endif
-
-						ImageInfo = FUnrealMutableImageInfo(pResult);
-
-						bFound = true;
-						break;
-					}
-
-				case UCustomizableSystemImageProvider::ValueType::Unreal_Deferred:
-					{
-						if (UTexture2D* UnrealDeferredTexture = Provider->GetTextureParameterValue(Id))
-						{
-							ImageInfo = FUnrealMutableImageInfo(*UnrealDeferredTexture);
-
-							bFound = true;
-						}
-
-						break;
-					}
-
-				default:
-					break;
-				}
-			}
-
-			if (!bFound)
-			{
-				UE_LOG(LogMutable, Warning, TEXT("Failed to cache external image %s. Missing result and source texture."), *Id.ToString());
-				return;
-			}
-			
-			if (bUser)
-			{
-				ImageInfo.ReferencesUser = true;			
-			}
-				else
-			{
-				ImageInfo.ReferencesSystem++;							
-			}
-			
-			GlobalExternalImages.Add(Id, ImageInfo);
-		}
-	}
-}
-
-
-void FUnrealMutableResourceProvider::UnCacheImage(FName Id, bool bUser)
-{
-	if (Id == NAME_None)
-	{
-		return;	
-	}
-	
-	FScopeLock Lock(&ExternalImagesLock);
-	if (FUnrealMutableImageInfo* Result = GlobalExternalImages.Find(Id))
-	{
-		if (bUser)
-		{
-			Result->ReferencesUser = false;			
-		}
-		else
-		{
-			--Result->ReferencesSystem;
-			check(Result->ReferencesSystem >= 0); // Something went wrong. The image has been uncached more times than it has been cached.
-		}
-		
-		if (Result->ReferencesUser + Result->ReferencesSystem == 0)
-		{
-			GlobalExternalImages.Remove(Id);		
-		}
-	}
-	else
-	{
-		UE_LOG(LogMutable, Warning, TEXT("Failed to uncache external image %s. Possible double free!"), *Id.ToString());
-	}
-}
-
-
-void FUnrealMutableResourceProvider::ClearCache(bool bUser)
-{
-	FScopeLock Lock(&ExternalImagesLock);
-	for (TTuple<FName, FUnrealMutableImageInfo> Tuple : GlobalExternalImages)
-	{
-		UnCacheImage(Tuple.Key, bUser);
-	}	
-}
-
-
-void FUnrealMutableResourceProvider::CacheImages(const mu::FParameters& Parameters)
-{
-	const int32 NumParams = Parameters.GetCount();
-	for (int32 ParamIndex = 0; ParamIndex < NumParams; ++ParamIndex)
-	{
-		if (Parameters.GetType(ParamIndex) != mu::EParameterType::Image)
-		{
-			continue;
-		}
-	
-		{
-			const FName TextureId = Parameters.GetImageValue(ParamIndex);
-			CacheImage(TextureId, false);
-		}
-
-		const int32 NumValues = Parameters.GetValueCount(ParamIndex);
-		for (int32 ValueIndex = 0; ValueIndex < NumValues; ++ValueIndex)
-		{
-			TSharedPtr<const mu::FRangeIndex> Range = Parameters.GetValueIndex(ParamIndex, ValueIndex);
-			const FName TextureId = Parameters.GetImageValue(ParamIndex, Range.Get());
-
-			CacheImage(TextureId, false);
-		}
-	}
-}
-
-
-void FUnrealMutableResourceProvider::UnCacheImages(const mu::FParameters& Parameters)
-{
-	const int32 NumParams = Parameters.GetCount();
-	for (int32 ParamIndex = 0; ParamIndex < NumParams; ++ParamIndex)
-	{
-		if (Parameters.GetType(ParamIndex) != mu::EParameterType::Image)
-		{
-			continue;
-		}
-	
-		{
-			const FName TextureId = Parameters.GetImageValue(ParamIndex);
-			UnCacheImage(TextureId, false);				
-		}
-
-		const int32 NumValues = Parameters.GetValueCount(ParamIndex);
-		for (int32 ValueIndex = 0; ValueIndex < NumValues; ++ValueIndex)
-		{
-			TSharedPtr<const mu::FRangeIndex> Range = Parameters.GetValueIndex(ParamIndex, ValueIndex);
-			const FName TextureId = Parameters.GetImageValue(ParamIndex, Range.Get());
-
-			UnCacheImage(TextureId, false);
-		}
-	}
-}
-
-
-#if WITH_EDITOR
-void FUnrealMutableResourceProvider::CacheRuntimeReferencedImages(const TSharedRef<const mu::FModel>& Model, const TArray<TSoftObjectPtr<const UTexture>>& RuntimeReferencedTextures)
+void FUnrealMutableResourceProvider::CacheRuntimeReferencedImages(const TSharedRef<const mu::FModel>& Model, const TArray<TSoftObjectPtr<UTexture2D>>& RuntimeReferencedTextures)
 {
 	check(IsInGameThread());
 	
@@ -828,10 +538,10 @@ void FUnrealMutableResourceProvider::CacheRuntimeReferencedImages(const TSharedR
 	FRuntimeReferencedImages& ModelImages = RuntimeReferencedImages.Add(&Model.Get());
 	ModelImages.Model = Model.ToWeakPtr();
 
-	ModelImages.SourceTextures.Reset();
-	for (const TSoftObjectPtr<const UTexture>& RuntimeReferencedTexture : RuntimeReferencedTextures)
+	ModelImages.Images.Reset();
+	for (const TSoftObjectPtr<UTexture2D>& RuntimeReferencedTexture : RuntimeReferencedTextures)
 	{
-		const UTexture* Texture = RuntimeReferencedTexture.Get(); // Is already loaded.
+		UTexture2D* Texture = RuntimeReferencedTexture.Get(); // Is already loaded.
 		if (!Texture)
 		{
 			UE_LOG(LogMutable, Warning, TEXT("Runtime Referenced Texture [%s] was not async loaded. Forcing load sync."), *RuntimeReferencedTexture->GetPathName());
@@ -844,7 +554,7 @@ void FUnrealMutableResourceProvider::CacheRuntimeReferencedImages(const TSharedR
 			}
 		}
 
-		ModelImages.SourceTextures.Emplace(*Texture); // Perform a CopyTornOff. Once done, we no longer need the texture loaded.
+		ModelImages.Images.Add(TStrongObjectPtr(Texture)); // Perform a CopyTornOff. Once done, we no longer need the texture loaded.
 	}
 }
 #endif
@@ -912,32 +622,7 @@ void FUnrealMutableResourceProvider::AddReferencedObjects(FReferenceCollector& C
 			}
 		}
 	}
-#else
-	bool bDoLock = CVarMutableLockExternalImagesDuringGC.GetValueOnAnyThread();
-
-	if (bDoLock)
-	{
-		ExternalImagesLock.Lock();
-	}
-
-	for (TPair<FName, FUnrealMutableImageInfo>& Image : GlobalExternalImages)
-	{
-		if (Image.Value.TextureToLoad)
-		{
-			Collector.AddReferencedObject(Image.Value.TextureToLoad);
-		}
-	}
-
-	if (bDoLock)
-	{
-		ExternalImagesLock.Unlock();
-	}
 #endif
-
-	for (FReferencedSkeletalMesh& Mesh : ReferencedSkeletalMeshes)
-	{
-		Collector.AddReferencedObject(Mesh.SkeletalMesh);
-	}
 }
 
 
@@ -947,22 +632,3 @@ void FUnrealMutableResourceProvider::SetCurrentObject(const TWeakObjectPtr<UCust
 	CurrentCustomizableObject = InObject;
 }
 
-
-FUnrealMutableResourceProvider::FUnrealMutableImageInfo::FUnrealMutableImageInfo(const TSharedPtr<mu::FImage>& InImage)
-{
-	check(IsInGameThread())
-	
-	Image = InImage;
-}
-
-
-FUnrealMutableResourceProvider::FUnrealMutableImageInfo::FUnrealMutableImageInfo(UTexture2D& Texture)
-{
-	check(IsInGameThread())
-
-#if WITH_EDITOR
-	SourceTextureData = MakeShared<FMutableSourceTextureData>(Texture);
-#else
-	TextureToLoad = &Texture;
-#endif
-}
diff --git a/Engine/Plugins/Mutable/Source/CustomizableObject/Private/MuCO/UnrealMutableImageProvider.h b/Engine/Plugins/Mutable/Source/CustomizableObject/Private/MuCO/UnrealMutableImageProvider.h
--- a/Engine/Plugins/Mutable/Source/CustomizableObject/Private/MuCO/UnrealMutableImageProvider.h
+++ b/Engine/Plugins/Mutable/Source/CustomizableObject/Private/MuCO/UnrealMutableImageProvider.h
@@ -13,8 +13,8 @@
 #include "Tasks/Task.h"
 
 class UCustomizableObject;
-class UTexture2D;
-
+class UTexture;
+class USkeletalMesh;
 
 /** Implementation of a mutable core provider for image parameters that are application-specific. */
 class FUnrealMutableResourceProvider : public TSharedFromThis<FUnrealMutableResourceProvider>, public mu::FExternalResourceProvider, public FGCObject
@@ -23,46 +23,18 @@ class FUnrealMutableResourceProvider : public TSharedFromThis<FUnrealMutableReso
 public:
 	// mu::ExternalResourceProvider interface
 	// Thread: worker
-	virtual TTuple<UE::Tasks::FTask, TFunction<void()>> GetImageAsync(FName Id, uint8 MipmapsToSkip, TFunction<void(TSharedPtr<mu::FImage>)>& ResultCallback) override;
+	virtual TTuple<UE::Tasks::FTask, TFunction<void()>> GetImageAsync(UTexture* Texture, uint8 MipmapsToSkip, TFunction<void(TSharedPtr<mu::FImage>)>& ResultCallback) override;
 	virtual TTuple<UE::Tasks::FTask, TFunction<void()>> GetReferencedImageAsync(const void* ModelPtr, int32 Id, uint8 MipmapsToSkip, TFunction<void(TSharedPtr<mu::FImage>)>& ResultCallback) override;
-	virtual mu::FExtendedImageDesc GetImageDesc(FName Id) override;
-	virtual TTuple<UE::Tasks::FTask, TFunction<void()>> GetMeshAsync(FName Id, int32 LODIndex, int32 SectionIndex, TFunction<void(TSharedPtr<mu::FMesh>)>& ResultCallback) override;
-
-	
-	// Own interface	
-	// Thread: Game
-	/** Add a reference to the image. If it was not cached it caches it.
-	 * @param bUser if true, adds a reference to the user reference counter. If false, adds a reference to the system reference counter. */
-	void CacheImage(FName Id, bool bUser);
-
-	/** Removes a reference to the image. If all references are removed, it uncaches the image.
-	 * @param bUser if true, removes a reference from the user reference counter. If false, removes a reference from the system reference counter. */
-	void UnCacheImage(FName Id, bool bUser);
-
-	/** Removes a reference to all images. All images which no longer have references will be uncached.
-	 * @param bUser if true, removes a references from the user reference counter. If false, removes a reference from the system reference counter. */
-	void ClearCache(bool bUser);
-
-	void CacheImages(const mu::FParameters& Parameters);
-	void UnCacheImages(const mu::FParameters& Parameters);
+	virtual mu::FExtendedImageDesc GetImageDesc(UTexture* Texture) override;
+	virtual TTuple<UE::Tasks::FTask, TFunction<void()>> GetMeshAsync(USkeletalMesh* SkeletalMesh, int32 LODIndex, int32 SectionIndex, TFunction<void(TSharedPtr<mu::FMesh>)>& ResultCallback) override;
 
 #if WITH_EDITOR
-	void CacheRuntimeReferencedImages(const TSharedRef<const mu::FModel>& Model, const TArray<TSoftObjectPtr<const UTexture>>& RuntimeReferencedTextures);
+	void CacheRuntimeReferencedImages(const TSharedRef<const mu::FModel>& Model, const TArray<TSoftObjectPtr<UTexture2D>>& RuntimeReferencedTextures);
 #endif
 	
-	/** List of actual image providers that have been registered to the CustomizableObjectSystem. */
-	TArray< TWeakObjectPtr<class UCustomizableSystemImageProvider> > ImageProviders;
-
 	// Register and unregister the current CO that we are working on. This is called by internal code when building instances.
 	void SetCurrentObject(const TWeakObjectPtr<UCustomizableObject>&);
-
-	struct FReferencedSkeletalMesh
-	{
-		TObjectPtr<class USkeletalMesh> SkeletalMesh;
-		uint32 ReferenceCount = 0;
-	};
-	TArray<FReferencedSkeletalMesh> ReferencedSkeletalMeshes;
-
+	
 	// FGCObject interface
 	virtual void AddReferencedObjects(FReferenceCollector& Collector) override;
 	virtual FString GetReferencerName() const override
@@ -71,31 +43,6 @@ public:
 	}
 
 private:
-	struct FUnrealMutableImageInfo
-	{
-		FUnrealMutableImageInfo() = default;
-
-		FUnrealMutableImageInfo(const TSharedPtr<mu::FImage>& InImage);
-
-		FUnrealMutableImageInfo(UTexture2D& Texture);
-
-		TSharedPtr<mu::FImage> Image;
-
-#if WITH_EDITOR
-		TSharedPtr<FMutableSourceTextureData> SourceTextureData;
-#else
-		/** If the above Image has not been loaded in the game thread, the TextureToLoad bulk data will be loaded
-		* from the Mutable thread when it's needed */
-		TObjectPtr<UTexture2D> TextureToLoad = nullptr;
-#endif
-		
-		/** true of the reference maintained by the user. */
-		bool ReferencesUser = false;
-		
-		/** Number of reference maintained by the system. */
-		int32 ReferencesSystem = 0;
-	};
-
 	/** Runtime reference to the CO being updated. It is needed for example for mesh conversion, to access to the bone id map.  */
 	TWeakObjectPtr<UCustomizableObject> CurrentCustomizableObject;
 
@@ -106,24 +53,15 @@ private:
 	static TSharedPtr<mu::FImage> CreateDummy();
 	static mu::FExtendedImageDesc CreateDummyDesc();
 
-	/** Map of Ids to external textures that may be required for any instance or Mutable texture mip under construction.
-	* This is only safely written from the game thread protected by the following critical section, and it
-	* is safely read from the mutable thread during the update of the instance or texture mips. */
-	TMap<FName, FUnrealMutableImageInfo> GlobalExternalImages;
-
 #if WITH_EDITOR
 	struct FRuntimeReferencedImages
 	{
-		TArray<FMutableSourceTextureData> SourceTextures;
 		TWeakPtr<const mu::FModel> Model;
+		TArray<TStrongObjectPtr<UTexture>> Images;
 	};
 	
 	TMap<const void*, FRuntimeReferencedImages> RuntimeReferencedImages;
 #endif
-	
-	/** Access to GlobalExternalImages must be protected with this because it may be accessed concurrently from the 
-	Game thread to modify it and from the Mutable thread to read it. */
-	FCriticalSection ExternalImagesLock;
 
 #if WITH_EDITOR
 	FCriticalSection RuntimeReferencedLock;
diff --git a/Engine/Plugins/Mutable/Source/CustomizableObject/Private/MuCO/UnrealToMutableTextureConversionUtils.cpp b/Engine/Plugins/Mutable/Source/CustomizableObject/Private/MuCO/UnrealToMutableTextureConversionUtils.cpp
--- a/Engine/Plugins/Mutable/Source/CustomizableObject/Private/MuCO/UnrealToMutableTextureConversionUtils.cpp
+++ b/Engine/Plugins/Mutable/Source/CustomizableObject/Private/MuCO/UnrealToMutableTextureConversionUtils.cpp
@@ -173,10 +173,8 @@ void BlurNormalForComposite(FImage& Image)
 } //namespace UnrealToMutableImageConversion_Internal
 
 
-FMutableSourceTextureData::FMutableSourceTextureData(const UTexture& Texture)
+FMutableSourceTextureData::FMutableSourceTextureData(const UTexture2D& Texture)
 {
-	check(IsInGameThread());
-
 	Source = Texture.Source.CopyTornOff();
 	bFlipGreenChannel = Texture.bFlipGreenChannel;
 	bHasAlphaChannel = Texture.AdjustMinAlpha != Texture.AdjustMaxAlpha &&
diff --git a/Engine/Plugins/Mutable/Source/CustomizableObject/Public/MuCO/CustomizableObject.h b/Engine/Plugins/Mutable/Source/CustomizableObject/Public/MuCO/CustomizableObject.h
--- a/Engine/Plugins/Mutable/Source/CustomizableObject/Public/MuCO/CustomizableObject.h
+++ b/Engine/Plugins/Mutable/Source/CustomizableObject/Public/MuCO/CustomizableObject.h
@@ -87,11 +87,11 @@ struct FProfileParameterDat
 
 	////
 	UPROPERTY()
-	TArray<FCustomizableObjectAssetParameterValue> TextureParameters;
+	TArray<FCustomizableObjectTextureParameterValue> TextureParameters;
 
 	////
 	UPROPERTY()
-	TArray<FCustomizableObjectAssetParameterValue> MeshParameters;
+	TArray<FCustomizableObjectSkeletalMeshParameterValue> SkeletalMeshParameters;
 
 	////
 	UPROPERTY()
@@ -458,10 +458,13 @@ public:
 	CUSTOMIZABLEOBJECT_API FCustomizableObjectProjector GetProjectorParameterDefaultValue(const FString& ParameterName) const;
 	
 	/** Get the default value of a parameter of type Texture.
-	  * @param InParameterName The name of the Projector parameter to get the default value of.
-	  * @return An id representing the default parameter's texture. */
-	CUSTOMIZABLEOBJECT_API FName GetTextureParameterDefaultValue(const FString& InParameterName) const;
+	  * @param InParameterName The name of the Projector parameter to get the default value of. */
+	CUSTOMIZABLEOBJECT_API UTexture* GetTextureParameterDefaultValue(const FString& InParameterName) const;
 
+	/** Get the default value of a parameter of type Skeletal Mesh.
+	  * @param InParameterName The name of the Projector parameter to get the default value of. */
+	CUSTOMIZABLEOBJECT_API USkeletalMesh* GetSkeletalMeshParameterDefaultValue(const FString& ParameterName) const;
+	
 	UE_DEPRECATED(5.6, "Use the ParameterName signature instead")
 	CUSTOMIZABLEOBJECT_API bool IsParameterMultidimensional(const int32& InParamIndex) const;
 	
diff --git a/Engine/Plugins/Mutable/Source/CustomizableObject/Public/MuCO/CustomizableObjectInstance.h b/Engine/Plugins/Mutable/Source/CustomizableObject/Public/MuCO/CustomizableObjectInstance.h
--- a/Engine/Plugins/Mutable/Source/CustomizableObject/Public/MuCO/CustomizableObjectInstance.h
+++ b/Engine/Plugins/Mutable/Source/CustomizableObject/Public/MuCO/CustomizableObjectInstance.h
@@ -289,7 +289,6 @@ public:
 	
 	// UObject interface.
 #if WITH_EDITOR
-	UE_API virtual void PostEditChangeProperty(FPropertyChangedEvent& PropertyChangedEvent) override;
 	UE_API virtual bool CanEditChange( const FProperty* InProperty ) const override;
 	UE_API virtual void PostTransacted(const FTransactionObjectEvent& TransactionEvent) override;
 #endif //WITH_EDITOR
@@ -440,12 +439,18 @@ public:
 	
 	/** Gets the value of a texture parameter with name "TextureParamName". */
 	UFUNCTION(BlueprintCallable, Category = CustomizableObjectInstance)
-	UE_API FName GetTextureParameterSelectedOption(const FString& TextureParamName, int32 RangeIndex = -1) const;
+	UE_API UTexture* GetTextureParameterSelectedOption(const FString& TextureParamName, int32 RangeIndex = -1) const;
 
 	/** Sets the texture value "TextureValue" of a texture parameter with index "TextureParamIndex". */
 	UFUNCTION(BlueprintCallable, Category = CustomizableObjectInstance)
-	UE_API void SetTextureParameterSelectedOption(const FString& TextureParamName, const FString& TextureValue, int32 RangeIndex = -1);
+	UE_API void SetTextureParameterSelectedOption(const FString& TextureParamName, UTexture* TextureValue, int32 RangeIndex = -1);
 
+	UFUNCTION(BlueprintCallable, Category = CustomizableObjectInstance)
+	UE_API USkeletalMesh* GetSkeletalMeshParameterSelectedOption(const FString& SkeletalMeshParamName, int32 RangeIndex = -1);
+
+	UFUNCTION(BlueprintCallable, Category = CustomizableObjectInstance)
+	UE_API void SetSkeletalMeshParameterSelectedOption(const FString& SkeletalMeshParamName, USkeletalMesh* SkeletalMeshValue, int32 RangeIndex = -1);
+	
 	/** Gets the value of a color parameter with name "ColorParamName". */
 	UFUNCTION(BlueprintCallable, Category = CustomizableObjectInstance)
 	UE_API FLinearColor GetColorParameterSelectedOption(const FString& ColorParamName) const;
@@ -740,14 +745,6 @@ private:
 	UPROPERTY()
 	TObjectPtr<UCustomizableInstancePrivate> PrivateData;
 
-public:
-#if WITH_EDITORONLY_DATA
-	/** Textures which can used as values in Texture Parameters. */
-	UPROPERTY(EditAnywhere, Category = TextureParameter)
-	TArray<TObjectPtr<UTexture2D>> TextureParameterDeclarations;
-#endif
-
-private:
 	// Deprecated properties	
 	UPROPERTY()
 	TObjectPtr<UCustomizableObject> CustomizableObject_DEPRECATED;
@@ -762,7 +759,7 @@ private:
 	TArray<FCustomizableObjectFloatParameterValue> FloatParameters_DEPRECATED;
 
 	UPROPERTY()																																						
-	TArray<FCustomizableObjectAssetParameterValue> TextureParameters_DEPRECATED;
+	TArray<FCustomizableObjectTextureParameterValue> TextureParameters_DEPRECATED;
 
 	UPROPERTY()
 	TArray<FCustomizableObjectVectorParameterValue> VectorParameters_DEPRECATED;
diff --git a/Engine/Plugins/Mutable/Source/CustomizableObject/Public/MuCO/CustomizableObjectInstanceDescriptor.h b/Engine/Plugins/Mutable/Source/CustomizableObject/Public/MuCO/CustomizableObjectInstanceDescriptor.h
--- a/Engine/Plugins/Mutable/Source/CustomizableObject/Public/MuCO/CustomizableObjectInstanceDescriptor.h
+++ b/Engine/Plugins/Mutable/Source/CustomizableObject/Public/MuCO/CustomizableObjectInstanceDescriptor.h
@@ -10,7 +10,7 @@
 
 #define UE_API CUSTOMIZABLEOBJECT_API
 
-class UTexture2D;
+class UTexture;
 enum class ECustomizableObjectProjectorType : uint8;
 
 class FArchive;
@@ -82,22 +82,6 @@ struct FCustomizableObjectInstanceDescriptor
 	// ------------------------------------------------------------
 	// FParameters
 	// ------------------------------------------------------------
-
-	UE_API const TArray<FCustomizableObjectBoolParameterValue>& GetBoolParameters() const;
-	
-	UE_API const TArray<FCustomizableObjectIntParameterValue>& GetIntParameters() const;
-
-	UE_API const TArray<FCustomizableObjectFloatParameterValue>& GetFloatParameters() const;
-	
-	UE_API const TArray<FCustomizableObjectAssetParameterValue>& GetTextureParameters() const;
-
-	UE_API const TArray<FCustomizableObjectAssetParameterValue>& GetMeshParameters() const;
-
-	UE_API const TArray<FCustomizableObjectVectorParameterValue>& GetVectorParameters() const;
-
-	UE_API const TArray<FCustomizableObjectProjectorParameterValue>& GetProjectorParameters() const;
-	
-	UE_API const TArray<FCustomizableObjectTransformParameterValue>& GetTransformParameters() const;
 	
 	/** Return true if there are any parameters. */
 	UE_API bool HasAnyParameters() const;
@@ -118,10 +102,16 @@ struct FCustomizableObjectInstanceDescriptor
 	UE_API void SetFloatParameterSelectedOption(const FString& FloatParamName, float FloatValue, int32 RangeIndex = -1);
 
 	/** Gets the value of a texture parameter with name "TextureParamName". */
-	UE_API FName GetTextureParameterSelectedOption(const FString& TextureParamName, int32 RangeIndex) const;
+	UE_API UTexture* GetTextureParameterSelectedOption(const FString& TextureParamName, int32 RangeIndex) const;
 
 	/** Sets the texture value "TextureValue" of a texture parameter with index "TextureParamIndex". */
-	UE_API void SetTextureParameterSelectedOption(const FString& TextureParamName, const FString& TextureValue, int32 RangeIndex);
+	UE_API void SetTextureParameterSelectedOption(const FString& TextureParamName, UTexture* TextureValue, int32 RangeIndex);
+	
+	/** Gets the value of a skeletalMesh parameter with name "SkeletalMeshParamName". */
+	UE_API USkeletalMesh* GetSkeletalMeshParameterSelectedOption(const FString& SkeletalMeshParamName, int32 RangeIndex = -1) const;
+
+	/** Sets the skeletalMesh value "SkeletalMeshValue" of a skeletalMesh parameter with index "SkeletalMeshParamIndex". */
+	UE_API void SetSkeletalMeshParameterSelectedOption(const FString& SkeletalMeshParamName, USkeletalMesh* SkeletalMeshValue, int32 RangeIndex = -1);
 	
 	/** Gets the value of a color parameter with name "ColorParamName". */
 	UE_API FLinearColor GetColorParameterSelectedOption(const FString& ColorParamName) const;
@@ -316,10 +306,10 @@ struct FCustomizableObjectInstanceDescriptor
 	TArray<FCustomizableObjectFloatParameterValue> FloatParameters;
 
 	UPROPERTY()
-	TArray<FCustomizableObjectAssetParameterValue> TextureParameters;
+	TArray<FCustomizableObjectTextureParameterValue> TextureParameters;
 
 	UPROPERTY()
-	TArray<FCustomizableObjectAssetParameterValue> MeshParameters;
+	TArray<FCustomizableObjectSkeletalMeshParameterValue> SkeletalMeshParameters;
 
 	UPROPERTY()
 	TArray<FCustomizableObjectVectorParameterValue> VectorParameters;
diff --git a/Engine/Plugins/Mutable/Source/CustomizableObject/Public/MuCO/CustomizableObjectParameterTypeDefinitions.h b/Engine/Plugins/Mutable/Source/CustomizableObject/Public/MuCO/CustomizableObjectParameterTypeDefinitions.h
--- a/Engine/Plugins/Mutable/Source/CustomizableObject/Public/MuCO/CustomizableObjectParameterTypeDefinitions.h
+++ b/Engine/Plugins/Mutable/Source/CustomizableObject/Public/MuCO/CustomizableObjectParameterTypeDefinitions.h
@@ -6,20 +6,22 @@
 
 #include "CustomizableObjectParameterTypeDefinitions.generated.h"
 
-class UTexture2D;
+class UTexture;
+class USkeletalMesh;
+
 
 UENUM(BlueprintType)
 enum class EMutableParameterType : uint8
 {
-	None		UMETA(DisplayName = "Unknown"),
-	Bool		UMETA(DisplayName = "Boolean"),
-	Int			UMETA(DisplayName = "Integer"),
-	Float		UMETA(DisplayName = "Float"),
-	Color		UMETA(DisplayName = "Color"),
-	Projector	UMETA(DisplayName = "Projector"),
-	Texture		UMETA(DisplayName = "Texture"),
-	Transform	UMETA(DisplayName = "Transform"),
-	Mesh		UMETA(DisplayName = "Mesh"),
+	None			UMETA(DisplayName = "Unknown"),
+	Bool			UMETA(DisplayName = "Boolean"),
+	Int				UMETA(DisplayName = "Integer"),
+	Float			UMETA(DisplayName = "Float"),
+	Color			UMETA(DisplayName = "Color"),
+	Projector		UMETA(DisplayName = "Projector"),
+	Texture			UMETA(DisplayName = "Texture"),
+	Transform		UMETA(DisplayName = "Transform"),
+	SkeletalMesh	UMETA(DisplayName = "SkeletalMesh"),
 };
 
 
@@ -157,42 +159,71 @@ inline uint32 GetTypeHash(const FCustomizableObjectFloatParameterValue& Key)
 
 
 USTRUCT(BlueprintType)
-struct FCustomizableObjectAssetParameterValue
+struct FCustomizableObjectTextureParameterValue
 {
 	GENERATED_USTRUCT_BODY()
 
-	inline static const FName DEFAULT_PARAMETER_VALUE = FName();
-	
 	UPROPERTY(Category = CustomizableObjectTextureParameterValue, VisibleAnywhere)
 	FString ParameterName;
 
 	UPROPERTY(Category = CustomizableObjectTextureParameterValue, VisibleAnywhere)
-	FName ParameterValue;
+	TObjectPtr<UTexture> ParameterValue;
 
 	UPROPERTY(Category = CustomizableObjectTextureParameterValue, VisibleAnywhere)
 	FGuid Id;
 
 	UPROPERTY(Category = CustomizableObjectTextureParameterValue, VisibleAnywhere)
-	TArray<FName> ParameterRangeValues;
+	TArray<TObjectPtr<UTexture>> ParameterRangeValues;
 };
 
-using FCustomizableObjectTextureParameterValue = FCustomizableObjectAssetParameterValue;
 
-
-inline uint32 GetTypeHash(const FCustomizableObjectAssetParameterValue& Key)
+inline uint32 GetTypeHash(const FCustomizableObjectTextureParameterValue& Key)
 {
 	uint32 Hash = GetTypeHash(Key.ParameterName);
 	Hash = HashCombine(Hash, GetTypeHash(Key.ParameterValue));
 
-	for (const FName& Value : Key.ParameterRangeValues)
+	for (const TObjectPtr<UTexture>& Value : Key.ParameterRangeValues)
     {
-    	Hash = HashCombine(Hash, GetTypeHash(Value));
+    	Hash = HashCombine(Hash, GetTypeHash(Value.GetPathName()));
     }
 
 	return Hash;
 }
 
 
+USTRUCT()
+struct FCustomizableObjectSkeletalMeshParameterValue
+{
+	GENERATED_USTRUCT_BODY()
+	
+	UPROPERTY(Category = CustomizableObjectTextureParameterValue, VisibleAnywhere)
+	FString ParameterName;
+
+	UPROPERTY(Category = CustomizableObjectTextureParameterValue, VisibleAnywhere)
+	TObjectPtr<USkeletalMesh> ParameterValue;
+
+	UPROPERTY(Category = CustomizableObjectTextureParameterValue, VisibleAnywhere)
+	FGuid Id;
+
+	UPROPERTY(Category = CustomizableObjectTextureParameterValue, VisibleAnywhere)
+	TArray<TObjectPtr<USkeletalMesh>> ParameterRangeValues;
+};
+
+
+inline uint32 GetTypeHash(const FCustomizableObjectSkeletalMeshParameterValue& Key)
+{
+	uint32 Hash = GetTypeHash(Key.ParameterName);
+	Hash = HashCombine(Hash, GetTypeHash(Key.ParameterValue));
+
+	for (const TObjectPtr<USkeletalMesh>& Value : Key.ParameterRangeValues)
+	{
+		Hash = HashCombine(Hash, GetTypeHash(Value.GetPathName()));
+	}
+
+	return Hash;
+}
+
+
 USTRUCT(BlueprintType)
 struct FCustomizableObjectVectorParameterValue
 {
diff --git a/Engine/Plugins/Mutable/Source/CustomizableObject/Public/MuCO/CustomizableObjectSystem.h b/Engine/Plugins/Mutable/Source/CustomizableObject/Public/MuCO/CustomizableObjectSystem.h
--- a/Engine/Plugins/Mutable/Source/CustomizableObject/Public/MuCO/CustomizableObjectSystem.h
+++ b/Engine/Plugins/Mutable/Source/CustomizableObject/Public/MuCO/CustomizableObjectSystem.h
@@ -20,7 +20,6 @@ class IConsoleVariable;
 class UCustomizableInstanceLODManagementBase;
 class ITargetPlatform;
 class UCustomizableObject;
-class UDefaultImageProvider;
 class USkeletalMesh;
 class UMaterialInterface;
 class UTexture2D;
@@ -141,54 +140,6 @@ struct FCustomizableObjectExternalTexture
 };
 
 
-/** Base class for Image provider.
- * 
- *  An image provider only needs to be alive when calling UpdateSkeletalMeshAsync. Once called it can be safety destructed.
- *  Mutable will cache those images in case they are needed due to MIP updates. */
-UCLASS(MinimalAPI, abstract)
-class UCustomizableSystemImageProvider : public UObject
-{
-public:
-	GENERATED_BODY()
-
-	enum class ValueType : uint8
-	{
-		// This texture is not provided by this provider.
-		None,
-
-		// Data will be provided with size and pointer
-		Raw,
-
-		// Data will be provided from an unreal texture, loaded in the game thread and kept in memory
-		Unreal,
-
-		// Data will be provided from an unreal texture, and will only be loaded when actually needed in the Mutable thread
-		Unreal_Deferred,
-
-		// Number of elements of this enum.
-		Count
-	};
-
-	// Query that Mutable will run to find out if a texture will be provided as an Unreal UTexture2D,
-	// or as a raw data blob.
-	virtual ValueType HasTextureParameterValue(const FName& ID) { return ValueType::None; }
-
-	// In case IsTextureParameterValueUnreal returns false, this will be used to query the texture size data.
-	virtual FIntVector GetTextureParameterValueSize(const FName& ID) { return FIntVector(0, 0, 0); }
-
-	// In case IsTextureParameterValueUnreal returns false, this will be used to query the texture data that must
-	// be copied in the preallocated buffer. The pixel format is assumed to be 4-channel RGBA, uint8_t per channel.
-	virtual void GetTextureParameterValueData(const FName& ID, uint8* OutData) {}
-
-	// In case IsTextureParameterValueUnreal returns true, this will be used to query the texture.
-	virtual UTexture2D* GetTextureParameterValue(const FName& ID) { return nullptr; }
-
-	// Used in the editor to show the list of available options.
-	// Only necessary if the images are required in editor previews.
-	virtual void GetTextureParameterValues(TArray<FCustomizableObjectExternalTexture>& OutValues) {}
-};
-
-
 UCLASS(MinimalAPI, BlueprintType)
 class UCustomizableObjectSystem : public UObject
 {
@@ -271,28 +222,6 @@ public:
 	UE_API void SetImagePixelFormatOverride(const mu::FImageOperator::FImagePixelFormatFunc&);
 #endif
 
-	/** [Texture FParameters] Get a list of all the possible values for external texture parameters according to the various providers registered with RegisterImageProvider. */
-	UE_API TArray<FCustomizableObjectExternalTexture> GetTextureParameterValues();
-
-	/** [Texture FParameters] Add a new image provider to the CustomizableObject System. This will be queried for when an external image ID is provided to mutable in an Texture Parameter node. */
-	UE_API void RegisterImageProvider(UCustomizableSystemImageProvider* Provider);
-
-	/** [Texture FParameters] Remove a previously registered provider. */
-	UE_API void UnregisterImageProvider(UCustomizableSystemImageProvider* Provider);
-
-	/** [Texture FParameters] Interface to actually cache Images in the Mutable system and make them available at run-time.
-		Any cached image has to be registered by an Image provider before caching it.
-		Have in mind that once an image has been cached, it will spend memory according to its size, except in the case
-		of images of type UCustomizableSystemImageProvider::ValueType::Unreal_Deferred, where only a very small amount of
-		memory is used and the real texel data is loaded when needed during an update and then immediately discarded */
-
-	/** [Texture FParameters] Cache an image which has to have been previously registered by an Image provider with the parameter id. */
-	UE_API void CacheImage(FName ImageId);
-	/** [Texture FParameters] Remove an image from the cache. */
-	UE_API void UnCacheImage(FName ImageId);
-	/** [Texture FParameters] Remove all images from the cache. */
-	UE_API void ClearImageCache();
-
 	// Show a warning on-screen and via a notification (if in Editor) and log an error when a CustomizableObject is
 	// being used and it's not compiled.  Callers can add additional information to the error log.
 	UE_API void AddUncompiledCOWarning(const UCustomizableObject& InObject, FString const* OptionalLogInfo = nullptr);
diff --git a/Engine/Plugins/Mutable/Source/CustomizableObject/Public/MuCO/DefaultImageProvider.h b/Engine/Plugins/Mutable/Source/CustomizableObject/Public/MuCO/DefaultImageProvider.h
deleted file mode 100644
diff --git a/Engine/Plugins/Mutable/Source/CustomizableObjectEditor/Private/MuCOE/CustomizableInstanceDetails.cpp b/Engine/Plugins/Mutable/Source/CustomizableObjectEditor/Private/MuCOE/CustomizableInstanceDetails.cpp
--- a/Engine/Plugins/Mutable/Source/CustomizableObjectEditor/Private/MuCOE/CustomizableInstanceDetails.cpp
+++ b/Engine/Plugins/Mutable/Source/CustomizableObjectEditor/Private/MuCOE/CustomizableInstanceDetails.cpp
@@ -15,14 +15,13 @@
 #include "Serialization/BufferArchive.h"
 #include "Toolkits/ToolkitManager.h"
 #include "Math/Transform.h"
-
+#include "PropertyCustomizationHelpers.h"
 #include "MuCO/CustomizableObjectInstancePrivate.h"
 #include "MuCO/CustomizableObject.h"
 #include "MuCO/CustomizableObjectInstance.h"
 #include "MuCO/CustomizableObjectSystem.h"
 #include "MuCOE/CustomizableObjectEditor.h"
 #include "MuCOE/UnrealEditorPortabilityHelpers.h"
-
 #include "AssetThumbnail.h"
 #include "Slate/DeferredCleanupSlateBrush.h"
 #include "SSearchableComboBox.h"
@@ -41,6 +40,8 @@
 #include "Widgets/Layout/SExpandableArea.h"
 #include "Widgets/Layout/SGridPanel.h"
 #include "Widgets/Layout/SUniformGridPanel.h"
+#include "Engine/Texture.h"
+#include "Engine/SkeletalMesh.h"
 
 class UObject;
 
@@ -113,9 +114,6 @@ void FCustomizableInstanceDetails::CustomizeDetails(const TSharedPtr<IDetailLayo
 	IDetailCategoryBuilder& ResourcesCategory = DetailBuilder->EditCategory("Generated Resources");
 	IDetailCategoryBuilder& VisibilitySettingsCategory = DetailBuilder->EditCategory("ParametersVisibility");
 	IDetailCategoryBuilder& ParametersCategory = DetailBuilder->EditCategory("Instance Parameters");
-	IDetailCategoryBuilder& OldParametersCategory = DetailBuilder->EditCategory("Old Instance Parameters");
-	IDetailCategoryBuilder& SkeletalMeshCategory = DetailBuilder->EditCategory("CustomizableSkeletalMesh");
-	IDetailCategoryBuilder& TextureParametersCategory = DetailBuilder->EditCategory("TextureParameter");
 
 	// Showing warning message in case that the instance has not been compiled
 	UCustomizableObject* CustomizableObject = CustomInstance->GetCustomizableObject();
@@ -140,9 +138,6 @@ void FCustomizableInstanceDetails::CustomizeDetails(const TSharedPtr<IDetailLayo
 	ResourcesCategory.InitiallyCollapsed(true);
 	IDetailPropertyRow* PrivateDataRow = ResourcesCategory.AddExternalObjects(Private, EPropertyLocation::Default, PrivatePropertyParams);
 
-	// Store which Texture Parameter values can be selected.
-	GenerateTextureParameterOptions();
-
 	// State Selector Widget
 	VisibilitySettingsCategory.AddCustomRow(LOCTEXT("CustomizableInstanceDetails_StateSelector", "State"))
 	.NameContent()
@@ -977,7 +972,23 @@ TSharedRef<SWidget> FCustomizableInstanceDetails::GenerateParameterWidget(const
 		}
 		case EMutableParameterType::Texture:
 		{
-			return GenerateTextureWidget(CustomizableObject, ParamName);
+			UTexture* Texture = CustomInstance->GetTextureParameterSelectedOption(ParamName);
+			const FString ObjectPath = Texture ? Texture->GetPathName() : TEXT("");
+
+			return SNew(SObjectPropertyEntryBox)
+				.AllowedClass(UTexture2D::StaticClass())
+				.OnObjectChanged(this, &FCustomizableInstanceDetails::OnTextureParameterSelected, ParamName)
+				.ObjectPath(ObjectPath);
+		}
+		case EMutableParameterType::SkeletalMesh:
+		{
+			USkeletalMesh* SkeletalMesh = CustomInstance->GetSkeletalMeshParameterSelectedOption(ParamName);
+			const FString ObjectPath = SkeletalMesh ? SkeletalMesh->GetPathName() : TEXT("");
+
+			return SNew(SObjectPropertyEntryBox)
+				.AllowedClass(USkeletalMesh::StaticClass())
+				.OnObjectChanged(this, &FCustomizableInstanceDetails::OnSkeletalMeshParameterSelected, ParamName)
+				.ObjectPath(ObjectPath);
 		}
 		case EMutableParameterType::Transform:
 		{
@@ -1308,108 +1319,23 @@ void FCustomizableInstanceDetails::OnFloatParameterCommited(float Value, ETextCo
 	}
 }
 
-
 // TEXTURE PARAMETERS -----------------------------------------------------------------------------------------------------------------
 
-TSharedRef<SWidget> FCustomizableInstanceDetails::GenerateTextureWidget(const UCustomizableObject& CustomizableObject, const FString& ParamName)
+void FCustomizableInstanceDetails::OnTextureParameterSelected(const FAssetData& AssetData, const FString ParamName)
 {
-	const FName ParameterValue = CustomInstance->GetTextureParameterSelectedOption(ParamName);
-	TSharedPtr<FString> InitiallySelected;
-
-	// Look for the value index
-	for (int32 ValueIndex = 0; ValueIndex < TextureParameterValueNames.Num(); ++ValueIndex)
-	{
-		if (ParameterValue == TextureParameterValues[ValueIndex])
-		{
-			InitiallySelected = TextureParameterValueNames[ValueIndex];
-			break;
-		}
-	}
-
-	return SNew(STextComboBox)
-		.OptionsSource(&TextureParameterValueNames)
-		.InitiallySelectedItem(InitiallySelected)
-		.OnSelectionChanged(this, &FCustomizableInstanceDetails::OnTextureParameterComboBoxSelectionChanged, ParamName);
+	CustomInstance->SetTextureParameterSelectedOption(ParamName, Cast<UTexture2D>(AssetData.GetAsset()));
+	UpdateInstance();
 }
 
 
-void FCustomizableInstanceDetails::GenerateTextureParameterOptions()
-{
-	UCustomizableObject* CustomizableObject = CustomInstance->GetCustomizableObject();
-
-	for (const FCustomizableObjectAssetParameterValue& TextureParameter : CustomInstance->GetPrivate()->GetDescriptor().GetTextureParameters())
-	{
-		// Get default values.
-		const FName& DefaultValue = CustomizableObject->GetTextureParameterDefaultValue(TextureParameter.ParameterName);
-
-		if (!TextureParameterValues.Contains(DefaultValue))
-		{
-			TextureParameterValueNames.Add(MakeShareable(new FString(DefaultValue.ToString())));
-			TextureParameterValues.Add(DefaultValue);
-		}
-
-		// Selected parameter value
-		const FName& SelectedValue = TextureParameter.ParameterValue;
-
-		if (!TextureParameterValues.Contains(SelectedValue))
-		{
-			TextureParameterValueNames.Add(MakeShareable(new FString(SelectedValue.ToString())));
-			TextureParameterValues.Add(SelectedValue);
-		}
-	}
+// SKELETAL MESH PARAMETERS -----------------------------------------------------------------------------------------------------------------
 
-	// Get values from registered providers providers.
-	TArray<FCustomizableObjectExternalTexture> Textures = UCustomizableObjectSystem::GetInstance()->GetTextureParameterValues();
-	for (int i = 0; i < Textures.Num(); ++i)
-	{
-		const FName& Value = FName(Textures[i].Value);
-
-		if (!TextureParameterValues.Contains(Value))
-		{
-			TextureParameterValueNames.Add(MakeShareable(new FString(Textures[i].Name)));
-			TextureParameterValues.Add(Value);
-		}
-	}
-
-	// Get values from TextureParameterDeclarations.
-	for (TObjectPtr<UTexture2D> Declaration : CustomInstance->TextureParameterDeclarations)
-	{
-		if (!Declaration)
-		{
-			continue;
-		}
-
-		const FName Value = FName(Declaration.GetPathName());
-
-		if (!TextureParameterValues.Contains(Value))
-		{
-			TextureParameterValueNames.Add(MakeShareable(new FString(Declaration.GetName())));
-			TextureParameterValues.Add(Value);
-		}
-	}
-}
-
-
-void FCustomizableInstanceDetails::OnTextureParameterComboBoxSelectionChanged(TSharedPtr<FString> Selection, ESelectInfo::Type SelectInfo, FString ParamName)
+void FCustomizableInstanceDetails::OnSkeletalMeshParameterSelected(const FAssetData& AssetData, const FString ParamName)
 {
-	if (!Selection.IsValid())
-	{
-		return;
-	}
-
-	const int32 Index = TextureParameterValueNames.Find(Selection);
-	check(Index != INDEX_NONE);
-	
-	FScopedTransaction LocalTransaction(FText::Format(LOCTEXT("OnTextureParameterSelected", "Set Texture Parameter: {0}"), FText::FromString(ParamName)));
-	CustomInstance->Modify();
-	CustomInstance->SetTextureParameterSelectedOption(ParamName, TextureParameterValues[Index].ToString());
+	CustomInstance->SetSkeletalMeshParameterSelectedOption(ParamName, Cast<USkeletalMesh>(AssetData.GetAsset()));
 	UpdateInstance();
-
-	// Non-continuous change: collect garbage.
-	GEngine->ForceGarbageCollection();
 }
 
-
 // COLOR PARAMETERS -----------------------------------------------------------------------------------------------------------------
 
 TSharedRef<SWidget> FCustomizableInstanceDetails::GenerateColorWidget(const FString& ParamName)
@@ -1886,7 +1812,7 @@ FReply FCustomizableInstanceDetails::OnProjectorCopyTransform(const FString Para
 
 	if ((ParameterIndexInObject >= 0) && (ProjectorParamIndex >= 0))
 	{
-		const TArray<FCustomizableObjectProjectorParameterValue>& ProjectorParameters = CustomInstance->GetPrivate()->GetDescriptor().GetProjectorParameters();
+		const TArray<FCustomizableObjectProjectorParameterValue>& ProjectorParameters = CustomInstance->GetPrivate()->GetDescriptor().ProjectorParameters;
 		FCustomizableObjectProjector Value;
 
 		if (RangeIndex == -1)
@@ -1966,7 +1892,7 @@ FReply FCustomizableInstanceDetails::OnProjectorResetTransform(const FString Par
 TSharedRef<SWidget> FCustomizableInstanceDetails::GenerateMultidimensionalProjector(const UCustomizableObject& CustomizableObject, const FString& ParamName, const int32 ParamIndexInObject)
 {
 	const TSharedPtr<ICustomizableObjectInstanceEditor> Editor = GetEditorChecked();
-	const TArray<FCustomizableObjectProjectorParameterValue>& ProjectorParameters = CustomInstance->GetPrivate()->GetDescriptor().GetProjectorParameters();
+	const TArray<FCustomizableObjectProjectorParameterValue>& ProjectorParameters = CustomInstance->GetPrivate()->GetDescriptor().ProjectorParameters;
 	const int32 ProjectorParamIndex = CustomInstance->GetPrivate()->FindProjectorParameterNameIndex(ParamName);
 
 	check(ProjectorParamIndex < ProjectorParameters.Num());
diff --git a/Engine/Plugins/Mutable/Source/CustomizableObjectEditor/Private/MuCOE/CustomizableInstanceDetails.h b/Engine/Plugins/Mutable/Source/CustomizableObjectEditor/Private/MuCOE/CustomizableInstanceDetails.h
--- a/Engine/Plugins/Mutable/Source/CustomizableObjectEditor/Private/MuCOE/CustomizableInstanceDetails.h
+++ b/Engine/Plugins/Mutable/Source/CustomizableObjectEditor/Private/MuCOE/CustomizableInstanceDetails.h
@@ -103,12 +103,13 @@ private:
 	void OnFloatParameterSpinBoxEnd(float Value, const FString ParamName, int32 RangeIndex);
 	void OnFloatParameterSliderEnd();
 	void OnFloatParameterCommited(float Value, ETextCommit::Type Type, const FString ParamName, int32 RangeIndex);	// Needed to have undo/redo
- 
+
 	// Texture FParameters Functions
-	TSharedRef<SWidget> GenerateTextureWidget(const UCustomizableObject& CustomizableObject, const FString& ParamName);
-	void GenerateTextureParameterOptions();
-	void OnTextureParameterComboBoxSelectionChanged(TSharedPtr<FString> Selection, ESelectInfo::Type SelectInfo, const FString ParamName);
+	void OnTextureParameterSelected(const FAssetData& AssetData, const FString ParamName);
 	
+	// Skeletal Mesh FParameters Functions
+	void OnSkeletalMeshParameterSelected(const FAssetData& AssetData, const FString ParamName);
+
 	// Color FParameters Functions
 	TSharedRef<SWidget> GenerateColorWidget(const FString& ParamName);
 	FLinearColor GetColorParameterValue(const FString ParamName) const;
@@ -171,12 +172,7 @@ private:
 
 	/** Array with all the possible states for the states combo box */
 	TArray< TSharedPtr<FString> > StateNames;
-
-	// These arrays store the textures available for texture parameters of the model.
-	// These come from the texture generators registered in the CustomizableObjectSystem
-	TArray<TSharedPtr<FString>> TextureParameterValueNames;
-	TArray<FName> TextureParameterValues;
-
+	
 	/** Weak pointer of the open editor */
 	TWeakPtr<ICustomizableObjectInstanceEditor> WeakEditor;
 
diff --git a/Engine/Plugins/Mutable/Source/CustomizableObjectEditor/Private/MuCOE/CustomizableObjectCompiler.cpp b/Engine/Plugins/Mutable/Source/CustomizableObjectEditor/Private/MuCOE/CustomizableObjectCompiler.cpp
--- a/Engine/Plugins/Mutable/Source/CustomizableObjectEditor/Private/MuCOE/CustomizableObjectCompiler.cpp
+++ b/Engine/Plugins/Mutable/Source/CustomizableObjectEditor/Private/MuCOE/CustomizableObjectCompiler.cpp
@@ -928,7 +928,7 @@ void FCustomizableObjectCompiler::CompileInternal(bool bAsync)
 	{
 		TArray<FMutableSourceTextureData> NewCompileTimeReferencedTextures;
 		{
-			for (const TPair<TSoftObjectPtr<const UTexture>, FMutableGraphGenerationContext::FGeneratedReferencedTexture>& Pair : GenerationContext.CompileTimeTextureMap)
+			for (const TPair<TSoftObjectPtr<UTexture2D>, FMutableGraphGenerationContext::FGeneratedReferencedTexture>& Pair : GenerationContext.CompileTimeTextureMap)
 			{
 				check(Pair.Value.ID == NewCompileTimeReferencedTextures.Num());
 
@@ -980,12 +980,15 @@ void FCustomizableObjectCompiler::CompileInternal(bool bAsync)
 			ModelResources->PassThroughTextures.Add(Pair.Key);
 		}
 
-		for (const TPair<TSoftObjectPtr<const UTexture>, FMutableGraphGenerationContext::FGeneratedReferencedTexture>& Pair : GenerationContext.RuntimeReferencedTextureMap)
+		for (const TPair<TSoftObjectPtr<UTexture2D>, FMutableGraphGenerationContext::FGeneratedReferencedTexture>& Pair : GenerationContext.RuntimeReferencedTextureMap)
 		{
 			check(Pair.Value.ID == ModelResources->RuntimeReferencedTextures.Num());
 			ModelResources->RuntimeReferencedTextures.Add(Pair.Key);
 		}
 
+		ModelResources->TextureParameterDefaultValues = MoveTemp(GenerationContext.TextureParameterDefaultValues);
+		ModelResources->SkeletalMeshParameterDefaultValues = MoveTemp(GenerationContext.SkeletalMeshParameterDefaultValues);
+
 		ModelResources->AnimBPs = MoveTemp(GenerationContext.AnimBPAssets);
 
 		ModelResources->MaterialSlotNames = MoveTemp(GenerationContext.ReferencedMaterialSlotNames);
@@ -1643,14 +1646,14 @@ mu::Ptr<mu::Node> FCustomizableObjectCompiler::Export(UCustomizableObject* Objec
 
 	// Pass out the referenced textures
 	OutRuntimeReferencedTextures.Empty();
-	for (const TPair<TSoftObjectPtr<const UTexture>, FMutableGraphGenerationContext::FGeneratedReferencedTexture>& Pair : GenerationContext.RuntimeReferencedTextureMap)
+	for (const TPair<TSoftObjectPtr<UTexture2D>, FMutableGraphGenerationContext::FGeneratedReferencedTexture>& Pair : GenerationContext.RuntimeReferencedTextureMap)
 	{
 		check(Pair.Value.ID == OutRuntimeReferencedTextures.Num());
 		OutRuntimeReferencedTextures.Add(Pair.Key);
 	}
 
 	OutCompilerReferencedTextures.Empty();
-	for (const TPair<TSoftObjectPtr<const UTexture>, FMutableGraphGenerationContext::FGeneratedReferencedTexture>& Pair : GenerationContext.CompileTimeTextureMap)
+	for (const TPair<TSoftObjectPtr<UTexture2D>, FMutableGraphGenerationContext::FGeneratedReferencedTexture>& Pair : GenerationContext.CompileTimeTextureMap)
 	{
 		check(Pair.Value.ID == OutCompilerReferencedTextures.Num());
 
diff --git a/Engine/Plugins/Mutable/Source/CustomizableObjectEditor/Private/MuCOE/CustomizableObjectEditorModule.cpp b/Engine/Plugins/Mutable/Source/CustomizableObjectEditor/Private/MuCOE/CustomizableObjectEditorModule.cpp
--- a/Engine/Plugins/Mutable/Source/CustomizableObjectEditor/Private/MuCOE/CustomizableObjectEditorModule.cpp
+++ b/Engine/Plugins/Mutable/Source/CustomizableObjectEditor/Private/MuCOE/CustomizableObjectEditorModule.cpp
@@ -56,7 +56,7 @@
 #include "MuCOE/Nodes/CustomizableObjectNodeModifierClipWithUVMask.h"
 #include "MuCOE/Nodes/CustomizableObjectNodeModifierClipDeform.h"
 #include "MuCOE/Nodes/CustomizableObjectNodeModifierTransformInMesh.h"
-#include "MuCOE/Nodes/CustomizableObjectNodeMeshParameter.h"
+#include "MuCOE/Nodes/CustomizableObjectNodeSkeletalMeshParameter.h"
 #include "MuCOE/Nodes/CustomizableObjectNodeMeshMorph.h"
 #include "MuCOE/Nodes/CustomizableObjectNodeMeshMorphDetails.h"
 #include "MuCOE/Nodes/CustomizableObjectNodeMeshReshapeCommon.h"
@@ -70,8 +70,8 @@
 #include "MuCOE/Nodes/CustomizableObjectNodeSkeletalMesh.h"
 #include "MuCOE/Nodes/CustomizableObjectNodeSkeletalMeshDetails.h"
 #include "MuCOE/Nodes/CustomizableObjectNodeStaticMesh.h"
-#include "MuCOE/Nodes/CustomizableObjectNodeMeshParameter.h"
-#include "MuCOE/Nodes/CustomizableObjectNodeMeshParameterDetails.h"
+#include "MuCOE/Nodes/CustomizableObjectNodeSkeletalMeshParameter.h"
+#include "MuCOE/Nodes/CustomizableObjectNodeSkeletalMeshParameterDetails.h"
 #include "MuCOE/Nodes/CustomizableObjectNodeTable.h"
 #include "MuCOE/Nodes/CustomizableObjectNodeTableDetails.h"
 #include "MuCOE/Nodes/CustomizableObjectNodeTunnel.h"
@@ -201,7 +201,7 @@ TMap<FString, FString> GetCompileOnlySelectedParameters(const UCustomizableObjec
 {
 	TMap<FString, FString> Parameters;
 	
-	const TArray<FCustomizableObjectIntParameterValue>& IntParameters = Instance.GetPrivate()->GetDescriptor().GetIntParameters();
+	const TArray<FCustomizableObjectIntParameterValue>& IntParameters = Instance.GetPrivate()->GetDescriptor().IntParameters;
 	Parameters.Reserve(IntParameters.Num());
 
 	for (const FCustomizableObjectIntParameterValue& IntParam : IntParameters)
@@ -267,7 +267,7 @@ void FCustomizableObjectEditorModule::StartupModule()
 	RegisterCustomDetails(PropertyModule, UCustomizableObjectNodeMaterial::StaticClass(), FOnGetDetailCustomizationInstance::CreateStatic(&FCustomizableObjectNodeMeshSectionDetails::MakeInstance));
 	RegisterCustomDetails(PropertyModule, UCustomizableObjectNodeSkeletalMesh::StaticClass(), FOnGetDetailCustomizationInstance::CreateStatic(&FCustomizableObjectNodeSkeletalMeshDetails::MakeInstance));
 	RegisterCustomDetails(PropertyModule, UCustomizableObjectNodeStaticMesh::StaticClass(), FOnGetDetailCustomizationInstance::CreateStatic(&FCustomizableObjectNodeDetails::MakeInstance));
-	RegisterCustomDetails(PropertyModule, UCustomizableObjectNodeMeshParameter::StaticClass(), FOnGetDetailCustomizationInstance::CreateStatic(&FCustomizableObjectNodeMeshParameterDetails::MakeInstance));
+	RegisterCustomDetails(PropertyModule, UCustomizableObjectNodeSkeletalMeshParameter::StaticClass(), FOnGetDetailCustomizationInstance::CreateStatic(&FCustomizableObjectNodeSkeletalMeshParameterDetails::MakeInstance));
 	RegisterCustomDetails(PropertyModule, UCustomizableObjectNodeTable::StaticClass(), FOnGetDetailCustomizationInstance::CreateStatic(&FCustomizableObjectNodeTableDetails::MakeInstance));
 	RegisterCustomDetails(PropertyModule, UCustomizableObjectNodeComponentMesh::StaticClass(), FOnGetDetailCustomizationInstance::CreateStatic(&FCustomizableObjectNodeComponentMeshDetails::MakeInstance));
 	RegisterCustomDetails(PropertyModule, UCustomizableObjectNodeMacroInstance::StaticClass(), FOnGetDetailCustomizationInstance::CreateStatic(&FCustomizableObjectNodeMacroInstanceDetails::MakeInstance));
@@ -999,7 +999,7 @@ TMap<FName, FGuid> FCustomizableObjectEditorModule::GetParticipatingObjects(cons
 			AddObjectSoft(TransformNode->ReferenceTexture);
 		}
 
-		else if (UCustomizableObjectNodeMeshParameter* MeshParamNode = Cast<UCustomizableObjectNodeMeshParameter>(&Node))
+		else if (UCustomizableObjectNodeSkeletalMeshParameter* MeshParamNode = Cast<UCustomizableObjectNodeSkeletalMeshParameter>(&Node))
 		{
 			AddObjectSoft(MeshParamNode->DefaultValue.ToSoftObjectPath());
 		}
diff --git a/Engine/Plugins/Mutable/Source/CustomizableObjectEditor/Private/MuCOE/EdGraphSchema_CustomizableObject.cpp b/Engine/Plugins/Mutable/Source/CustomizableObjectEditor/Private/MuCOE/EdGraphSchema_CustomizableObject.cpp
--- a/Engine/Plugins/Mutable/Source/CustomizableObjectEditor/Private/MuCOE/EdGraphSchema_CustomizableObject.cpp
+++ b/Engine/Plugins/Mutable/Source/CustomizableObjectEditor/Private/MuCOE/EdGraphSchema_CustomizableObject.cpp
@@ -53,7 +53,7 @@
 #include "MuCOE/Nodes/CustomizableObjectNodeMeshReshape.h"
 #include "MuCOE/Nodes/CustomizableObjectNodeMeshSwitch.h"
 #include "MuCOE/Nodes/CustomizableObjectNodeMeshVariation.h"
-#include "MuCOE/Nodes/CustomizableObjectNodeMeshParameter.h"
+#include "MuCOE/Nodes/CustomizableObjectNodeSkeletalMeshParameter.h"
 #include "MuCOE/Nodes/CustomizableObjectNodeObjectChild.h"
 #include "MuCOE/Nodes/CustomizableObjectNodeObjectGroup.h"
 #include "MuCOE/Nodes/CustomizableObjectNodeComponentPassthroughMesh.h"
@@ -608,7 +608,7 @@ void UEdGraphSchema_CustomizableObject::GetGraphContextActions(FGraphContextMenu
 			ContextMenuBuilder.CreateTemplateNode<UCustomizableObjectNodeModifierClipDeform>(),
 			ContextMenuBuilder.CreateTemplateNode<UCustomizableObjectNodeModifierTransformInMesh>(),
 			ContextMenuBuilder.CreateTemplateNode<UCustomizableObjectNodeTextureParameter>(),
-			ContextMenuBuilder.CreateTemplateNode<UCustomizableObjectNodeMeshParameter>(),
+			ContextMenuBuilder.CreateTemplateNode<UCustomizableObjectNodeSkeletalMeshParameter>(),
 		};
 
 		AddNewNodeCategoryActionsFiltered(ExperimentalTemplateNodes, ContextMenuBuilder, NC_Experimental.ToString(), 2, Filter);
diff --git a/Engine/Plugins/Mutable/Source/CustomizableObjectEditor/Private/MuCOE/GenerateMutableSource/GenerateMutableSource.cpp b/Engine/Plugins/Mutable/Source/CustomizableObjectEditor/Private/MuCOE/GenerateMutableSource/GenerateMutableSource.cpp
--- a/Engine/Plugins/Mutable/Source/CustomizableObjectEditor/Private/MuCOE/GenerateMutableSource/GenerateMutableSource.cpp
+++ b/Engine/Plugins/Mutable/Source/CustomizableObjectEditor/Private/MuCOE/GenerateMutableSource/GenerateMutableSource.cpp
@@ -1779,12 +1779,12 @@ TSharedPtr<mu::FImage> GenerateImageConstant(UTexture* Texture, FMutableGraphGen
 	}
 	else if (bIsCompileTime)
 	{
-        Entry = &GenerationContext.CompileTimeTextureMap.FindOrAdd(Texture, InvalidEntry);
+        Entry = &GenerationContext.CompileTimeTextureMap.FindOrAdd(CastChecked<UTexture2D>(Texture), InvalidEntry);
 		Num = GenerationContext.CompileTimeTextureMap.Num();
 	}
 	else
 	{
-		Entry = &GenerationContext.RuntimeReferencedTextureMap.FindOrAdd(Texture, InvalidEntry);
+		Entry = &GenerationContext.RuntimeReferencedTextureMap.FindOrAdd(CastChecked<UTexture2D>(Texture), InvalidEntry);
 		Num = GenerationContext.RuntimeReferencedTextureMap.Num();
 	}
 
diff --git a/Engine/Plugins/Mutable/Source/CustomizableObjectEditor/Private/MuCOE/GenerateMutableSource/GenerateMutableSource.h b/Engine/Plugins/Mutable/Source/CustomizableObjectEditor/Private/MuCOE/GenerateMutableSource/GenerateMutableSource.h
--- a/Engine/Plugins/Mutable/Source/CustomizableObjectEditor/Private/MuCOE/GenerateMutableSource/GenerateMutableSource.h
+++ b/Engine/Plugins/Mutable/Source/CustomizableObjectEditor/Private/MuCOE/GenerateMutableSource/GenerateMutableSource.h
@@ -655,11 +655,11 @@ struct FMutableGraphGenerationContext
 
 	// Cache of runtime images and their IDs used in the core to identify them.
 	// These textures will remain as external references even in optimized models.
-	TMap<TSoftObjectPtr<const UTexture>, FGeneratedReferencedTexture> RuntimeReferencedTextureMap;
+	TMap<TSoftObjectPtr<UTexture2D>, FGeneratedReferencedTexture> RuntimeReferencedTextureMap;
 	
 	// Cache of runtime pass-through images and their IDs used in the core to identify them
 	// These textures will become mutable images in the compiled model.
-	TMap<TSoftObjectPtr<const UTexture>, FGeneratedReferencedTexture> CompileTimeTextureMap;
+	TMap<TSoftObjectPtr<UTexture2D>, FGeneratedReferencedTexture> CompileTimeTextureMap;
 
 	// Mutable meshes already build for source UStaticMesh or USkeletalMesh.
 	struct FGeneratedMeshData
@@ -886,6 +886,9 @@ struct FMutableGraphGenerationContext
 	// This should be automatized (create a struct that each time that a Macro is stacked, it adds 1 to the count.
 	// it also substracts 1 each time we remove a Macro from the stack)
 	TArray<const UCustomizableObjectNodeMacroInstance*> MacroNodesStack;
+	
+	TMap<FName, TObjectPtr<UTexture>> TextureParameterDefaultValues;
+	TMap<FName, TObjectPtr<USkeletalMesh>> SkeletalMeshParameterDefaultValues;
 };
 
 
diff --git a/Engine/Plugins/Mutable/Source/CustomizableObjectEditor/Private/MuCOE/GenerateMutableSource/GenerateMutableSourceImage.cpp b/Engine/Plugins/Mutable/Source/CustomizableObjectEditor/Private/MuCOE/GenerateMutableSource/GenerateMutableSourceImage.cpp
--- a/Engine/Plugins/Mutable/Source/CustomizableObjectEditor/Private/MuCOE/GenerateMutableSource/GenerateMutableSourceImage.cpp
+++ b/Engine/Plugins/Mutable/Source/CustomizableObjectEditor/Private/MuCOE/GenerateMutableSource/GenerateMutableSourceImage.cpp
@@ -163,7 +163,7 @@ mu::Ptr<mu::NodeImage> GenerateMutableSourceImage(const UEdGraphPin* Pin, FMutab
 
 		if (TypedNodeParam->DefaultValue)
 		{
-			TextureNode->DefaultValue = FName(TypedNodeParam->DefaultValue->GetPathName());
+			GenerationContext.TextureParameterDefaultValues.Add(FName(TypedNodeParam->GetParameterName()), TypedNodeParam->DefaultValue);
 		}
 
 		GenerationContext.ParameterUIDataMap.Add(TypedNodeParam->GetParameterName(&GenerationContext.MacroNodesStack), FMutableParameterData(
diff --git a/Engine/Plugins/Mutable/Source/CustomizableObjectEditor/Private/MuCOE/GenerateMutableSource/GenerateMutableSourceMesh.cpp b/Engine/Plugins/Mutable/Source/CustomizableObjectEditor/Private/MuCOE/GenerateMutableSource/GenerateMutableSourceMesh.cpp
--- a/Engine/Plugins/Mutable/Source/CustomizableObjectEditor/Private/MuCOE/GenerateMutableSource/GenerateMutableSourceMesh.cpp
+++ b/Engine/Plugins/Mutable/Source/CustomizableObjectEditor/Private/MuCOE/GenerateMutableSource/GenerateMutableSourceMesh.cpp
@@ -35,7 +35,7 @@
 #include "MuCOE/Nodes/CustomizableObjectNodeMeshReshapeCommon.h"
 #include "MuCOE/Nodes/CustomizableObjectNodeMeshSwitch.h"
 #include "MuCOE/Nodes/CustomizableObjectNodeMeshVariation.h"
-#include "MuCOE/Nodes/CustomizableObjectNodeMeshParameter.h"
+#include "MuCOE/Nodes/CustomizableObjectNodeSkeletalMeshParameter.h"
 #include "MuCOE/Nodes/CustomizableObjectNodeSkeletalMesh.h"
 #include "MuCOE/Nodes/CustomizableObjectNodeStaticMesh.h"
 #include "MuCOE/Nodes/CustomizableObjectNodeStaticString.h"
@@ -4579,7 +4579,7 @@ mu::Ptr<mu::NodeMesh> GenerateMutableSourceMesh(const UEdGraphPin* Pin,
 		}
 	}
 
-	else if (const UCustomizableObjectNodeMeshParameter* TypedNodeParam = Cast<UCustomizableObjectNodeMeshParameter>(Node))
+	else if (const UCustomizableObjectNodeSkeletalMeshParameter* TypedNodeParam = Cast<UCustomizableObjectNodeSkeletalMeshParameter>(Node))
 	{
 		mu::Ptr<mu::NodeMeshParameter> ParameterNode = new mu::NodeMeshParameter();
 
@@ -4588,7 +4588,7 @@ mu::Ptr<mu::NodeMesh> GenerateMutableSourceMesh(const UEdGraphPin* Pin,
 
 		if (TypedNodeParam->DefaultValue)
 		{
-			ParameterNode->DefaultValue = FName(TypedNodeParam->DefaultValue->GetPathName());
+			GenerationContext.SkeletalMeshParameterDefaultValues.Add(FName(TypedNodeParam->ParameterName), MutablePrivate::LoadObject(TypedNodeParam->DefaultValue));
 		}
 
 		const int32 LODIndex = 0;
diff --git a/Engine/Plugins/Mutable/Source/CustomizableObjectEditor/Private/MuCOE/GraphTraversal.cpp b/Engine/Plugins/Mutable/Source/CustomizableObjectEditor/Private/MuCOE/GraphTraversal.cpp
--- a/Engine/Plugins/Mutable/Source/CustomizableObjectEditor/Private/MuCOE/GraphTraversal.cpp
+++ b/Engine/Plugins/Mutable/Source/CustomizableObjectEditor/Private/MuCOE/GraphTraversal.cpp
@@ -18,7 +18,7 @@
 #include "MuCOE/Nodes/CustomizableObjectNodeMeshReshape.h"
 #include "MuCOE/Nodes/CustomizableObjectNodeMeshSwitch.h"
 #include "MuCOE/Nodes/CustomizableObjectNodeMeshVariation.h"
-#include "MuCOE/Nodes/CustomizableObjectNodeMeshParameter.h"
+#include "MuCOE/Nodes/CustomizableObjectNodeSkeletalMeshParameter.h"
 #include "MuCOE/Nodes/CustomizableObjectNodeObject.h"
 #include "MuCOE/Nodes/CustomizableObjectNodeSkeletalMesh.h"
 #include "MuCOE/Nodes/CustomizableObjectNodeStaticMesh.h"
@@ -775,7 +775,7 @@ const UEdGraphPin* FindMeshBaseSource(const UEdGraphPin& Pin, const bool bOnlyLo
 		}
 	}
 
-	else if (Cast<UCustomizableObjectNodeMeshParameter>(Node))
+	else if (Cast<UCustomizableObjectNodeSkeletalMeshParameter>(Node))
 	{
 		if (!bOnlyLookForStaticMesh)
 		{
diff --git a/Engine/Plugins/Mutable/Source/CustomizableObjectEditor/Private/MuCOE/Nodes/CustomizableObjectNodeMeshParameter.cpp b/Engine/Plugins/Mutable/Source/CustomizableObjectEditor/Private/MuCOE/Nodes/CustomizableObjectNodeMeshParameter.cpp
deleted file mode 100644
diff --git a/Engine/Plugins/Mutable/Source/CustomizableObjectEditor/Private/MuCOE/Nodes/CustomizableObjectNodeMeshParameter.h b/Engine/Plugins/Mutable/Source/CustomizableObjectEditor/Private/MuCOE/Nodes/CustomizableObjectNodeMeshParameter.h
deleted file mode 100644
diff --git a/Engine/Plugins/Mutable/Source/CustomizableObjectEditor/Private/MuCOE/Nodes/CustomizableObjectNodeMeshParameterDetails.h b/Engine/Plugins/Mutable/Source/CustomizableObjectEditor/Private/MuCOE/Nodes/CustomizableObjectNodeMeshParameterDetails.h
deleted file mode 100644
diff --git a/Engine/Plugins/Mutable/Source/CustomizableObjectEditor/Private/MuCOE/Nodes/CustomizableObjectNodeMeshParametersDetails.cpp b/Engine/Plugins/Mutable/Source/CustomizableObjectEditor/Private/MuCOE/Nodes/CustomizableObjectNodeMeshParametersDetails.cpp
deleted file mode 100644
diff --git a/Engine/Plugins/Mutable/Source/CustomizableObjectEditor/Private/MuCOE/Nodes/CustomizableObjectNodeObjectDetails.cpp b/Engine/Plugins/Mutable/Source/CustomizableObjectEditor/Private/MuCOE/Nodes/CustomizableObjectNodeObjectDetails.cpp
--- a/Engine/Plugins/Mutable/Source/CustomizableObjectEditor/Private/MuCOE/Nodes/CustomizableObjectNodeObjectDetails.cpp
+++ b/Engine/Plugins/Mutable/Source/CustomizableObjectEditor/Private/MuCOE/Nodes/CustomizableObjectNodeObjectDetails.cpp
@@ -19,7 +19,7 @@
 #include "MuCOE/Nodes/CustomizableObjectNodeGroupProjectorParameter.h"
 #include "MuCOE/Nodes/CustomizableObjectNodeTable.h"
 #include "MuCOE/Nodes/CustomizableObjectNodeTextureParameter.h"
-#include "MuCOE/Nodes/CustomizableObjectNodeMeshParameter.h"
+#include "MuCOE/Nodes/CustomizableObjectNodeSkeletalMeshParameter.h"
 #include "MuCOE/SCustomizableObjectNodeObjectRTMorphTargetOverride.h"
 #include "MuCOE/UnrealEditorPortabilityHelpers.h"
 #include "MuCOE/Widgets/CustomizableObjectLODReductionSettings.h"
@@ -462,10 +462,10 @@ void FCustomizableObjectNodeObjectDetails::FillParameterNamesArray()
 			}
 		}
 
-		TArray<UCustomizableObjectNodeMeshParameter*> MeshParameterNodes;
+		TArray<UCustomizableObjectNodeSkeletalMeshParameter*> MeshParameterNodes;
 		Source->GetNodesOfClass(MeshParameterNodes);
 
-		for (const UCustomizableObjectNodeMeshParameter* MeshParameterNode : MeshParameterNodes)
+		for (const UCustomizableObjectNodeSkeletalMeshParameter* MeshParameterNode : MeshParameterNodes)
 		{
 			if (MeshParameterNode)
 			{
diff --git a/Engine/Plugins/Mutable/Source/CustomizableObjectEditor/Private/MuCOE/Nodes/CustomizableObjectNodeSkeletalMeshParameter.cpp b/Engine/Plugins/Mutable/Source/CustomizableObjectEditor/Private/MuCOE/Nodes/CustomizableObjectNodeSkeletalMeshParameter.cpp
new file mode 100644
--- /dev/null
+++ b/Engine/Plugins/Mutable/Source/CustomizableObjectEditor/Private/MuCOE/Nodes/CustomizableObjectNodeSkeletalMeshParameter.cpp
@@ -0,0 +1,344 @@
+// Copyright Epic Games, Inc. All Rights Reserved.
+
+#include "MuCOE/Nodes/CustomizableObjectNodeSkeletalMeshParameter.h"
+
+#include "MuCOE/EdGraphSchema_CustomizableObject.h"
+#include "MuCOE/CustomizableObjectLayout.h"
+#include "MuCOE/CustomizableObjectMacroLibrary/CustomizableObjectMacroLibrary.h"
+#include "MuCO/CustomizableObjectCustomVersion.h"
+#include "Rendering/SkeletalMeshModel.h"
+#include "Engine/SkeletalMesh.h"
+
+class UCustomizableObjectNodeRemapPins;
+
+#define LOCTEXT_NAMESPACE "CustomizableObjectEditor"
+
+void UCustomizableObjectNodeSkeletalMeshParameter::AllocateDefaultPins(UCustomizableObjectNodeRemapPins* RemapPins)
+{
+	USkeletalMesh* SkeletalMesh = DefaultValue.LoadSynchronous();
+
+	if (!SkeletalMesh)
+	{
+		return;
+	}
+
+	NamePin = CustomCreatePin(EGPD_Input, UEdGraphSchema_CustomizableObject::PC_String, FName("Name"));
+
+	const UEdGraphSchema_CustomizableObject* Schema = GetDefault<UEdGraphSchema_CustomizableObject>();
+
+	if (const FSkeletalMeshModel* ImportedModel = SkeletalMesh->GetImportedModel())
+	{
+		int32 LODIndex = 0;
+		{
+			const int32 NumSections = ImportedModel->LODModels[LODIndex].Sections.Num();
+			for (int32 SectionIndex = 0; SectionIndex < NumSections; ++SectionIndex)
+			{
+				// Ignore disabled sections.
+				if (ImportedModel->LODModels[LODIndex].Sections[SectionIndex].bDisabled)
+				{
+					continue;
+				}
+
+				UMaterialInterface* MaterialInterface = GetMaterialInterfaceFor( SectionIndex);
+
+				FString Section = FString::Printf(TEXT("Section %i"), SectionIndex);
+				if (MaterialInterface)
+				{
+					Section.Append(FString::Printf(TEXT(" : %s"), *MaterialInterface->GetName()));
+				}
+
+				// Mesh
+				{
+					UCustomizableObjectNodeSkeletalMeshParameterPinDataSection* PinData = NewObject<UCustomizableObjectNodeSkeletalMeshParameterPinDataSection>(this);
+					PinData->Init( SectionIndex, ImportedModel->LODModels[LODIndex].NumTexCoords );
+
+					const FString MeshPinName = FString::Printf(TEXT("Section %i - Mesh"), SectionIndex);
+
+					UEdGraphPin* Pin = CustomCreatePin(EGPD_Output, Schema->PC_Mesh, FName(*MeshPinName), PinData);
+					Pin->PinFriendlyName = FText::FromString(FString::Printf(TEXT("%s"), *Section));
+				}
+			}
+		}
+	}
+
+}
+
+
+void UCustomizableObjectNodeSkeletalMeshParameter::PostEditChangeProperty(FPropertyChangedEvent& PropertyChangedEvent)
+{
+	Super::PostEditChangeProperty(PropertyChangedEvent);
+	
+	FProperty* PropertyThatChanged = PropertyChangedEvent.Property;
+	if (PropertyThatChanged && PropertyThatChanged->GetName() == GET_MEMBER_NAME_CHECKED(UCustomizableObjectNodeSkeletalMeshParameter, DefaultValue))
+	{
+		ReconstructNode();
+	}
+}
+
+
+bool UCustomizableObjectNodeSkeletalMeshParameter::IsExperimental() const
+{
+	return true;
+}
+
+
+void UCustomizableObjectNodeSkeletalMeshParameter::BackwardsCompatibleFixup(int32 CustomizableObjectCustomVersion)
+{
+	Super::BackwardsCompatibleFixup(CustomizableObjectCustomVersion);
+
+	if (CustomizableObjectCustomVersion == FCustomizableObjectCustomVersion::EnableMutableMacrosNewVersion)
+	{
+		if (!NamePin.Get())
+		{
+			NamePin = CustomCreatePin(EGPD_Input, UEdGraphSchema_CustomizableObject::PC_String, FName("Name"));
+		}
+	}
+}
+
+
+FText UCustomizableObjectNodeSkeletalMeshParameter::GetNodeTitle(ENodeTitleType::Type TitleType) const
+{
+	const UEdGraphPin* ParamNamePin = NamePin.Get();
+
+	if (TitleType == ENodeTitleType::ListView || (ParamNamePin && ParamNamePin->LinkedTo.Num()))
+	{
+		return LOCTEXT("Skeletal_Mesh_Parameter", "Skeletal Mesh Parameter");
+	}
+	else if (TitleType == ENodeTitleType::EditableTitle)
+	{
+		return FText::Format(LOCTEXT("Skeletal_Mesh_Parameter_EditableTitle", "{0}"), FText::FromString(ParameterName));
+	}
+	else
+	{
+		return FText::Format(LOCTEXT("Skeletal_Mesh_Parameter_Title", "{0}\nSkeletal Mesh Parameter"), FText::FromString(ParameterName));
+	}
+}
+
+
+FLinearColor UCustomizableObjectNodeSkeletalMeshParameter::GetNodeTitleColor() const
+{
+	const UEdGraphSchema_CustomizableObject* Schema = GetDefault<UEdGraphSchema_CustomizableObject>();
+	return Schema->GetPinTypeColor(Schema->PC_Mesh);
+}
+
+
+FText UCustomizableObjectNodeSkeletalMeshParameter::GetTooltipText() const
+{
+	return LOCTEXT("Mesh_Parameter_Tooltip", "Expose a runtime modifiable Mesh parameter from the Customizable Object.");
+}
+
+
+void UCustomizableObjectNodeSkeletalMeshParameter::OnRenameNode(const FString& NewName)
+{
+	if (!NewName.IsEmpty())
+	{
+		ParameterName = NewName;
+	}
+}
+
+
+bool UCustomizableObjectNodeSkeletalMeshParameter::GetCanRenameNode() const
+{
+	const UEdGraphPin* ParamNamePin = NamePin.Get();
+
+	if (ParamNamePin && ParamNamePin->LinkedTo.Num())
+	{
+		return false;
+	}
+
+	return true;
+}
+
+
+void UCustomizableObjectNodeSkeletalMeshParameter::PinConnectionListChanged(UEdGraphPin* Pin)
+{
+	if (Pin == NamePin)
+	{
+		GetGraph()->NotifyGraphChanged();
+	}
+}
+
+
+UTexture2D* UCustomizableObjectNodeSkeletalMeshParameter::FindTextureForPin(const UEdGraphPin* Pin) const
+{
+	return nullptr;
+}
+
+
+TArray<UCustomizableObjectLayout*> UCustomizableObjectNodeSkeletalMeshParameter::GetLayouts(const UEdGraphPin& MeshPin) const
+{
+	const UCustomizableObjectNodeSkeletalMeshParameterPinDataSection& MeshPinData = GetPinData<UCustomizableObjectNodeSkeletalMeshParameterPinDataSection>(MeshPin);
+	return MeshPinData.Layouts;
+}
+
+
+TSoftObjectPtr<UObject> UCustomizableObjectNodeSkeletalMeshParameter::GetMesh() const
+{
+	return DefaultValue;
+}
+
+
+UEdGraphPin* UCustomizableObjectNodeSkeletalMeshParameter::GetMeshPin(const int32 LODIndex, const int32 SectionIndex) const
+{
+	for (UEdGraphPin* Pin : GetAllNonOrphanPins())
+	{
+		if (const UCustomizableObjectNodeSkeletalMeshParameterPinDataSection* PinData = Cast<UCustomizableObjectNodeSkeletalMeshParameterPinDataSection>(GetPinData(*Pin)))
+		{
+			if (PinData->GetSectionIndex() == SectionIndex)
+			{
+				return Pin;
+			}
+		}
+	}
+
+	return nullptr;
+}
+
+
+void UCustomizableObjectNodeSkeletalMeshParameter::GetPinSection(const UEdGraphPin& Pin, int32& OutLODIndex, int32& OutSectionIndex, int32& OutLayoutIndex) const
+{
+	OutLODIndex = -1;
+	OutSectionIndex = -1;
+	OutLayoutIndex = -1;
+
+	if (const UCustomizableObjectNodeSkeletalMeshParameterPinDataSection* PinData = Cast<UCustomizableObjectNodeSkeletalMeshParameterPinDataSection>(GetPinData(Pin)))
+	{
+		OutLODIndex = 0;
+		OutSectionIndex = PinData->GetSectionIndex();
+	}
+}
+
+
+void UCustomizableObjectNodeSkeletalMeshParameter::GetPinSection(const UEdGraphPin& Pin, int32& OutSectionIndex) const
+{
+	if (const UCustomizableObjectNodeSkeletalMeshParameterPinDataSection* PinData = Cast<UCustomizableObjectNodeSkeletalMeshParameterPinDataSection>(GetPinData(Pin)))
+	{
+		OutSectionIndex = PinData->GetSectionIndex();
+		return;
+	}
+
+	OutSectionIndex = -1;
+}
+
+
+UMaterialInterface* UCustomizableObjectNodeSkeletalMeshParameter::GetMaterialInterfaceFor(const int32 SectionIndex) const
+{
+	if (FSkeletalMaterial* SkeletalMaterial = GetSkeletalMaterialFor(SectionIndex))
+	{
+		return SkeletalMaterial->MaterialInterface;
+	}
+
+	return nullptr;
+}
+
+
+FSkeletalMaterial* UCustomizableObjectNodeSkeletalMeshParameter::GetSkeletalMaterialFor(const int32 SectionIndex) const
+{
+	USkeletalMesh* SkeletalMesh = DefaultValue.LoadSynchronous();
+
+	if (!SkeletalMesh)
+	{
+		return nullptr;
+	}
+
+	const int32 SkeletalMeshMaterialIndex = GetSkeletalMaterialIndexFor(SectionIndex);
+	if (SkeletalMesh->GetMaterials().IsValidIndex(SkeletalMeshMaterialIndex))
+	{
+		return &SkeletalMesh->GetMaterials()[SkeletalMeshMaterialIndex];
+	}
+
+	return nullptr;
+}
+
+int32 UCustomizableObjectNodeSkeletalMeshParameter::GetSkeletalMaterialIndexFor(const int32 SectionIndex) const
+{
+	USkeletalMesh* SkeletalMesh = DefaultValue.LoadSynchronous();
+
+	if (!SkeletalMesh)
+	{
+		return INDEX_NONE;
+	}
+
+	int32 LODIndex = 0;
+
+	// We assume that LODIndex and MaterialIndex are valid for the imported model
+	int32 SkeletalMeshMaterialIndex = INDEX_NONE;
+
+	// Check if we have lod info map to get the correct material index
+	if (const FSkeletalMeshLODInfo* LodInfo = SkeletalMesh->GetLODInfo(LODIndex))
+	{
+		if (LodInfo->LODMaterialMap.IsValidIndex(SectionIndex))
+		{
+			SkeletalMeshMaterialIndex = LodInfo->LODMaterialMap[SectionIndex];
+		}
+	}
+
+	// Only deduce index when the explicit mapping is not found or there is no remap
+	if (SkeletalMeshMaterialIndex == INDEX_NONE)
+	{
+		FSkeletalMeshModel* ImportedModel = SkeletalMesh->GetImportedModel();
+		if (ImportedModel && ImportedModel->LODModels.IsValidIndex(LODIndex) && ImportedModel->LODModels[LODIndex].Sections.IsValidIndex(SectionIndex))
+		{
+			SkeletalMeshMaterialIndex = ImportedModel->LODModels[LODIndex].Sections[SectionIndex].MaterialIndex;
+		}
+	}
+
+	return SkeletalMeshMaterialIndex;
+}
+
+
+int32 UCustomizableObjectNodeSkeletalMeshParameter::GetSkeletalMaterialIndexFor(const UEdGraphPin& Pin) const
+{
+	int32 SectionIndex;
+	GetPinSection(Pin, SectionIndex);
+
+	return GetSkeletalMaterialIndexFor(SectionIndex);
+}
+
+
+void UCustomizableObjectNodeSkeletalMeshParameterPinDataSection::Init(int32 InSectionIndex, int32 NumTexCoords)
+{
+	SectionIndex = InSectionIndex;
+
+	if (NumTexCoords > 0)
+	{
+		UObject* Outer = GetOuter();
+
+		Layouts.SetNum(NumTexCoords);
+
+		const int LODIndex = 0;
+
+		for (int32 Index = 0; Index < NumTexCoords; ++Index)
+		{
+			Layouts[Index] = NewObject<UCustomizableObjectLayout>(Outer);
+			Layouts[Index]->SetLayout(LODIndex, InSectionIndex, Index);
+		}
+	}
+}
+
+
+int32 UCustomizableObjectNodeSkeletalMeshParameterPinDataSection::GetSectionIndex() const
+{
+	return SectionIndex;
+}
+
+
+bool UCustomizableObjectNodeSkeletalMeshParameterPinDataSection::Equals(const UCustomizableObjectNodePinData& Other) const
+{
+	if (GetClass() != Other.GetClass())
+	{
+		return false;
+	}
+
+	const UCustomizableObjectNodeSkeletalMeshParameterPinDataSection& OtherTyped = static_cast<const UCustomizableObjectNodeSkeletalMeshParameterPinDataSection&>(Other);
+	if (SectionIndex != OtherTyped.SectionIndex)
+	{
+		return false;
+	}
+
+	return Super::Equals(Other);
+}
+
+
+#undef LOCTEXT_NAMESPACE
+
diff --git a/Engine/Plugins/Mutable/Source/CustomizableObjectEditor/Private/MuCOE/Nodes/CustomizableObjectNodeSkeletalMeshParameter.h b/Engine/Plugins/Mutable/Source/CustomizableObjectEditor/Private/MuCOE/Nodes/CustomizableObjectNodeSkeletalMeshParameter.h
new file mode 100644
--- /dev/null
+++ b/Engine/Plugins/Mutable/Source/CustomizableObjectEditor/Private/MuCOE/Nodes/CustomizableObjectNodeSkeletalMeshParameter.h
@@ -0,0 +1,96 @@
+// Copyright Epic Games, Inc. All Rights Reserved.
+
+#pragma once
+
+#include "MuCOE/Nodes/CustomizableObjectNodeMesh.h"
+
+#include "CustomizableObjectNodeSkeletalMeshParameter.generated.h"
+
+namespace ENodeTitleType { enum Type : int; }
+
+class UCustomizableObjectLayout;
+class UMaterialInterface;
+struct FSkeletalMaterial;
+
+
+/** PinData of a pin that belongs to a Skeletal Mesh Section. */
+UCLASS()
+class UCustomizableObjectNodeSkeletalMeshParameterPinDataSection : public UCustomizableObjectNodePinData
+{
+	GENERATED_BODY()
+
+public:
+	void Init(int32 InSectionIndex, int32 NumTexCoords);
+
+	int32 GetSectionIndex() const;
+
+	/** Layouts related to this Mesh pin */
+	UPROPERTY()
+	TArray<TObjectPtr<UCustomizableObjectLayout>> Layouts;
+
+protected:
+
+	virtual bool Equals(const UCustomizableObjectNodePinData& Other) const override;
+
+private:
+
+	UPROPERTY()
+	int32 SectionIndex = -1;
+
+};
+
+
+UCLASS()
+class CUSTOMIZABLEOBJECTEDITOR_API UCustomizableObjectNodeSkeletalMeshParameter : public UCustomizableObjectNodeMesh
+{
+public:
+	GENERATED_BODY()
+
+	/** Default value of the parameter. */
+	UPROPERTY(EditAnywhere, Category=CustomizableObject)
+	TSoftObjectPtr<USkeletalMesh> DefaultValue;
+
+	UPROPERTY(EditAnywhere, Category=CustomizableObject)
+	FString ParameterName = "Mesh Param";
+
+	UPROPERTY(EditAnywhere, Category = UI, meta = (DisplayName = "Parameter UI Metadata"))
+	FMutableParamUIMetadata ParamUIMetadata;
+
+	UPROPERTY()
+	FEdGraphPinReference NamePin;
+
+public:
+
+	// UObject interface.
+	virtual void PostEditChangeProperty(FPropertyChangedEvent& PropertyChangedEvent) override;
+
+	// EdGraphNode interface
+	virtual FText GetNodeTitle(ENodeTitleType::Type TitleType) const override;
+	virtual FLinearColor GetNodeTitleColor() const override;
+	virtual FText GetTooltipText() const override;
+	virtual void OnRenameNode(const FString& NewName) override;
+	virtual bool GetCanRenameNode() const override;
+	virtual void PinConnectionListChanged(UEdGraphPin* Pin) override;
+
+	// UCustomizableObjectNode interface
+	virtual void AllocateDefaultPins(UCustomizableObjectNodeRemapPins* RemapPins) override;
+	virtual bool IsAffectedByLOD() const override { return false; }
+	virtual bool IsExperimental() const override;
+	virtual void BackwardsCompatibleFixup(int32 CustomizableObjectCustomVersion) override;
+
+	// UCustomizableObjectNodeMesh interface
+	virtual UTexture2D* FindTextureForPin(const UEdGraphPin* Pin) const override;
+	virtual TArray<UCustomizableObjectLayout*> GetLayouts(const UEdGraphPin& MeshPin) const override;
+	virtual TSoftObjectPtr<UObject> GetMesh() const override;
+	virtual UEdGraphPin* GetMeshPin(int32 LOD, int32 SectionIndex) const override;
+	virtual void GetPinSection(const UEdGraphPin& Pin, int32& OutLODIndex, int32& OutSectionIndex, int32& OutLayoutIndex) const override;
+
+	// Own interface
+	void GetPinSection(const UEdGraphPin& Pin, int32& OutSectionIndex) const;
+	UMaterialInterface* GetMaterialInterfaceFor(const int32 SectionIndex) const;
+	FSkeletalMaterial* GetSkeletalMaterialFor(const int32 SectionIndex) const;
+	int32 GetSkeletalMaterialIndexFor(const UEdGraphPin& Pin) const;
+	int32 GetSkeletalMaterialIndexFor(const int32 SectionIndex) const;
+
+};
+
diff --git a/Engine/Plugins/Mutable/Source/CustomizableObjectEditor/Private/MuCOE/Nodes/CustomizableObjectNodeSkeletalMeshParameterDetails.h b/Engine/Plugins/Mutable/Source/CustomizableObjectEditor/Private/MuCOE/Nodes/CustomizableObjectNodeSkeletalMeshParameterDetails.h
new file mode 100644
--- /dev/null
+++ b/Engine/Plugins/Mutable/Source/CustomizableObjectEditor/Private/MuCOE/Nodes/CustomizableObjectNodeSkeletalMeshParameterDetails.h
@@ -0,0 +1,27 @@
+// Copyright Epic Games, Inc. All Rights Reserved.
+
+#pragma once
+
+#include "CustomizableObjectNodeDetails.h"
+#include "IDetailCustomization.h"
+
+class IDetailLayoutBuilder;
+struct FLayoutEditorMeshSection;
+
+
+class FCustomizableObjectNodeSkeletalMeshParameterDetails : public FCustomizableObjectNodeDetails
+{
+public:
+	/** Makes a new instance of this detail layout class for a specific detail view requesting it */
+	static TSharedRef<IDetailCustomization> MakeInstance();
+
+	/** ILayoutDetails interface */
+	virtual void CustomizeDetails(IDetailLayoutBuilder& DetailBuilder) override;
+
+private:
+
+	void GenerateMeshSectionOptions(TArray<FLayoutEditorMeshSection>& OutMeshSections);
+
+	// Pointer to the node represented in this details
+	class UCustomizableObjectNodeSkeletalMeshParameter* Node;
+};
diff --git a/Engine/Plugins/Mutable/Source/CustomizableObjectEditor/Private/MuCOE/Nodes/CustomizableObjectNodeSkeletalMeshParametersDetails.cpp b/Engine/Plugins/Mutable/Source/CustomizableObjectEditor/Private/MuCOE/Nodes/CustomizableObjectNodeSkeletalMeshParametersDetails.cpp
new file mode 100644
--- /dev/null
+++ b/Engine/Plugins/Mutable/Source/CustomizableObjectEditor/Private/MuCOE/Nodes/CustomizableObjectNodeSkeletalMeshParametersDetails.cpp
@@ -0,0 +1,89 @@
+// Copyright Epic Games, Inc. All Rights Reserved.
+
+#include "MuCOE/Nodes/CustomizableObjectNodeSkeletalMeshParameterDetails.h"
+
+#include "DetailLayoutBuilder.h"
+#include "DetailWidgetRow.h"
+#include "IDetailsView.h"
+#include "MuCOE/Nodes/CustomizableObjectNodeSkeletalMeshParameter.h"
+#include "MuCOE/PinViewer/SPinViewer.h"
+#include "MuCOE/SCustomizableObjectLayoutEditor.h"
+
+#define LOCTEXT_NAMESPACE "CustomizableObjectNodeMaterialDetails"
+
+
+TSharedRef<IDetailCustomization> FCustomizableObjectNodeSkeletalMeshParameterDetails::MakeInstance()
+{
+	return MakeShareable(new FCustomizableObjectNodeSkeletalMeshParameterDetails);
+}
+
+
+void FCustomizableObjectNodeSkeletalMeshParameterDetails::CustomizeDetails(IDetailLayoutBuilder& DetailBuilder)
+{
+	FCustomizableObjectNodeDetails::CustomizeDetails(DetailBuilder);
+
+	TSharedPtr<const IDetailsView> DetailsView = DetailBuilder.GetDetailsViewSharedPtr();
+
+    if (DetailsView && DetailsView->GetSelectedObjects().Num())
+	{
+		Node = Cast<UCustomizableObjectNodeSkeletalMeshParameter>(DetailsView->GetSelectedObjects()[0].Get());
+	}
+
+    if (!Node)
+    {
+        return;
+    }
+
+    IDetailCategoryBuilder& CustomizableObject = DetailBuilder.EditCategory("CustomizableObject");
+
+	TSharedRef<IPropertyHandle> DefaultValueProperty = DetailBuilder.GetProperty(GET_MEMBER_NAME_CHECKED(UCustomizableObjectNodeSkeletalMeshParameter, DefaultValue));
+
+	TArray<FLayoutEditorMeshSection> MeshSectionsAndLayouts;
+	GenerateMeshSectionOptions(MeshSectionsAndLayouts);
+
+	TSharedPtr<SCustomizableObjectLayoutEditor> LayoutBlocksEditor = SNew(SCustomizableObjectLayoutEditor)
+		.Node(Node)
+		.MeshSections(MeshSectionsAndLayouts);
+
+	FCustomizableObjectLayoutEditorDetailsBuilder LayoutEditorBuilder;
+	LayoutEditorBuilder.LayoutEditor = LayoutBlocksEditor;
+	LayoutEditorBuilder.bShowLayoutSelector = true;
+	LayoutEditorBuilder.bShowPackagingStrategy = true;
+	LayoutEditorBuilder.bShowAutomaticGenerationSettings = true;
+	LayoutEditorBuilder.bShowGridSize = true;
+	LayoutEditorBuilder.bShowMaxGridSize = true;
+	LayoutEditorBuilder.bShowReductionMethods = true;
+	LayoutEditorBuilder.bShowWarningSettings = true;
+
+	LayoutEditorBuilder.CustomizeDetails(DetailBuilder);
+
+	LayoutBlocksEditor->UpdateLayout(nullptr);
+}
+
+
+void FCustomizableObjectNodeSkeletalMeshParameterDetails::GenerateMeshSectionOptions(TArray<FLayoutEditorMeshSection>& OutMeshSections)
+{
+	OutMeshSections.Empty();
+
+	if (!Node)
+	{
+		return;
+	}
+
+	for (const UEdGraphPin* Pin : Node->GetAllNonOrphanPins())
+	{
+		if (const UCustomizableObjectNodeSkeletalMeshParameterPinDataSection* PinData = Cast<UCustomizableObjectNodeSkeletalMeshParameterPinDataSection>(Node->GetPinData(*Pin)))
+		{
+			FLayoutEditorMeshSection& MeshSection = OutMeshSections.AddDefaulted_GetRef();
+			MeshSection.MeshName = MakeShareable(new FString(Pin->PinFriendlyName.ToString()));
+
+			for (UCustomizableObjectLayout* Layout : PinData->Layouts)
+			{
+				MeshSection.Layouts.Add(Layout);
+			}
+		}
+	}
+}
+
+
+#undef LOCTEXT_NAMESPACE
diff --git a/Engine/Plugins/Mutable/Source/CustomizableObjectEditor/Private/MuCOE/Nodes/SCustomizableObjectNodeStaticStringPin.cpp b/Engine/Plugins/Mutable/Source/CustomizableObjectEditor/Private/MuCOE/Nodes/SCustomizableObjectNodeStaticStringPin.cpp
--- a/Engine/Plugins/Mutable/Source/CustomizableObjectEditor/Private/MuCOE/Nodes/SCustomizableObjectNodeStaticStringPin.cpp
+++ b/Engine/Plugins/Mutable/Source/CustomizableObjectEditor/Private/MuCOE/Nodes/SCustomizableObjectNodeStaticStringPin.cpp
@@ -6,7 +6,7 @@
 #include "MuCOE/Nodes/CustomizableObjectNodeComponentMesh.h"
 #include "MuCOE/Nodes/CustomizableObjectNodeComponentMeshAddTo.h"
 #include "MuCOE/Nodes/CustomizableObjectNodeMeshMorph.h"
-#include "MuCOE/Nodes/CustomizableObjectNodeMeshParameter.h"
+#include "MuCOE/Nodes/CustomizableObjectNodeSkeletalMeshParameter.h"
 #include "MuCOE/Nodes/CustomizableObjectNodeObjectChild.h"
 #include "MuCOE/Nodes/CustomizableObjectNodeObjectGroup.h"
 #include "MuCOE/Nodes/CustomizableObjectNodeParameter.h"
@@ -53,7 +53,7 @@ FText SCustomizableObjectNodeStaticStringPin::GetNodeStringValue() const
 	{
 		return FText::FromName(AddComponentNode->GetParentComponentName());
 	}
-	else if (const UCustomizableObjectNodeMeshParameter* MeshParamNode = Cast<UCustomizableObjectNodeMeshParameter>(Node))
+	else if (const UCustomizableObjectNodeSkeletalMeshParameter* MeshParamNode = Cast<UCustomizableObjectNodeSkeletalMeshParameter>(Node))
 	{
 		return FText::FromString(MeshParamNode->ParameterName);
 	}
@@ -100,7 +100,7 @@ void SCustomizableObjectNodeStaticStringPin::OnTextCommited(const FText& InValue
 	{
 		AddComponentNode->SetParentComponentName(FName(*InValue.ToString()));
 	}
-	else if (UCustomizableObjectNodeMeshParameter* MeshParamNode = Cast<UCustomizableObjectNodeMeshParameter>(Node))
+	else if (UCustomizableObjectNodeSkeletalMeshParameter* MeshParamNode = Cast<UCustomizableObjectNodeSkeletalMeshParameter>(Node))
 	{
 		MeshParamNode->ParameterName = InValue.ToString();
 	}
diff --git a/Engine/Plugins/Mutable/Source/CustomizableObjectEditor/Private/MuCOE/SMutableCodeViewer.cpp b/Engine/Plugins/Mutable/Source/CustomizableObjectEditor/Private/MuCOE/SMutableCodeViewer.cpp
--- a/Engine/Plugins/Mutable/Source/CustomizableObjectEditor/Private/MuCOE/SMutableCodeViewer.cpp
+++ b/Engine/Plugins/Mutable/Source/CustomizableObjectEditor/Private/MuCOE/SMutableCodeViewer.cpp
@@ -2782,7 +2782,7 @@ namespace
 	public:
 
 
-		virtual TTuple<UE::Tasks::FTask, TFunction<void()>> GetImageAsync(FName Id, uint8 MipmapsToSkip, TFunction<void(TSharedPtr<mu::FImage>)>& ResultCallback) override
+		virtual TTuple<UE::Tasks::FTask, TFunction<void()>> GetImageAsync(UTexture* Texture, uint8 MipmapsToSkip, TFunction<void(TSharedPtr<mu::FImage>)>& ResultCallback) override
 		{
 			MUTABLE_CPUPROFILER_SCOPE(TestImageProvider_GetImage);
 
@@ -2824,7 +2824,7 @@ namespace
 			return MakeTuple(UE::Tasks::MakeCompletedTask<void>(), []() -> void {});
 		}
 
-		virtual mu::FExtendedImageDesc GetImageDesc(FName Id) override
+		virtual mu::FExtendedImageDesc GetImageDesc(UTexture* Texture) override
 		{
 			return mu::FExtendedImageDesc{ IMAGE_DESC };
 		}
@@ -2834,24 +2834,34 @@ namespace
 		{
 			check(ReferencedTextures.IsValidIndex(Id));
 
-			TSoftObjectPtr<const UTexture> TexturePtr = ReferencedTextures[Id].Get();
-			const UTexture2D* Texture = Cast<UTexture2D>( TexturePtr.Get() );
-			check(Texture);
+			const UTexture* Texture = ReferencedTextures[Id].Get();
+			if (!Texture)
+			{
+				UE_LOG(LogMutable, Warning, TEXT("Failed to load Referenced Image [%i]. Nullptr."), Id);
+				return MakeTuple(UE::Tasks::MakeCompletedTask<void>(), []() -> void {});
+			}
+			
+			const UTexture2D* Texture2D = Cast<UTexture2D>(Texture);
+			if (!Texture2D)
+			{
+				UE_LOG(LogMutable, Warning, TEXT("Invalid Referenced Image [%i, %s]. Is not a UTexture2D."), Id, *Texture->GetName());
+				return MakeTuple(UE::Tasks::MakeCompletedTask<void>(), []() -> void {});
+			}
 			
 			// In the editor the src data can be directly accessed
-			int32 MipIndex = (MipmapsToSkip < Texture->GetPlatformData()->Mips.Num()) ? MipmapsToSkip : Texture->GetPlatformData()->Mips.Num() - 1;
+			int32 MipIndex = (MipmapsToSkip < Texture2D->GetPlatformData()->Mips.Num()) ? MipmapsToSkip : Texture2D->GetPlatformData()->Mips.Num() - 1;
 			check(MipIndex >= 0);
 
 			TSharedPtr<mu::FImage> ResultImage = MakeShared<mu::FImage>();
 
-			FMutableSourceTextureData Tex(*Texture);
+			FMutableSourceTextureData Tex(*Texture2D);
 			EUnrealToMutableConversionError Error = ConvertTextureUnrealSourceToMutable(ResultImage.Get(), Tex, MipmapsToSkip);
 
 			if (Error != EUnrealToMutableConversionError::Success)
 			{
 				// This could happen in the editor, because some source textures may have changed while there was a background compilation.
 				// We just show a warning and move on. This cannot happen during cooks, so it is fine.
-				UE_LOG(LogMutable, Warning, TEXT("Failed to load some source texture data for [%s]. Some textures may be corrupted."), *Texture->GetName());
+				UE_LOG(LogMutable, Warning, TEXT("Failed to load some source texture data for Referenced Image [%i, %s]. Some textures may be corrupted."), Id, *Texture->GetName());
 			}
 
 			ResultCallback(ResultImage);
@@ -2860,17 +2870,15 @@ namespace
 		}
 
 
-		virtual TTuple<UE::Tasks::FTask, TFunction<void()>> GetMeshAsync(FName Id, int32 LODIndex, int32 SectionIndex, TFunction<void(TSharedPtr<mu::FMesh>)>& ResultCallback) override
+		virtual TTuple<UE::Tasks::FTask, TFunction<void()>> GetMeshAsync(USkeletalMesh* SkeletalMesh, int32 LODIndex, int32 SectionIndex, TFunction<void(TSharedPtr<mu::FMesh>)>& ResultCallback) override
 		{
 			// Thread: worker
 			MUTABLE_CPUPROFILER_SCOPE(FUnrealMutableImageProvider::GetMeshAsync);
 
-			TSharedPtr<mu::FMesh> Result = MakeShared<mu::FMesh>();
-
-			// TODO: Do it async
-			USkeletalMesh* SkeletalMesh = Cast<USkeletalMesh>(FSoftObjectPath(Id.ToString()).TryLoad());
 			check(SkeletalMesh);
 
+			TSharedPtr<mu::FMesh> Result = MakeShared<mu::FMesh>();
+
 			UE::Tasks::FTask ConversionTask = UnrealConversionUtils::ConvertSkeletalMeshFromRuntimeData(SkeletalMesh, LODIndex, SectionIndex, nullptr, Result.Get());
 
 			return MakeTuple(
diff --git a/Engine/Plugins/Mutable/Source/MutableRuntime/Internal/MuR/ParametersPrivate.h b/Engine/Plugins/Mutable/Source/MutableRuntime/Internal/MuR/ParametersPrivate.h
--- a/Engine/Plugins/Mutable/Source/MutableRuntime/Internal/MuR/ParametersPrivate.h
+++ b/Engine/Plugins/Mutable/Source/MutableRuntime/Internal/MuR/ParametersPrivate.h
@@ -11,6 +11,9 @@
 #include "MuR/Image.h"
 #include "MuR/Model.h"
 
+#include "Engine/Texture.h"
+#include "Engine/SkeletalMesh.h"
+
 #define UE_API MUTABLERUNTIME_API
 
 namespace mu
@@ -251,12 +254,13 @@ namespace Private
 	using FParamFloatType = float;
 	using FParamColorType = FVector4f;
 	using FParamProjectorType = Private::TIndirectObject<FProjector>;
-	using FParamAssetType = FName;
+	using FParamTextureType = TStrongObjectPtr<UTexture>;
+	using FParamSkeletalMeshType = TStrongObjectPtr<USkeletalMesh>;
 	using FParamStringType = Private::TIndirectObject<FString>;
 	using FParamMatrixType = Private::TIndirectObject<FMatrix44f>;
 	
 	using FParameterValue = TVariant<
-            FParamBoolType, FParamIntType, FParamFloatType, FParamColorType, FParamProjectorType, FParamAssetType, FParamStringType, FParamMatrixType>;
+            FParamBoolType, FParamIntType, FParamFloatType, FParamColorType, FParamProjectorType, FParamTextureType, FParamSkeletalMeshType, FParamStringType, FParamMatrixType>;
 
     // static_assert to track ParameterValue size changes. It is ok to change if needed.
     static_assert(sizeof(FParameterValue) == 8*4, "ParameterValue size has changed.");
diff --git a/Engine/Plugins/Mutable/Source/MutableRuntime/Private/MuR/CodeRunner.cpp b/Engine/Plugins/Mutable/Source/MutableRuntime/Private/MuR/CodeRunner.cpp
--- a/Engine/Plugins/Mutable/Source/MutableRuntime/Private/MuR/CodeRunner.cpp
+++ b/Engine/Plugins/Mutable/Source/MutableRuntime/Private/MuR/CodeRunner.cpp
@@ -219,7 +219,7 @@ namespace mu
 			if (Id.ReferenceResourceId < 0)
 			{
 				// It's a parameter image
-				return System->ExternalResourceProvider->GetImageAsync(Id.ParameterId, MipmapsToSkip, ResultCallback);
+				return System->ExternalResourceProvider->GetImageAsync(Id.ImageParameter, MipmapsToSkip, ResultCallback);
 			}
 			else
 			{
@@ -248,7 +248,7 @@ namespace mu
 			if (Id.ReferenceResourceId < 0)
 			{
 				// It's a parameter mesh
-				return System->ExternalResourceProvider->GetMeshAsync(Id.ParameterId, LODIndex, SectionIndex, ResultCallback);
+				return System->ExternalResourceProvider->GetMeshAsync(Id.MeshParameter, LODIndex, SectionIndex, ResultCallback);
 			}
 			else
 			{
@@ -266,7 +266,7 @@ namespace mu
 	}
 
 	
-	FExtendedImageDesc CodeRunner::GetExternalImageDesc(FName Id)
+	FExtendedImageDesc CodeRunner::GetExternalImageDesc(UTexture* Image)
 	{
 		MUTABLE_CPUPROFILER_SCOPE(GetExternalImageDesc);
 
@@ -274,7 +274,7 @@ namespace mu
 
 		if (System->ExternalResourceProvider)
 		{
-			return System->ExternalResourceProvider->GetImageDesc(Id);
+			return System->ExternalResourceProvider->GetImageDesc(Image);
 		}
 		else
 		{
@@ -6380,10 +6380,10 @@ namespace mu
 			check(Item.Stage == 0);
 
 			OP::ParameterArgs Args = Program.GetOpArgs<OP::ParameterArgs>(Item.At);
-			FName Id = pParams->GetImageValue(Args.variable);
+			UTexture* Image = pParams->GetImageValue(Args.variable);
 
 			FExtendedImageDesc& Result = ImageDescResults.FindOrAdd(Item.At);
-			Result = GetExternalImageDesc(Id);
+			Result = GetExternalImageDesc(Image);
 
 			StoreValidDesc(Item);
 			break;
diff --git a/Engine/Plugins/Mutable/Source/MutableRuntime/Private/MuR/CodeRunner.h b/Engine/Plugins/Mutable/Source/MutableRuntime/Private/MuR/CodeRunner.h
--- a/Engine/Plugins/Mutable/Source/MutableRuntime/Private/MuR/CodeRunner.h
+++ b/Engine/Plugins/Mutable/Source/MutableRuntime/Private/MuR/CodeRunner.h
@@ -139,12 +139,13 @@ namespace  mu
 			int32 ReferenceResourceId = -1;
 
 			/** If it is an image or mesh parameter.*/
-			FName ParameterId; 
+			UTexture* ImageParameter = nullptr; 
+			USkeletalMesh* MeshParameter = nullptr; 
 		};
 
 		/** Load an external image asynchronously, returns an event to wait for complition and a cleanup function that must be called once the event has completed. */
 		TTuple<UE::Tasks::FTask, TFunction<void()>> LoadExternalImageAsync(FExternalResourceId Id, uint8 MipmapsToSkip, TFunction<void(TSharedPtr<FImage>)>& ResultCallback);
- 	    mu::FExtendedImageDesc GetExternalImageDesc(FName Id);
+ 	    mu::FExtendedImageDesc GetExternalImageDesc(UTexture* Id);
 
 		/** Load an external mesh asynchronously, returns an event to wait for complition and a cleanup function that must be called once the event has completed. */
 		TTuple<UE::Tasks::FTask, TFunction<void()>> LoadExternalMeshAsync(FExternalResourceId Id, int32 LODIndex, int32 SectionIndex, TFunction<void(TSharedPtr<FMesh>)>& ResultCallback);
diff --git a/Engine/Plugins/Mutable/Source/MutableRuntime/Private/MuR/CodeRunnerTasked.cpp b/Engine/Plugins/Mutable/Source/MutableRuntime/Private/MuR/CodeRunnerTasked.cpp
--- a/Engine/Plugins/Mutable/Source/MutableRuntime/Private/MuR/CodeRunnerTasked.cpp
+++ b/Engine/Plugins/Mutable/Source/MutableRuntime/Private/MuR/CodeRunnerTasked.cpp
@@ -2578,7 +2578,7 @@ namespace mu
 			OP::ParameterArgs args = Program.GetOpArgs<OP::ParameterArgs>(item.At);
 			TSharedPtr<FRangeIndex> Index = BuildCurrentOpRangeIndex(item, Params, Model.Get(), args.variable);
 
-			const FName Id = Params->GetImageValue(args.variable, Index.Get());
+			UTexture* Image = Params->GetImageValue(args.variable, Index.Get());
 
 			check(ImageLOD < TNumericLimits<uint8>::Max() && ImageLOD >= 0);
 			check(ImageLOD + static_cast<int32>(item.ExecutionOptions) < TNumericLimits<uint8>::Max());
@@ -2586,7 +2586,7 @@ namespace mu
 			const uint8 MipmapsToSkip = item.ExecutionOptions + static_cast<uint8>(ImageLOD);
 
 			CodeRunner::FExternalResourceId FullId;
-			FullId.ParameterId = Id;
+			FullId.ImageParameter = Image;
 			Issued = MakeShared<FImageExternalLoadTask>(item, MipmapsToSkip, FullId);
 
 			break;
@@ -2616,13 +2616,13 @@ namespace mu
 			OP::MeshParameterArgs Args = Program.GetOpArgs<OP::MeshParameterArgs>(item.At);
 			TSharedPtr<FRangeIndex> Index = BuildCurrentOpRangeIndex(item, Params, Model.Get(), Args.variable);
 
-			const FName Id = Params->GetMeshValue(Args.variable, Index.Get());
+			USkeletalMesh* Id = Params->GetMeshValue(Args.variable, Index.Get());
 
 			int32 LODIndex = Args.LOD;
 			int32 SectionIndex = Args.Section;
 
 			CodeRunner::FExternalResourceId FullId;
-			FullId.ParameterId = Id;
+			FullId.MeshParameter = Id;
 			Issued = MakeShared<FMeshExternalLoadTask>(item, FullId, LODIndex, SectionIndex, Args.MeshID);
 
 			break;
diff --git a/Engine/Plugins/Mutable/Source/MutableRuntime/Private/MuR/Model.cpp b/Engine/Plugins/Mutable/Source/MutableRuntime/Private/MuR/Model.cpp
--- a/Engine/Plugins/Mutable/Source/MutableRuntime/Private/MuR/Model.cpp
+++ b/Engine/Plugins/Mutable/Source/MutableRuntime/Private/MuR/Model.cpp
@@ -19,6 +19,9 @@
 #include "MuR/Types.h"
 #include "Templates/Tuple.h"
 
+class UTexture;
+class USkeletalMesh;
+
 
 namespace mu
 {
@@ -365,23 +368,7 @@ namespace mu
     	if (OutProjectionAngle) *OutProjectionAngle = Projector.projectionAngle;
 	}
 
-
-	FName FModel::GetImageDefaultValue(int32 Index) const
-    {
-	    check(m_pD->Program.Parameters.IsValidIndex(Index));
-		check(m_pD->Program.Parameters[Index].Type == EParameterType::Image);
-
-        // Early out in case of invalid parameters
-        if (!m_pD->Program.Parameters.IsValidIndex(Index) ||
-            m_pD->Program.Parameters[Index].Type != EParameterType::Image)
-        {
-            return {};
-        }
-		
-        return m_pD->Program.Parameters[Index].DefaultValue.Get<FParamAssetType>();
-    }
-
-
+	
     int32 FModel::GetRomCount() const
     {
     	return m_pD->Program.Roms.Num();
@@ -431,8 +418,11 @@ namespace mu
     }
 
 
-    TSharedPtr<FParameters> FModel::NewParameters(TSharedPtr<const FModel> Model, const FParameters* OldParameters )
-    {
+	TSharedPtr<FParameters> FModel::NewParameters(TSharedPtr<const FModel> Model,
+		const FParameters* OldParameters,
+		const TMap<FName, TObjectPtr<UTexture>>* ImageDefaults,
+		const TMap<FName, TObjectPtr<USkeletalMesh>>* MeshDefaults)
+	{
 		LLM_SCOPE_BYNAME(TEXT("MutableRuntime"));
 
         TSharedPtr<FParameters> pRes = MakeShared<FParameters>();
@@ -500,6 +490,30 @@ namespace mu
             }
         }
 
+    	if (ImageDefaults)
+    	{
+    		for (const TPair<FName, TObjectPtr<UTexture>>& Pair : *ImageDefaults)
+    		{
+    			int32 Index = pRes->GetPrivate()->Find(Pair.Key.ToString());
+				check(Index != INDEX_NONE);
+    			check(pRes->GetType(Index) == EParameterType::Image);
+    			
+    			pRes->SetImageValue(Index, Pair.Value);
+    		}
+    	}
+    	
+    	if (MeshDefaults)
+    	{
+    		for (const TPair<FName, TObjectPtr<USkeletalMesh>>& Pair : *MeshDefaults)
+    		{
+    			int32 Index = pRes->GetPrivate()->Find(Pair.Key.ToString());
+    			check(Index != INDEX_NONE);
+    			check(pRes->GetType(Index) == EParameterType::Mesh);
+    			
+    			pRes->SetMeshValue(Index, Pair.Value);
+    		}
+    	}
+
         return pRes;
     }
 
diff --git a/Engine/Plugins/Mutable/Source/MutableRuntime/Private/MuR/Parameters.cpp b/Engine/Plugins/Mutable/Source/MutableRuntime/Private/MuR/Parameters.cpp
--- a/Engine/Plugins/Mutable/Source/MutableRuntime/Private/MuR/Parameters.cpp
+++ b/Engine/Plugins/Mutable/Source/MutableRuntime/Private/MuR/Parameters.cpp
@@ -592,7 +592,7 @@ namespace mu
     }
 
 
-	FName FParameters::GetMeshValue(int32 Index, const FRangeIndex* Pos) const
+	USkeletalMesh* FParameters::GetMeshValue(int32 Index, const FRangeIndex* Pos) const
 	{
 		check(Index >= 0 && Index < (int32)m_pD->Values.Num());
 		check(GetType(Index) == EParameterType::Mesh);
@@ -610,7 +610,7 @@ namespace mu
 		// Single value case
 		if (!Pos)
 		{
-			return m_pD->Values[Index].Get<FParamAssetType>();
+			return m_pD->Values[Index].Get<FParamSkeletalMeshType>().Get();
 		}
 
 		// Multivalue case
@@ -622,17 +622,17 @@ namespace mu
 			const FParameterValue* it = m.Find(Pos->Values);
 			if (it)
 			{
-				return it->Get<FParamAssetType>();
+				return it->Get<FParamSkeletalMeshType>().Get();
 			}
 		}
 
 		// Multivalue parameter, but no multivalue set. Return single value.
-		return m_pD->Values[Index].Get<FParamAssetType>();
+		return m_pD->Values[Index].Get<FParamSkeletalMeshType>().Get();
 
 	}
 
 
-	void FParameters::SetMeshValue(int32 Index, FName id, const FRangeIndex* Pos)
+	void FParameters::SetMeshValue(int32 Index, USkeletalMesh* Value, const FRangeIndex* Pos)
 	{
 		check(Index >= 0 && Index < (int32)m_pD->Values.Num());
 		check(GetType(Index) == EParameterType::Mesh);
@@ -640,7 +640,7 @@ namespace mu
 		// Single value case
 		if (!Pos)
 		{
-			m_pD->Values[Index].Set<FParamAssetType>(id);
+			m_pD->Values[Index].Set<FParamSkeletalMeshType>(TStrongObjectPtr(Value));
 		}
 
 		// Multivalue case
@@ -655,12 +655,12 @@ namespace mu
 
 			TMap< TArray<int32>, FParameterValue >& m = m_pD->MultiValues[Index];
 			FParameterValue& it = m.FindOrAdd(Pos->Values);
-			it.Set<FParamAssetType>(id);
+			it.Set<FParamSkeletalMeshType>(TStrongObjectPtr(Value));
 		}
 	}
 
 
-	FName FParameters::GetImageValue( int32 Index, const FRangeIndex* Pos ) const
+	UTexture* FParameters::GetImageValue( int32 Index, const FRangeIndex* Pos ) const
     {
         check( Index >= 0 && Index < (int32)m_pD->Values.Num() );
         check( GetType( Index ) == EParameterType::Image );
@@ -678,7 +678,7 @@ namespace mu
 		// Single value case
 		if (!Pos)
 		{
-			return m_pD->Values[Index].Get<FParamAssetType>();
+			return m_pD->Values[Index].Get<FParamTextureType>().Get();
 		}
 
 		// Multivalue case
@@ -690,17 +690,16 @@ namespace mu
 			const FParameterValue* it = m.Find(Pos->Values);
 			if (it)
 			{
-				return it->Get<FParamAssetType>();
+				return it->Get<FParamTextureType>().Get();
 			}
 		}
 
 		// Multivalue parameter, but no multivalue set. Return single value.
-		return m_pD->Values[Index].Get<FParamAssetType>();
-
+		return m_pD->Values[Index].Get<FParamTextureType>().Get();
     }
 
 
-    void FParameters::SetImageValue( int32 Index, FName id, const FRangeIndex* Pos )
+    void FParameters::SetImageValue( int32 Index, UTexture* Value, const FRangeIndex* Pos )
     {
         check( Index >= 0 && Index < (int32)m_pD->Values.Num() );
         check( GetType( Index ) == EParameterType::Image );
@@ -708,7 +707,7 @@ namespace mu
 		// Single value case
 		if (!Pos)
 		{
-			m_pD->Values[Index].Set<FParamAssetType>(id);
+			m_pD->Values[Index].Set<FParamTextureType>(TStrongObjectPtr(Value));
 		}
 
 		// Multivalue case
@@ -723,7 +722,7 @@ namespace mu
 
 			TMap< TArray<int32>, FParameterValue >& m = m_pD->MultiValues[Index];
 			FParameterValue& it = m.FindOrAdd(Pos->Values);
-			it.Set<FParamAssetType>(id);
+			it.Set<FParamTextureType>(TStrongObjectPtr(Value));
 		}
     }
 
diff --git a/Engine/Plugins/Mutable/Source/MutableRuntime/Private/MuR/System.cpp b/Engine/Plugins/Mutable/Source/MutableRuntime/Private/MuR/System.cpp
--- a/Engine/Plugins/Mutable/Source/MutableRuntime/Private/MuR/System.cpp
+++ b/Engine/Plugins/Mutable/Source/MutableRuntime/Private/MuR/System.cpp
@@ -1904,8 +1904,8 @@ namespace mu
 
 			case EParameterType::Image:
 			{
-				DataSize = sizeof(FName);
-				Blob.Serialize((void*)&Params->GetPrivate()->Values[ParamIndex].Get<FParamAssetType>(), DataSize);
+				FString Path = Params->GetPrivate()->Values[ParamIndex].Get<FParamTextureType>()->GetPathName();
+				Blob.Serialize(GetData(Path), Path.GetAllocatedSize());
 
 				// Multi-values
 				if (Params->GetPrivate()->HasMultipleValues(ParamIndex))
@@ -1914,7 +1914,27 @@ namespace mu
 					AddMultiValueKeys(Blob, Multi);
 					for (const TPair< TArray<int32>, FParameterValue >& v : Multi)
 					{
-						Blob.Serialize((void*)&v.Value.Get<FParamAssetType>(), DataSize);
+						Path = v.Value.Get<FParamTextureType>()->GetPathName();
+						Blob.Serialize(GetData(Path), Path.GetAllocatedSize());
+					}
+				}
+				break;
+			}
+
+			case EParameterType::Mesh:
+			{
+				FString Path = Params->GetPrivate()->Values[ParamIndex].Get<FParamSkeletalMeshType>()->GetPathName();
+				Blob.Serialize(GetData(Path), Path.GetAllocatedSize());
+
+				// Multi-values
+				if (Params->GetPrivate()->HasMultipleValues(ParamIndex))
+				{
+					const TMap< TArray<int32>, FParameterValue >& Multi = Params->GetPrivate()->MultiValues[ParamIndex];
+					AddMultiValueKeys(Blob, Multi);
+					for (const TPair< TArray<int32>, FParameterValue >& v : Multi)
+					{
+						Path = v.Value.Get<FParamSkeletalMeshType>()->GetPathName();
+						Blob.Serialize(GetData(Path), Path.GetAllocatedSize());
 					}
 				}
 				break;
diff --git a/Engine/Plugins/Mutable/Source/MutableRuntime/Public/MuR/Model.h b/Engine/Plugins/Mutable/Source/MutableRuntime/Public/MuR/Model.h
--- a/Engine/Plugins/Mutable/Source/MutableRuntime/Public/MuR/Model.h
+++ b/Engine/Plugins/Mutable/Source/MutableRuntime/Public/MuR/Model.h
@@ -67,7 +67,10 @@ namespace mu
 		//! Create a set of new parameters of the model with the default values.
 		//! If old parameters are provided, they will be reused when possible instead of the
 		//! default values.
-        static UE_API TSharedPtr<FParameters> NewParameters( TSharedPtr<const FModel> Model, const FParameters* Old = nullptr );
+        static UE_API TSharedPtr<FParameters> NewParameters(TSharedPtr<const FModel> Model,
+        	const FParameters* OldParameters = nullptr,
+        	const TMap<FName, TObjectPtr<UTexture>>* ImageDefaults = nullptr,
+        	const TMap<FName, TObjectPtr<USkeletalMesh>>* MeshDefaults = nullptr);
 
 		/** Return true if the parameter is multi-dimensional */
 		UE_API bool IsParameterMultidimensional(int32 ParameterIndex) const;
@@ -130,12 +133,6 @@ namespace mu
     	UE_API void GetProjectorDefaultValue(int32 Index, EProjectorType* OutProjectionType, FVector3f* OutPos,
     	 	FVector3f* OutDir, FVector3f* OutUp, FVector3f* OutScale, float* OutProjectionAngle) const;
 
-        //! Return the default value of an image parameter.
-        //! \pre The parameter specified by index is a T_IMAGE.
-        //! \param Index Index of the parameter from 0 to GetCount()-1
-		//! \return The externalId specified when setting the image value (\see SetImageValue)
-		UE_API FName GetImageDefaultValue(int32 Index) const;
-
     	UE_API int32 GetRomCount() const;
 
 #if WITH_EDITOR
diff --git a/Engine/Plugins/Mutable/Source/MutableRuntime/Public/MuR/Parameters.h b/Engine/Plugins/Mutable/Source/MutableRuntime/Public/MuR/Parameters.h
--- a/Engine/Plugins/Mutable/Source/MutableRuntime/Public/MuR/Parameters.h
+++ b/Engine/Plugins/Mutable/Source/MutableRuntime/Public/MuR/Parameters.h
@@ -8,6 +8,9 @@
 #include "Templates/SharedPointer.h"
 #include "VectorTypes.h"
 
+class UTexture;
+class USkeletalMesh;
+
 #define UE_API MUTABLERUNTIME_API
 
 
@@ -311,19 +314,19 @@ namespace mu
         //! \param index Index of the parameter from 0 to GetCount()-1
 		//! \param pos Only for multidimensional parameters: relevant position to set in the ranges
 	    //! \return The externalId specified when setting the image value (\see SetImageValue)
-		UE_API FName GetImageValue( int32 Index, const FRangeIndex* Pos = nullptr) const;
+		UE_API UTexture* GetImageValue( int32 Index, const FRangeIndex* Pos = nullptr) const;
 
 		/** */
-		UE_API FName GetMeshValue(int32 Index, const FRangeIndex* Pos = nullptr) const;
+		UE_API USkeletalMesh* GetMeshValue(int32 Index, const FRangeIndex* Pos = nullptr) const;
 
         //! If the parameter is of the image type, set its value.
         //! \param index Index of the parameter from 0 to GetCount()-1
         //! \param externalId Application-specific id used to identify this image during replication.
 		//! \param pos Only for multidimensional parameters: relevant position to set in the ranges
-		UE_API void SetImageValue( int32 Index, FName externalId, const FRangeIndex* = nullptr);
+		UE_API void SetImageValue( int32 Index, UTexture* externalId, const FRangeIndex* = nullptr);
 
 		/** */
-		UE_API void SetMeshValue(int32 Index, FName externalId, const FRangeIndex* = nullptr);
+		UE_API void SetMeshValue(int32 Index, USkeletalMesh* externalId, const FRangeIndex* = nullptr);
 
         //! Return the value of a float parameter.
         //! \pre The parameter specified by index is a T_FLOAT.
diff --git a/Engine/Plugins/Mutable/Source/MutableRuntime/Public/MuR/Serialisation.h b/Engine/Plugins/Mutable/Source/MutableRuntime/Public/MuR/Serialisation.h
--- a/Engine/Plugins/Mutable/Source/MutableRuntime/Public/MuR/Serialisation.h
+++ b/Engine/Plugins/Mutable/Source/MutableRuntime/Public/MuR/Serialisation.h
@@ -15,6 +15,10 @@
 #include "Containers/Map.h"
 #include "Containers/ContainerAllocationPolicies.h"
 #include "Templates/SharedPointer.h"
+#include "UObject/StrongObjectPtr.h"
+
+class UTexture;
+class USkeletalMesh;
 
 #define UE_API MUTABLERUNTIME_API
 
@@ -406,6 +410,21 @@ namespace mu
 	// Bool size is not a standard
 	MUTABLERUNTIME_API void operator<<(FOutputArchive&, const bool&);
 	MUTABLERUNTIME_API void operator>>(FInputArchive&, bool&);
+
+
+	// Do not serialize pointers to UObjects.
+	template<typename T>
+	void operator<<(FOutputArchive& Arch, const TStrongObjectPtr<T>& V)
+	{
+    	// Do Nothing. UObjects values are saved in the CO.	
+	}
+
+	template<typename T>
+	void operator>>(FInputArchive& Arch, TStrongObjectPtr<T>& V)
+	{
+    	// Do Nothing. UObjects values are saved in the CO.	
+	}
 }
 
+
 #undef UE_API
diff --git a/Engine/Plugins/Mutable/Source/MutableRuntime/Public/MuR/System.h b/Engine/Plugins/Mutable/Source/MutableRuntime/Public/MuR/System.h
--- a/Engine/Plugins/Mutable/Source/MutableRuntime/Public/MuR/System.h
+++ b/Engine/Plugins/Mutable/Source/MutableRuntime/Public/MuR/System.h
@@ -15,6 +15,9 @@
 
 #include "System.generated.h"
 
+class UTexture;
+class USkeletalMesh;
+
 #define UE_API MUTABLERUNTIME_API
 
 
@@ -79,13 +82,13 @@ namespace mu
         virtual ~FExternalResourceProvider() = default;
 
         /** Returns the completion event and a cleanup function that must be called once event is completed. */
-		virtual TTuple<UE::Tasks::FTask, TFunction<void()>> GetImageAsync(FName Id, uint8 MipmapsToSkip, TFunction<void(TSharedPtr<FImage>)>& ResultCallback) = 0;
+		virtual TTuple<UE::Tasks::FTask, TFunction<void()>> GetImageAsync(UTexture* Texture, uint8 MipmapsToSkip, TFunction<void(TSharedPtr<FImage>)>& ResultCallback) = 0;
 		virtual TTuple<UE::Tasks::FTask, TFunction<void()>> GetReferencedImageAsync(const void* ModelPtr, int32 Id, uint8 MipmapsToSkip, TFunction<void(TSharedPtr<FImage>)>& ResultCallback) { check(false); return {}; }
 
-		virtual mu::FExtendedImageDesc GetImageDesc(FName Id) = 0;
+		virtual mu::FExtendedImageDesc GetImageDesc(UTexture* Texture) = 0;
 
 		/** Returns the completion event and a cleanup function that must be called once event is completed. */
-		virtual TTuple<UE::Tasks::FTask, TFunction<void()>> GetMeshAsync(FName Id, int32 InLODIndex, int32 InSectionIndex, TFunction<void(TSharedPtr<FMesh>)>& ResultCallback) = 0;
+		virtual TTuple<UE::Tasks::FTask, TFunction<void()>> GetMeshAsync(USkeletalMesh* SkeleltalMesh, int32 InLODIndex, int32 InSectionIndex, TFunction<void(TSharedPtr<FMesh>)>& ResultCallback) = 0;
     };
 
 
diff --git a/Engine/Plugins/Mutable/Source/MutableTools/Private/MuT/CodeGenerator_Image.cpp b/Engine/Plugins/Mutable/Source/MutableTools/Private/MuT/CodeGenerator_Image.cpp
--- a/Engine/Plugins/Mutable/Source/MutableTools/Private/MuT/CodeGenerator_Image.cpp
+++ b/Engine/Plugins/Mutable/Source/MutableTools/Private/MuT/CodeGenerator_Image.cpp
@@ -267,7 +267,7 @@ namespace mu
 				bool bParseOk = FGuid::Parse(Node->UID, op->Parameter.UID);
 				check(bParseOk);
 				op->Parameter.Type = EParameterType::Image;
-				op->Parameter.DefaultValue.Set<FParamAssetType>(Node->DefaultValue);
+				op->Parameter.DefaultValue.Set<FParamTextureType>(nullptr);
 
 				FirstPass.ParameterNodes.GenericParametersCache.Add(Node, op);
 			}
diff --git a/Engine/Plugins/Mutable/Source/MutableTools/Private/MuT/CodeGenerator_Mesh.cpp b/Engine/Plugins/Mutable/Source/MutableTools/Private/MuT/CodeGenerator_Mesh.cpp
--- a/Engine/Plugins/Mutable/Source/MutableTools/Private/MuT/CodeGenerator_Mesh.cpp
+++ b/Engine/Plugins/Mutable/Source/MutableTools/Private/MuT/CodeGenerator_Mesh.cpp
@@ -1317,7 +1317,7 @@ namespace mu
 				bool bParseOk = FGuid::Parse(Node->UID, param.UID);
 				check(bParseOk);
 				param.Type = EParameterType::Mesh;
-				param.DefaultValue.Set<FParamAssetType>(Node->DefaultValue);
+				param.DefaultValue.Set<FParamSkeletalMeshType>(nullptr);
 
 				Op = new ASTOpParameter();
 				Op->Type = EOpType::ME_PARAMETER;
diff --git a/Engine/Plugins/Mutable/Source/MutableTools/Public/MuT/NodeImageParameter.h b/Engine/Plugins/Mutable/Source/MutableTools/Public/MuT/NodeImageParameter.h
--- a/Engine/Plugins/Mutable/Source/MutableTools/Public/MuT/NodeImageParameter.h
+++ b/Engine/Plugins/Mutable/Source/MutableTools/Public/MuT/NodeImageParameter.h
@@ -15,10 +15,6 @@ namespace mu
     class MUTABLETOOLS_API NodeImageParameter : public NodeImage
 	{
 	public:
-
-		/** Full asset name for the default value of the parameter. */
-		FName DefaultValue;
-		
 		/** Name of the parameter */
 		FString Name;
 
diff --git a/Engine/Plugins/Mutable/Source/MutableTools/Public/MuT/NodeMeshParameter.h b/Engine/Plugins/Mutable/Source/MutableTools/Public/MuT/NodeMeshParameter.h
--- a/Engine/Plugins/Mutable/Source/MutableTools/Public/MuT/NodeMeshParameter.h
+++ b/Engine/Plugins/Mutable/Source/MutableTools/Public/MuT/NodeMeshParameter.h
@@ -16,10 +16,6 @@ namespace mu
 	class MUTABLETOOLS_API NodeMeshParameter : public NodeMesh
 	{
 	public:
-
-		/** Full asset name for the default value of the parameter. */
-		FName DefaultValue;
-
 		/** */
 		TArray<Ptr<NodeLayout>> Layouts;
 
