[Mutable] Extend ME_ADDTAGS opertation to also add other metadata, skeleton ids and external resource ids.

- ME_ADDTAGS ops now are called ME_ADDMETADATA
- Metadata is now stored in lists that can potentally be shared between operations. This can save a considerable amount of program size in some cases.
- Added log at the end of a compilation indicating the program size.

#rnx

diff --git a/Engine/Plugins/Mutable/Source/CustomizableObject/Internal/MuCO/CustomizableObjectPrivate.h b/Engine/Plugins/Mutable/Source/CustomizableObject/Internal/MuCO/CustomizableObjectPrivate.h
--- a/Engine/Plugins/Mutable/Source/CustomizableObject/Internal/MuCO/CustomizableObjectPrivate.h
+++ b/Engine/Plugins/Mutable/Source/CustomizableObject/Internal/MuCO/CustomizableObjectPrivate.h
@@ -1619,6 +1619,8 @@ enum class ECustomizableObjectVersions : int32
 	IncorrectBonePoseMerging,
 
 	CoreParameterUObjects,
+
+	MoveMeshMetadataToOperation,
 	
 	// -----<new versions can be added above this line>--------
 	LastCustomizableObjectVersion
@@ -1885,7 +1887,7 @@ public:
 
 #endif
 	
-	static constexpr int32 DerivedDataVersion = 0x34789a1c;
+	static constexpr int32 DerivedDataVersion = 0x7290caa3;
 };
 
 #if WITH_EDITOR
@@ -1906,4 +1908,4 @@ CUSTOMIZABLEOBJECT_API TObjectPtr<UModelResources> LoadModelResources_Internal(F
 CUSTOMIZABLEOBJECT_API const TSharedPtr<FModelStreamableBulkData> LoadModelStreamableBulk_Internal(FArchive& Ar);
 CUSTOMIZABLEOBJECT_API const TSharedPtr<mu::FModel, ESPMode::ThreadSafe> LoadModel_Internal(FArchive& Ar);
 
-#endif
\ No newline at end of file
+#endif
diff --git a/Engine/Plugins/Mutable/Source/CustomizableObjectEditor/Private/MuCOE/SMutableCodeViewer.cpp b/Engine/Plugins/Mutable/Source/CustomizableObjectEditor/Private/MuCOE/SMutableCodeViewer.cpp
--- a/Engine/Plugins/Mutable/Source/CustomizableObjectEditor/Private/MuCOE/SMutableCodeViewer.cpp
+++ b/Engine/Plugins/Mutable/Source/CustomizableObjectEditor/Private/MuCOE/SMutableCodeViewer.cpp
@@ -1535,20 +1535,40 @@ void SMutableCodeViewer::GenerateElementRecursive(const int32& InStateIndex, mu:
 		break;
 	}
 
-	case mu::EOpType::ME_ADDTAGS:
+	case mu::EOpType::ME_ADDMETADATA:
 	{
-		const uint8* OpData = InProgram.GetOpArgsPointer(InParentAddress);
+		mu::OP::MeshAddMetadataArgs Args = InProgram.GetOpArgs<mu::OP::MeshAddMetadataArgs>(InParentAddress);
 
-		mu::OP::ADDRESS SourceAddress;
-		FMemory::Memcpy(&SourceAddress, OpData, sizeof(mu::OP::ADDRESS));
-		OpData += sizeof(mu::OP::ADDRESS);
+		using OpEnumFlags = mu::OP::MeshAddMetadataArgs::EnumFlags;
+
+		// Set count to 1 if is not a list.
+		int32 TagCount      = int32(!EnumHasAnyFlags(Args.Flags, OpEnumFlags::IsTagList));
+		int32 ResourceCount = int32(!EnumHasAnyFlags(Args.Flags, OpEnumFlags::IsResourceList));
+		int32 SkeletonCount = int32(!EnumHasAnyFlags(Args.Flags, OpEnumFlags::IsSkeletonList));
+
+		if (!TagCount)
+		{
+			TagCount = InProgram.ConstantUInt32Lists.IsValidIndex(Args.Tags.ListAddress) 
+					? InProgram.ConstantUInt32Lists[Args.Tags.ListAddress].Num() 
+					: 0;
+		}
 
-		uint16 TagCount;
-		FMemory::Memcpy(&TagCount, OpData, sizeof(uint16));
-		OpData += sizeof(uint16);
+		if (!ResourceCount)
+		{
+			ResourceCount = InProgram.ConstantUInt64Lists.IsValidIndex(Args.ResourceIds.ListAddress) 
+					? InProgram.ConstantUInt64Lists[Args.ResourceIds.ListAddress].Num() 
+					: 0;
+		}
+
+		if (!SkeletonCount)
+		{
+			SkeletonCount = InProgram.ConstantUInt32Lists.IsValidIndex(Args.SkeletonIds.ListAddress) 
+					? InProgram.ConstantUInt32Lists[Args.SkeletonIds.ListAddress].Num() 
+					: 0;
+		}
 
-		FString Caption = FString::Printf(TEXT("add %d tags to "), TagCount);
-		AddOpFunc(SourceAddress, Caption);
+		FString Caption = FString::Printf(TEXT("add %d tags %d resources and %d skeletons to "), TagCount, ResourceCount, SkeletonCount);
+		AddOpFunc(Args.Source, Caption);
 
 		break;
 	}
@@ -2595,7 +2615,7 @@ bool SMutableCodeViewer::IsConstantResourceUsedByOperation(const int32 IndexOnCo
 						Data += sizeof(int32);
 					}
 				}
-				else if (OperationType == mu::EOpType::ME_ADDTAGS)
+				else if (OperationType == mu::EOpType::ME_ADDMETADATA)
 				{
 					 const uint8* OpData = InProgram.GetOpArgsPointer(OperationAddress);
 					
diff --git a/Engine/Plugins/Mutable/Source/MutableRuntime/Internal/MuR/ModelPrivate.h b/Engine/Plugins/Mutable/Source/MutableRuntime/Internal/MuR/ModelPrivate.h
--- a/Engine/Plugins/Mutable/Source/MutableRuntime/Internal/MuR/ModelPrivate.h
+++ b/Engine/Plugins/Mutable/Source/MutableRuntime/Internal/MuR/ModelPrivate.h
@@ -49,7 +49,7 @@ namespace mu
 	*/
     struct FRomDataRuntime
     {
-		//! Size of the rom
+		/** Size of the rom */
 		uint32 Size : 30;
 
 		/** Index of the resource in its type-specific array. See ERomDataType. */
@@ -59,7 +59,7 @@ namespace mu
 		uint32 IsHighRes : 1;
 
 		// TODO: Store the offset here and delete the FModelStreamableBlock map?
-		//! Offset in file
+		/** Offset in file */
 		// uint32 Offset = 0;
     };
 
@@ -185,26 +185,30 @@ namespace mu
 
         struct FState
         {
-            //! Name of the state
+            /** Name of the state */
             FString Name;
 
-            //! First instruction of the full build of an instance in this state
+            /** First instruction of the full build of an instance in this state */
             OP::ADDRESS Root = 0;
 
-            //! List of parameters index (to FProgram::Parameters) of the runtime parameters of
-            //! this state.
+            /** 
+			 * List of parameters index (to FProgram::Parameters) of the runtime parameters of
+             * this state.
+			 */
 			TArray<int> m_runtimeParameters;
 
-            //! List of instructions that need to be cached to efficiently update this state
+            /** List of instructions that need to be cached to efficiently update this state */
 			TArray<OP::ADDRESS> m_updateCache;
 
-            //! List of root instructions for the dynamic resources that depend on the runtime
-            //! parameters of this state, with a mask of relevant runtime parameters.
-            //! The mask has a bit on for every runtime parameter in the m_runtimeParameters array.
-			//! The uint64 is linked to MUTABLE_MAX_RUNTIME_PARAMETERS_PER_STATE
+            /** 
+			 * List of root instructions for the dynamic resources that depend on the runtime
+			 * parameters of this state, with a mask of relevant runtime parameters.
+             * The mask has a bit on for every runtime parameter in the m_runtimeParameters array.
+			 * The uint64 is linked to MUTABLE_MAX_RUNTIME_PARAMETERS_PER_STATE
+			 */
 			TArray< TPair<OP::ADDRESS,uint64> > m_dynamicResources;
 
-            //!
+            /** */
             inline void Serialise( FOutputArchive& arch ) const
             {
                 arch << Name;
@@ -214,7 +218,7 @@ namespace mu
                 arch << m_dynamicResources;
             }
 
-            //!
+            /** */
             inline void Unserialise( FInputArchive& arch )
             {
                 arch >> Name;
@@ -243,7 +247,7 @@ namespace mu
                 return res;
             }
 
-            //!
+            /** */
             bool IsUpdateCache( OP::ADDRESS at ) const
             {
                 bool res = false;
@@ -259,7 +263,7 @@ namespace mu
                 return res;
             }
 
-            //!
+            /** */ 
             void AddUpdateCache( OP::ADDRESS at )
             {
                 if ( !IsUpdateCache(at) )
@@ -269,13 +273,13 @@ namespace mu
             }
         };
 
-        //! Location in the ByteCode of the beginning of each operation
+        /** Location in the ByteCode of the beginning of each operation */
 		TArray<uint32> OpAddress;
 
-        //! Byte-coded representation of the program, using variable-sized op data.
+        /** Byte-coded representation of the program, using variable-sized op data. */
 		TArray<uint8> ByteCode;
 
-        //!
+        /** */
 		TArray<FState> States;
 
 		/** Data for every rom required in-game. */
@@ -298,10 +302,10 @@ namespace mu
 		*/
 		TMap<uint32, TSharedPtr<const FImage>> ConstantImageLODsStreamed;
 
-		//! Constant image mip chain indices: ranges in this array are defined in FImageLODRange and the indices here refer to ConstantImageLODs.
+		/** Constant image mip chain indices: ranges in this array are defined in FImageLODRange and the indices here refer to ConstantImageLODs. */
 		TArray<FConstantResourceIndex> ConstantImageLODIndices;
 
-		//! Constant image data.
+		/** Constant image data. */
 		TArray<FImageLODRange> ConstantImages;
 
 		/** Constant mesh content indices: ranges in this array are defined in FMeshContentRange and the indices here refer to ConstantMeshes. */
@@ -321,113 +325,127 @@ namespace mu
 		*/
 		TMap<uint32, TSharedPtr<const FMesh>> ConstantMeshesStreamed;
 
-		//! Constant FExtensionData
+		/** Constant FExtensionData */
 		TArray<FExtensionDataConstant> ConstantExtensionData;
 
-        //! Constant string data
+        /** Constant string data */
 		TArray<FString> ConstantStrings;
 
-        //! Constant layout data
+		/** */
+		TArray<TArray<uint32>> ConstantUInt32Lists;
+
+		/** */
+		TArray<TArray<uint64>> ConstantUInt64Lists;
+
+        /** Constant layout data */
 		TArray<TSharedPtr<const FLayout>> ConstantLayouts;
 
-        //! Constant projectors
+        /** Constant projectors */
 		TArray<FProjector> ConstantProjectors;
 
-        //! Constant matrices, usually used for transforms
+        /** Constant matrices, usually used for transforms */
 		TArray<FMatrix44f> ConstantMatrices;
 
-		//! Constant shapes
+		/** Constant shapes */
 		TArray<FShape> ConstantShapes;
 
-        //! Constant curves
+        /** Constant curves */
 		TArray<FRichCurve> ConstantCurves;
 
-        //! Constant skeletons
+        /** Constant skeletons */
 		TArray<TSharedPtr<const FSkeleton>> ConstantSkeletons;
 
-		//! Constant Physics Bodies
+		/** Constant Physics Bodies */
 		TArray<TSharedPtr<const FPhysicsBody>> ConstantPhysicsBodies;
 
-        //! FParameters of the model.
-        //! The value stored here is the default value.
+        /** FParameters of the model. */
+        /** The value stored here is the default value. */
 		TArray<FParameterDesc> Parameters;
 
-        //! Ranges for iteration of the model operations.
+        /** Ranges for iteration of the model operations. */
 		TArray<FRangeDesc> Ranges;
 
-        //! List of parameter lists. These are used in several places, like storing the
-        //! pregenerated list of parameters influencing a resource.
-        //! The parameter lists are sorted.
+        /** 
+		 * List of parameter lists. These are used in several places, like storing the
+         * pregenerated list of parameters influencing a resource.
+         * The parameter lists are sorted. 
+		 */
 		TArray<TArray<uint16>> ParameterLists;
 
 #if WITH_EDITOR
-		//! State of the program. True unless the streamed resources were destroyed,
-		//! which could happen in the editor after recompiling the CO.
+		/** 
+		 * State of the program. True unless the streamed resources were destroyed,
+		 * which could happen in the editor after recompiling the CO. 
+		 */
 		bool bIsValid = true;
 #endif
-        //!
-        void Serialise( FOutputArchive& arch ) const
+        /** */
+        void Serialise(FOutputArchive& Arch) const
         {
-            arch << OpAddress;
-            arch << ByteCode;
-            arch << States;
-			arch << Roms;
-			arch << RomsCompileData;
-			arch << ConstantImageLODsPermanent;
-			arch << ConstantImageLODIndices;
-			arch << ConstantImages;
-			arch << ConstantMeshesPermanent;
-			arch << ConstantMeshContentIndices;
-			arch << ConstantMeshes;
-			arch << ConstantExtensionData;
-			arch << ConstantStrings;
-            arch << ConstantLayouts;
-            arch << ConstantProjectors;
-			arch << ConstantMatrices;
-			arch << ConstantShapes;
-            arch << ConstantCurves;
-            arch << ConstantSkeletons;
-			arch << ConstantPhysicsBodies;
-            arch << Parameters;
-            arch << Ranges;
-            arch << ParameterLists;
+            Arch << OpAddress;
+            Arch << ByteCode;
+            Arch << States;
+			Arch << Roms;
+			Arch << RomsCompileData;
+			Arch << ConstantImageLODsPermanent;
+			Arch << ConstantImageLODIndices;
+			Arch << ConstantImages;
+			Arch << ConstantMeshesPermanent;
+			Arch << ConstantMeshContentIndices;
+			Arch << ConstantMeshes;
+			Arch << ConstantExtensionData;
+			Arch << ConstantStrings;
+			Arch << ConstantUInt32Lists;
+			Arch << ConstantUInt64Lists;
+            Arch << ConstantLayouts;
+            Arch << ConstantProjectors;
+			Arch << ConstantMatrices;
+			Arch << ConstantShapes;
+            Arch << ConstantCurves;
+            Arch << ConstantSkeletons;
+			Arch << ConstantPhysicsBodies;
+            Arch << Parameters;
+            Arch << Ranges;
+            Arch << ParameterLists;
         }
 
-        //!
-        void Unserialise( FInputArchive& arch )
+        /** */
+        void Unserialise(FInputArchive& Arch)
         {
-            arch >> OpAddress;
-            arch >> ByteCode;
-            arch >> States;
-			arch >> Roms;
-			arch >> RomsCompileData;
-			arch >> ConstantImageLODsPermanent;
-			arch >> ConstantImageLODIndices;
-			arch >> ConstantImages;
-			arch >> ConstantMeshesPermanent;
-			arch >> ConstantMeshContentIndices;
-			arch >> ConstantMeshes;
-			arch >> ConstantExtensionData;
-			arch >> ConstantStrings;
-            arch >> ConstantLayouts;
-            arch >> ConstantProjectors;
-			arch >> ConstantMatrices;
-			arch >> ConstantShapes;
-            arch >> ConstantCurves;
-            arch >> ConstantSkeletons;
-			arch >> ConstantPhysicsBodies;
-            arch >> Parameters;
-            arch >> Ranges;
-            arch >> ParameterLists;
+            Arch >> OpAddress;
+            Arch >> ByteCode;
+            Arch >> States;
+			Arch >> Roms;
+			Arch >> RomsCompileData;
+			Arch >> ConstantImageLODsPermanent;
+			Arch >> ConstantImageLODIndices;
+			Arch >> ConstantImages;
+			Arch >> ConstantMeshesPermanent;
+			Arch >> ConstantMeshContentIndices;
+			Arch >> ConstantMeshes;
+			Arch >> ConstantExtensionData;
+			Arch >> ConstantStrings;
+			Arch >> ConstantUInt32Lists;
+			Arch >> ConstantUInt64Lists;
+            Arch >> ConstantLayouts;
+            Arch >> ConstantProjectors;
+			Arch >> ConstantMatrices;
+			Arch >> ConstantShapes;
+            Arch >> ConstantCurves;
+            Arch >> ConstantSkeletons;
+			Arch >> ConstantPhysicsBodies;
+            Arch >> Parameters;
+            Arch >> Ranges;
+            Arch >> ParameterLists;
         }
 
-        //! Debug method that sanity-checks the program with a variety of tests.
+        /** Debug method that sanity-checks the program with a variety of tests. */
         void Check();
 
-        //! Debug method that logs the top used instruction types.
+        /** Debug method that logs the top used instruction types. */
         void LogHistogram() const;
 
-		//! Return true if the given ROM is loaded.
+		/** Return true if the given ROM is loaded. */
 		FORCEINLINE bool IsRomLoaded(int32 RomIndex) const
 		{
 			switch (Roms[RomIndex].ResourceType)
@@ -575,21 +593,35 @@ namespace mu
             return index;
         }
 
-        OP::ADDRESS AddConstant( const FString& str )
-        {            
-            // Ensure unique
-            for (int32 i=0; i<ConstantStrings.Num(); ++i)
-            {
-                if (ConstantStrings[i]==str)
-                {
-                    return (OP::ADDRESS)i;
-                }
-            }
-
-            OP::ADDRESS index = OP::ADDRESS( ConstantStrings.Num() );
-            ConstantStrings.Add( str );
-            return index;
+        OP::ADDRESS AddConstant(const FString& Str)
+        {     
+			return (OP::ADDRESS)ConstantStrings.AddUnique(Str);
         }
+		
+		OP::ADDRESS AddConstant(const TArray<uint32>& UInt32List)
+		{
+			return (OP::ADDRESS)ConstantUInt32Lists.AddUnique(UInt32List);
+		}
+
+		OP::ADDRESS AddConstant(const TArray<uint64>& UInt64List)
+		{
+			return (OP::ADDRESS)ConstantUInt64Lists.AddUnique(UInt64List);
+		}
+
+		OP::ADDRESS AddConstant(const TArray<FString>& StringList)
+		{
+			const int32 NumStrings = StringList.Num();
+
+			TArray<uint32> StringAddrs;
+			StringAddrs.SetNum(NumStrings);
+
+			for (int32 StringIndex = 0; StringIndex < NumStrings; ++StringIndex)
+			{
+				StringAddrs[StringIndex] = AddConstant(StringList[StringIndex]);
+			}
+
+			return (OP::ADDRESS)AddConstant(StringAddrs);
+		}
 
         OP::ADDRESS AddConstant( const FMatrix44f& m )
         {
@@ -669,7 +701,7 @@ namespace mu
 		}
 
 
-        //! Get a constant image, assuming at least some mips are loaded. The image constant will be composed with lodaded mips if necessary.
+        /** Get a constant image, assuming at least some mips are loaded. The image constant will be composed with lodaded mips if necessary. */
 		template <typename CreateImageFunc>
         void GetConstant( uint32 ConstantIndex, TSharedPtr<const FImage>& res, int32 MipsToSkip, const CreateImageFunc& CreateImage) const
         {
diff --git a/Engine/Plugins/Mutable/Source/MutableRuntime/Internal/MuR/Operations.h b/Engine/Plugins/Mutable/Source/MutableRuntime/Internal/MuR/Operations.h
--- a/Engine/Plugins/Mutable/Source/MutableRuntime/Internal/MuR/Operations.h
+++ b/Engine/Plugins/Mutable/Source/MutableRuntime/Internal/MuR/Operations.h
@@ -272,8 +272,8 @@ namespace mu
 		//! Optimize skinning before adding a mesh to the component
 		ME_OPTIMIZESKINNING,
 
-		//! Add a set of tags to a mesh
-		ME_ADDTAGS,
+		//! Add a metadata to a mesh
+		ME_ADDMETADATA,
 
     	//! Transform with a 4x4 matrix the geometry channels of a mesh that are bounded by another mesh
     	ME_TRANSFORMWITHMESH,
@@ -959,6 +959,43 @@ namespace mu
 		{
 			ADDRESS source;
 		};
+         
+        struct MeshAddMetadataArgs
+        {
+            enum class EnumFlags : uint32
+            {
+                None           = 0,
+                IsTagList      = 1 << 0,
+                IsResourceList = 1 << 1,
+                IsSkeletonList = 1 << 2
+            };
+
+            union TagAddressOrUInt32ListAddress
+            {
+                ADDRESS TagAddress;
+                ADDRESS ListAddress;
+            };
+
+            union SkeletonIdOrUInt32ListAddress
+            {
+                uint32 SkeletonId;
+                ADDRESS ListAddress;
+            };
+
+            union ResourceIdOrUInt64ListAddress
+            {
+                uint64 ResourceId;
+                ADDRESS ListAddress;
+            };
+
+            EnumFlags Flags; 
+            ADDRESS Source;
+            TagAddressOrUInt32ListAddress Tags;
+            SkeletonIdOrUInt32ListAddress SkeletonIds;
+            ResourceIdOrUInt64ListAddress ResourceIds;
+        };
+        static_assert(sizeof(MeshAddMetadataArgs) == sizeof(uint64)*3, "MeshAddMetadataArgs has an unexpected size.");
+
 
         struct InstanceAddArgs
         {
@@ -1027,8 +1064,10 @@ namespace mu
 		};
     };
 
+
     typedef TOperation<uint32> OP;
 
+    ENUM_CLASS_FLAGS(OP::MeshAddMetadataArgs::EnumFlags);
 
 	//! Types of data handled by the Mutable runtime.
 	enum class EDataType : uint8
diff --git a/Engine/Plugins/Mutable/Source/MutableRuntime/Private/MuR/CodeRunner.cpp b/Engine/Plugins/Mutable/Source/MutableRuntime/Private/MuR/CodeRunner.cpp
--- a/Engine/Plugins/Mutable/Source/MutableRuntime/Private/MuR/CodeRunner.cpp
+++ b/Engine/Plugins/Mutable/Source/MutableRuntime/Private/MuR/CodeRunner.cpp
@@ -2948,27 +2948,21 @@ namespace mu
             break;
         }
 
-        case EOpType::ME_ADDTAGS:
+        case EOpType::ME_ADDMETADATA:
 		{
-			MUTABLE_CPUPROFILER_SCOPE(ME_ADDTAGS)
+			MUTABLE_CPUPROFILER_SCOPE(ME_ADDMETADATA)
 
-			// Decode op
-			// TODO: Partial decode for each stage
-			const uint8* Data = Program.GetOpArgsPointer(Item.At);
-
-			OP::ADDRESS Source;
-			FMemory::Memcpy(&Source, Data, sizeof(OP::ADDRESS));
-			Data += sizeof(OP::ADDRESS);
+			const OP::MeshAddMetadataArgs OpArgs = Program.GetOpArgs<OP::MeshAddMetadataArgs>(Item.At);
 
 			// Schedule next stages
 			switch (Item.Stage)
 			{
 			case 0:
 			{
-				if (Source)
+				if (OpArgs.Source)
 				{
 					// Request the source
-					AddOp(FScheduledOp(Item.At, Item, 1), FScheduledOp(Source, Item));
+					AddOp(FScheduledOp(Item.At, Item, 1), FScheduledOp(OpArgs.Source, Item));
 				}
 				else
 				{
@@ -2979,46 +2973,90 @@ namespace mu
 
 			case 1:
 			{
-				MUTABLE_CPUPROFILER_SCOPE(ME_ADDTAGS_2)
+				MUTABLE_CPUPROFILER_SCOPE(ME_ADDMETADATA_2)
 
-				if (!EnumHasAnyFlags(EMeshContentFlags(Item.ExecutionOptions), EMeshContentFlags::MetaData))
-				{
-					TSharedPtr<const FMesh> SourceMesh = LoadMesh(FCacheAddress(Source, Item));
-					StoreMesh(Item, SourceMesh);
-				}
-				else
+				TSharedPtr<const FMesh> SourceMesh = LoadMesh(FCacheAddress(OpArgs.Source, Item));
+
+				TSharedPtr<FMesh> Result = nullptr;
+				if (SourceMesh)
 				{
-					TSharedPtr<const FMesh> SourceMesh = LoadMesh(FCacheAddress(Source, Item));
+					Result = CloneOrTakeOver(SourceMesh);
 
-					if (!SourceMesh)
+					using OpEnumFlags = OP::MeshAddMetadataArgs::EnumFlags;
+					bool bIsTagList      = EnumHasAnyFlags(OpArgs.Flags, OpEnumFlags::IsTagList);
+					bool bIsResourceList = EnumHasAnyFlags(OpArgs.Flags, OpEnumFlags::IsResourceList);
+					bool bIsSkeletonList = EnumHasAnyFlags(OpArgs.Flags, OpEnumFlags::IsSkeletonList);
+
+					if (bIsTagList)
 					{
-						StoreMesh(Item, nullptr);
+						if (Program.ConstantUInt32Lists.IsValidIndex(OpArgs.Tags.ListAddress))
+						{
+							const TArray<uint32>& ConstantTagList = Program.ConstantUInt32Lists[OpArgs.Tags.ListAddress];
+							Result->Tags.Reserve(Result->Tags.Num() + ConstantTagList.Num());
+
+							for (uint32 TagIndex : ConstantTagList)
+							{
+								if (Program.ConstantStrings.IsValidIndex(TagIndex))
+								{
+									Result->Tags.Add(Program.ConstantStrings[TagIndex]);
+								}
+								else
+								{
+									check(false);
+								}
+							}
+						}
+						else
+						{
+							check(false);
+						}
 					}
 					else
 					{
-						TSharedPtr<FMesh> Result = CloneOrTakeOver(SourceMesh);
-
-						// Decode the tags
-						uint16 TagCount;
-						FMemory::Memcpy(&TagCount, Data, sizeof(uint16));
-						Data += sizeof(uint16);
-
-						int32 FirstMeshTagIndex = Result->Tags.Num();
-						Result->Tags.SetNum(FirstMeshTagIndex+TagCount);
-						for (uint16 TagIndex = 0; TagIndex < TagCount; ++TagIndex)
+						if (Program.ConstantStrings.IsValidIndex(OpArgs.Tags.TagAddress))
+						{
+							Result->Tags.Add(Program.ConstantStrings[OpArgs.Tags.TagAddress]);
+						}
+						else
 						{
-							OP::ADDRESS TagConstant;
-							FMemory::Memcpy(&TagConstant, Data, sizeof(OP::ADDRESS));
-							Data += sizeof(OP::ADDRESS);
+							check(false);
+						}
+					}
 
-							check(TagConstant < (uint32)InModel->GetPrivate()->Program.ConstantStrings.Num());
-							const FString& Name = Program.ConstantStrings[TagConstant];
-							Result->Tags[FirstMeshTagIndex+TagIndex] = Name;
+					if (bIsResourceList)
+					{
+						if (Program.ConstantUInt64Lists.IsValidIndex(OpArgs.ResourceIds.ListAddress))
+						{
+							Result->StreamedResources.Append(Program.ConstantUInt64Lists[OpArgs.ResourceIds.ListAddress]);
+						}
+						else
+						{
+							check(false);
 						}
+					}
+					else
+					{
+						Result->StreamedResources.Add(OpArgs.ResourceIds.ResourceId);
+					}
 
-						StoreMesh(Item, Result);
+					if (bIsSkeletonList)
+					{
+						if (Program.ConstantUInt32Lists.IsValidIndex(OpArgs.SkeletonIds.ListAddress))
+						{
+							Result->SkeletonIDs.Append(Program.ConstantUInt32Lists[OpArgs.SkeletonIds.ListAddress]);
+						}
+						else
+						{
+							check(false);
+						}
+					}
+					else
+					{
+						Result->SkeletonIDs.Add(OpArgs.SkeletonIds.SkeletonId);
 					}
 				}
+
+				StoreMesh(Item, Result);
 				break;
 			}
 
diff --git a/Engine/Plugins/Mutable/Source/MutableRuntime/Private/MuR/Operations.cpp b/Engine/Plugins/Mutable/Source/MutableRuntime/Private/MuR/Operations.cpp
--- a/Engine/Plugins/Mutable/Source/MutableRuntime/Private/MuR/Operations.cpp
+++ b/Engine/Plugins/Mutable/Source/MutableRuntime/Private/MuR/Operations.cpp
@@ -121,7 +121,7 @@ namespace mu
 		{ EDataType::Mesh,			},	// ME_CLIPDEFORM
 		{ EDataType::Mesh,			},	// ME_MORPHRESHAPE
 		{ EDataType::Mesh,			},	// ME_OPTIMIZESKINNING
-		{ EDataType::Mesh,			},	// ME_ADDTAGS
+		{ EDataType::Mesh,			},	// ME_ADDMETADATA
 		{ EDataType::Mesh,			},	// ME_TRANSFORMWITHMESH
 
 		{ EDataType::Instance,		},	// IN_ADDMESH
@@ -636,12 +636,10 @@ namespace mu
             break;
         }
 
-		case EOpType::ME_ADDTAGS:
+		case EOpType::ME_ADDMETADATA:
 		{
-			const uint8_t* data = program.GetOpArgsPointer(at);
-			mu::OP::ADDRESS source;
-			FMemory::Memcpy(&source, data, sizeof(OP::ADDRESS)); data += sizeof(OP::ADDRESS);
-			f(source);
+			OP::MeshAddMetadataArgs Args = program.GetOpArgs<OP::MeshAddMetadataArgs>(at);
+			f(Args.Source);
 			break;
 		}
 
diff --git a/Engine/Plugins/Mutable/Source/MutableTools/Private/MuT/ASTOpColorArithmetic.cpp b/Engine/Plugins/Mutable/Source/MutableTools/Private/MuT/ASTOpColorArithmetic.cpp
--- a/Engine/Plugins/Mutable/Source/MutableTools/Private/MuT/ASTOpColorArithmetic.cpp
+++ b/Engine/Plugins/Mutable/Source/MutableTools/Private/MuT/ASTOpColorArithmetic.cpp
@@ -1,7 +1,6 @@
 // Copyright Epic Games, Inc. All Rights Reserved.
 
 #include "MuT/ASTOpColorArithmetic.h"
-#include "MuT/ASTOpMeshAddTags.h"
 
 #include "HAL/PlatformMath.h"
 #include "MuR/ModelPrivate.h"
diff --git a/Engine/Plugins/Mutable/Source/MutableTools/Private/MuT/ASTOpImageRasterMesh.cpp b/Engine/Plugins/Mutable/Source/MutableTools/Private/MuT/ASTOpImageRasterMesh.cpp
--- a/Engine/Plugins/Mutable/Source/MutableTools/Private/MuT/ASTOpImageRasterMesh.cpp
+++ b/Engine/Plugins/Mutable/Source/MutableTools/Private/MuT/ASTOpImageRasterMesh.cpp
@@ -7,7 +7,7 @@
 #include "MuT/ASTOpMeshMorph.h"
 #include "MuT/ASTOpConditional.h"
 #include "MuT/ASTOpSwitch.h"
-#include "MuT/ASTOpMeshAddTags.h"
+#include "MuT/ASTOpMeshAddMetadata.h"
 #include "MuT/ASTOpMeshProject.h"
 #include "MuR/ModelPrivate.h"
 #include "MuR/RefCounted.h"
@@ -271,10 +271,10 @@ namespace mu
 			break;
 		}
 
-		case EOpType::ME_ADDTAGS:
+		case EOpType::ME_ADDMETADATA:
 		{
-			// Ignore tags
-			const ASTOpMeshAddTags* typedSource = static_cast<const ASTOpMeshAddTags*>(sourceAt.get());
+			// Ignore metadata
+			const ASTOpMeshAddMetadata* typedSource = static_cast<const ASTOpMeshAddMetadata*>(sourceAt.get());
 			Ptr<ASTOpImageRasterMesh> rasterOp = mu::Clone<ASTOpImageRasterMesh>(this);
 			rasterOp->mesh = typedSource->Source.child();
 			at = rasterOp;
diff --git a/Engine/Plugins/Mutable/Source/MutableTools/Private/MuT/ASTOpLayoutFromMesh.cpp b/Engine/Plugins/Mutable/Source/MutableTools/Private/MuT/ASTOpLayoutFromMesh.cpp
--- a/Engine/Plugins/Mutable/Source/MutableTools/Private/MuT/ASTOpLayoutFromMesh.cpp
+++ b/Engine/Plugins/Mutable/Source/MutableTools/Private/MuT/ASTOpLayoutFromMesh.cpp
@@ -19,7 +19,7 @@
 #include "MuT/ASTOpMeshFormat.h"
 #include "MuT/ASTOpMeshApplyShape.h"
 #include "MuT/ASTOpMeshBindShape.h"
-#include "MuT/ASTOpMeshAddTags.h"
+#include "MuT/ASTOpMeshAddMetadata.h"
 #include "MuT/ASTOpMeshMerge.h"
 #include "MuT/ASTOpSwitch.h"
 #include "MuT/ASTOpLayoutMerge.h"
@@ -187,10 +187,10 @@ namespace mu
 					break;
 				}
 
-				case EOpType::ME_ADDTAGS:
+				case EOpType::ME_ADDMETADATA:
 				{
 					// Sink, ignoring the op
-					const ASTOpMeshAddTags* Typed = static_cast<const ASTOpMeshAddTags*>(at.get());
+					const ASTOpMeshAddMetadata* Typed = static_cast<const ASTOpMeshAddMetadata*>(at.get());
 					newAt = Visit(Typed->Source.child(), CurrentSinkingOp);
 					break;
 				}
diff --git a/Engine/Plugins/Mutable/Source/MutableTools/Private/MuT/ASTOpMeshAddMetadata.cpp b/Engine/Plugins/Mutable/Source/MutableTools/Private/MuT/ASTOpMeshAddMetadata.cpp
new file mode 100644
--- /dev/null
+++ b/Engine/Plugins/Mutable/Source/MutableTools/Private/MuT/ASTOpMeshAddMetadata.cpp
@@ -0,0 +1,138 @@
+// Copyright Epic Games, Inc. All Rights Reserved.
+
+#include "MuT/ASTOpMeshAddMetadata.h"
+
+#include "MuT/ASTOpMeshMorph.h"
+#include "HAL/PlatformMath.h"
+#include "MuR/ModelPrivate.h"
+#include "MuR/RefCounted.h"
+#include "MuR/Types.h"
+
+
+namespace mu
+{
+
+	ASTOpMeshAddMetadata::ASTOpMeshAddMetadata()
+		: Source(this)
+	{
+	}
+
+
+	ASTOpMeshAddMetadata::~ASTOpMeshAddMetadata()
+	{
+		// Explicit call needed to avoid recursive destruction
+		ASTOp::RemoveChildren();
+	}
+
+
+	bool ASTOpMeshAddMetadata::IsEqual(const ASTOp& OtherUntyped) const
+	{
+		if (OtherUntyped.GetOpType() == GetOpType())
+		{
+			const ASTOpMeshAddMetadata* Other = static_cast<const ASTOpMeshAddMetadata*>(&OtherUntyped);
+			return 
+				Source      == Other->Source      && 
+				Tags        == Other->Tags        && 
+				ResourceIds == Other->ResourceIds &&
+				SkeletonIds == Other->SkeletonIds;
+		}
+
+		return false;
+	}
+
+
+	mu::Ptr<ASTOp> ASTOpMeshAddMetadata::Clone(MapChildFuncRef MapChild) const
+	{
+		Ptr<ASTOpMeshAddMetadata> NewOp = new ASTOpMeshAddMetadata();
+		NewOp->Source = MapChild(Source.child());
+		NewOp->Tags = Tags;
+		NewOp->ResourceIds = ResourceIds;
+		NewOp->SkeletonIds = SkeletonIds;
+
+		return NewOp;
+	}
+
+
+	void ASTOpMeshAddMetadata::ForEachChild(const TFunctionRef<void(ASTChild&)> Func)
+	{
+		Func(Source);
+	}
+
+
+	uint64 ASTOpMeshAddMetadata::Hash() const
+	{
+		uint64 Result = std::hash<ASTOp*>()(Source.child().get());
+		return Result;
+	}
+
+
+	void ASTOpMeshAddMetadata::Link(FProgram& Program, FLinkerOptions*)
+	{
+		// Already linked?
+		if (!linkedAddress)
+		{
+			linkedAddress = (OP::ADDRESS)Program.OpAddress.Num();
+	
+			OP::MeshAddMetadataArgs Args;
+			FMemory::Memzero(Args);
+		
+			Args.Source = Source ? Source->linkedAddress : 0;
+			
+			using OpEnumFlags = OP::MeshAddMetadataArgs::EnumFlags;
+
+			Args.Flags = OpEnumFlags::None;
+			EnumAddFlags(Args.Flags, Tags.Num()        != 1 ? OpEnumFlags::IsTagList      : OpEnumFlags::None);
+			EnumAddFlags(Args.Flags, ResourceIds.Num() != 1 ? OpEnumFlags::IsResourceList : OpEnumFlags::None);
+			EnumAddFlags(Args.Flags, SkeletonIds.Num() != 1 ? OpEnumFlags::IsSkeletonList : OpEnumFlags::None);
+
+			bool bIsTagList      = EnumHasAnyFlags(Args.Flags, OpEnumFlags::IsTagList);
+			bool bIsResourceList = EnumHasAnyFlags(Args.Flags, OpEnumFlags::IsResourceList);
+			bool bIsSkeletonList = EnumHasAnyFlags(Args.Flags, OpEnumFlags::IsSkeletonList);
+			if (bIsTagList)
+			{
+				Args.Tags.ListAddress = Program.AddConstant(Tags);
+			}
+			else
+			{
+				check(Tags.Num() > 0);
+				Args.Tags.TagAddress = Program.AddConstant(Tags[0]);
+			}
+
+			if (bIsResourceList)
+			{
+				Args.ResourceIds.ListAddress = Program.AddConstant(ResourceIds);
+			}
+			else
+			{
+				check(ResourceIds.Num() > 0);
+				Args.ResourceIds.ResourceId = ResourceIds[0];
+			}
+
+			if (bIsSkeletonList)
+			{
+				Args.SkeletonIds.ListAddress = Program.AddConstant(SkeletonIds);
+			}
+			else
+			{
+				check(SkeletonIds.Num() > 0);
+				Args.SkeletonIds.SkeletonId = SkeletonIds[0];
+			}
+
+			Program.OpAddress.Add((uint32)Program.ByteCode.Num());
+			AppendCode(Program.ByteCode, EOpType::ME_ADDMETADATA);
+			AppendCode(Program.ByteCode, Args);
+		}
+	}
+
+
+	FSourceDataDescriptor ASTOpMeshAddMetadata::GetSourceDataDescriptor(FGetSourceDataDescriptorContext* Context) const
+	{
+		if (Source)
+		{
+			return Source->GetSourceDataDescriptor(Context);
+		}
+
+		return {};
+	}
+
+}
diff --git a/Engine/Plugins/Mutable/Source/MutableTools/Private/MuT/ASTOpMeshAddMetadata.h b/Engine/Plugins/Mutable/Source/MutableTools/Private/MuT/ASTOpMeshAddMetadata.h
new file mode 100644
--- /dev/null
+++ b/Engine/Plugins/Mutable/Source/MutableTools/Private/MuT/ASTOpMeshAddMetadata.h
@@ -0,0 +1,49 @@
+// Copyright Epic Games, Inc. All Rights Reserved.
+
+#pragma once
+
+#include "HAL/Platform.h"
+#include "MuR/Operations.h"
+#include "MuR/Ptr.h"
+#include "MuT/AST.h"
+
+
+namespace mu
+{
+	struct FProgram;
+
+	/** Add tags to a mesh. */
+	class ASTOpMeshAddMetadata final : public ASTOp
+	{
+	public:
+
+		/** Source mesh to add tags to. */
+		ASTChild Source;
+
+		/** Tags to add. */
+		TArray<FString> Tags;
+	
+		/** ResourceIds to add. */
+		TArray<uint64> ResourceIds;
+
+		/** SkeletonIds to add. */
+		TArray<uint32> SkeletonIds;
+
+	public:
+
+		ASTOpMeshAddMetadata();
+		ASTOpMeshAddMetadata(const ASTOpMeshAddMetadata&) = delete;
+		virtual ~ASTOpMeshAddMetadata() override;
+
+		// ASTOp interface
+		virtual EOpType GetOpType() const override { return EOpType::ME_ADDMETADATA; }
+		virtual uint64 Hash() const override;
+		virtual void ForEachChild(const TFunctionRef<void(ASTChild&)>) override;
+		virtual bool IsEqual(const ASTOp& OtherUntyped) const override;
+		virtual Ptr<ASTOp> Clone(MapChildFuncRef MapChild) const override;
+		virtual void Link(FProgram& Program, FLinkerOptions*) override;
+		virtual FSourceDataDescriptor GetSourceDataDescriptor(FGetSourceDataDescriptorContext*) const override;
+	};
+
+}
+
diff --git a/Engine/Plugins/Mutable/Source/MutableTools/Private/MuT/ASTOpMeshAddTags.cpp b/Engine/Plugins/Mutable/Source/MutableTools/Private/MuT/ASTOpMeshAddTags.cpp
deleted file mode 100644
diff --git a/Engine/Plugins/Mutable/Source/MutableTools/Private/MuT/ASTOpMeshAddTags.h b/Engine/Plugins/Mutable/Source/MutableTools/Private/MuT/ASTOpMeshAddTags.h
deleted file mode 100644
diff --git a/Engine/Plugins/Mutable/Source/MutableTools/Private/MuT/ASTOpMeshApplyLayout.cpp b/Engine/Plugins/Mutable/Source/MutableTools/Private/MuT/ASTOpMeshApplyLayout.cpp
--- a/Engine/Plugins/Mutable/Source/MutableTools/Private/MuT/ASTOpMeshApplyLayout.cpp
+++ b/Engine/Plugins/Mutable/Source/MutableTools/Private/MuT/ASTOpMeshApplyLayout.cpp
@@ -1,8 +1,8 @@
 // Copyright Epic Games, Inc. All Rights Reserved.
 
 #include "MuT/ASTOpMeshApplyLayout.h"
-#include "MuT/ASTOpMeshAddTags.h"
 #include "MuT/ASTOpSwitch.h"
+#include "MuT/ASTOpMeshAddMetadata.h"
 
 #include "HAL/PlatformMath.h"
 #include "MuR/ModelPrivate.h"
@@ -103,9 +103,9 @@ namespace mu
 
 		switch (MeshOpType)
 		{
-			case EOpType::ME_ADDTAGS:
+			case EOpType::ME_ADDMETADATA:
 			{
-				Ptr<ASTOpMeshAddTags> New = mu::Clone<ASTOpMeshAddTags>(MeshAt);
+				Ptr<ASTOpMeshAddMetadata> New = mu::Clone<ASTOpMeshAddMetadata>(MeshAt);
 				if (New->Source)
 				{
 					Ptr<ASTOpMeshApplyLayout> NewApplyPose = mu::Clone<ASTOpMeshApplyLayout>(this);
diff --git a/Engine/Plugins/Mutable/Source/MutableTools/Private/MuT/ASTOpMeshApplyPose.cpp b/Engine/Plugins/Mutable/Source/MutableTools/Private/MuT/ASTOpMeshApplyPose.cpp
--- a/Engine/Plugins/Mutable/Source/MutableTools/Private/MuT/ASTOpMeshApplyPose.cpp
+++ b/Engine/Plugins/Mutable/Source/MutableTools/Private/MuT/ASTOpMeshApplyPose.cpp
@@ -1,7 +1,7 @@
 // Copyright Epic Games, Inc. All Rights Reserved.
 
 #include "MuT/ASTOpMeshApplyPose.h"
-#include "MuT/ASTOpMeshAddTags.h"
+#include "MuT/ASTOpMeshAddMetadata.h"
 
 #include "HAL/PlatformMath.h"
 #include "MuR/ModelPrivate.h"
@@ -96,9 +96,9 @@ namespace mu
 
 		switch (MeshType)
 		{
-			case EOpType::ME_ADDTAGS:
+			case EOpType::ME_ADDMETADATA:
 			{
-				Ptr<ASTOpMeshAddTags> New = mu::Clone<ASTOpMeshAddTags>(MeshAt);
+				Ptr<ASTOpMeshAddMetadata> New = mu::Clone<ASTOpMeshAddMetadata>(MeshAt);
 				if (New->Source)
 				{
 					Ptr<ASTOpMeshApplyPose> NewApplyPose = mu::Clone<ASTOpMeshApplyPose>(this);
diff --git a/Engine/Plugins/Mutable/Source/MutableTools/Private/MuT/ASTOpMeshBindShape.cpp b/Engine/Plugins/Mutable/Source/MutableTools/Private/MuT/ASTOpMeshBindShape.cpp
--- a/Engine/Plugins/Mutable/Source/MutableTools/Private/MuT/ASTOpMeshBindShape.cpp
+++ b/Engine/Plugins/Mutable/Source/MutableTools/Private/MuT/ASTOpMeshBindShape.cpp
@@ -10,8 +10,8 @@
 #include "MuT/ASTOpConditional.h"
 #include "MuT/ASTOpSwitch.h"
 #include "MuT/ASTOpMeshRemoveMask.h"
-#include "MuT/ASTOpMeshAddTags.h"
 #include "MuT/ASTOpMeshPrepareLayout.h"
+#include "MuT/ASTOpMeshAddMetadata.h"
 
 
 namespace mu
@@ -375,9 +375,9 @@ namespace mu
 				break;
 			}
 
-			case EOpType::ME_ADDTAGS:
+			case EOpType::ME_ADDMETADATA:
 			{
-				Ptr<ASTOpMeshAddTags> New = mu::Clone<ASTOpMeshAddTags>(MeshAt);
+				Ptr<ASTOpMeshAddMetadata> New = mu::Clone<ASTOpMeshAddMetadata>(MeshAt);
 				if (New->Source)
 				{
 					Ptr<ASTOpMeshBindShape> NewBind = mu::Clone<ASTOpMeshBindShape>(this);
@@ -465,11 +465,11 @@ namespace mu
 				break;
 			}
 
-			case EOpType::ME_ADDTAGS:
+			case EOpType::ME_ADDMETADATA:
 			{
 				// Ignore the tags in the shape
 				Ptr<ASTOpMeshBindShape> NewBind = mu::Clone<ASTOpMeshBindShape>(this);
-				const ASTOpMeshAddTags* New = static_cast<const ASTOpMeshAddTags*>(ShapeAt.get());
+				const ASTOpMeshAddMetadata* New = static_cast<const ASTOpMeshAddMetadata*>(ShapeAt.get());
 				NewBind->Shape = New->Source.child();
 				NewOp = NewBind;
 				break;
diff --git a/Engine/Plugins/Mutable/Source/MutableTools/Private/MuT/ASTOpMeshClipMorphPlane.cpp b/Engine/Plugins/Mutable/Source/MutableTools/Private/MuT/ASTOpMeshClipMorphPlane.cpp
--- a/Engine/Plugins/Mutable/Source/MutableTools/Private/MuT/ASTOpMeshClipMorphPlane.cpp
+++ b/Engine/Plugins/Mutable/Source/MutableTools/Private/MuT/ASTOpMeshClipMorphPlane.cpp
@@ -2,7 +2,7 @@
 
 #include "MuT/ASTOpMeshClipMorphPlane.h"
 
-#include "MuT/ASTOpMeshAddTags.h"
+#include "MuT/ASTOpMeshAddMetadata.h"
 #include "MuR/ModelPrivate.h"
 #include "MuR/RefCounted.h"
 #include "MuR/Types.h"
@@ -133,18 +133,18 @@ namespace mu
 		switch (SourceType)
 		{
 
-		case EOpType::ME_ADDTAGS:
+		case EOpType::ME_ADDMETADATA:
 		{
-			Ptr<ASTOpMeshAddTags> NewAddTags = mu::Clone<ASTOpMeshAddTags>(Source.child());
+			Ptr<ASTOpMeshAddMetadata> NewAddMetadata = mu::Clone<ASTOpMeshAddMetadata>(Source.child());
 
-			if (NewAddTags->Source)
+			if (NewAddMetadata->Source)
 			{
 				Ptr<ASTOpMeshClipMorphPlane> New = mu::Clone<ASTOpMeshClipMorphPlane>(this);
-				New->Source = NewAddTags->Source.child();
-				NewAddTags->Source = New;
+				New->Source = NewAddMetadata->Source.child();
+				NewAddMetadata->Source = New;
 			}
 
-			NewOp = NewAddTags;
+			NewOp = NewAddMetadata;
 			break;
 		}
 
diff --git a/Engine/Plugins/Mutable/Source/MutableTools/Private/MuT/ASTOpMeshClipWithMesh.cpp b/Engine/Plugins/Mutable/Source/MutableTools/Private/MuT/ASTOpMeshClipWithMesh.cpp
--- a/Engine/Plugins/Mutable/Source/MutableTools/Private/MuT/ASTOpMeshClipWithMesh.cpp
+++ b/Engine/Plugins/Mutable/Source/MutableTools/Private/MuT/ASTOpMeshClipWithMesh.cpp
@@ -1,7 +1,7 @@
 // Copyright Epic Games, Inc. All Rights Reserved.
 
 #include "MuT/ASTOpMeshClipWithMesh.h"
-#include "MuT/ASTOpMeshAddTags.h"
+#include "MuT/ASTOpMeshAddMetadata.h"
 
 #include "HAL/PlatformMath.h"
 #include "MuR/ModelPrivate.h"
@@ -96,9 +96,9 @@ namespace mu
 
 		switch (MeshType)
 		{
-			case EOpType::ME_ADDTAGS:
+			case EOpType::ME_ADDMETADATA:
 			{
-				Ptr<ASTOpMeshAddTags> New = mu::Clone<ASTOpMeshAddTags>(MeshAt);
+				Ptr<ASTOpMeshAddMetadata> New = mu::Clone<ASTOpMeshAddMetadata>(MeshAt);
 				if (New->Source)
 				{
 					Ptr<ASTOpMeshClipWithMesh> NewApplyPose = mu::Clone<ASTOpMeshClipWithMesh>(this);
diff --git a/Engine/Plugins/Mutable/Source/MutableTools/Private/MuT/ASTOpMeshDifference.cpp b/Engine/Plugins/Mutable/Source/MutableTools/Private/MuT/ASTOpMeshDifference.cpp
--- a/Engine/Plugins/Mutable/Source/MutableTools/Private/MuT/ASTOpMeshDifference.cpp
+++ b/Engine/Plugins/Mutable/Source/MutableTools/Private/MuT/ASTOpMeshDifference.cpp
@@ -9,7 +9,7 @@
 #include "MuT/ASTOpConditional.h"
 #include "MuT/ASTOpSwitch.h"
 #include "MuT/ASTOpMeshRemoveMask.h"
-#include "MuT/ASTOpMeshAddTags.h"
+#include "MuT/ASTOpMeshAddMetadata.h"
 
 
 namespace mu
@@ -263,9 +263,9 @@ namespace mu
 				break;
 			}
 
-			case EOpType::ME_ADDTAGS:
+			case EOpType::ME_ADDMETADATA:
 			{
-				Ptr<ASTOpMeshAddTags> NewAdd = mu::Clone<ASTOpMeshAddTags>(BaseAt);
+				Ptr<ASTOpMeshAddMetadata> NewAdd = mu::Clone<ASTOpMeshAddMetadata>(BaseAt);
 
 				if (NewAdd->Source)
 				{
@@ -340,10 +340,10 @@ namespace mu
 				break;
 			}
 
-			case EOpType::ME_ADDTAGS:
+			case EOpType::ME_ADDMETADATA:
 			{
 				// Ignore tags in this branch
-				const ASTOpMeshAddTags* Add = static_cast<const ASTOpMeshAddTags*>(TargetAt.get());
+				const ASTOpMeshAddMetadata* Add = static_cast<const ASTOpMeshAddMetadata*>(TargetAt.get());
 
 				Ptr<ASTOpMeshDifference> NewDiff = mu::Clone<ASTOpMeshDifference>(this);
 				NewDiff->Target = Add->Source.child();
diff --git a/Engine/Plugins/Mutable/Source/MutableTools/Private/MuT/ASTOpMeshExtractLayoutBlocks.cpp b/Engine/Plugins/Mutable/Source/MutableTools/Private/MuT/ASTOpMeshExtractLayoutBlocks.cpp
--- a/Engine/Plugins/Mutable/Source/MutableTools/Private/MuT/ASTOpMeshExtractLayoutBlocks.cpp
+++ b/Engine/Plugins/Mutable/Source/MutableTools/Private/MuT/ASTOpMeshExtractLayoutBlocks.cpp
@@ -4,7 +4,7 @@
 
 #include "MuT/ASTOpSwitch.h"
 #include "MuT/ASTOpConditional.h"
-#include "MuT/ASTOpMeshAddTags.h"
+#include "MuT/ASTOpMeshAddMetadata.h"
 #include "MuT/ASTOpMeshMorph.h"
 #include "MuT/ASTOpMeshRemoveMask.h"
 #include "MuT/ASTOpMeshClipMorphPlane.h"
@@ -213,9 +213,9 @@ namespace mu
 			break;
 		}
 
-		case EOpType::ME_ADDTAGS:
+		case EOpType::ME_ADDMETADATA:
 		{
-			Ptr<ASTOpMeshAddTags> NewOp = mu::Clone<ASTOpMeshAddTags>(at);
+			Ptr<ASTOpMeshAddMetadata> NewOp = mu::Clone<ASTOpMeshAddMetadata>(at);
 			NewOp->Source = Visit(NewOp->Source.child(), currentSinkOp);
 			newAt = NewOp;
 			break;
diff --git a/Engine/Plugins/Mutable/Source/MutableTools/Private/MuT/ASTOpMeshFormat.cpp b/Engine/Plugins/Mutable/Source/MutableTools/Private/MuT/ASTOpMeshFormat.cpp
--- a/Engine/Plugins/Mutable/Source/MutableTools/Private/MuT/ASTOpMeshFormat.cpp
+++ b/Engine/Plugins/Mutable/Source/MutableTools/Private/MuT/ASTOpMeshFormat.cpp
@@ -17,7 +17,7 @@
 #include "MuT/ASTOpMeshTransformWithBoundingMesh.h"
 #include "MuT/ASTOpMeshRemoveMask.h"
 #include "MuT/ASTOpMeshMorph.h"
-#include "MuT/ASTOpMeshAddTags.h"
+#include "MuT/ASTOpMeshAddMetadata.h"
 #include "MuT/ASTOpMeshApplyPose.h"
 #include "MuT/ASTOpMeshApplyLayout.h"
 #include "MuT/ASTOpMeshMerge.h"
@@ -309,9 +309,9 @@ mu::Ptr<ASTOp> Sink_MeshFormatAST::Visit(const mu::Ptr<ASTOp>& at, const ASTOpMe
 		break;
 	}
 
-	case EOpType::ME_ADDTAGS:
+	case EOpType::ME_ADDMETADATA:
 	{
-		Ptr<ASTOpMeshAddTags> newOp = mu::Clone<ASTOpMeshAddTags>(at);
+		Ptr<ASTOpMeshAddMetadata> newOp = mu::Clone<ASTOpMeshAddMetadata>(at);
 		newOp->Source = Visit(newOp->Source.child(), currentFormatOp);
 		newAt = newOp;
 		break;
diff --git a/Engine/Plugins/Mutable/Source/MutableTools/Private/MuT/ASTOpMeshMaskClipMesh.cpp b/Engine/Plugins/Mutable/Source/MutableTools/Private/MuT/ASTOpMeshMaskClipMesh.cpp
--- a/Engine/Plugins/Mutable/Source/MutableTools/Private/MuT/ASTOpMeshMaskClipMesh.cpp
+++ b/Engine/Plugins/Mutable/Source/MutableTools/Private/MuT/ASTOpMeshMaskClipMesh.cpp
@@ -8,8 +8,8 @@
 #include "MuR/Types.h"
 #include "MuT/ASTOpConditional.h"
 #include "MuT/ASTOpMeshRemoveMask.h"
-#include "MuT/ASTOpMeshAddTags.h"
 #include "MuT/ASTOpMeshPrepareLayout.h"
+#include "MuT/ASTOpMeshAddMetadata.h"
 #include "MuT/ASTOpSwitch.h"
 
 
@@ -179,9 +179,9 @@ namespace mu
 					break;
 				}
 
-				case EOpType::ME_ADDTAGS:
+				case EOpType::ME_ADDMETADATA:
 				{
-					Ptr<ASTOpMeshAddTags> newOp = mu::Clone<ASTOpMeshAddTags>(at);
+					Ptr<ASTOpMeshAddMetadata> newOp = mu::Clone<ASTOpMeshAddMetadata>(at);
 					newOp->Source = Visit(newOp->Source.child());
 					newAt = newOp;
 					break;
@@ -322,10 +322,10 @@ namespace mu
 					break;
 				}
 
-				case EOpType::ME_ADDTAGS:
+				case EOpType::ME_ADDMETADATA:
 				{
 					// Ignore tags in the clip mesh branch
-					const ASTOpMeshAddTags* AddOp = static_cast<const ASTOpMeshAddTags*>(at.get());
+					const ASTOpMeshAddMetadata* AddOp = static_cast<const ASTOpMeshAddMetadata*>(at.get());
 					newAt = Visit(AddOp->Source.child());
 					break;
 				}
@@ -386,4 +386,4 @@ namespace mu
 	}
 
 }
- 
\ No newline at end of file
+ 
diff --git a/Engine/Plugins/Mutable/Source/MutableTools/Private/MuT/ASTOpMeshMaskClipUVMask.cpp b/Engine/Plugins/Mutable/Source/MutableTools/Private/MuT/ASTOpMeshMaskClipUVMask.cpp
--- a/Engine/Plugins/Mutable/Source/MutableTools/Private/MuT/ASTOpMeshMaskClipUVMask.cpp
+++ b/Engine/Plugins/Mutable/Source/MutableTools/Private/MuT/ASTOpMeshMaskClipUVMask.cpp
@@ -8,8 +8,8 @@
 #include "MuR/Types.h"
 #include "MuT/ASTOpConditional.h"
 #include "MuT/ASTOpMeshRemoveMask.h"
-#include "MuT/ASTOpMeshAddTags.h"
 #include "MuT/ASTOpMeshPrepareLayout.h"
+#include "MuT/ASTOpMeshAddMetadata.h"
 #include "MuT/ASTOpSwitch.h"
 
 
@@ -186,9 +186,9 @@ namespace mu
 					break;
 				}
 
-				case EOpType::ME_ADDTAGS:
+				case EOpType::ME_ADDMETADATA:
 				{
-					Ptr<ASTOpMeshAddTags> newOp = mu::Clone<ASTOpMeshAddTags>(at);
+					Ptr<ASTOpMeshAddMetadata> newOp = mu::Clone<ASTOpMeshAddMetadata>(at);
 					newOp->Source = Visit(newOp->Source.child());
 					newAt = newOp;
 					break;
@@ -394,4 +394,4 @@ namespace mu
 	}
 
 }
- 
\ No newline at end of file
+
diff --git a/Engine/Plugins/Mutable/Source/MutableTools/Private/MuT/ASTOpMeshMaskDiff.cpp b/Engine/Plugins/Mutable/Source/MutableTools/Private/MuT/ASTOpMeshMaskDiff.cpp
--- a/Engine/Plugins/Mutable/Source/MutableTools/Private/MuT/ASTOpMeshMaskDiff.cpp
+++ b/Engine/Plugins/Mutable/Source/MutableTools/Private/MuT/ASTOpMeshMaskDiff.cpp
@@ -1,7 +1,7 @@
 // Copyright Epic Games, Inc. All Rights Reserved.
 
 #include "MuT/ASTOpMeshMaskDiff.h"
-#include "MuT/ASTOpMeshAddTags.h"
+#include "MuT/ASTOpMeshAddMetadata.h"
 #include "MuT/ASTOpSwitch.h"
 #include "MuT/ASTOpMeshPrepareLayout.h"
 #include "MuR/ModelPrivate.h"
@@ -94,10 +94,10 @@ namespace mu
 		switch (FragmentType)
 		{
 
-		case EOpType::ME_ADDTAGS:
+		case EOpType::ME_ADDMETADATA:
 		{
 			// Tags in the fragment can be ignored.
-			const ASTOpMeshAddTags* Add = static_cast<const ASTOpMeshAddTags*>(Fragment.child().get());
+			const ASTOpMeshAddMetadata* Add = static_cast<const ASTOpMeshAddMetadata*>(Fragment.child().get());
 
 			Ptr<ASTOpMeshMaskDiff> NewAt = mu::Clone<ASTOpMeshMaskDiff>(this);
 			NewAt->Fragment = Add->Source.child();
@@ -154,10 +154,10 @@ namespace mu
 			switch (SourceType)
 			{
 
-			case EOpType::ME_ADDTAGS:
+			case EOpType::ME_ADDMETADATA:
 			{
 				// Tags in the source can be ignored for mask generation.
-				const ASTOpMeshAddTags* Add = static_cast<const ASTOpMeshAddTags*>(Source.child().get());
+				const ASTOpMeshAddMetadata* Add = static_cast<const ASTOpMeshAddMetadata*>(Source.child().get());
 
 				Ptr<ASTOpMeshMaskDiff> NewAt = mu::Clone<ASTOpMeshMaskDiff>(this);
 				NewAt->Source = Add->Source.child();
diff --git a/Engine/Plugins/Mutable/Source/MutableTools/Private/MuT/ASTOpMeshMerge.cpp b/Engine/Plugins/Mutable/Source/MutableTools/Private/MuT/ASTOpMeshMerge.cpp
--- a/Engine/Plugins/Mutable/Source/MutableTools/Private/MuT/ASTOpMeshMerge.cpp
+++ b/Engine/Plugins/Mutable/Source/MutableTools/Private/MuT/ASTOpMeshMerge.cpp
@@ -1,7 +1,6 @@
 // Copyright Epic Games, Inc. All Rights Reserved.
 
 #include "MuT/ASTOpMeshMerge.h"
-#include "MuT/ASTOpMeshAddTags.h"
 
 #include "HAL/PlatformMath.h"
 #include "MuR/ModelPrivate.h"
diff --git a/Engine/Plugins/Mutable/Source/MutableTools/Private/MuT/ASTOpMeshMorph.cpp b/Engine/Plugins/Mutable/Source/MutableTools/Private/MuT/ASTOpMeshMorph.cpp
--- a/Engine/Plugins/Mutable/Source/MutableTools/Private/MuT/ASTOpMeshMorph.cpp
+++ b/Engine/Plugins/Mutable/Source/MutableTools/Private/MuT/ASTOpMeshMorph.cpp
@@ -2,7 +2,7 @@
 
 #include "MuT/ASTOpMeshMorph.h"
 
-#include "MuT/ASTOpMeshAddTags.h"
+#include "MuT/ASTOpMeshAddMetadata.h"
 #include "MuR/ModelPrivate.h"
 #include "MuR/RefCounted.h"
 #include "MuR/Types.h"
@@ -107,19 +107,19 @@ namespace mu
 		switch (BaseType)
 		{
 
-		case EOpType::ME_ADDTAGS:
+		case EOpType::ME_ADDMETADATA:
 		{
 			// Add the base tags after morphing
-			Ptr<ASTOpMeshAddTags> NewAddTags = mu::Clone<ASTOpMeshAddTags>(Base.child());
+			Ptr<ASTOpMeshAddMetadata> NewAddMetadata = mu::Clone<ASTOpMeshAddMetadata>(Base.child());
 
-			if (NewAddTags->Source)
+			if (NewAddMetadata->Source)
 			{
 				Ptr<ASTOpMeshMorph> New = mu::Clone<ASTOpMeshMorph>(this);
-				New->Base = NewAddTags->Source.child();
-				NewAddTags->Source = New;
+				New->Base = NewAddMetadata->Source.child();
+				NewAddMetadata->Source = New;
 			}
 
-			NewOp = NewAddTags;
+			NewOp = NewAddMetadata;
 			break;
 		}
 
@@ -136,13 +136,13 @@ namespace mu
 			switch (MorphType)
 			{
 
-			case EOpType::ME_ADDTAGS:
+			case EOpType::ME_ADDMETADATA:
 			{
 				// Ignore the morph target tags
-				const ASTOpMeshAddTags* AddTags = static_cast<const ASTOpMeshAddTags*>(Target.child().get());
+				const ASTOpMeshAddMetadata* AddMetadata = static_cast<const ASTOpMeshAddMetadata*>(Target.child().get());
 
 				Ptr<ASTOpMeshMorph> New = mu::Clone<ASTOpMeshMorph>(this);
-				New->Target = AddTags->Source.child();
+				New->Target = AddMetadata->Source.child();
 				NewOp = New;
 
 				break;
diff --git a/Engine/Plugins/Mutable/Source/MutableTools/Private/MuT/ASTOpMeshOptimizeSkinning.cpp b/Engine/Plugins/Mutable/Source/MutableTools/Private/MuT/ASTOpMeshOptimizeSkinning.cpp
--- a/Engine/Plugins/Mutable/Source/MutableTools/Private/MuT/ASTOpMeshOptimizeSkinning.cpp
+++ b/Engine/Plugins/Mutable/Source/MutableTools/Private/MuT/ASTOpMeshOptimizeSkinning.cpp
@@ -4,7 +4,7 @@
 
 #include "MuT/ASTOpConditional.h"
 #include "MuT/ASTOpMeshMorph.h"
-#include "MuT/ASTOpMeshAddTags.h"
+#include "MuT/ASTOpMeshAddMetadata.h"
 #include "MuT/ASTOpMeshApplyLayout.h"
 #include "MuT/ASTOpSwitch.h"
 #include "MuR/ModelPrivate.h"
@@ -152,9 +152,9 @@ namespace mu
 			break;
 		}
 
-		case EOpType::ME_ADDTAGS:
+		case EOpType::ME_ADDMETADATA:
 		{
-			Ptr<ASTOpMeshAddTags> NewOp = mu::Clone<ASTOpMeshAddTags>(at);
+			Ptr<ASTOpMeshAddMetadata> NewOp = mu::Clone<ASTOpMeshAddMetadata>(at);
 			NewOp->Source = Visit(NewOp->Source.child(), CurrentSinkOp);
 			NewAt = NewOp;
 			break;
diff --git a/Engine/Plugins/Mutable/Source/MutableTools/Private/MuT/ASTOpMeshPrepareLayout.cpp b/Engine/Plugins/Mutable/Source/MutableTools/Private/MuT/ASTOpMeshPrepareLayout.cpp
--- a/Engine/Plugins/Mutable/Source/MutableTools/Private/MuT/ASTOpMeshPrepareLayout.cpp
+++ b/Engine/Plugins/Mutable/Source/MutableTools/Private/MuT/ASTOpMeshPrepareLayout.cpp
@@ -1,7 +1,7 @@
 // Copyright Epic Games, Inc. All Rights Reserved.
 
 #include "MuT/ASTOpMeshPrepareLayout.h"
-#include "MuT/ASTOpMeshAddTags.h"
+#include "MuT/ASTOpMeshAddMetadata.h"
 
 #include "HAL/PlatformMath.h"
 #include "MuR/ModelPrivate.h"
diff --git a/Engine/Plugins/Mutable/Source/MutableTools/Private/MuT/ASTOpMeshProject.cpp b/Engine/Plugins/Mutable/Source/MutableTools/Private/MuT/ASTOpMeshProject.cpp
--- a/Engine/Plugins/Mutable/Source/MutableTools/Private/MuT/ASTOpMeshProject.cpp
+++ b/Engine/Plugins/Mutable/Source/MutableTools/Private/MuT/ASTOpMeshProject.cpp
@@ -1,7 +1,7 @@
 // Copyright Epic Games, Inc. All Rights Reserved.
 
 #include "MuT/ASTOpMeshProject.h"
-#include "MuT/ASTOpMeshAddTags.h"
+#include "MuT/ASTOpMeshAddMetadata.h"
 #include "MuT/ASTOpConditional.h"
 #include "MuT/ASTOpSwitch.h"
 
@@ -149,19 +149,19 @@ namespace mu
 			break;
 		}
 
-		case EOpType::ME_ADDTAGS:
+		case EOpType::ME_ADDMETADATA:
 		{
 			// Add the tags after layout
-			Ptr<ASTOpMeshAddTags> NewAddTags = mu::Clone<ASTOpMeshAddTags>(SourceAt);
+			Ptr<ASTOpMeshAddMetadata> NewAddMetadata = mu::Clone<ASTOpMeshAddMetadata>(SourceAt);
 
-			if (NewAddTags->Source)
+			if (NewAddMetadata->Source)
 			{
 				Ptr<ASTOpMeshProject> NewAt = mu::Clone<ASTOpMeshProject>(this);
-				NewAt->Mesh = NewAddTags->Source.child();
-				NewAddTags->Source = NewAt;
+				NewAt->Mesh = NewAddMetadata->Source.child();
+				NewAddMetadata->Source = NewAt;
 			}
 
-			NewOp = NewAddTags;
+			NewOp = NewAddMetadata;
 			break;
 		}
 
diff --git a/Engine/Plugins/Mutable/Source/MutableTools/Private/MuT/ASTOpMeshRemoveMask.cpp b/Engine/Plugins/Mutable/Source/MutableTools/Private/MuT/ASTOpMeshRemoveMask.cpp
--- a/Engine/Plugins/Mutable/Source/MutableTools/Private/MuT/ASTOpMeshRemoveMask.cpp
+++ b/Engine/Plugins/Mutable/Source/MutableTools/Private/MuT/ASTOpMeshRemoveMask.cpp
@@ -3,7 +3,7 @@
 #include "MuT/ASTOpMeshRemoveMask.h"
 
 #include "MuT/ASTOpMeshMorph.h"
-#include "MuT/ASTOpMeshAddTags.h"
+#include "MuT/ASTOpMeshAddMetadata.h"
 #include "HAL/PlatformMath.h"
 #include "MuR/ModelPrivate.h"
 #include "MuR/RefCounted.h"
@@ -180,9 +180,9 @@ namespace mu
 					break;
 				}
 
-				case EOpType::ME_ADDTAGS:
+				case EOpType::ME_ADDMETADATA:
 				{
-					mu::Ptr<ASTOpMeshAddTags> newOp = mu::Clone<ASTOpMeshAddTags>(at);
+					mu::Ptr<ASTOpMeshAddMetadata> newOp = mu::Clone<ASTOpMeshAddMetadata>(at);
 					newOp->Source = Visit(newOp->Source.child());
 					newAt = newOp;
 					break;
@@ -253,7 +253,7 @@ namespace mu
 				EOpType RemoveType = removes[RemoveIndex].Value->GetOpType();
 				switch (RemoveType)
 				{
-				case EOpType::ME_ADDTAGS:
+				case EOpType::ME_ADDMETADATA:
 				{
 					// It can be ignored.
 					if (!NewOp)
@@ -261,7 +261,7 @@ namespace mu
 						NewOp = mu::Clone<ASTOpMeshRemoveMask>(this);
 					}
 
-					const ASTOpMeshAddTags* Add = static_cast<const ASTOpMeshAddTags*>(removes[RemoveIndex].Value.child().get());
+					const ASTOpMeshAddMetadata* Add = static_cast<const ASTOpMeshAddMetadata*>(removes[RemoveIndex].Value.child().get());
 					NewOp->removes[RemoveIndex].Value = Add->Source.child();
 
 					break;
diff --git a/Engine/Plugins/Mutable/Source/MutableTools/Private/MuT/ASTOpMeshSetSkeleton.cpp b/Engine/Plugins/Mutable/Source/MutableTools/Private/MuT/ASTOpMeshSetSkeleton.cpp
--- a/Engine/Plugins/Mutable/Source/MutableTools/Private/MuT/ASTOpMeshSetSkeleton.cpp
+++ b/Engine/Plugins/Mutable/Source/MutableTools/Private/MuT/ASTOpMeshSetSkeleton.cpp
@@ -1,7 +1,6 @@
 // Copyright Epic Games, Inc. All Rights Reserved.
 
 #include "MuT/ASTOpMeshSetSkeleton.h"
-#include "MuT/ASTOpMeshAddTags.h"
 
 #include "HAL/PlatformMath.h"
 #include "MuR/ModelPrivate.h"
diff --git a/Engine/Plugins/Mutable/Source/MutableTools/Private/MuT/ASTOpScalarArithmetic.cpp b/Engine/Plugins/Mutable/Source/MutableTools/Private/MuT/ASTOpScalarArithmetic.cpp
--- a/Engine/Plugins/Mutable/Source/MutableTools/Private/MuT/ASTOpScalarArithmetic.cpp
+++ b/Engine/Plugins/Mutable/Source/MutableTools/Private/MuT/ASTOpScalarArithmetic.cpp
@@ -1,7 +1,6 @@
 // Copyright Epic Games, Inc. All Rights Reserved.
 
 #include "MuT/ASTOpScalarArithmetic.h"
-#include "MuT/ASTOpMeshAddTags.h"
 
 #include "HAL/PlatformMath.h"
 #include "MuR/ModelPrivate.h"
diff --git a/Engine/Plugins/Mutable/Source/MutableTools/Private/MuT/ASTOpSwitch.cpp b/Engine/Plugins/Mutable/Source/MutableTools/Private/MuT/ASTOpSwitch.cpp
--- a/Engine/Plugins/Mutable/Source/MutableTools/Private/MuT/ASTOpSwitch.cpp
+++ b/Engine/Plugins/Mutable/Source/MutableTools/Private/MuT/ASTOpSwitch.cpp
@@ -4,7 +4,7 @@
 
 #include "MuT/ASTOpParameter.h"
 #include "MuT/ASTOpConstantInt.h"
-#include "MuT/ASTOpMeshAddTags.h"
+#include "MuT/ASTOpMeshAddMetadata.h"
 #include "Containers/Map.h"
 #include "Misc/AssertionMacros.h"
 #include "MuR/ImagePrivate.h"
@@ -615,17 +615,21 @@ namespace mu
 		{
 			switch (BranchOpType)
 			{
-			case EOpType::ME_ADDTAGS:
+			case EOpType::ME_ADDMETADATA:
 			{
 				// Move the add tags out of the switch if all tags are the same
-				bool bAllTagsAreTheSame = true;
+				bool bAllMetadataIsTheSame = true;
 				TArray<FString> Tags;
+				TArray<uint64> ResourceIds;
+				TArray<uint32> SkeletonIds;
 
 				if (Default)
 				{
-					check(Default->GetOpType()==EOpType::ME_ADDTAGS);
-					const ASTOpMeshAddTags* Typed = static_cast<const ASTOpMeshAddTags*>(Default.child().get());
+					check(Default->GetOpType() == EOpType::ME_ADDMETADATA);
+					const ASTOpMeshAddMetadata* Typed = static_cast<const ASTOpMeshAddMetadata*>(Default.child().get());
 					Tags = Typed->Tags;
+					ResourceIds = Typed->ResourceIds;
+					SkeletonIds = Typed->SkeletonIds;
 				}
 
 				for (const FCase& Case : Cases)
@@ -635,36 +639,43 @@ namespace mu
 						continue;
 					}
 
-					check(Case.Branch->GetOpType() == EOpType::ME_ADDTAGS);
-					const ASTOpMeshAddTags* Typed = static_cast<const ASTOpMeshAddTags*>(Case.Branch.child().get());
-					if (Tags.IsEmpty())
+					check(Case.Branch->GetOpType() == EOpType::ME_ADDMETADATA);
+					const ASTOpMeshAddMetadata* Typed = static_cast<const ASTOpMeshAddMetadata*>(Case.Branch.child().get());
+					if (Tags.IsEmpty() && ResourceIds.IsEmpty() && SkeletonIds.IsEmpty())
 					{
 						Tags = Typed->Tags;
+						ResourceIds = Typed->ResourceIds;
+						SkeletonIds = Typed->SkeletonIds;
 					}
-					else if (Typed->Tags!=Tags)
+					else if (
+							Typed->Tags != Tags || 
+							Typed->ResourceIds != ResourceIds || 
+							Typed->SkeletonIds != SkeletonIds)
 					{
-						bAllTagsAreTheSame = false;
+						bAllMetadataIsTheSame = false;
 						break;
 					}
 				}
 
-				if (bAllTagsAreTheSame)
+				if (bAllMetadataIsTheSame)
 				{
-					Ptr<ASTOpMeshAddTags> New = new ASTOpMeshAddTags();
+					Ptr<ASTOpMeshAddMetadata> New = new ASTOpMeshAddMetadata();
 					New->Tags = Tags;
-					
+					New->ResourceIds = ResourceIds;
+					New->SkeletonIds = SkeletonIds;
+
 					{
 						Ptr<ASTOpSwitch> NewSwitch = mu::Clone<ASTOpSwitch>(this);
 
 						// Replace all branches removing the "add tags" operation.
 						if (Default)
 						{
-							check(Default->GetOpType() == EOpType::ME_ADDTAGS);
-							const ASTOpMeshAddTags* Typed = static_cast<const ASTOpMeshAddTags*>(Default.child().get());
+							check(Default->GetOpType() == EOpType::ME_ADDMETADATA);
+							const ASTOpMeshAddMetadata* Typed = static_cast<const ASTOpMeshAddMetadata*>(Default.child().get());
 							NewSwitch->Default = Typed->Source.child();
 						}
 
-						for (int32 CaseIndex=0; CaseIndex<Cases.Num();++CaseIndex)
+						for (int32 CaseIndex = 0; CaseIndex < Cases.Num(); ++CaseIndex)
 						{
 							const FCase& SourceCase = Cases[CaseIndex];
 							if (!SourceCase.Branch)
@@ -674,8 +685,8 @@ namespace mu
 
 							FCase& NewCase = NewSwitch->Cases[CaseIndex];
 
-							check(SourceCase.Branch->GetOpType() == EOpType::ME_ADDTAGS);
-							const ASTOpMeshAddTags* Typed = static_cast<const ASTOpMeshAddTags*>(SourceCase.Branch.child().get());
+							check(SourceCase.Branch->GetOpType() == EOpType::ME_ADDMETADATA);
+							const ASTOpMeshAddMetadata* Typed = static_cast<const ASTOpMeshAddMetadata*>(SourceCase.Branch.child().get());
 							NewCase.Branch = Typed->Source.child();
 						}
 
diff --git a/Engine/Plugins/Mutable/Source/MutableTools/Private/MuT/CodeGenerator_Mesh.cpp b/Engine/Plugins/Mutable/Source/MutableTools/Private/MuT/CodeGenerator_Mesh.cpp
--- a/Engine/Plugins/Mutable/Source/MutableTools/Private/MuT/CodeGenerator_Mesh.cpp
+++ b/Engine/Plugins/Mutable/Source/MutableTools/Private/MuT/CodeGenerator_Mesh.cpp
@@ -31,8 +31,8 @@
 #include "MuT/ASTOpMeshTransform.h"
 #include "MuT/ASTOpMeshDifference.h"
 #include "MuT/ASTOpMeshMorph.h"
-#include "MuT/ASTOpMeshAddTags.h"
 #include "MuT/ASTOpMeshPrepareLayout.h"
+#include "MuT/ASTOpMeshAddMetadata.h"
 #include "MuT/ASTOpSwitch.h"
 #include "MuT/ASTOpReferenceResource.h"
 #include "MuT/CodeGenerator.h"
@@ -777,19 +777,35 @@ namespace mu
 					return FMeshGenerationResult();
 				}
 
-				// Separate the tags from the mesh
+				// Separate metadata from the mesh
 				TArray<FString> Tags = ReferenceMesh->Tags;
-				if (!Mesh->Tags.IsEmpty())
-				{
-					Tags.Append(Mesh->Tags);
-					TSharedPtr<FMesh> TaglessMesh = CloneOrTakeOver(Mesh);
-					TaglessMesh->Tags.SetNum(0, EAllowShrinking::No);
-					Mesh = TaglessMesh;
-				}
+				TArray<uint64> ResourceIds = ReferenceMesh->StreamedResources;
+				TArray<uint32> SkeletonIds = ReferenceMesh->SkeletonIDs;
 
-				for (uint64 StreamedResource : ReferenceMesh->StreamedResources)
+				if (Mesh->Tags.Num() + Mesh->StreamedResources.Num() + Mesh->SkeletonIDs.Num() > 0)
 				{
-					Mesh->StreamedResources.AddUnique(StreamedResource);
+					for (FString& Tag : Mesh->Tags)
+					{
+						Tags.AddUnique(Tag);
+					}
+
+					for (uint64 ResourceId : Mesh->StreamedResources)
+					{
+						ResourceIds.AddUnique(ResourceId);
+					}
+
+					for (uint32 SkeletonId : Mesh->SkeletonIDs)
+					{
+						SkeletonIds.AddUnique(SkeletonId);
+					}
+
+					TSharedPtr<FMesh> MetadatalessMesh = CloneOrTakeOver(Mesh);
+
+					MetadatalessMesh->Tags.SetNum(0, EAllowShrinking::No);
+					MetadatalessMesh->StreamedResources.SetNum(0, EAllowShrinking::No);
+					MetadatalessMesh->SkeletonIDs.SetNum(0, EAllowShrinking::No);
+
+					Mesh = MetadatalessMesh;
 				}
 
 				FMeshGenerationDynamicOptions Options = InOptions.GetResult();
@@ -853,13 +869,17 @@ namespace mu
 
 				Result.BaseMeshOp = LastMeshOp;
 
-				// Add the tags operation
-				if (Tags.Num())
+				// Add the metadata operation
+				if (Tags.Num() + ResourceIds.Num() + SkeletonIds.Num() > 0)
 				{
-					Ptr<ASTOpMeshAddTags> AddTagsOp = new ASTOpMeshAddTags;
-					AddTagsOp->Source = LastMeshOp;
-					AddTagsOp->Tags = Tags;
-					LastMeshOp = AddTagsOp;
+					Ptr<ASTOpMeshAddMetadata> AddMetadataOp = new ASTOpMeshAddMetadata;
+					AddMetadataOp->Source = LastMeshOp;
+
+					AddMetadataOp->Tags        = MoveTemp(Tags);
+					AddMetadataOp->ResourceIds = MoveTemp(ResourceIds);
+					AddMetadataOp->SkeletonIds = MoveTemp(SkeletonIds);
+
+					LastMeshOp = AddMetadataOp;
 				}
 
 				Result.MeshOp = LastMeshOp;
diff --git a/Engine/Plugins/Mutable/Source/MutableTools/Private/MuT/CodeOptimiser.cpp b/Engine/Plugins/Mutable/Source/MutableTools/Private/MuT/CodeOptimiser.cpp
--- a/Engine/Plugins/Mutable/Source/MutableTools/Private/MuT/CodeOptimiser.cpp
+++ b/Engine/Plugins/Mutable/Source/MutableTools/Private/MuT/CodeOptimiser.cpp
@@ -15,7 +15,7 @@
 #include "MuT/ASTOpMeshClipMorphPlane.h"
 #include "MuT/ASTOpMeshApplyPose.h"
 #include "MuT/ASTOpMeshPrepareLayout.h"
-#include "MuT/ASTOpMeshAddTags.h"
+#include "MuT/ASTOpMeshAddMetadata.h"
 #include "MuT/ASTOpImageRasterMesh.h"
 #include "MuT/ASTOpReferenceResource.h"
 
@@ -583,10 +583,10 @@ namespace mu
 						bHasSpecialOpInSubgraph = true;
 						break;
 
-					case EOpType::ME_ADDTAGS:
+					case EOpType::ME_ADDMETADATA:
 					{
-						const ASTOpMeshAddTags* AddTags = static_cast<const ASTOpMeshAddTags*>(SubgraphRoot.get());
-						if (AddTags->Source && AddTags->Source->GetParentCount()==1)
+						const ASTOpMeshAddMetadata* AddMetadata = static_cast<const ASTOpMeshAddMetadata*>(SubgraphRoot.get());
+						if (AddMetadata->Source && AddMetadata->Source->GetParentCount() == 1)
 						{
 							// We could resolve the constant if there are no other constraints since this op's mesh op is only referenced once.
 							bGetFromChildren = true;
diff --git a/Engine/Plugins/Mutable/Source/MutableTools/Private/MuT/CodeOptimiser_Semantic.cpp b/Engine/Plugins/Mutable/Source/MutableTools/Private/MuT/CodeOptimiser_Semantic.cpp
--- a/Engine/Plugins/Mutable/Source/MutableTools/Private/MuT/CodeOptimiser_Semantic.cpp
+++ b/Engine/Plugins/Mutable/Source/MutableTools/Private/MuT/CodeOptimiser_Semantic.cpp
@@ -17,7 +17,7 @@
 #include "MuT/ASTOpImageRasterMesh.h"
 #include "MuT/ASTOpImageResize.h"
 #include "MuT/ASTOpMeshMorph.h"
-#include "MuT/ASTOpMeshAddTags.h"
+#include "MuT/ASTOpMeshAddMetadata.h"
 #include "MuT/ASTOpSwitch.h"
 #include "MuT/CodeOptimiser.h"
 #include "MuR/Image.h"
diff --git a/Engine/Plugins/Mutable/Source/MutableTools/Private/MuT/Compiler.cpp b/Engine/Plugins/Mutable/Source/MutableTools/Private/MuT/Compiler.cpp
--- a/Engine/Plugins/Mutable/Source/MutableTools/Private/MuT/Compiler.cpp
+++ b/Engine/Plugins/Mutable/Source/MutableTools/Private/MuT/Compiler.cpp
@@ -159,7 +159,7 @@ namespace mu
 		{ true,		{ 0, 0, 0, 0, 0, 0, 0, 0, 0, 0, 0, 0, 0, 0, 0, 0, 0, 0, 0, 0, 0, 0, 0, 0, 0, 0, 0, 0, 0 }		},	// ME_CLIPDEFORM
 		{ true,		{ 0, 0, 0, 0, 0, 0, 0, 0, 0, 0, 0, 0, 0, 0, 0, 0, 0, 0, 0, 0, 0, 0, 0, 0, 0, 0, 0, 0, 0 }		},	// ME_MORPHRESHAPE
 		{ true,		{ 0, 0, 0, 0, 0, 0, 0, 0, 0, 0, 0, 0, 0, 0, 0, 0, 0, 0, 0, 0, 0, 0, 0, 0, 0, 0, 0, 0, 0 }		},	// ME_OPTIMIZESKINNING
-		{ true,		{ 0, 0, 0, 0, 0, 0, 0, 0, 0, 0, 0, 0, 0, 0, 0, 0, 0, 0, 0, 0, 0, 0, 0, 0, 0, 0, 0, 0, 0 }		},	// ME_ADDTAGS
+		{ true,		{ 0, 0, 0, 0, 0, 0, 0, 0, 0, 0, 0, 0, 0, 0, 0, 0, 0, 0, 0, 0, 0, 0, 0, 0, 0, 0, 0, 0, 0 }		},	// ME_ADDMETADATA
 		{ true,		{ 0, 0, 0, 0, 0, 0, 0, 0, 0, 0, 0, 0, 0, 0, 0, 0, 0, 0, 0, 0, 0, 0, 0, 0, 0, 0, 0, 0, 0 }		},	// ME_TRANSFORMWITHMESH
 
 		{ false,	{ 0, 0, 0, 0, 0, 0, 0, 0, 0, 0, 0, 0, 0, 0, 0, 0, 0, 0, 0, 0, 0, 0, 0, 0, 0, 0, 0, 0, 0 }		},	// IN_ADDMESH
@@ -666,8 +666,6 @@ namespace mu
             Program.States.Add(State.state);
         }
 
-		UE_LOG(LogMutableCore, Verbose, TEXT("(int) %s : %ld"), TEXT("Program size"), int64(Program.OpAddress.Num()));
-
         // Merge the log in the right order
         GenErrorLog->Merge( m_pD->ErrorLog.Get() );
         m_pD->ErrorLog = GenErrorLog;
@@ -675,8 +673,6 @@ namespace mu
 		// Pack data
 		m_pD->GenerateRoms(Result.Get(), LinkerOptions.AdditionalData );
 
-		UE_LOG(LogMutableCore, Verbose, TEXT("(int) %s : %ld"), TEXT("Program size"), int64(Program.OpAddress.Num()));
-
 		// We are not touching the program anymore. Ensure we are not wasting memory.
 		Program.Roms.Shrink();
 		Program.ConstantImageLODsPermanent.Shrink();
@@ -684,10 +680,28 @@ namespace mu
 		Program.ConstantImages.Shrink();
 		Program.ConstantMeshesPermanent.Shrink();
 		Program.ConstantStrings.Shrink();
+		Program.ConstantUInt32Lists.Shrink();
+		Program.ConstantUInt64Lists.Shrink();
 		Program.ConstantSkeletons.Shrink();
 		Program.ConstantPhysicsBodies.Shrink();
 		Program.Parameters.Shrink();
 
+		FOutputSizeStream ProgramSizeStream;
+		{
+			MUTABLE_CPUPROFILER_SCOPE(ComputeProgramSize)
+
+			FOutputArchive ProgramMemoryArch(&ProgramSizeStream);
+			Program.Serialise(ProgramMemoryArch);
+		}
+
+		const double ProgramSize  = double(ProgramSizeStream.GetBufferSize());
+		const double ByteCodeSize = double(Program.ByteCode.Num() + Program.OpAddress.Num() * sizeof(uint32));
+		constexpr double ByteToMiB = 1.0/(1024.0*1024.0); 
+
+		UE_LOG(LogMutableCore, Log, TEXT("Program Size : %lf MiB"), ProgramSize * ByteToMiB);
+		UE_LOG(LogMutableCore, Log, TEXT("Program ByteCode Size: %lf MiB"), ByteCodeSize * ByteToMiB);
+		UE_LOG(LogMutableCore, Log, TEXT("Program Num Ops : %d"), Program.OpAddress.Num());
+
         return Result;
     }
 
diff --git a/Engine/Plugins/Mutable/Source/MutableTools/Private/MuT/ErrorLog.cpp b/Engine/Plugins/Mutable/Source/MutableTools/Private/MuT/ErrorLog.cpp
--- a/Engine/Plugins/Mutable/Source/MutableTools/Private/MuT/ErrorLog.cpp
+++ b/Engine/Plugins/Mutable/Source/MutableTools/Private/MuT/ErrorLog.cpp
@@ -305,7 +305,7 @@ namespace mu
 		TEXT("ME_CLIPDEFORM	   "),
 		TEXT("ME_MORPHRESHAPE  "),
 		TEXT("ME_OPTIMIZESKIN  "),
-		TEXT("ME_ADDTAGS       "),
+		TEXT("ME_ADDMETADATA   "),
 		TEXT("ME_TRANSFORMWITHMESH"),
 
 		TEXT("IN_ADDMESH       "),
