[Mutable] remove FConstantResourceIndex POD array serialization as it seems to cause link errors on some platforms.

- Update CO DD version to prevent any unexpected serialization mismatch.

#rnx

diff --git a/Engine/Plugins/Mutable/Source/CustomizableObject/Internal/MuCO/CustomizableObjectPrivate.h b/Engine/Plugins/Mutable/Source/CustomizableObject/Internal/MuCO/CustomizableObjectPrivate.h
--- a/Engine/Plugins/Mutable/Source/CustomizableObject/Internal/MuCO/CustomizableObjectPrivate.h
+++ b/Engine/Plugins/Mutable/Source/CustomizableObject/Internal/MuCO/CustomizableObjectPrivate.h
@@ -1887,7 +1887,7 @@ public:
 
 #endif
 	
-	static constexpr int32 DerivedDataVersion = 0x7290caa3;
+	static constexpr int32 DerivedDataVersion = 0x1a376aa8;
 };
 
 #if WITH_EDITOR
diff --git a/Engine/Plugins/Mutable/Source/MutableRuntime/Internal/MuR/ModelPrivate.h b/Engine/Plugins/Mutable/Source/MutableRuntime/Internal/MuR/ModelPrivate.h
--- a/Engine/Plugins/Mutable/Source/MutableRuntime/Internal/MuR/ModelPrivate.h
+++ b/Engine/Plugins/Mutable/Source/MutableRuntime/Internal/MuR/ModelPrivate.h
@@ -35,7 +35,7 @@ namespace mu
 	};
 	static_assert(sizeof(FConstantResourceIndex) == 4);
 	MUTABLE_DEFINE_POD_SERIALISABLE(FConstantResourceIndex);
-	MUTABLE_DEFINE_POD_VECTOR_SERIALISABLE(FConstantResourceIndex);
+	//MUTABLE_DEFINE_POD_VECTOR_SERIALISABLE(FConstantResourceIndex);
 
 	/** This is encoded with minimal bits. Make sure to review all uses if extended. */
 	enum class ERomDataType : uint32
diff --git a/Engine/Plugins/Mutable/Source/MutableRuntime/Private/MuR/Model.cpp b/Engine/Plugins/Mutable/Source/MutableRuntime/Private/MuR/Model.cpp
--- a/Engine/Plugins/Mutable/Source/MutableRuntime/Private/MuR/Model.cpp
+++ b/Engine/Plugins/Mutable/Source/MutableRuntime/Private/MuR/Model.cpp
@@ -30,7 +30,7 @@ namespace mu
 	MUTABLE_IMPLEMENT_POD_SERIALISABLE(FImageLODRange);
 	MUTABLE_IMPLEMENT_POD_SERIALISABLE(FMeshContentRange);
 	MUTABLE_IMPLEMENT_POD_SERIALISABLE(FConstantResourceIndex);
-	MUTABLE_IMPLEMENT_POD_VECTOR_SERIALISABLE(FConstantResourceIndex);
+	//MUTABLE_IMPLEMENT_POD_VECTOR_SERIALISABLE(FConstantResourceIndex);
 
 	
     //---------------------------------------------------------------------------------------------
