[Mutable] Fix crash when trying to reload paremeters when no ModelResources is present.

#rnx

diff --git a/Engine/Plugins/Mutable/Source/CustomizableObject/Private/MuCO/CustomizableObjectInstanceDescriptor.cpp b/Engine/Plugins/Mutable/Source/CustomizableObject/Private/MuCO/CustomizableObjectInstanceDescriptor.cpp
--- a/Engine/Plugins/Mutable/Source/CustomizableObject/Private/MuCO/CustomizableObjectInstanceDescriptor.cpp
+++ b/Engine/Plugins/Mutable/Source/CustomizableObject/Private/MuCO/CustomizableObjectInstanceDescriptor.cpp
@@ -883,6 +883,12 @@ void FCustomizableObjectInstanceDescriptor::ReloadParameters()
 		UE_LOG(LogMutable, Warning, TEXT("[ReloadParametersFromObject] No model in object [%s], generated empty parameters for [%s] "), *CustomizableObject->GetName(), *CustomizableObject->GetName());
 		return;
 	}
+
+	if(!CustomizableObject->GetPrivate()->GetModelResources())
+	{
+		UE_LOG(LogMutable, Warning, TEXT("[ReloadParametersFromObject] No model resources in object [%s], generated empty parameters for [%s] "), *CustomizableObject->GetName(), *CustomizableObject->GetName());
+		return;
+	}
  
 	const UModelResources& ModelResources = CustomizableObject->GetPrivate()->GetModelResourcesChecked();
 	
