[Mutable] Fix "[ReloadParametersFromObject] No model resources in object" warning when compiling a CO.

#rnx

diff --git a/Engine/Plugins/Mutable/Source/CustomizableObject/Private/MuCO/CustomizableObjectInstanceDescriptor.cpp b/Engine/Plugins/Mutable/Source/CustomizableObject/Private/MuCO/CustomizableObjectInstanceDescriptor.cpp
--- a/Engine/Plugins/Mutable/Source/CustomizableObject/Private/MuCO/CustomizableObjectInstanceDescriptor.cpp
+++ b/Engine/Plugins/Mutable/Source/CustomizableObject/Private/MuCO/CustomizableObjectInstanceDescriptor.cpp
@@ -884,19 +884,19 @@ void FCustomizableObjectInstanceDescriptor::ReloadParameters()
 		return;
 	}
 
+	const UModelResources* ModelResources = CustomizableObject->GetPrivate()->GetModelResources();
+	
 	if(!CustomizableObject->GetPrivate()->GetModelResources())
 	{
 		UE_LOG(LogMutable, Warning, TEXT("[ReloadParametersFromObject] No model resources in object [%s], generated empty parameters for [%s] "), *CustomizableObject->GetName(), *CustomizableObject->GetName());
 		return;
 	}
- 
-	const UModelResources& ModelResources = CustomizableObject->GetPrivate()->GetModelResourcesChecked();
 	
 	TSharedPtr<mu::FParameters> MutableParameters = mu::FModel::NewParameters(
 		CustomizableObject->GetPrivate()->GetModel(),
 		nullptr,
-		&ModelResources.TextureParameterDefaultValues,
-		&ModelResources.SkeletalMeshParameterDefaultValues);
+		&ModelResources->TextureParameterDefaultValues,
+		&ModelResources->SkeletalMeshParameterDefaultValues);
  
 	int32 ParamCount = MutableParameters->GetCount();
 	for (int32 ParamIndex = 0; ParamIndex < ParamCount; ++ParamIndex)
