[Mutable] Remove DLLEXPORT directive from MUTABLE_IMPLEMENT_POD_VECTOR_SERIALISABLE.

- Enable vector serialization of FConstantResourceIndex type.

#rnx

diff --git a/Engine/Plugins/Mutable/Source/MutableRuntime/Internal/MuR/ModelPrivate.h b/Engine/Plugins/Mutable/Source/MutableRuntime/Internal/MuR/ModelPrivate.h
--- a/Engine/Plugins/Mutable/Source/MutableRuntime/Internal/MuR/ModelPrivate.h
+++ b/Engine/Plugins/Mutable/Source/MutableRuntime/Internal/MuR/ModelPrivate.h
@@ -35,7 +35,7 @@ namespace mu
 	};
 	static_assert(sizeof(FConstantResourceIndex) == 4);
 	MUTABLE_DEFINE_POD_SERIALISABLE(FConstantResourceIndex);
-	//MUTABLE_DEFINE_POD_VECTOR_SERIALISABLE(FConstantResourceIndex);
+	MUTABLE_DEFINE_POD_VECTOR_SERIALISABLE(FConstantResourceIndex);
 
 	/** This is encoded with minimal bits. Make sure to review all uses if extended. */
 	enum class ERomDataType : uint32
diff --git a/Engine/Plugins/Mutable/Source/MutableRuntime/Internal/MuR/SerialisationPrivate.h b/Engine/Plugins/Mutable/Source/MutableRuntime/Internal/MuR/SerialisationPrivate.h
--- a/Engine/Plugins/Mutable/Source/MutableRuntime/Internal/MuR/SerialisationPrivate.h
+++ b/Engine/Plugins/Mutable/Source/MutableRuntime/Internal/MuR/SerialisationPrivate.h
@@ -65,29 +65,28 @@ namespace mu
     }																 \
 		
 
-#define MUTABLE_IMPLEMENT_POD_VECTOR_SERIALISABLE(Type)                        \
-    template<class Alloc>                                                      \
-	void DLLEXPORT operator<<(FOutputArchive& Arch, const TArray<Type, Alloc>& V)         \
-	{                                                                          \
-		uint32 Num = uint32(V.Num());                                          \
-		Arch << Num;                                                           \
-		if (Num)                                                               \
-		{                                                                      \
-			Arch.Stream->Write(V.GetData(), Num * sizeof(Type));			   \
-		}                                                                      \
-	}                                                                          \
-                                                                               \
-    template<class Alloc>                                                      \
-	void DLLEXPORT operator>>(FInputArchive& Arch, TArray<Type, Alloc>& V)     \
-	{                                                                          \
-		uint32 Num;                                                            \
-		Arch >> Num;                                                           \
-		V.SetNum(Num);                                                         \
-		if (Num)                                                               \
-		{                                                                      \
-			Arch.Stream->Read(V.GetData(), Num * sizeof(Type));					   \
-		}                                                                      \
-	}                                                                          \
+#define MUTABLE_IMPLEMENT_POD_VECTOR_SERIALISABLE(Type)                           \
+	template<class Alloc>                                                         \
+	void operator<<(FOutputArchive& Arch, const TArray<Type, Alloc>& V)           \
+	{                                                                             \
+		uint32 Num = uint32(V.Num());                                             \
+		Arch << Num;                                                              \
+		if (Num)                                                                  \
+		{                                                                         \
+			Arch.Stream->Write(V.GetData(), Num * sizeof(Type));                  \
+		}                                                                         \
+	}                                                                             \
+	template<class Alloc>                                                         \
+	void operator>>(FInputArchive& Arch, TArray<Type, Alloc>& V)                  \
+	{                                                                             \
+		uint32 Num;                                                               \
+		Arch >> Num;                                                              \
+		V.SetNum(Num);                                                            \
+		if (Num)                                                                  \
+		{                                                                         \
+			Arch.Stream->Read(V.GetData(), Num * sizeof(Type));                   \
+		}                                                                         \
+	}                                                                             \
 
 	/** TVariant custom serialize. Based on the default serialization. */
 	template <typename... Ts>
diff --git a/Engine/Plugins/Mutable/Source/MutableRuntime/Private/MuR/Model.cpp b/Engine/Plugins/Mutable/Source/MutableRuntime/Private/MuR/Model.cpp
--- a/Engine/Plugins/Mutable/Source/MutableRuntime/Private/MuR/Model.cpp
+++ b/Engine/Plugins/Mutable/Source/MutableRuntime/Private/MuR/Model.cpp
@@ -30,7 +30,7 @@ namespace mu
 	MUTABLE_IMPLEMENT_POD_SERIALISABLE(FImageLODRange);
 	MUTABLE_IMPLEMENT_POD_SERIALISABLE(FMeshContentRange);
 	MUTABLE_IMPLEMENT_POD_SERIALISABLE(FConstantResourceIndex);
-	//MUTABLE_IMPLEMENT_POD_VECTOR_SERIALISABLE(FConstantResourceIndex);
+	MUTABLE_IMPLEMENT_POD_VECTOR_SERIALISABLE(FConstantResourceIndex);
 
 	
     //---------------------------------------------------------------------------------------------
diff --git a/Engine/Plugins/Mutable/Source/MutableRuntime/Public/MuR/Serialisation.h b/Engine/Plugins/Mutable/Source/MutableRuntime/Public/MuR/Serialisation.h
--- a/Engine/Plugins/Mutable/Source/MutableRuntime/Public/MuR/Serialisation.h
+++ b/Engine/Plugins/Mutable/Source/MutableRuntime/Public/MuR/Serialisation.h
@@ -23,26 +23,27 @@ class USkeletalMesh;
 #define UE_API MUTABLERUNTIME_API
 
 namespace mu
-{    
-    class FImage;
+{
+	class FImage;
 	class FModel;
 	enum class EDataType : uint8;
 
-#define MUTABLE_DEFINE_POD_SERIALISABLE(Type)									\
-	void DLLEXPORT operator<<(mu::FOutputArchive& Arch, const Type& T);			\
-	void DLLEXPORT operator>>(mu::FInputArchive& Arch, Type& T);
+#define MUTABLE_DEFINE_POD_SERIALISABLE(Type)                                         \
+	void DLLEXPORT operator<<(mu::FOutputArchive& Arch, const Type& T);               \
+	void DLLEXPORT operator>>(mu::FInputArchive& Arch, Type& T);                      \
+
+#define MUTABLE_DEFINE_POD_VECTOR_SERIALISABLE(Type)                                   \
+	template<class Alloc>                                                              \
+	void DLLEXPORT operator<<(mu::FOutputArchive& Arch, const TArray<Type, Alloc>& V); \
+	template<class Alloc>                                                              \
+	void DLLEXPORT operator>>(mu::FInputArchive& Arch, TArray<Type, Alloc>& V);        \
+
+
+#define MUTABLE_DEFINE_ENUM_SERIALISABLE(Type)                                      \
+	void DLLEXPORT operator<<(mu::FOutputArchive& Arch, const Type& T);             \
+	void DLLEXPORT operator>>(mu::FInputArchive& Arch, Type& T);                    \
 
-#define MUTABLE_DEFINE_POD_VECTOR_SERIALISABLE(Type)							\
-	template<typename Alloc>													\
-	void DLLEXPORT operator<<(mu::FOutputArchive& Arch, const TArray<Type, Alloc>& V);	\
-	template<typename Alloc>													\
-	void DLLEXPORT operator>>(mu::FInputArchive& Arch, TArray<Type, Alloc>& V);
 
-#define MUTABLE_DEFINE_ENUM_SERIALISABLE(Type)									\
-    void DLLEXPORT operator<<(mu::FOutputArchive& Arch, const Type& T);			\
-    void DLLEXPORT operator>>(mu::FInputArchive& Arch, Type& T);
-	
-	
     /** This repesents a handle to Mutable resource. */
     template<class R>
     class TResourceProxy : public RefCounted
