IncrementalCook: Add accesstracking scopes around some system-specific callbacks that act on behalf of a source package and load (possibly hidden) dependencies for that package. This fixes occurrences of the warning "Unsolicited package %s was loaded outside of any cook operation on a referencer package...".
#rnx
#incrementalcook

diff --git a/Engine/Plugins/Mutable/Source/CustomizableObjectEditor/Private/MuCOE/CustomizableObjectCompiler.cpp b/Engine/Plugins/Mutable/Source/CustomizableObjectEditor/Private/MuCOE/CustomizableObjectCompiler.cpp
--- a/Engine/Plugins/Mutable/Source/CustomizableObjectEditor/Private/MuCOE/CustomizableObjectCompiler.cpp
+++ b/Engine/Plugins/Mutable/Source/CustomizableObjectEditor/Private/MuCOE/CustomizableObjectCompiler.cpp
@@ -25,6 +25,8 @@
 #include "Widgets/Notifications/SNotificationList.h"
 #include "Interfaces/ITargetPlatform.h"
 #include "Misc/App.h"
+#include "Misc/PackageAccessTracking.h"
+#include "Misc/PackageAccessTrackingOps.h"
 #include "MuCO/CustomizableObjectSystemPrivate.h"
 #include "MuCOE/CompileRequest.h"
 #include "MuR/Model.h"
@@ -69,6 +71,8 @@ bool FCustomizableObjectCompiler::Tick(bool bBlocking)
 	{
 		bFinished = false;
 	}
+	FName CurrentPackageName = CurrentObject ? CurrentObject->GetPackage()->GetFName() : NAME_None;
+	UE_TRACK_REFERENCING_PACKAGE_SCOPED(CurrentPackageName, PackageAccessTrackingOps::NAME_PostLoad);
 
 	if (LoadModelDataFromDDCEvent.IsValid())
 	{
@@ -269,6 +273,7 @@ void FCustomizableObjectCompiler::Compile(const TSharedRef<FCompilationRequest>&
 	}
 
 	TRACE_BEGIN_REGION(UE_MUTABLE_COMPILE_REGION);
+	UE_TRACK_REFERENCING_PACKAGE_SCOPED(Object, PackageAccessTrackingOps::NAME_PostLoad);
 
 	check(!CurrentRequest);
 
diff --git a/Engine/Source/Editor/UnrealEd/Private/CookOnTheFlyServer.cpp b/Engine/Source/Editor/UnrealEd/Private/CookOnTheFlyServer.cpp
--- a/Engine/Source/Editor/UnrealEd/Private/CookOnTheFlyServer.cpp
+++ b/Engine/Source/Editor/UnrealEd/Private/CookOnTheFlyServer.cpp
@@ -4673,7 +4673,7 @@ void UCookOnTheFlyServer::ProcessUnsolicitedPackages(TArray<FName>* OutDiscovere
 			}
 		};
 
-	TArray< FPackageStreamEvent> PackageStream = PackageTracker->GetPackageStream();
+	TArray<FPackageStreamEvent> PackageStream = PackageTracker->GetPackageStream();
 	for (FPackageStreamEvent& PackageStreamEvent : PackageStream)
 	{
 		FName PackageName = PackageStreamEvent.PackageName;
@@ -5259,6 +5259,7 @@ void UCookOnTheFlyServer::PostLoadPackageFixup(UE::Cook::FPackageData& PackageDa
 	}
 
 	UE_SCOPED_HIERARCHICAL_COOKTIMER(PostLoadPackageFixup);
+	UE_TRACK_REFERENCING_PACKAGE_SCOPED(Package, PackageAccessTrackingOps::NAME_PostLoad);
 
 	// Perform special processing for UWorld
 	World->PersistentLevel->HandleLegacyMapBuildData();
diff --git a/Engine/Source/Editor/UnrealEd/Private/Cooker/CookGenerationHelper.cpp b/Engine/Source/Editor/UnrealEd/Private/Cooker/CookGenerationHelper.cpp
--- a/Engine/Source/Editor/UnrealEd/Private/Cooker/CookGenerationHelper.cpp
+++ b/Engine/Source/Editor/UnrealEd/Private/Cooker/CookGenerationHelper.cpp
@@ -1589,6 +1589,10 @@ UPackage* FGenerationHelper::TryCreateGeneratedPackage(FCookGenerationInfo& Gene
 #if ENABLE_COOK_STATS
 		++DetailedCookStats::NumRequestedLoads;
 #endif
+		// Mark that we are the referencer of the package. This is not necessary for cook correctness, because we 
+		// early exit GeneratedPackages from ProcessUnsolicitedPackages, but it is useful for debug tools to have
+		// PackageTracker know the referencer of the package.
+		UE_TRACK_REFERENCING_PACKAGE_SCOPED(GetOwner().GetPackageName(), PackageAccessTrackingOps::NAME_PostLoad);
 		GeneratedPackage = CreatePackage(*GeneratedPackageName);
 	}
 	GeneratedPackage->SetSavedHash(GeneratedInfo.PackageHash);
