[Mutable] Fix deadlock when closing the editor while Mutable is streaming Skeletal Meshes.

#rnx

diff --git a/Engine/Plugins/Mutable/Source/CustomizableObject/Internal/MuCO/FMutableTaskGraph.h b/Engine/Plugins/Mutable/Source/CustomizableObject/Internal/MuCO/FMutableTaskGraph.h
--- a/Engine/Plugins/Mutable/Source/CustomizableObject/Internal/MuCO/FMutableTaskGraph.h
+++ b/Engine/Plugins/Mutable/Source/CustomizableObject/Internal/MuCO/FMutableTaskGraph.h
@@ -95,9 +95,11 @@ private:
 	/** Check if the Mutable Thread is locked without using the lock. Does not update the region. */
 	bool IsMutableThreadLockedNoLock() const;
 
+public: 
 	/** Unlock the Mutable Thread. Releases all blocked tasks. */
-	void UnlockMutableThread();
-	
+	void UnlockMutableThread(); // TODO Make private once UE-281921
+
+private:
 	/** A Mutable Task Low Priority will only be launched if:
 	 * - No other low priority task is running.
 	 * - Is allowed to launch low priority tasks.
diff --git a/Engine/Plugins/Mutable/Source/CustomizableObject/Private/MuCO/CustomizableObjectSystem.cpp b/Engine/Plugins/Mutable/Source/CustomizableObject/Private/MuCO/CustomizableObjectSystem.cpp
--- a/Engine/Plugins/Mutable/Source/CustomizableObject/Private/MuCO/CustomizableObjectSystem.cpp
+++ b/Engine/Plugins/Mutable/Source/CustomizableObject/Private/MuCO/CustomizableObjectSystem.cpp
@@ -3887,13 +3887,18 @@ int32 UCustomizableObjectSystem::TickInternal(const bool bBlocking)
 	return 0;
 #endif
 
-	if (IsEngineExitRequested())
+	if (!Private)
 	{
 		return 0;
 	}
 
-	if (!Private)
-	{
+	if (IsEngineExitRequested())
+	{	
+		// TODO BEGIN. Remove once UE-281921  
+		GetPrivate()->MutableTaskGraph.UnlockMutableThread();
+   		GetPrivate()->MutableTaskGraph.AllowLaunchingMutableTaskLowPriority(true, false);
+		// TODO END
+
 		return 0;
 	}
 
@@ -3906,7 +3911,7 @@ int32 UCustomizableObjectSystem::TickInternal(const bool bBlocking)
 			return 0;
 		}
 	}
-
+	
 	// \TODO: Review: We should never compile an object from this tick, so this could be removed
 #if WITH_EDITOR
 
