[Mutable] Optimize MeshVertexIdIteratorConst::Get().

#rnx

diff --git a/Engine/Plugins/Mutable/Source/MutableRuntime/Internal/MuR/MeshPrivate.h b/Engine/Plugins/Mutable/Source/MutableRuntime/Internal/MuR/MeshPrivate.h
--- a/Engine/Plugins/Mutable/Source/MutableRuntime/Internal/MuR/MeshPrivate.h
+++ b/Engine/Plugins/Mutable/Source/MutableRuntime/Internal/MuR/MeshPrivate.h
@@ -643,10 +643,10 @@ namespace mu
 	private:
 
 		/** Current Id that the iterator is pointing at. */
-		int32 CurrentIdIndex=0;
+		int32 CurrentIdIndex = 0;
+		int32 NumVertices = 0;
 
-		/** */
-		const FMesh* Mesh=nullptr;
+		uint64 MeshIDPrefix = 0;
 
 		/** Buffer iterator in case there is an actual Id BufferIndex. */
 		UntypedMeshBufferIteratorConst BufferIterator;
@@ -664,20 +664,13 @@ namespace mu
 				return;
 			}
 
-			Mesh = InMesh;
-			if (Mesh)
-			{
-				BufferIterator = UntypedMeshBufferIteratorConst(Mesh->VertexBuffers, EMeshBufferSemantic::VertexIndex, 0 );
-			}
+			NumVertices = InMesh->GetVertexCount();
+			MeshIDPrefix = (uint64(InMesh->MeshIDPrefix) << 32);
+			BufferIterator = UntypedMeshBufferIteratorConst(InMesh->VertexBuffers, EMeshBufferSemantic::VertexIndex, 0);
 		}		
 
 		inline void operator++()
 		{
-			if (!Mesh)
-			{
-				return;
-			}
-
 			CurrentIdIndex++;			
 
 			if (BufferIterator.ptr())
@@ -715,46 +708,35 @@ namespace mu
 
 		bool IsValid() const
 		{
-			return 
-				Mesh 
-				&&
-				(CurrentIdIndex < Mesh->GetVertexCount())
-				;
+			return CurrentIdIndex < NumVertices;
 		}
 
 		uint64 Get() const
 		{
-			check(Mesh);
-			check(CurrentIdIndex < Mesh->GetVertexCount());
-
-			uint64 Id=FMesh::InvalidVertexId;
+			check(CurrentIdIndex < NumVertices);
 
-			// Is it implicit?
-			if (!BufferIterator.ptr())
+			switch (BufferIterator.GetFormat())
 			{
-				// The id is just prefix and index
-				Id = (uint64(Mesh->MeshIDPrefix) << 32) | uint64(CurrentIdIndex);
-			}
-
-			// Is it relative?
-			else if (BufferIterator.GetFormat() == EMeshBufferFormat::UInt32)
+			case EMeshBufferFormat::UInt32: // relative
 			{
 				// There is a BufferIndex storing IDs without prefix because it is the same for all vertices.
-				uint32 RelativeId = BufferIterator.GetAsUINT32();
-				Id = (uint64(Mesh->MeshIDPrefix) << 32) | uint64(RelativeId);
+				return MeshIDPrefix | uint64(*reinterpret_cast<const uint32*>(BufferIterator.ptr()));
 			}
-
-			// Is it explicit?
-			else if (BufferIterator.GetFormat() == EMeshBufferFormat::UInt64)
+			case EMeshBufferFormat::UInt64: // explicit
 			{
-				Id = BufferIterator.GetAsUINT64();
+				return *reinterpret_cast<const uint64*>(BufferIterator.ptr());
 			}
-			else
+			case EMeshBufferFormat::None: // implicit
+			{
+				// The id is just prefix and index
+				return MeshIDPrefix | uint64(CurrentIdIndex);
+			}
+			default:
 			{
-				check(false);
+				unimplemented();
+				return FMesh::InvalidVertexId;
+			}
 			}
-
-			return Id;
 		}
 
 	};
