[Mutable] Move PhysicsBody constants to mesh roms.

#rnx

diff --git a/Engine/Plugins/Mutable/Source/CustomizableObject/Internal/MuCO/CustomizableObjectPrivate.h b/Engine/Plugins/Mutable/Source/CustomizableObject/Internal/MuCO/CustomizableObjectPrivate.h
--- a/Engine/Plugins/Mutable/Source/CustomizableObject/Internal/MuCO/CustomizableObjectPrivate.h
+++ b/Engine/Plugins/Mutable/Source/CustomizableObject/Internal/MuCO/CustomizableObjectPrivate.h
@@ -1622,6 +1622,8 @@ enum class ECustomizableObjectVersions : int32
 
 	MoveMeshMetadataToOperation,
 	
+	MovePhysicsBodiesToRoms,
+
 	// -----<new versions can be added above this line>--------
 	LastCustomizableObjectVersion
 };
diff --git a/Engine/Plugins/Mutable/Source/CustomizableObjectEditor/Private/MuCOE/SMutableCodeViewer.cpp b/Engine/Plugins/Mutable/Source/CustomizableObjectEditor/Private/MuCOE/SMutableCodeViewer.cpp
--- a/Engine/Plugins/Mutable/Source/CustomizableObjectEditor/Private/MuCOE/SMutableCodeViewer.cpp
+++ b/Engine/Plugins/Mutable/Source/CustomizableObjectEditor/Private/MuCOE/SMutableCodeViewer.cpp
@@ -2763,7 +2763,7 @@ bool SMutableCodeViewer::IsConstantResourceUsedByOperation(const int32 IndexOnCo
 			{
 				if (OperationType == mu::EOpType::ME_CONSTANT)
 				{
-					bResourceLocated = IndexOnConstantsArray == InProgram.GetOpArgs<mu::OP::MeshConstantArgs>(OperationAddress).PhysicsBody;
+					bResourceLocated = IndexOnConstantsArray == InProgram.GetOpArgs<mu::OP::MeshConstantArgs>(OperationAddress).Value;
 				}
 				break;
 			}
diff --git a/Engine/Plugins/Mutable/Source/MutableRuntime/Internal/MuR/ModelPrivate.h b/Engine/Plugins/Mutable/Source/MutableRuntime/Internal/MuR/ModelPrivate.h
--- a/Engine/Plugins/Mutable/Source/MutableRuntime/Internal/MuR/ModelPrivate.h
+++ b/Engine/Plugins/Mutable/Source/MutableRuntime/Internal/MuR/ModelPrivate.h
@@ -403,7 +403,6 @@ namespace mu
 			Arch << ConstantShapes;
             Arch << ConstantCurves;
             Arch << ConstantSkeletons;
-			Arch << ConstantPhysicsBodies;
             Arch << Parameters;
             Arch << Ranges;
             Arch << ParameterLists;
@@ -433,7 +432,6 @@ namespace mu
 			Arch >> ConstantShapes;
             Arch >> ConstantCurves;
             Arch >> ConstantSkeletons;
-			Arch >> ConstantPhysicsBodies;
             Arch >> Parameters;
             Arch >> Ranges;
             Arch >> ParameterLists;
@@ -800,7 +798,6 @@ namespace mu
         void GetConstant(
 				uint32 MeshConstantIndex, 
 				int32 SkeletonConstantIndex, 
-				int32 PhysicsConstantIndex, 
 				TSharedPtr<const FMesh>& OutMesh, 
 				EMeshContentFlags FilterContentFlags, 
 				CreateMeshFunc&& CreateMesh) const
@@ -938,6 +935,8 @@ namespace mu
 				}
 				else
 				{
+					Result->PhysicsBody = PhysicsMesh->PhysicsBody;
+
 					for (const TPair<EMeshBufferType, FMeshBufferSet>& AdditionalBuffer : PhysicsMesh->AdditionalBuffers)
 					{
 						const bool bIsPhysicsBufferType = 
@@ -990,12 +989,6 @@ namespace mu
 				Result->Skeleton = ConstantSkeletons[SkeletonConstantIndex];
 			}
 
-			check(ConstantPhysicsBodies.Num() > PhysicsConstantIndex);
-			if (ConstantPhysicsBodies.IsValidIndex(PhysicsConstantIndex))
-			{
-				Result->PhysicsBody = ConstantPhysicsBodies[PhysicsConstantIndex];
-			}
-
 			OutMesh = Result;
 		}
 
diff --git a/Engine/Plugins/Mutable/Source/MutableRuntime/Internal/MuR/Operations.h b/Engine/Plugins/Mutable/Source/MutableRuntime/Internal/MuR/Operations.h
--- a/Engine/Plugins/Mutable/Source/MutableRuntime/Internal/MuR/Operations.h
+++ b/Engine/Plugins/Mutable/Source/MutableRuntime/Internal/MuR/Operations.h
@@ -413,10 +413,6 @@ namespace mu
             // If not negative, index of the skeleton to set to the mesh from the skeleton
             // constant array.
             int32 Skeleton;
-
-            // If not negative, index of the physics body to set to the mesh from the physics body
-            // constant array.
-			int32 PhysicsBody;
         };
 
 		struct ParameterArgs
diff --git a/Engine/Plugins/Mutable/Source/MutableRuntime/Private/MuR/CodeRunner.cpp b/Engine/Plugins/Mutable/Source/MutableRuntime/Private/MuR/CodeRunner.cpp
--- a/Engine/Plugins/Mutable/Source/MutableRuntime/Private/MuR/CodeRunner.cpp
+++ b/Engine/Plugins/Mutable/Source/MutableRuntime/Private/MuR/CodeRunner.cpp
@@ -1035,7 +1035,7 @@ namespace mu
 			EMeshContentFlags MeshContentFlags = (EMeshContentFlags)Item.ExecutionOptions;
 
 			TSharedPtr<const FMesh> Source;
-			Program.GetConstant(Args.Value, Args.Skeleton, Args.PhysicsBody, Source, MeshContentFlags, 
+			Program.GetConstant(Args.Value, Args.Skeleton, Source, MeshContentFlags, 
 					[this](int32 BudgetReserve){ return CreateMesh(BudgetReserve); });
 			
 			if (!Source)
diff --git a/Engine/Plugins/Mutable/Source/MutableRuntime/Public/MuR/Mesh.h b/Engine/Plugins/Mutable/Source/MutableRuntime/Public/MuR/Mesh.h
--- a/Engine/Plugins/Mutable/Source/MutableRuntime/Public/MuR/Mesh.h
+++ b/Engine/Plugins/Mutable/Source/MutableRuntime/Public/MuR/Mesh.h
@@ -447,7 +447,7 @@ namespace mu
 		TArray<uint32> SkeletonIDs;
 
 		/** 
-		 * This skeleton and physics body are not owned and may be used by other meshes, so it cannot be modified
+		 * This skeleton and physicsbody are not owned and may be used by other meshes, so it cannot be modified
 		 * once the mesh has been fully created.
 		 */
 		TSharedPtr<const FSkeleton> Skeleton;
@@ -522,6 +522,17 @@ namespace mu
 					bEqual = false;
 				}
 			}
+			if (bEqual && PhysicsBody != Other.PhysicsBody)
+			{
+				if (PhysicsBody && Other.PhysicsBody)
+				{
+					bEqual = (*PhysicsBody == *Other.PhysicsBody);
+				}
+				else
+				{
+					bEqual = false;
+				}
+			}
 			if (bEqual) bEqual = (StreamedResources == Other.StreamedResources);
 			if (bEqual) bEqual = (Surfaces == Other.Surfaces);
 			if (bEqual) bEqual = (Tags == Other.Tags);
diff --git a/Engine/Plugins/Mutable/Source/MutableTools/Private/MuT/ASTOpConstantResource.cpp b/Engine/Plugins/Mutable/Source/MutableTools/Private/MuT/ASTOpConstantResource.cpp
--- a/Engine/Plugins/Mutable/Source/MutableTools/Private/MuT/ASTOpConstantResource.cpp
+++ b/Engine/Plugins/Mutable/Source/MutableTools/Private/MuT/ASTOpConstantResource.cpp
@@ -364,7 +364,8 @@ namespace mu
 			// PhysicsMeshData
 			{
 				const EMeshCopyFlags PhysicsDataCopyFlags = 
-						EMeshCopyFlags::WithAdditionalPhysics;
+						EMeshCopyFlags::WithAdditionalPhysics |
+						EMeshCopyFlags::WithPhysicsBody;
 
 				TSharedPtr<FMesh> MeshPhysicsData = MeshData->Clone(PhysicsDataCopyFlags);
 			
@@ -599,13 +600,6 @@ namespace mu
 					MeshValue->SetSkeleton(nullptr);
 				}
 
-				Args.PhysicsBody = -1;
-				if (TSharedPtr<const FPhysicsBody> MeshPhysicsBody = MeshValue->GetPhysicsBody())
-				{
-					Args.PhysicsBody = Program.AddConstant(MeshPhysicsBody);
-					MeshValue->SetPhysicsBody(nullptr);
-				}
-
 				Args.Value = AddConstantMesh(Program, MeshValue, *Options);
 				int32 DataDescIndex = Options->AdditionalData.SourceMeshPerConstant.Add(SourceDataDescriptor);
 				check(DataDescIndex == Args.Value);
