[Mutable] Skeletal Mesh Parameter improvements.

- Fixed multiple issues with Morph nodes and the Skeletal Mesh Parameter node.
- Skeletal Mesh Parameter node, added Reference Value.
- Renamed Mesh Component to Skeletal Mesh Component.
- Removed redundant node name titles. If there is a pin, now the name is not shown in the title.
- Ordered some pins. Name should always be the first one.
- GetCanRenameNode set to false by default.
- BaseNode added name pin.

diff --git a/Engine/Plugins/Mutable/Source/CustomizableObject/Internal/MuCO/CustomizableObjectCustomVersion.h b/Engine/Plugins/Mutable/Source/CustomizableObject/Internal/MuCO/CustomizableObjectCustomVersion.h
--- a/Engine/Plugins/Mutable/Source/CustomizableObject/Internal/MuCO/CustomizableObjectCustomVersion.h
+++ b/Engine/Plugins/Mutable/Source/CustomizableObject/Internal/MuCO/CustomizableObjectCustomVersion.h
@@ -194,7 +194,8 @@ struct FCustomizableObjectCustomVersion
 			
 		EnableMutableMacrosNewVersion,
 
-
+		BaseObjectNodeNamePin,
+		
 		// -----<new versions can be added above this line>-------------------------------------------------
 		VersionPlusOne,
 		LatestVersion = VersionPlusOne - 1
diff --git a/Engine/Plugins/Mutable/Source/CustomizableObjectEditor/Private/MuCOE/CustomizableObjectEditorModule.cpp b/Engine/Plugins/Mutable/Source/CustomizableObjectEditor/Private/MuCOE/CustomizableObjectEditorModule.cpp
--- a/Engine/Plugins/Mutable/Source/CustomizableObjectEditor/Private/MuCOE/CustomizableObjectEditorModule.cpp
+++ b/Engine/Plugins/Mutable/Source/CustomizableObjectEditor/Private/MuCOE/CustomizableObjectEditorModule.cpp
@@ -104,6 +104,7 @@
 #include "Nodes/CustomizableObjectNodeComponentMeshDetails.h"
 #include "Nodes/CustomizableObjectNodeModifierTransformInMeshDetails.h"
 #include "MuCOE/GenerateMutableSource/GenerateMutableSourceTable.h"
+#include "Nodes/CustomizableObjectNodeTextureParameter.h"
 #include "PhysicsEngine/PhysicsAsset.h"
 
 
@@ -1002,6 +1003,13 @@ TMap<FName, FGuid> FCustomizableObjectEditorModule::GetParticipatingObjects(cons
 		else if (UCustomizableObjectNodeSkeletalMeshParameter* MeshParamNode = Cast<UCustomizableObjectNodeSkeletalMeshParameter>(&Node))
 		{
 			AddObjectSoft(MeshParamNode->DefaultValue.ToSoftObjectPath());
+			AddObjectSoft(MeshParamNode->ReferenceValue.ToSoftObjectPath());
+		}
+
+		else if (UCustomizableObjectNodeTextureParameter* TextureParameterNode = Cast<UCustomizableObjectNodeTextureParameter>(&Node))
+		{
+			AddObjectSoft(TextureParameterNode->DefaultValue);
+			AddObjectSoft(TextureParameterNode->ReferenceValue);
 		}
 
 		else if (UCustomizableObjectNodeMaterial* MatNode = Cast<UCustomizableObjectNodeMaterial>(&Node))
diff --git a/Engine/Plugins/Mutable/Source/CustomizableObjectEditor/Private/MuCOE/GenerateMutableSource/GenerateMutableSourceMesh.cpp b/Engine/Plugins/Mutable/Source/CustomizableObjectEditor/Private/MuCOE/GenerateMutableSource/GenerateMutableSourceMesh.cpp
--- a/Engine/Plugins/Mutable/Source/CustomizableObjectEditor/Private/MuCOE/GenerateMutableSource/GenerateMutableSourceMesh.cpp
+++ b/Engine/Plugins/Mutable/Source/CustomizableObjectEditor/Private/MuCOE/GenerateMutableSource/GenerateMutableSourceMesh.cpp
@@ -2920,11 +2920,11 @@ TSharedPtr<mu::FMesh> BuildMorphedMutableMesh(const UEdGraphPin* BaseSourcePin,
 	USkeletalMesh* SkeletalMesh = nullptr;
 	UCustomizableObjectNode* Node = Cast<UCustomizableObjectNode>(BaseSourcePin->GetOwningNode());
 
-	if (const UCustomizableObjectNodeSkeletalMesh* TypedNodeSkeletalMesh = Cast<UCustomizableObjectNodeSkeletalMesh>(Node))
+	if (const UCustomizableObjectNodeMesh* TypedNodeSkeletalMesh = Cast<UCustomizableObjectNodeMesh>(Node))
 	{
 		int32 LayoutIndex;
 		TypedNodeSkeletalMesh->GetPinSection(*BaseSourcePin, LODIndexConnected, SectionIndexConnected, LayoutIndex);
-		SkeletalMesh = MutablePrivate::LoadObject(TypedNodeSkeletalMesh->SkeletalMesh);
+		SkeletalMesh = Cast<USkeletalMesh>(MutablePrivate::LoadObject(TypedNodeSkeletalMesh->GetMesh()));
 	}
 
 	else if (const UCustomizableObjectNodeTable* TypedNodeTable = Cast<UCustomizableObjectNodeTable>(Node))
@@ -2937,6 +2937,11 @@ TSharedPtr<mu::FMesh> BuildMorphedMutableMesh(const UEdGraphPin* BaseSourcePin,
 			SkeletalMesh = Cast<USkeletalMesh>(GenerationContext.LoadObject(TypedNodeTable->GetSkeletalMeshAt(BaseSourcePin, DataTable, RowName)));
 		}
 	}
+	
+	else
+	{
+		unimplemented();	
+	}
 
 	TSharedPtr<mu::FMesh> MorphedSourceMesh;
 
@@ -4643,9 +4648,14 @@ mu::Ptr<mu::NodeMesh> GenerateMutableSourceMesh(const UEdGraphPin* Pin,
 		// For constants we also have:
 		// animbp instance?
 		// gameplay tags?
-		// morphs?
 
 		Result = ParameterNode;
+
+		// Applying Mesh Morph Node
+		if (GenerationContext.MeshMorphStack.Num())
+		{
+			Result = GenerateMorphMesh(Pin, GenerationContext.MeshMorphStack, 0, Result, GenerationContext, bOnlyConnectedLOD);
+		}
 	}
 
 	else if (const UCustomizableObjectNodeMacroInstance* TypedNodeMacro = Cast<UCustomizableObjectNodeMacroInstance>(Node))
diff --git a/Engine/Plugins/Mutable/Source/CustomizableObjectEditor/Private/MuCOE/Nodes/CustomizableObjectNode.cpp b/Engine/Plugins/Mutable/Source/CustomizableObjectEditor/Private/MuCOE/Nodes/CustomizableObjectNode.cpp
--- a/Engine/Plugins/Mutable/Source/CustomizableObjectEditor/Private/MuCOE/Nodes/CustomizableObjectNode.cpp
+++ b/Engine/Plugins/Mutable/Source/CustomizableObjectEditor/Private/MuCOE/Nodes/CustomizableObjectNode.cpp
@@ -338,6 +338,12 @@ TSharedPtr<SGraphNode> UCustomizableObjectNode::CreateVisualWidget()
 }
 
 
+bool UCustomizableObjectNode::GetCanRenameNode() const
+{
+	return false;
+}
+
+
 void UCustomizableObjectNode::AutowireNewNode(UEdGraphPin* FromPin)
 {
 	if (!FromPin)
diff --git a/Engine/Plugins/Mutable/Source/CustomizableObjectEditor/Private/MuCOE/Nodes/CustomizableObjectNodeComponent.cpp b/Engine/Plugins/Mutable/Source/CustomizableObjectEditor/Private/MuCOE/Nodes/CustomizableObjectNodeComponent.cpp
--- a/Engine/Plugins/Mutable/Source/CustomizableObjectEditor/Private/MuCOE/Nodes/CustomizableObjectNodeComponent.cpp
+++ b/Engine/Plugins/Mutable/Source/CustomizableObjectEditor/Private/MuCOE/Nodes/CustomizableObjectNodeComponent.cpp
@@ -42,19 +42,6 @@ void UCustomizableObjectNodeComponent::BackwardsCompatibleFixup(int32 Customizab
 }
 
 
-bool UCustomizableObjectNodeComponent::GetCanRenameNode() const
-{
-	const UEdGraphPin* NamePin = GetComponentNamePin();
-
-	if (NamePin && NamePin->LinkedTo.Num())
-	{
-		return false;
-	}
-
-	return true;
-}
-
-
 void UCustomizableObjectNodeComponent::OnRenameNode(const FString& NewName)
 {
 	if (!NewName.IsEmpty())
diff --git a/Engine/Plugins/Mutable/Source/CustomizableObjectEditor/Private/MuCOE/Nodes/CustomizableObjectNodeComponent.h b/Engine/Plugins/Mutable/Source/CustomizableObjectEditor/Private/MuCOE/Nodes/CustomizableObjectNodeComponent.h
--- a/Engine/Plugins/Mutable/Source/CustomizableObjectEditor/Private/MuCOE/Nodes/CustomizableObjectNodeComponent.h
+++ b/Engine/Plugins/Mutable/Source/CustomizableObjectEditor/Private/MuCOE/Nodes/CustomizableObjectNodeComponent.h
@@ -20,7 +20,6 @@ public:
 	virtual void AllocateDefaultPins(UCustomizableObjectNodeRemapPins* RemapPins) override;
 	virtual FLinearColor GetNodeTitleColor() const override;
 	virtual void PinConnectionListChanged(UEdGraphPin* Pin) override;
-	virtual bool GetCanRenameNode() const override;
 	virtual void OnRenameNode(const FString& NewName) override;
 
 	// UCustomizableObjectNode interface
diff --git a/Engine/Plugins/Mutable/Source/CustomizableObjectEditor/Private/MuCOE/Nodes/CustomizableObjectNodeComponentMesh.cpp b/Engine/Plugins/Mutable/Source/CustomizableObjectEditor/Private/MuCOE/Nodes/CustomizableObjectNodeComponentMesh.cpp
--- a/Engine/Plugins/Mutable/Source/CustomizableObjectEditor/Private/MuCOE/Nodes/CustomizableObjectNodeComponentMesh.cpp
+++ b/Engine/Plugins/Mutable/Source/CustomizableObjectEditor/Private/MuCOE/Nodes/CustomizableObjectNodeComponentMesh.cpp
@@ -58,6 +58,8 @@ void UCustomizableObjectNodeComponentMesh::BackwardsCompatibleFixup(int32 Custom
 
 void UCustomizableObjectNodeComponentMesh::AllocateDefaultPins(UCustomizableObjectNodeRemapPins* RemapPins)
 {
+	Super::AllocateDefaultPins(RemapPins);
+
 	// Mesh Component Pins
 	const FString PinFriendlyName = TEXT("Overlay Material");
 	const FString PinName = PinFriendlyName + FString(TEXT("_Input_Pin"));
@@ -75,31 +77,12 @@ void UCustomizableObjectNodeComponentMesh::AllocateDefaultPins(UCustomizableObje
 		UEdGraphPin* Pin = CustomCreatePin(EGPD_Input, UEdGraphSchema_CustomizableObject::PC_Material, FName(*LODName), true);
 		LODPins.Add(Pin);
 	}
-
-	Super::AllocateDefaultPins(RemapPins);
 }
 
 
 FText UCustomizableObjectNodeComponentMesh::GetNodeTitle(ENodeTitleType::Type TitleType) const
 {
-	FText NodeTitle;
-	FFormatNamedArguments Args;
-	Args.Add(TEXT("ComponentName"), FText::FromName(ComponentName));
-
-	if (TitleType == ENodeTitleType::ListView || (GetComponentNamePin() && GetComponentNamePin()->LinkedTo.Num()))
-	{
-		NodeTitle = LOCTEXT("Component_Mesh", "Mesh Component");
-	}
-	else if (TitleType == ENodeTitleType::EditableTitle)
-	{
-		NodeTitle = LOCTEXT("ComponentMesh_Edit", "{ComponentName}");
-	}
-	else
-	{
-		NodeTitle = LOCTEXT("ComponentMesh_Title", "{ComponentName}\nMesh Component");
-	}
-
-	return FText::Format(NodeTitle, Args);
+	return LOCTEXT("Component_Mesh", "Skeletal Mesh Component");
 }
 
 
diff --git a/Engine/Plugins/Mutable/Source/CustomizableObjectEditor/Private/MuCOE/Nodes/CustomizableObjectNodeComponentMeshAddTo.cpp b/Engine/Plugins/Mutable/Source/CustomizableObjectEditor/Private/MuCOE/Nodes/CustomizableObjectNodeComponentMeshAddTo.cpp
--- a/Engine/Plugins/Mutable/Source/CustomizableObjectEditor/Private/MuCOE/Nodes/CustomizableObjectNodeComponentMeshAddTo.cpp
+++ b/Engine/Plugins/Mutable/Source/CustomizableObjectEditor/Private/MuCOE/Nodes/CustomizableObjectNodeComponentMeshAddTo.cpp
@@ -77,25 +77,7 @@ void UCustomizableObjectNodeComponentMeshAddTo::BackwardsCompatibleFixup(int32 C
 
 FText UCustomizableObjectNodeComponentMeshAddTo::GetNodeTitle(ENodeTitleType::Type TitleType) const
 {
-	const UEdGraphPin* ParentNamePin = GetParentComponentNamePin();
-	FText NodeTitle;
-	FFormatNamedArguments Args;
-	Args.Add(TEXT("ComponentName"), FText::FromName(ParentComponentName));
-
-	if (TitleType == ENodeTitleType::ListView || (ParentNamePin && ParentNamePin->LinkedTo.Num()))
-	{
-		NodeTitle = LOCTEXT("ComponentMeshAdd", "Add To Mesh Component");
-	}
-	else if (TitleType == ENodeTitleType::EditableTitle)
-	{
-		NodeTitle = LOCTEXT("ComponentMeshAdd_Edit", "{ComponentName}");
-	}
-	else
-	{
-		NodeTitle = LOCTEXT("ComponentMeshAdd_Title", "{ComponentName}\nAdd To Mesh Component");
-	}
-
-	return FText::Format(NodeTitle, Args);
+	return LOCTEXT("ComponentMeshAdd", "Add To Skeletal Mesh Component");
 }
 
 
diff --git a/Engine/Plugins/Mutable/Source/CustomizableObjectEditor/Private/MuCOE/Nodes/CustomizableObjectNodeComponentPassthroughMesh.cpp b/Engine/Plugins/Mutable/Source/CustomizableObjectEditor/Private/MuCOE/Nodes/CustomizableObjectNodeComponentPassthroughMesh.cpp
--- a/Engine/Plugins/Mutable/Source/CustomizableObjectEditor/Private/MuCOE/Nodes/CustomizableObjectNodeComponentPassthroughMesh.cpp
+++ b/Engine/Plugins/Mutable/Source/CustomizableObjectEditor/Private/MuCOE/Nodes/CustomizableObjectNodeComponentPassthroughMesh.cpp
@@ -145,24 +145,7 @@ bool UCustomizableObjectNodeComponentPassthroughMesh::IsExperimental() const
 
 FText UCustomizableObjectNodeComponentPassthroughMesh::GetNodeTitle(ENodeTitleType::Type TitleType) const
 {
-	FText NodeTitle;
-	FFormatNamedArguments Args;
-	Args.Add(TEXT("ComponentName"), FText::FromName(ComponentName));
-
-	if (TitleType == ENodeTitleType::ListView || (GetComponentNamePin() && GetComponentNamePin()->LinkedTo.Num()))
-	{
-		NodeTitle = LOCTEXT("ComponentPassthroughMesh", "Passthrough Mesh Component");
-	}
-	else if (TitleType == ENodeTitleType::EditableTitle)
-	{
-		NodeTitle = LOCTEXT("ComponentPassthrough_Edit", "{ComponentName}");
-	}
-	else
-	{
-		NodeTitle = LOCTEXT("ComponentPassthroughMesh_Title", "{ComponentName}\n Passthrough Mesh Component");
-	}
-
-	return FText::Format(NodeTitle, Args);
+	return LOCTEXT("ComponentPassthroughMesh", "Passthrough Skeletal Mesh Component");
 }
 
 
diff --git a/Engine/Plugins/Mutable/Source/CustomizableObjectEditor/Private/MuCOE/Nodes/CustomizableObjectNodeExposePin.cpp b/Engine/Plugins/Mutable/Source/CustomizableObjectEditor/Private/MuCOE/Nodes/CustomizableObjectNodeExposePin.cpp
--- a/Engine/Plugins/Mutable/Source/CustomizableObjectEditor/Private/MuCOE/Nodes/CustomizableObjectNodeExposePin.cpp
+++ b/Engine/Plugins/Mutable/Source/CustomizableObjectEditor/Private/MuCOE/Nodes/CustomizableObjectNodeExposePin.cpp
@@ -44,25 +44,10 @@ bool UCustomizableObjectNodeExposePin::IsNodeSupportedInMacros() const
 
 FText UCustomizableObjectNodeExposePin::GetNodeTitle(ENodeTitleType::Type TitleType) const
 {
-	FText NodeTitle;
 	FFormatNamedArguments Args;
-	Args.Add(TEXT("NodeName"), FText::FromString(Name));
 	Args.Add(TEXT("PinType"), UEdGraphSchema_CustomizableObject::GetPinCategoryFriendlyName(PinType));
 
-	if (TitleType == ENodeTitleType::ListView)
-	{
-		NodeTitle = LOCTEXT("Expose_Pin_Title_ListView", "Export {PinType} Pin");
-	}
-	else if (TitleType == ENodeTitleType::EditableTitle)
-	{
-		NodeTitle = LOCTEXT("Expose_Pin_Title_Edit", "{NodeName}");
-	}
-	else
-	{
-		NodeTitle = LOCTEXT("Expose_Pin_Title", "{NodeName}\nExport {PinType} Pin");
-	}
-
-	return FText::Format(NodeTitle, Args);
+	return FText::Format(LOCTEXT("Expose_Pin_Title_ListView", "Export {PinType} Pin"), Args);
 }
 
 
diff --git a/Engine/Plugins/Mutable/Source/CustomizableObjectEditor/Private/MuCOE/Nodes/CustomizableObjectNodeMacroInstance.cpp b/Engine/Plugins/Mutable/Source/CustomizableObjectEditor/Private/MuCOE/Nodes/CustomizableObjectNodeMacroInstance.cpp
--- a/Engine/Plugins/Mutable/Source/CustomizableObjectEditor/Private/MuCOE/Nodes/CustomizableObjectNodeMacroInstance.cpp
+++ b/Engine/Plugins/Mutable/Source/CustomizableObjectEditor/Private/MuCOE/Nodes/CustomizableObjectNodeMacroInstance.cpp
@@ -85,12 +85,6 @@ FLinearColor UCustomizableObjectNodeMacroInstance::GetNodeTitleColor() const
 }
 
 
-bool UCustomizableObjectNodeMacroInstance::GetCanRenameNode() const
-{
-	return false;
-}
-
-
 bool UCustomizableObjectNodeMacroInstance::CanUserDeleteNode() const
 {
 	return true;
diff --git a/Engine/Plugins/Mutable/Source/CustomizableObjectEditor/Private/MuCOE/Nodes/CustomizableObjectNodeMacroInstance.h b/Engine/Plugins/Mutable/Source/CustomizableObjectEditor/Private/MuCOE/Nodes/CustomizableObjectNodeMacroInstance.h
--- a/Engine/Plugins/Mutable/Source/CustomizableObjectEditor/Private/MuCOE/Nodes/CustomizableObjectNodeMacroInstance.h
+++ b/Engine/Plugins/Mutable/Source/CustomizableObjectEditor/Private/MuCOE/Nodes/CustomizableObjectNodeMacroInstance.h
@@ -59,7 +59,6 @@ public:
 	virtual FText GetNodeTitle(ENodeTitleType::Type TitleType) const override;
 	virtual FText GetTooltipText() const override;
 	virtual FLinearColor GetNodeTitleColor() const override;
-	virtual bool GetCanRenameNode() const override;
 	virtual bool CanUserDeleteNode() const override;
 	UObject* GetJumpTargetForDoubleClick() const override;
 	virtual bool CanJumpToDefinition() const override;
diff --git a/Engine/Plugins/Mutable/Source/CustomizableObjectEditor/Private/MuCOE/Nodes/CustomizableObjectNodeMesh.h b/Engine/Plugins/Mutable/Source/CustomizableObjectEditor/Private/MuCOE/Nodes/CustomizableObjectNodeMesh.h
--- a/Engine/Plugins/Mutable/Source/CustomizableObjectEditor/Private/MuCOE/Nodes/CustomizableObjectNodeMesh.h
+++ b/Engine/Plugins/Mutable/Source/CustomizableObjectEditor/Private/MuCOE/Nodes/CustomizableObjectNodeMesh.h
@@ -13,7 +13,7 @@ class UTexture2D;
 
 
 UCLASS(abstract)
-class CUSTOMIZABLEOBJECTEDITOR_API UCustomizableObjectNodeMesh : public UCustomizableObjectNode
+class CUSTOMIZABLEOBJECTEDITOR_API UCustomizableObjectNodeMesh : public UCustomizableObjectNode // TODO Convert this to an UInterface (UE-282488)
 {
 public:
 	GENERATED_BODY()
diff --git a/Engine/Plugins/Mutable/Source/CustomizableObjectEditor/Private/MuCOE/Nodes/CustomizableObjectNodeMeshMorph.h b/Engine/Plugins/Mutable/Source/CustomizableObjectEditor/Private/MuCOE/Nodes/CustomizableObjectNodeMeshMorph.h
--- a/Engine/Plugins/Mutable/Source/CustomizableObjectEditor/Private/MuCOE/Nodes/CustomizableObjectNodeMeshMorph.h
+++ b/Engine/Plugins/Mutable/Source/CustomizableObjectEditor/Private/MuCOE/Nodes/CustomizableObjectNodeMeshMorph.h
@@ -22,15 +22,15 @@ public:
 	UCustomizableObjectNodeMeshMorph();
 
 	// Begin EdGraphNode interface
-	FText GetNodeTitle(ENodeTitleType::Type TitleType) const override;
-	FLinearColor GetNodeTitleColor() const override;
-	FText GetTooltipText() const override;
+	virtual FText GetNodeTitle(ENodeTitleType::Type TitleType) const override;
+	virtual FLinearColor GetNodeTitleColor() const override;
+	virtual FText GetTooltipText() const override;
 
 	// UCustomizableObjectNode interface
-	void AllocateDefaultPins(UCustomizableObjectNodeRemapPins* RemapPins) override;
-	bool IsNodeOutDatedAndNeedsRefresh() override;
-	FString GetRefreshMessage() const override;
-	void PinConnectionListChanged(UEdGraphPin* Pin) override;
+	virtual void AllocateDefaultPins(UCustomizableObjectNodeRemapPins* RemapPins) override;
+	virtual bool IsNodeOutDatedAndNeedsRefresh() override;
+	virtual FString GetRefreshMessage() const override;
+	virtual void PinConnectionListChanged(UEdGraphPin* Pin) override;
 	virtual bool CreateStaticStringPinWidget() const override;
 	virtual void BackwardsCompatibleFixup(int32 CustomizableObjectCustomVersion) override;
 
diff --git a/Engine/Plugins/Mutable/Source/CustomizableObjectEditor/Private/MuCOE/Nodes/CustomizableObjectNodeMeshMorphDetails.cpp b/Engine/Plugins/Mutable/Source/CustomizableObjectEditor/Private/MuCOE/Nodes/CustomizableObjectNodeMeshMorphDetails.cpp
--- a/Engine/Plugins/Mutable/Source/CustomizableObjectEditor/Private/MuCOE/Nodes/CustomizableObjectNodeMeshMorphDetails.cpp
+++ b/Engine/Plugins/Mutable/Source/CustomizableObjectEditor/Private/MuCOE/Nodes/CustomizableObjectNodeMeshMorphDetails.cpp
@@ -54,14 +54,18 @@ void FCustomizableObjectNodeMeshMorphDetails::CustomizeDetails( IDetailLayoutBui
 		if ( BaseSourcePin )
 		{
 			USkeletalMesh* SkeletalMesh = nullptr;
-			if ( const UCustomizableObjectNodeSkeletalMesh* TypedNodeSkeletalMesh = Cast<UCustomizableObjectNodeSkeletalMesh>(BaseSourcePin->GetOwningNode()) )
+			if ( const UCustomizableObjectNodeMesh* TypedNodeSkeletalMesh = Cast<UCustomizableObjectNodeMesh>(BaseSourcePin->GetOwningNode()) )
 			{
-				SkeletalMesh = MutablePrivate::LoadObject(TypedNodeSkeletalMesh->SkeletalMesh);
+				SkeletalMesh = Cast<USkeletalMesh>(MutablePrivate::LoadObject(TypedNodeSkeletalMesh->GetMesh()));
 			}
 			else if (const UCustomizableObjectNodeTable* TypedNodeTable = Cast<UCustomizableObjectNodeTable>(BaseSourcePin->GetOwningNode()))
 			{
 				SkeletalMesh = TypedNodeTable->GetColumnDefaultAssetByType<USkeletalMesh>(BaseSourcePin);
 			}
+			else
+			{
+				unimplemented();
+			}
 
 			if (SkeletalMesh)
 			{
diff --git a/Engine/Plugins/Mutable/Source/CustomizableObjectEditor/Private/MuCOE/Nodes/CustomizableObjectNodeMeshMorphStackApplication.cpp b/Engine/Plugins/Mutable/Source/CustomizableObjectEditor/Private/MuCOE/Nodes/CustomizableObjectNodeMeshMorphStackApplication.cpp
--- a/Engine/Plugins/Mutable/Source/CustomizableObjectEditor/Private/MuCOE/Nodes/CustomizableObjectNodeMeshMorphStackApplication.cpp
+++ b/Engine/Plugins/Mutable/Source/CustomizableObjectEditor/Private/MuCOE/Nodes/CustomizableObjectNodeMeshMorphStackApplication.cpp
@@ -107,14 +107,18 @@ TArray<FString> UCustomizableObjectNodeMeshMorphStackApplication::GetMorphList()
 	{
 		USkeletalMesh* SkeletalMesh = nullptr;
 
-		if (const UCustomizableObjectNodeSkeletalMesh* SkeletalMeshNode = Cast< UCustomizableObjectNodeSkeletalMesh >(MeshNodePin->GetOwningNode()))
+		if (const UCustomizableObjectNodeMesh* MeshNode = Cast< UCustomizableObjectNodeMesh >(MeshNodePin->GetOwningNode()))
 		{
-			SkeletalMesh = MutablePrivate::LoadObject(SkeletalMeshNode->SkeletalMesh);
+			SkeletalMesh = Cast<USkeletalMesh>(MutablePrivate::LoadObject(MeshNode->GetMesh()));
 		}
 		else if (const UCustomizableObjectNodeTable* TableNode = Cast< UCustomizableObjectNodeTable >(MeshNodePin->GetOwningNode()))
 		{
 			SkeletalMesh = TableNode->GetColumnDefaultAssetByType<USkeletalMesh>(MeshNodePin);
 		}
+		else
+		{
+			unimplemented();
+		}
 
 		if (SkeletalMesh)
 		{
diff --git a/Engine/Plugins/Mutable/Source/CustomizableObjectEditor/Private/MuCOE/Nodes/CustomizableObjectNodeMeshMorphStackDefinition.cpp b/Engine/Plugins/Mutable/Source/CustomizableObjectEditor/Private/MuCOE/Nodes/CustomizableObjectNodeMeshMorphStackDefinition.cpp
--- a/Engine/Plugins/Mutable/Source/CustomizableObjectEditor/Private/MuCOE/Nodes/CustomizableObjectNodeMeshMorphStackDefinition.cpp
+++ b/Engine/Plugins/Mutable/Source/CustomizableObjectEditor/Private/MuCOE/Nodes/CustomizableObjectNodeMeshMorphStackDefinition.cpp
@@ -91,21 +91,25 @@ void UCustomizableObjectNodeMeshMorphStackDefinition::UpdateMorphList()
 	{
 		return;
 	}
-
+	
 	const UEdGraphPin* MeshNodePin = FindMeshBaseSource(*OutputMeshPin, false);
 
 	if (MeshNodePin && MeshNodePin->GetOwningNode())
 	{
 		USkeletalMesh* SkeletalMesh = nullptr;
 
-		if (const UCustomizableObjectNodeSkeletalMesh* SkeletalMeshNode = Cast< UCustomizableObjectNodeSkeletalMesh >(MeshNodePin->GetOwningNode()))
+		if (const UCustomizableObjectNodeMesh* MeshNode = Cast<UCustomizableObjectNodeMesh>(MeshNodePin->GetOwningNode()))
 		{
-			SkeletalMesh = MutablePrivate::LoadObject(SkeletalMeshNode->SkeletalMesh);
+			SkeletalMesh = Cast<USkeletalMesh>(MutablePrivate::LoadObject(MeshNode->GetMesh()));
 		}
-		else if (const UCustomizableObjectNodeTable* TableNode = Cast< UCustomizableObjectNodeTable >(MeshNodePin->GetOwningNode()))
+		else if (const UCustomizableObjectNodeTable* TableNode = Cast<UCustomizableObjectNodeTable>(MeshNodePin->GetOwningNode()))
 		{
 			SkeletalMesh = TableNode->GetColumnDefaultAssetByType<USkeletalMesh>(MeshNodePin);
 		}
+		else
+		{
+			unimplemented();
+		}
 
 		if (SkeletalMesh)
 		{
diff --git a/Engine/Plugins/Mutable/Source/CustomizableObjectEditor/Private/MuCOE/Nodes/CustomizableObjectNodeObject.cpp b/Engine/Plugins/Mutable/Source/CustomizableObjectEditor/Private/MuCOE/Nodes/CustomizableObjectNodeObject.cpp
--- a/Engine/Plugins/Mutable/Source/CustomizableObjectEditor/Private/MuCOE/Nodes/CustomizableObjectNodeObject.cpp
+++ b/Engine/Plugins/Mutable/Source/CustomizableObjectEditor/Private/MuCOE/Nodes/CustomizableObjectNodeObject.cpp
@@ -2,6 +2,7 @@
 
 #include "MuCOE/Nodes/CustomizableObjectNodeObject.h"
 
+#include "CustomizableObjectNodeStaticString.h"
 #include "MuCO/ICustomizableObjectModule.h"
 #include "MuCO/CustomizableObjectExtension.h"
 #include "MuCO/CustomizableObjectPrivate.h"
@@ -535,6 +536,14 @@ void UCustomizableObjectNodeObject::BackwardsCompatibleFixup(int32 CustomizableO
 			}
 		}
 	}
+
+	if (CustomizableObjectCustomVersion == FCustomizableObjectCustomVersion::EnableMutableMacrosNewVersion)
+	{
+		if (!NamePin.Get())
+		{
+			NamePin = CustomCreatePin(EGPD_Input, UEdGraphSchema_CustomizableObject::PC_String, "Name");
+		}
+	}
 }
 
 
@@ -573,6 +582,8 @@ void UCustomizableObjectNodeObject::AllocateDefaultPins(UCustomizableObjectNodeR
 {
 	const UEdGraphSchema_CustomizableObject* Schema = GetDefault<UEdGraphSchema_CustomizableObject>();
 
+	NamePin = CustomCreatePin(EGPD_Input, UEdGraphSchema_CustomizableObject::PC_String, "Name");
+	
 	CustomCreatePin(EGPD_Input, Schema->PC_Component, ComponentsPinName, true);
 	
 	CustomCreatePin(EGPD_Input, Schema->PC_Modifier, ModifiersPinName, true);
@@ -608,34 +619,13 @@ bool UCustomizableObjectNodeObject::IsNodeSupportedInMacros() const
 
 FText UCustomizableObjectNodeObject::GetNodeTitle(ENodeTitleType::Type TitleType) const
 {
-	FFormatNamedArguments Args;
-	Args.Add(TEXT("ObjectName"), FText::FromString(ObjectName));
-
-	if (TitleType == ENodeTitleType::ListView)
-	{
-		if (bIsBase)
-		{
-			return LOCTEXT("Base_Object", "Base Object");
-		}
-		else
-		{
-			return LOCTEXT("Base_Object_Deprecated", "Base Object (Deprecated)");
-		}
-	}
-	else if (TitleType == ENodeTitleType::EditableTitle)
+	if (bIsBase)
 	{
-		return FText::Format(LOCTEXT("Base_Object_Edit", "{ObjectName}"), Args);
+		return LOCTEXT("Base_Object", "Base Object");
 	}
 	else
 	{
-		if (bIsBase)
-		{
-			return FText::Format(LOCTEXT("Base_Object_Title", "{ObjectName}\nBase Object"), Args);
-		}
-		else
-		{
-			return FText::Format(LOCTEXT("Child_Object_Title_Deprecated", "{ObjectName}\nChild Object (Deprecated)"), Args);
-		}
+		return LOCTEXT("Base_Object_Deprecated", "Base Object (Deprecated)");
 	}
 }
 
@@ -727,13 +717,35 @@ void UCustomizableObjectNodeObject::OnRenameNode(const FString& NewName)
 {
 	if (!NewName.IsEmpty())
 	{
-		ObjectName = NewName;
+		SetObjectName(NewName);
+	}
+}
+
+
+void UCustomizableObjectNodeObject::PinConnectionListChanged(UEdGraphPin* Pin)
+{
+	if (Pin == NamePin.Get())
+	{
+		GetGraph()->NotifyGraphChanged();
 	}
 }
 
 
-FString UCustomizableObjectNodeObject::GetObjectName(TArray<const UCustomizableObjectNodeMacroInstance*>* MacroContext ) const
+FString UCustomizableObjectNodeObject::GetObjectName(TArray<const UCustomizableObjectNodeMacroInstance*>* MacroContext) const
 {
+	if (const UEdGraphPin* Pin = NamePin.Get())
+	{
+		if (const UEdGraphPin* LinkedPin = FollowInputPin(*Pin))
+		{
+			const UEdGraphPin* StringPin = GraphTraversal::FindIOPinSourceThroughMacroContext(*LinkedPin, MacroContext);
+
+			if (const UCustomizableObjectNodeStaticString* StringNode = StringPin ? Cast<UCustomizableObjectNodeStaticString>(StringPin->GetOwningNode()) : nullptr)
+			{
+				return StringNode->Value;
+			}
+		}
+	}
+
 	return ObjectName;
 }
 
diff --git a/Engine/Plugins/Mutable/Source/CustomizableObjectEditor/Private/MuCOE/Nodes/CustomizableObjectNodeObject.h b/Engine/Plugins/Mutable/Source/CustomizableObjectEditor/Private/MuCOE/Nodes/CustomizableObjectNodeObject.h
--- a/Engine/Plugins/Mutable/Source/CustomizableObjectEditor/Private/MuCOE/Nodes/CustomizableObjectNodeObject.h
+++ b/Engine/Plugins/Mutable/Source/CustomizableObjectEditor/Private/MuCOE/Nodes/CustomizableObjectNodeObject.h
@@ -147,30 +147,36 @@ public:
 
 	UPROPERTY()
 	TArray<FComponentSettings> ComponentSettings_DEPRECATED;
-		
+
+private:
+	UPROPERTY()
+	FEdGraphPinReference NamePin;
+
+public:
 	// UObject interface.
 	virtual void PostEditChangeProperty(FPropertyChangedEvent& PropertyChangedEvent) override;
 
 	// EdGraphNode interface
-	FText GetNodeTitle(ENodeTitleType::Type TitleType) const override;
-	FLinearColor GetNodeTitleColor() const override;
-	FText GetTooltipText() const override;
-	void PrepareForCopying() override;
-	void PostPasteNode() override;
-	void PostDuplicate(bool bDuplicateForPIE) override;
+	virtual FText GetNodeTitle(ENodeTitleType::Type TitleType) const override;
+	virtual FLinearColor GetNodeTitleColor() const override;
+	virtual FText GetTooltipText() const override;
+	virtual void PrepareForCopying() override;
+	virtual void PostPasteNode() override;
+	virtual void PostDuplicate(bool bDuplicateForPIE) override;
 	virtual bool GetCanRenameNode() const override { return true; }
 	virtual void OnRenameNode(const FString& NewName) override;
+	virtual void PinConnectionListChanged(UEdGraphPin* Pin) override;
 
 	// UCustomizableObjectNode interface
 	virtual void BackwardsCompatibleFixup(int32 CustomizableObjectCustomVersion) override;
 	virtual void PostBackwardsCompatibleFixup() override;
 	virtual void AllocateDefaultPins(UCustomizableObjectNodeRemapPins* RemapPins) override;
 	virtual bool IsNodeSupportedInMacros() const override;
-
+	
 	// Own Interface
 	void SetParentObject(UCustomizableObject* CustomizableParentObject);
-	virtual FString GetObjectName(TArray<const UCustomizableObjectNodeMacroInstance*>* MacroContext = nullptr) const;
 	void SetObjectName(const FString& Name);
+	virtual FString GetObjectName(TArray<const UCustomizableObjectNodeMacroInstance*>* MacroContext = nullptr) const;
 
 	// Own interface
 	UPROPERTY()
diff --git a/Engine/Plugins/Mutable/Source/CustomizableObjectEditor/Private/MuCOE/Nodes/CustomizableObjectNodeObjectChild.cpp b/Engine/Plugins/Mutable/Source/CustomizableObjectEditor/Private/MuCOE/Nodes/CustomizableObjectNodeObjectChild.cpp
--- a/Engine/Plugins/Mutable/Source/CustomizableObjectEditor/Private/MuCOE/Nodes/CustomizableObjectNodeObjectChild.cpp
+++ b/Engine/Plugins/Mutable/Source/CustomizableObjectEditor/Private/MuCOE/Nodes/CustomizableObjectNodeObjectChild.cpp
@@ -20,22 +20,7 @@ UCustomizableObjectNodeObjectChild::UCustomizableObjectNodeObjectChild()
 
 FText UCustomizableObjectNodeObjectChild::GetNodeTitle(ENodeTitleType::Type TitleType) const
 {
-	const UEdGraphPin* ChildNamePin = NamePin.Get();
-	FFormatNamedArguments Args;
-	Args.Add(TEXT("ObjectName"), FText::FromString(ObjectName));
-
-	if (TitleType == ENodeTitleType::ListView || (ChildNamePin && ChildNamePin->LinkedTo.Num()))
-	{
-		return FText::Format(LOCTEXT("Child_Object_Title_List", "Child Object"), Args);
-	}
-	else if (TitleType == ENodeTitleType::EditableTitle)
-	{
-		return FText::Format(LOCTEXT("Child_Object_Edit", "{ObjectName}"), Args);
-	}
-	else
-	{
-		return FText::Format(LOCTEXT("Child_Object_Title", "{ObjectName}\nChild Object"), Args);
-	}
+	return LOCTEXT("Child_Object_Title_List", "Child Object");
 }
 
 
@@ -57,55 +42,10 @@ bool UCustomizableObjectNodeObjectChild::CanDuplicateNode() const
 }
 
 
-void UCustomizableObjectNodeObjectChild::PinConnectionListChanged(UEdGraphPin* Pin)
-{
-	if (Pin == NamePin)
-	{
-		GetGraph()->NotifyGraphChanged();
-	}
-}
-
-
-bool UCustomizableObjectNodeObjectChild::GetCanRenameNode() const
-{
-	const UEdGraphPin* ChildNamePin = NamePin.Get();
-
-	if (ChildNamePin && ChildNamePin->LinkedTo.Num())
-	{
-		return false;
-	}
-
-	return true;
-}
-
-
 FText UCustomizableObjectNodeObjectChild::GetTooltipText() const
 {
 	return LOCTEXT("Child_Object_Tooltip",
 		"Defines a customizable object children in the same asset as its parent, to ease the addition of small Customizable Objects directly into\ntheir parents asset. Functionally equivalent to the Base Object Node when it has a parent defined. It can be a children of the root\nobject or of any children, allowing arbitrary nesting of objects. Defines materials that can be added to its parent, modify it, remove\nparts of it or change any of its parameters. Also defines properties for others to use or modify.");
-
-}
-
-
-void UCustomizableObjectNodeObjectChild::AllocateDefaultPins(UCustomizableObjectNodeRemapPins* RemapPins)
-{
-	Super::AllocateDefaultPins(RemapPins);
-
-	NamePin = CustomCreatePin(EGPD_Input, UEdGraphSchema_CustomizableObject::PC_String, "Name");
-}
-
-
-void UCustomizableObjectNodeObjectChild::BackwardsCompatibleFixup(int32 CustomizableObjectCustomVersion)
-{
-	Super::BackwardsCompatibleFixup(CustomizableObjectCustomVersion);
-
-	if (CustomizableObjectCustomVersion == FCustomizableObjectCustomVersion::EnableMutableMacrosNewVersion)
-	{
-		if (!NamePin.Get())
-		{
-			NamePin = CustomCreatePin(EGPD_Input, UEdGraphSchema_CustomizableObject::PC_String, "Name");
-		}
-	}
 }
 
 
@@ -115,23 +55,5 @@ bool UCustomizableObjectNodeObjectChild::IsNodeSupportedInMacros() const
 }
 
 
-FString UCustomizableObjectNodeObjectChild::GetObjectName(TArray<const UCustomizableObjectNodeMacroInstance*>* MacroContext) const
-{
-	if (NamePin.Get())
-	{
-		if (const UEdGraphPin* LinkedPin = FollowInputPin(*NamePin.Get()))
-		{
-			const UEdGraphPin* StringPin = GraphTraversal::FindIOPinSourceThroughMacroContext(*LinkedPin, MacroContext);
-
-			if (const UCustomizableObjectNodeStaticString* StringNode = StringPin ? Cast<UCustomizableObjectNodeStaticString>(StringPin->GetOwningNode()) : nullptr)
-			{
-				return StringNode->Value;
-			}
-		}
-	}
-
-	return ObjectName;
-}
-
 
 #undef LOCTEXT_NAMESPACE
diff --git a/Engine/Plugins/Mutable/Source/CustomizableObjectEditor/Private/MuCOE/Nodes/CustomizableObjectNodeObjectChild.h b/Engine/Plugins/Mutable/Source/CustomizableObjectEditor/Private/MuCOE/Nodes/CustomizableObjectNodeObjectChild.h
--- a/Engine/Plugins/Mutable/Source/CustomizableObjectEditor/Private/MuCOE/Nodes/CustomizableObjectNodeObjectChild.h
+++ b/Engine/Plugins/Mutable/Source/CustomizableObjectEditor/Private/MuCOE/Nodes/CustomizableObjectNodeObjectChild.h
@@ -16,28 +16,15 @@ public:
 	GENERATED_BODY()
 
 	UCustomizableObjectNodeObjectChild();
-
-	UPROPERTY()
-	FEdGraphPinReference NamePin;
-
-public:
-
+	
 	// UEdGraphNode interface
 	virtual FText GetNodeTitle(ENodeTitleType::Type TitleType) const override;
 	virtual void PrepareForCopying() override;
 	virtual FText GetTooltipText() const override;
 	virtual bool CanUserDeleteNode() const override;
 	virtual bool CanDuplicateNode() const override;
-	virtual void PinConnectionListChanged(UEdGraphPin* Pin) override;
-	virtual bool GetCanRenameNode() const;
-
+	
 	// UCustomizableObjectNode interface
 	virtual bool IsNodeSupportedInMacros() const override;
-	virtual void AllocateDefaultPins(UCustomizableObjectNodeRemapPins* RemapPins) override;
-	virtual void BackwardsCompatibleFixup(int32 CustomizableObjectCustomVersion) override;
-
-	// UCustomizableObjectNodeObject interface
-	virtual FString GetObjectName(TArray<const UCustomizableObjectNodeMacroInstance*>* MacroContext = nullptr) const override;
-
 };
 
diff --git a/Engine/Plugins/Mutable/Source/CustomizableObjectEditor/Private/MuCOE/Nodes/CustomizableObjectNodeObjectGroup.cpp b/Engine/Plugins/Mutable/Source/CustomizableObjectEditor/Private/MuCOE/Nodes/CustomizableObjectNodeObjectGroup.cpp
--- a/Engine/Plugins/Mutable/Source/CustomizableObjectEditor/Private/MuCOE/Nodes/CustomizableObjectNodeObjectGroup.cpp
+++ b/Engine/Plugins/Mutable/Source/CustomizableObjectEditor/Private/MuCOE/Nodes/CustomizableObjectNodeObjectGroup.cpp
@@ -88,25 +88,7 @@ void UCustomizableObjectNodeObjectGroup::PostEditChangeProperty(FPropertyChanged
 
 FText UCustomizableObjectNodeObjectGroup::GetNodeTitle(ENodeTitleType::Type TitleType) const
 {
-	const UEdGraphPin* GroupNamePin = NamePin.Get();
-	FText NodeTitle;
-	FFormatNamedArguments Args;
-	Args.Add(TEXT("GroupName"), FText::FromString(GroupName));
-
-	if (TitleType == ENodeTitleType::ListView || (GroupNamePin && GroupNamePin->LinkedTo.Num()))
-	{
-		NodeTitle = LOCTEXT("Object_Group", "Object Group");
-	}
-	else if(TitleType == ENodeTitleType::EditableTitle)
-	{
-		NodeTitle = LOCTEXT("ObjectGroup_Edit", "{GroupName}");
-	}
-	else
-	{
-		NodeTitle = LOCTEXT("Group_Object_Title", "{GroupName}\nObject Group");
-	}
-
-	return FText::Format(NodeTitle, Args);
+	return LOCTEXT("Object_Group", "Object Group");;
 }
 
 
@@ -132,22 +114,9 @@ void UCustomizableObjectNodeObjectGroup::OnRenameNode(const FString& NewName)
 }
 
 
-bool UCustomizableObjectNodeObjectGroup::GetCanRenameNode() const
-{
-	const UEdGraphPin* GroupNamePin = NamePin.Get();
-
-	if (GroupNamePin && GroupNamePin->LinkedTo.Num())
-	{
-		return false;
-	}
-
-	return true;
-}
-
-
 void UCustomizableObjectNodeObjectGroup::PinConnectionListChanged(UEdGraphPin* Pin)
 {
-	if (Pin == NamePin)
+	if (Pin == NamePin.Get())
 	{
 		GetGraph()->NotifyGraphChanged();
 	}
@@ -158,15 +127,13 @@ void UCustomizableObjectNodeObjectGroup::AllocateDefaultPins(UCustomizableObject
 {
 	const UEdGraphSchema_CustomizableObject* Schema = GetDefault<UEdGraphSchema_CustomizableObject>();
 
-	UEdGraphPin* ObjectsPin = CustomCreatePin(EGPD_Input, Schema->PC_Object, FName("Objects"), true);
-	ObjectsPin->bDefaultValueIsIgnored = true;
+	NamePin = CustomCreatePin(EGPD_Input, UEdGraphSchema_CustomizableObject::PC_String, FName("Name"));
 
-	UEdGraphPin* ProjectorsPin = CustomCreatePin(EGPD_Input, Schema->PC_GroupProjector, FName("Projectors"), true);
-	ObjectsPin->bDefaultValueIsIgnored = true;
+	CustomCreatePin(EGPD_Input, Schema->PC_Object, FName("Objects"), true);
+	
+	CustomCreatePin(EGPD_Input, Schema->PC_GroupProjector, FName("Projectors"), true);
 
 	CustomCreatePin(EGPD_Output, Schema->PC_Object, FName("Group"));
-
-	NamePin = CustomCreatePin(EGPD_Input, UEdGraphSchema_CustomizableObject::PC_String, FName("Name"));
 }
 
 
diff --git a/Engine/Plugins/Mutable/Source/CustomizableObjectEditor/Private/MuCOE/Nodes/CustomizableObjectNodeObjectGroup.h b/Engine/Plugins/Mutable/Source/CustomizableObjectEditor/Private/MuCOE/Nodes/CustomizableObjectNodeObjectGroup.h
--- a/Engine/Plugins/Mutable/Source/CustomizableObjectEditor/Private/MuCOE/Nodes/CustomizableObjectNodeObjectGroup.h
+++ b/Engine/Plugins/Mutable/Source/CustomizableObjectEditor/Private/MuCOE/Nodes/CustomizableObjectNodeObjectGroup.h
@@ -49,7 +49,6 @@ public:
 	virtual FText GetNodeTitle(ENodeTitleType::Type TitleType) const override;
 	virtual FLinearColor GetNodeTitleColor() const override;
 	virtual FText GetTooltipText() const override;
-	virtual bool GetCanRenameNode() const override;
 	virtual void OnRenameNode(const FString& NewName) override;
 	virtual void PinConnectionListChanged(UEdGraphPin* Pin) override;
 
diff --git a/Engine/Plugins/Mutable/Source/CustomizableObjectEditor/Private/MuCOE/Nodes/CustomizableObjectNodeParameter.cpp b/Engine/Plugins/Mutable/Source/CustomizableObjectEditor/Private/MuCOE/Nodes/CustomizableObjectNodeParameter.cpp
--- a/Engine/Plugins/Mutable/Source/CustomizableObjectEditor/Private/MuCOE/Nodes/CustomizableObjectNodeParameter.cpp
+++ b/Engine/Plugins/Mutable/Source/CustomizableObjectEditor/Private/MuCOE/Nodes/CustomizableObjectNodeParameter.cpp
@@ -44,26 +44,10 @@ void UCustomizableObjectNodeParameter::BackwardsCompatibleFixup(int32 Customizab
 
 FText UCustomizableObjectNodeParameter::GetNodeTitle(ENodeTitleType::Type TitleType) const
 {
-	FText NodeTitle;
-	const UEdGraphPin* ParamNamePin = NamePin.Get();
 	FFormatNamedArguments Args;
-	Args.Add(TEXT("ParameterName"), FText::FromString(ParameterName));
 	Args.Add(TEXT("Type"), UEdGraphSchema_CustomizableObject::GetPinCategoryFriendlyName(GetCategory()));
 
-	if (TitleType == ENodeTitleType::ListView || (ParamNamePin && ParamNamePin->LinkedTo.Num()))
-	{
-		NodeTitle = LOCTEXT("ParameterTitle_ListView", "{Type} Parameter");
-	}
-	else if (TitleType == ENodeTitleType::EditableTitle)
-	{
-		NodeTitle = LOCTEXT("ParameterTitle_EditableTitle", "{ParameterName}");
-	}
-	else
-	{
-		NodeTitle = LOCTEXT("ParameterTitle_Title", "{ParameterName}\n{Type} Parameter");
-	}
-
-	return FText::Format(NodeTitle, Args);
+	return FText::Format(LOCTEXT("ParameterTitle_ListView", "{Type} Parameter"), Args);
 }
 
 
@@ -88,22 +72,9 @@ void UCustomizableObjectNodeParameter::OnRenameNode(const FString& NewName)
 }
 
 
-bool UCustomizableObjectNodeParameter::GetCanRenameNode() const
-{
-	const UEdGraphPin* ParamNamePin = NamePin.Get();
-
-	if (ParamNamePin && ParamNamePin->LinkedTo.Num())
-	{
-		return false;
-	}
-
-	return true;
-}
-
-
 void UCustomizableObjectNodeParameter::PinConnectionListChanged(UEdGraphPin* Pin)
 {
-	if (Pin == NamePin)
+	if (Pin == NamePin.Get())
 	{
 		GetGraph()->NotifyGraphChanged();
 	}
diff --git a/Engine/Plugins/Mutable/Source/CustomizableObjectEditor/Private/MuCOE/Nodes/CustomizableObjectNodeParameter.h b/Engine/Plugins/Mutable/Source/CustomizableObjectEditor/Private/MuCOE/Nodes/CustomizableObjectNodeParameter.h
--- a/Engine/Plugins/Mutable/Source/CustomizableObjectEditor/Private/MuCOE/Nodes/CustomizableObjectNodeParameter.h
+++ b/Engine/Plugins/Mutable/Source/CustomizableObjectEditor/Private/MuCOE/Nodes/CustomizableObjectNodeParameter.h
@@ -36,7 +36,6 @@ public:
 	virtual FLinearColor GetNodeTitleColor() const override;
 	virtual FText GetTooltipText() const override;
 	virtual void OnRenameNode(const FString& NewName) override;
-	virtual bool GetCanRenameNode() const override;
 	virtual void PinConnectionListChanged(UEdGraphPin* Pin) override;
 
 	// UCustomizableObjectNode interface
diff --git a/Engine/Plugins/Mutable/Source/CustomizableObjectEditor/Private/MuCOE/Nodes/CustomizableObjectNodeSkeletalMeshParameter.cpp b/Engine/Plugins/Mutable/Source/CustomizableObjectEditor/Private/MuCOE/Nodes/CustomizableObjectNodeSkeletalMeshParameter.cpp
--- a/Engine/Plugins/Mutable/Source/CustomizableObjectEditor/Private/MuCOE/Nodes/CustomizableObjectNodeSkeletalMeshParameter.cpp
+++ b/Engine/Plugins/Mutable/Source/CustomizableObjectEditor/Private/MuCOE/Nodes/CustomizableObjectNodeSkeletalMeshParameter.cpp
@@ -15,7 +15,7 @@ class UCustomizableObjectNodeRemapPins;
 
 void UCustomizableObjectNodeSkeletalMeshParameter::AllocateDefaultPins(UCustomizableObjectNodeRemapPins* RemapPins)
 {
-	USkeletalMesh* SkeletalMesh = DefaultValue.LoadSynchronous();
+	USkeletalMesh* SkeletalMesh = ReferenceValue.LoadSynchronous();
 
 	if (!SkeletalMesh)
 	{
@@ -69,7 +69,7 @@ void UCustomizableObjectNodeSkeletalMeshParameter::PostEditChangeProperty(FPrope
 	Super::PostEditChangeProperty(PropertyChangedEvent);
 	
 	FProperty* PropertyThatChanged = PropertyChangedEvent.Property;
-	if (PropertyThatChanged && PropertyThatChanged->GetName() == GET_MEMBER_NAME_CHECKED(UCustomizableObjectNodeSkeletalMeshParameter, DefaultValue))
+	if (PropertyThatChanged && PropertyThatChanged->GetName() == GET_MEMBER_NAME_CHECKED(UCustomizableObjectNodeSkeletalMeshParameter, ReferenceValue))
 	{
 		ReconstructNode();
 	}
@@ -98,20 +98,7 @@ void UCustomizableObjectNodeSkeletalMeshParameter::BackwardsCompatibleFixup(int3
 
 FText UCustomizableObjectNodeSkeletalMeshParameter::GetNodeTitle(ENodeTitleType::Type TitleType) const
 {
-	const UEdGraphPin* ParamNamePin = NamePin.Get();
-
-	if (TitleType == ENodeTitleType::ListView || (ParamNamePin && ParamNamePin->LinkedTo.Num()))
-	{
-		return LOCTEXT("Skeletal_Mesh_Parameter", "Skeletal Mesh Parameter");
-	}
-	else if (TitleType == ENodeTitleType::EditableTitle)
-	{
-		return FText::Format(LOCTEXT("Skeletal_Mesh_Parameter_EditableTitle", "{0}"), FText::FromString(ParameterName));
-	}
-	else
-	{
-		return FText::Format(LOCTEXT("Skeletal_Mesh_Parameter_Title", "{0}\nSkeletal Mesh Parameter"), FText::FromString(ParameterName));
-	}
+	return LOCTEXT("Skeletal_Mesh_Parameter", "Skeletal Mesh Parameter");
 }
 
 
@@ -137,22 +124,9 @@ void UCustomizableObjectNodeSkeletalMeshParameter::OnRenameNode(const FString& N
 }
 
 
-bool UCustomizableObjectNodeSkeletalMeshParameter::GetCanRenameNode() const
-{
-	const UEdGraphPin* ParamNamePin = NamePin.Get();
-
-	if (ParamNamePin && ParamNamePin->LinkedTo.Num())
-	{
-		return false;
-	}
-
-	return true;
-}
-
-
 void UCustomizableObjectNodeSkeletalMeshParameter::PinConnectionListChanged(UEdGraphPin* Pin)
 {
-	if (Pin == NamePin)
+	if (Pin == NamePin.Get())
 	{
 		GetGraph()->NotifyGraphChanged();
 	}
@@ -174,7 +148,7 @@ TArray<UCustomizableObjectLayout*> UCustomizableObjectNodeSkeletalMeshParameter:
 
 TSoftObjectPtr<UObject> UCustomizableObjectNodeSkeletalMeshParameter::GetMesh() const
 {
-	return DefaultValue;
+	return ReferenceValue;
 }
 
 
@@ -234,7 +208,7 @@ UMaterialInterface* UCustomizableObjectNodeSkeletalMeshParameter::GetMaterialInt
 
 FSkeletalMaterial* UCustomizableObjectNodeSkeletalMeshParameter::GetSkeletalMaterialFor(const int32 SectionIndex) const
 {
-	USkeletalMesh* SkeletalMesh = DefaultValue.LoadSynchronous();
+	USkeletalMesh* SkeletalMesh = ReferenceValue.LoadSynchronous();
 
 	if (!SkeletalMesh)
 	{
@@ -252,7 +226,7 @@ FSkeletalMaterial* UCustomizableObjectNodeSkeletalMeshParameter::GetSkeletalMate
 
 int32 UCustomizableObjectNodeSkeletalMeshParameter::GetSkeletalMaterialIndexFor(const int32 SectionIndex) const
 {
-	USkeletalMesh* SkeletalMesh = DefaultValue.LoadSynchronous();
+	USkeletalMesh* SkeletalMesh = ReferenceValue.LoadSynchronous();
 
 	if (!SkeletalMesh)
 	{
diff --git a/Engine/Plugins/Mutable/Source/CustomizableObjectEditor/Private/MuCOE/Nodes/CustomizableObjectNodeSkeletalMeshParameter.h b/Engine/Plugins/Mutable/Source/CustomizableObjectEditor/Private/MuCOE/Nodes/CustomizableObjectNodeSkeletalMeshParameter.h
--- a/Engine/Plugins/Mutable/Source/CustomizableObjectEditor/Private/MuCOE/Nodes/CustomizableObjectNodeSkeletalMeshParameter.h
+++ b/Engine/Plugins/Mutable/Source/CustomizableObjectEditor/Private/MuCOE/Nodes/CustomizableObjectNodeSkeletalMeshParameter.h
@@ -51,6 +51,9 @@ public:
 	TSoftObjectPtr<USkeletalMesh> DefaultValue;
 
 	UPROPERTY(EditAnywhere, Category=CustomizableObject)
+	TSoftObjectPtr<USkeletalMesh> ReferenceValue;
+	
+	UPROPERTY()
 	FString ParameterName = "Mesh Param";
 
 	UPROPERTY(EditAnywhere, Category = UI, meta = (DisplayName = "Parameter UI Metadata"))
@@ -69,7 +72,6 @@ public:
 	virtual FLinearColor GetNodeTitleColor() const override;
 	virtual FText GetTooltipText() const override;
 	virtual void OnRenameNode(const FString& NewName) override;
-	virtual bool GetCanRenameNode() const override;
 	virtual void PinConnectionListChanged(UEdGraphPin* Pin) override;
 
 	// UCustomizableObjectNode interface
diff --git a/Engine/Plugins/Mutable/Source/CustomizableObjectEditor/Private/MuCOE/Nodes/CustomizableObjectNodeSkeletalMeshParametersDetails.cpp b/Engine/Plugins/Mutable/Source/CustomizableObjectEditor/Private/MuCOE/Nodes/CustomizableObjectNodeSkeletalMeshParametersDetails.cpp
--- a/Engine/Plugins/Mutable/Source/CustomizableObjectEditor/Private/MuCOE/Nodes/CustomizableObjectNodeSkeletalMeshParametersDetails.cpp
+++ b/Engine/Plugins/Mutable/Source/CustomizableObjectEditor/Private/MuCOE/Nodes/CustomizableObjectNodeSkeletalMeshParametersDetails.cpp
@@ -36,7 +36,7 @@ void FCustomizableObjectNodeSkeletalMeshParameterDetails::CustomizeDetails(IDeta
 
     IDetailCategoryBuilder& CustomizableObject = DetailBuilder.EditCategory("CustomizableObject");
 
-	TSharedRef<IPropertyHandle> DefaultValueProperty = DetailBuilder.GetProperty(GET_MEMBER_NAME_CHECKED(UCustomizableObjectNodeSkeletalMeshParameter, DefaultValue));
+	TSharedRef<IPropertyHandle> DefaultValueProperty = DetailBuilder.GetProperty(GET_MEMBER_NAME_CHECKED(UCustomizableObjectNodeSkeletalMeshParameter, ReferenceValue));
 
 	TArray<FLayoutEditorMeshSection> MeshSectionsAndLayouts;
 	GenerateMeshSectionOptions(MeshSectionsAndLayouts);
diff --git a/Engine/Plugins/Mutable/Source/CustomizableObjectEditor/Private/MuCOE/Nodes/CustomizableObjectNodeStaticString.cpp b/Engine/Plugins/Mutable/Source/CustomizableObjectEditor/Private/MuCOE/Nodes/CustomizableObjectNodeStaticString.cpp
--- a/Engine/Plugins/Mutable/Source/CustomizableObjectEditor/Private/MuCOE/Nodes/CustomizableObjectNodeStaticString.cpp
+++ b/Engine/Plugins/Mutable/Source/CustomizableObjectEditor/Private/MuCOE/Nodes/CustomizableObjectNodeStaticString.cpp
@@ -26,12 +26,6 @@ FLinearColor UCustomizableObjectNodeStaticString::GetNodeTitleColor() const
 }
 
 
-bool UCustomizableObjectNodeStaticString::GetCanRenameNode() const
-{
-	return false;
-}
-
-
 void UCustomizableObjectNodeStaticString::AllocateDefaultPins(UCustomizableObjectNodeRemapPins* RemapPins)
 {
 	const UEdGraphSchema_CustomizableObject* Schema = GetDefault<UEdGraphSchema_CustomizableObject>();
diff --git a/Engine/Plugins/Mutable/Source/CustomizableObjectEditor/Private/MuCOE/Nodes/CustomizableObjectNodeStaticString.h b/Engine/Plugins/Mutable/Source/CustomizableObjectEditor/Private/MuCOE/Nodes/CustomizableObjectNodeStaticString.h
--- a/Engine/Plugins/Mutable/Source/CustomizableObjectEditor/Private/MuCOE/Nodes/CustomizableObjectNodeStaticString.h
+++ b/Engine/Plugins/Mutable/Source/CustomizableObjectEditor/Private/MuCOE/Nodes/CustomizableObjectNodeStaticString.h
@@ -23,8 +23,7 @@ public:
 	virtual FText GetNodeTitle(ENodeTitleType::Type TitleType) const override;
 	virtual FText GetTooltipText() const override;
 	virtual FLinearColor GetNodeTitleColor() const override;
-	virtual bool GetCanRenameNode() const override;
-
+	
 	// UCustomizableObjectNode interface
 	void AllocateDefaultPins(UCustomizableObjectNodeRemapPins* RemapPins) override;
 
diff --git a/Engine/Plugins/Mutable/Source/CustomizableObjectEditor/Private/MuCOE/Nodes/CustomizableObjectNodeTunnel.cpp b/Engine/Plugins/Mutable/Source/CustomizableObjectEditor/Private/MuCOE/Nodes/CustomizableObjectNodeTunnel.cpp
--- a/Engine/Plugins/Mutable/Source/CustomizableObjectEditor/Private/MuCOE/Nodes/CustomizableObjectNodeTunnel.cpp
+++ b/Engine/Plugins/Mutable/Source/CustomizableObjectEditor/Private/MuCOE/Nodes/CustomizableObjectNodeTunnel.cpp
@@ -72,12 +72,6 @@ FLinearColor UCustomizableObjectNodeTunnel::GetNodeTitleColor() const
 }
 
 
-bool UCustomizableObjectNodeTunnel::GetCanRenameNode() const
-{
-	return false;
-}
-
-
 bool UCustomizableObjectNodeTunnel::CanUserDeleteNode() const
 {
 	return false;
diff --git a/Engine/Plugins/Mutable/Source/CustomizableObjectEditor/Private/MuCOE/Nodes/CustomizableObjectNodeTunnel.h b/Engine/Plugins/Mutable/Source/CustomizableObjectEditor/Private/MuCOE/Nodes/CustomizableObjectNodeTunnel.h
--- a/Engine/Plugins/Mutable/Source/CustomizableObjectEditor/Private/MuCOE/Nodes/CustomizableObjectNodeTunnel.h
+++ b/Engine/Plugins/Mutable/Source/CustomizableObjectEditor/Private/MuCOE/Nodes/CustomizableObjectNodeTunnel.h
@@ -60,7 +60,6 @@ public:
 	virtual FText GetNodeTitle(ENodeTitleType::Type TitleType) const override;
 	virtual FText GetTooltipText() const override;
 	virtual FLinearColor GetNodeTitleColor() const override;
-	virtual bool GetCanRenameNode() const override;
 	virtual bool CanUserDeleteNode() const override;
 	virtual bool CanDuplicateNode() const override;
 
diff --git a/Engine/Plugins/Mutable/Source/CustomizableObjectEditor/Private/MuCOE/Nodes/SCustomizableObjectNodeStaticStringPin.cpp b/Engine/Plugins/Mutable/Source/CustomizableObjectEditor/Private/MuCOE/Nodes/SCustomizableObjectNodeStaticStringPin.cpp
--- a/Engine/Plugins/Mutable/Source/CustomizableObjectEditor/Private/MuCOE/Nodes/SCustomizableObjectNodeStaticStringPin.cpp
+++ b/Engine/Plugins/Mutable/Source/CustomizableObjectEditor/Private/MuCOE/Nodes/SCustomizableObjectNodeStaticStringPin.cpp
@@ -57,9 +57,9 @@ FText SCustomizableObjectNodeStaticStringPin::GetNodeStringValue() const
 	{
 		return FText::FromString(MeshParamNode->ParameterName);
 	}
-	else if (const UCustomizableObjectNodeObjectChild* ChildObjectNode = Cast<UCustomizableObjectNodeObjectChild>(Node))
+	else if (const UCustomizableObjectNodeObject* ObjectNode = Cast<UCustomizableObjectNodeObject>(Node))
 	{
-		return FText::FromString(ChildObjectNode->GetObjectName());
+		return FText::FromString(ObjectNode->GetObjectName());
 	}
 	else if (const UCustomizableObjectNodeObjectGroup* GroupNode = Cast<UCustomizableObjectNodeObjectGroup>(Node))
 	{
@@ -79,7 +79,11 @@ FText SCustomizableObjectNodeStaticStringPin::GetNodeStringValue() const
 			}
 		}
 	}
-
+	else
+	{
+		unimplemented();
+	}
+	
 	return FText();
 }
 
@@ -104,9 +108,9 @@ void SCustomizableObjectNodeStaticStringPin::OnTextCommited(const FText& InValue
 	{
 		MeshParamNode->ParameterName = InValue.ToString();
 	}
-	else if (UCustomizableObjectNodeObjectChild* ChildObjectNode = Cast<UCustomizableObjectNodeObjectChild>(Node))
+	else if (UCustomizableObjectNodeObject* ObjectNode = Cast<UCustomizableObjectNodeObject>(Node))
 	{
-		ChildObjectNode->SetObjectName(InValue.ToString());
+		ObjectNode->SetObjectName(InValue.ToString());
 	}
 	else if (UCustomizableObjectNodeObjectGroup* GroupNode = Cast<UCustomizableObjectNodeObjectGroup>(Node))
 	{
@@ -128,6 +132,10 @@ void SCustomizableObjectNodeStaticStringPin::OnTextCommited(const FText& InValue
 			}
 		}
 	}
+	else
+	{
+		unimplemented();
+	}
 }
 
 
diff --git a/Engine/Plugins/Mutable/Source/CustomizableObjectEditor/Public/MuCOE/Nodes/CustomizableObjectNode.h b/Engine/Plugins/Mutable/Source/CustomizableObjectEditor/Public/MuCOE/Nodes/CustomizableObjectNode.h
--- a/Engine/Plugins/Mutable/Source/CustomizableObjectEditor/Public/MuCOE/Nodes/CustomizableObjectNode.h
+++ b/Engine/Plugins/Mutable/Source/CustomizableObjectEditor/Public/MuCOE/Nodes/CustomizableObjectNode.h
@@ -101,6 +101,7 @@ public:
 	virtual void PinConnectionListChanged(UEdGraphPin* Pin) override;
 	virtual void DestroyNode() override;
 	virtual TSharedPtr<SGraphNode> CreateVisualWidget() override;
+	virtual bool GetCanRenameNode() const override;
 
 	/** Allocates the default pins using the empty remap pins action. Usually called from CreateNode. */
 	virtual void AllocateDefaultPins() override final; // Final. Override AllocateDefaultPins(UCustomizableObjectNodeRemapPins* RemapPins) instead
