[Mutable] Fixed recurrent crash of unimplemented if case.

#rnx

diff --git a/Engine/Plugins/Mutable/Source/CustomizableObjectEditor/Private/MuCOE/GraphTraversal.cpp b/Engine/Plugins/Mutable/Source/CustomizableObjectEditor/Private/MuCOE/GraphTraversal.cpp
--- a/Engine/Plugins/Mutable/Source/CustomizableObjectEditor/Private/MuCOE/GraphTraversal.cpp
+++ b/Engine/Plugins/Mutable/Source/CustomizableObjectEditor/Private/MuCOE/GraphTraversal.cpp
@@ -793,6 +793,16 @@ const UEdGraphPin* FindMeshBaseSource(const UEdGraphPin& Pin, const bool bOnlyLo
 		}
 	}
 
+	else if (const UCustomizableObjectNodeTunnel* NodeTunnel = Cast<UCustomizableObjectNodeTunnel>(Node))
+	{
+		const UEdGraphPin* ConnectedPin = GraphTraversal::FindIOPinSourceThroughMacroContext(Pin, MacroContext);
+
+		if (ConnectedPin)
+		{
+			return FindMeshBaseSource(*ConnectedPin, bOnlyLookForStaticMesh, MacroContext);
+		}
+	}
+
 	else
 	{
 		unimplemented(); // Case missing.
