[Mutable] Optimize ImageMakeGrowMap parallel for batching.

#rnx

diff --git a/Engine/Plugins/Mutable/Source/MutableRuntime/Private/MuR/CodeRunner.cpp b/Engine/Plugins/Mutable/Source/MutableRuntime/Private/MuR/CodeRunner.cpp
--- a/Engine/Plugins/Mutable/Source/MutableRuntime/Private/MuR/CodeRunner.cpp
+++ b/Engine/Plugins/Mutable/Source/MutableRuntime/Private/MuR/CodeRunner.cpp
@@ -5018,10 +5018,13 @@ namespace mu
 
                 TSharedPtr<FImage> Result = CreateImage( Mask->GetSizeX(), Mask->GetSizeY(), Mask->GetLODCount(), EImageFormat::L_UByte, EInitializationType::NotInitialized);
 
-                ImageMakeGrowMap(Result.Get(), Mask.Get(), Args.border );
+				TSharedPtr<FImage> OwnedMask = CloneOrTakeOver(Mask);
+
+                ImageMakeGrowMap(Result.Get(), OwnedMask.Get(), Args.border);
 				Result->Flags |= FImage::IF_CANNOT_BE_SCALED;
 
 				Release(Mask);
+				Release(OwnedMask);
                 StoreImage( Item, Result);
                 break;
             }
diff --git a/Engine/Plugins/Mutable/Source/MutableRuntime/Private/MuR/OpImageDisplace.h b/Engine/Plugins/Mutable/Source/MutableRuntime/Private/MuR/OpImageDisplace.h
--- a/Engine/Plugins/Mutable/Source/MutableRuntime/Private/MuR/OpImageDisplace.h
+++ b/Engine/Plugins/Mutable/Source/MutableRuntime/Private/MuR/OpImageDisplace.h
@@ -8,7 +8,6 @@
 #include "MuR/ConvertData.h"
 #include "Async/ParallelFor.h"
 
-
 namespace mu
 {
 
@@ -32,7 +31,7 @@ namespace mu
 	//---------------------------------------------------------------------------------------------
 	//!
 	//---------------------------------------------------------------------------------------------
-	inline void ImageMakeGrowMap(FImage* ResultImage, const FImage* MaskImage, int32 InBorder)
+	inline void ImageMakeGrowMap(FImage* ResultImage, FImage* MaskImage, int32 InBorder)
 	{
 		MUTABLE_CPUPROFILER_SCOPE(ImageMakeGrowMap);
 
@@ -57,81 +56,105 @@ namespace mu
 			FMemory::Memset(Buffer.GetData(), ZeroOffset, Buffer.Num());
 		}
 
-		TSharedPtr<const FImage> ThisMaskCopy;
-		const FImage* ThisMask = MaskImage;
+		TSharedPtr<FImage> NextMask = MaskImage->Clone();
+
 		for (int32 B = 0; B < InBorder; ++B)
 		{
-			TSharedPtr<FImage> NextMask = ThisMask->Clone();
+			if (B > 0)
+			{
+				MaskImage->DataStorage = NextMask->DataStorage;
+			}
 
-			for(int32 CurrentMip = 0; CurrentMip < MipCount; ++CurrentMip)
-			{ 
-				FIntVector2 MipSize = ResultImage->CalculateMipSize(CurrentMip);
+			// TODO: Replace with a single parallel for.
+			const auto ProcessMip = [MaskImage, NextMask, ResultImage](int32 MipIndex)
+				{
+					uint8* ResultData = ResultImage->GetMipData(MipIndex);
+					const uint8* MaskData = MaskImage->GetMipData(MipIndex);
+					uint8* NextMaskData = NextMask->GetMipData(MipIndex);
 
-				uint8* ResultData = ResultImage->GetMipData(CurrentMip);
-				const uint8* ThisMaskData = ThisMask->GetMipData(CurrentMip);
-				uint8* NextMaskData = NextMask->GetMipData(CurrentMip);
+					FIntVector2 MipSize = ResultImage->CalculateMipSize(MipIndex);
 
-				//for (int32 Y=0; Y< MipSize.Y; ++Y)
-				const auto ProcessRow = [
-					ThisMaskData, NextMaskData, MipSize, ResultData
-				] (int32 Y)
-				{
-					uint8* NextData = NextMaskData + MipSize.X * Y;
+					const int32 NumRowsPerBatch = FMath::Min(FMath::DivideAndRoundUp((1 << 14), MipSize.X), MipSize.Y);
+					const int32 NumBatches = FMath::DivideAndRoundUp<int32>(MipSize.Y, NumRowsPerBatch);
 
-					const uint8* S0 = Y > 0 ? ThisMaskData + (Y - 1) * MipSize.X : nullptr;
-					const uint8* S1 = ThisMaskData + Y * MipSize.X;
-					bool bNotLastRow = Y < MipSize.Y - 1;
-					const uint8* S2 = bNotLastRow ? ThisMaskData + (Y + 1) * MipSize.X : nullptr;
+					const auto ProcessBatch = [
+						MaskData, NextMaskData, MipSize, ResultData, NumBatches, NumRowsPerBatch
+					] (int32 BatchIndex)
+						{
+							const int32 BatchBeginRowY = BatchIndex * NumRowsPerBatch;
+							const int32 BatchEndRowY = FMath::Min(BatchBeginRowY + NumRowsPerBatch, MipSize.Y);
 
-					const uint8* R0 = Y > 0 ? ResultData + (Y - 1) * MipSize.X : 0;
-					uint8* R1 = ResultData + Y * MipSize.X;
+							uint8* NextData = NextMaskData + MipSize.X * BatchBeginRowY;
 
-					const uint8* R2 = bNotLastRow ? ResultData + (Y + 1) * MipSize.X : nullptr;
+							for (int32 Y = BatchBeginRowY; Y < BatchEndRowY; ++Y)
+							{
+								const uint8* S0 = Y > 0 ? MaskData + (Y - 1) * MipSize.X : nullptr;
+								const uint8* S1 = MaskData + Y * MipSize.X;
+								bool bNotLastRow = Y < MipSize.Y - 1;
+								const uint8* S2 = bNotLastRow ? MaskData + (Y + 1) * MipSize.X : nullptr;
+
+								const uint8* R0 = Y > 0 ? ResultData + (Y - 1) * MipSize.X : 0;
+								uint8* R1 = ResultData + Y * MipSize.X;
+
+								const uint8* R2 = bNotLastRow ? ResultData + (Y + 1) * MipSize.X : nullptr;
+
+								for (int32 X = 0; X < MipSize.X; ++X)
+								{
+									int32 Dx, Dy;
+									if (S1[X])
+									{
+										*NextData = 255;
+									}
+									else if (Y > 0 && S0[X])
+									{
+										MutableDecodeOffset(R0[X], Dx, Dy);
+										R1[X] = MutableEncodeOffset(Dx, Dy - 1);
+										*NextData = 255;
+									}
+									else if (bNotLastRow && S2 && R2 && S2[X])
+									{
+										MutableDecodeOffset(R2[X], Dx, Dy);
+										R1[X] = MutableEncodeOffset(Dx, Dy + 1);
+										*NextData = 255;
+									}
+									else if (X < MipSize.X - 1 && S1[X + 1])
+									{
+										MutableDecodeOffset(R1[X + 1], Dx, Dy);
+										R1[X] = MutableEncodeOffset(Dx + 1, Dy);
+										*NextData = 255;
+									}
+									else if (X > 0 && S1[X - 1])
+									{
+										MutableDecodeOffset(R1[X - 1], Dx, Dy);
+										R1[X] = MutableEncodeOffset(Dx - 1, Dy);
+										*NextData = 255;
+									}
+
+									++NextData;
+								}
+							}
 
-					for (int32 X = 0; X < MipSize.X; ++X)
-					{
-						int32 Dx, Dy;
-						if (S1[X])
-						{
-							*NextData = 255;
-						}
-						else if (Y > 0 && S0[X])
-						{
-							MutableDecodeOffset(R0[X], Dx, Dy);
-							R1[X] = MutableEncodeOffset(Dx, Dy - 1);
-							*NextData = 255;
-						}
-						else if (bNotLastRow && S2 && R2 && S2[X])
-						{
-							MutableDecodeOffset(R2[X], Dx, Dy);
-							R1[X] = MutableEncodeOffset(Dx, Dy + 1);
-							*NextData = 255;
-						}
-						else if (X < MipSize.X - 1 && S1[X + 1])
+						};
+
+						if (NumBatches == 1)
 						{
-							MutableDecodeOffset(R1[X + 1], Dx, Dy);
-							R1[X] = MutableEncodeOffset(Dx + 1, Dy);
-							*NextData = 255;
+							ProcessBatch(0);
 						}
-						else if (X > 0 && S1[X - 1])
+						else
 						{
-							MutableDecodeOffset(R1[X - 1], Dx, Dy);
-							R1[X] = MutableEncodeOffset(Dx - 1, Dy);
-							*NextData = 255;
+							ParallelFor(NumBatches, ProcessBatch);
 						}
-
-						++NextData;
-					}
-
 				};
 
-				ParallelFor(MipSize.Y, ProcessRow);
+			if (MipCount == 1)
+			{
+				ProcessMip(0);
+			}
+			else
+			{
+				ParallelFor(MipCount, ProcessMip);
 			}
-
-			ThisMaskCopy = NextMask;
-			ThisMask = ThisMaskCopy.Get();
 		}
-
 	}
 
 
