[mutable] Fixed order of the "Clip Mesh with Plane and Morph" node properties (descriptor)
- New order follows the expected expansion behavior for UI content (left to right and top to bottom) where the Component Name goes on top and the "BoneName" dropdown goes under it (and not the other way around).
- Updated Viewport API to support hiding a Clip Morph gizmo without deselecting the currently selected nodes.
- Fixed the incapability of the resetting of only the name of the component being used.
- Fixed resetting behavior for the "BoneName" property. Previously it was resetting also the component name, making the "BoneName" dropdown disappear in the process.
- Include cleanup.
- Changes in the node the descriptor represents. Now the node is the one who handles its data when the properties change during the runtime.

diff --git a/Engine/Plugins/Mutable/Source/CustomizableObjectEditor/Private/MuCOE/CustomizableObjectEditor.cpp b/Engine/Plugins/Mutable/Source/CustomizableObjectEditor/Private/MuCOE/CustomizableObjectEditor.cpp
--- a/Engine/Plugins/Mutable/Source/CustomizableObjectEditor/Private/MuCOE/CustomizableObjectEditor.cpp
+++ b/Engine/Plugins/Mutable/Source/CustomizableObjectEditor/Private/MuCOE/CustomizableObjectEditor.cpp
@@ -812,7 +812,7 @@ void FCustomizableObjectEditor::ShowGizmoClipMorph(UCustomizableObjectNodeModifi
 }
 
 
-void FCustomizableObjectEditor::HideGizmoClipMorph()
+void FCustomizableObjectEditor::HideGizmoClipMorph(bool bClearGraphSelectionSet /**= true*/)
 {
 	if (GizmoType != EGizmoType::ClipMorph)
 	{
@@ -823,16 +823,19 @@ void FCustomizableObjectEditor::HideGizmoClipMorph()
 
 	Viewport->HideGizmoClipMorph();
 
-	const FGraphPanelSelectionSet SelectedNodes = GraphEditor->GetSelectedNodes();
-	for (FGraphPanelSelectionSet::TConstIterator NodeIt( SelectedNodes ); NodeIt; ++NodeIt)
+	if (bClearGraphSelectionSet)
 	{
-		const UObject* Node = *NodeIt;
-		if (Node->IsA<UCustomizableObjectNodeModifierClipMorph>())
+		const FGraphPanelSelectionSet SelectedNodes = GraphEditor->GetSelectedNodes();
+		for (FGraphPanelSelectionSet::TConstIterator NodeIt( SelectedNodes ); NodeIt; ++NodeIt)
 		{
-			GraphEditor->ClearSelectionSet();
-			break;
-		}
-	}	
+			const UObject* Node = *NodeIt;
+			if (Node->IsA<UCustomizableObjectNodeModifierClipMorph>())
+			{
+				GraphEditor->ClearSelectionSet();
+				break;
+			}
+		}	
+	}
 }
 
 
diff --git a/Engine/Plugins/Mutable/Source/CustomizableObjectEditor/Private/MuCOE/CustomizableObjectEditor.h b/Engine/Plugins/Mutable/Source/CustomizableObjectEditor/Private/MuCOE/CustomizableObjectEditor.h
--- a/Engine/Plugins/Mutable/Source/CustomizableObjectEditor/Private/MuCOE/CustomizableObjectEditor.h
+++ b/Engine/Plugins/Mutable/Source/CustomizableObjectEditor/Private/MuCOE/CustomizableObjectEditor.h
@@ -129,7 +129,7 @@ public:
 	virtual void ShowGizmoProjectorParameter(const FString& ParamName, int32 RangeIndex = -1) override;
 	virtual void HideGizmoProjectorParameter() override;
 	virtual void ShowGizmoClipMorph(UCustomizableObjectNodeModifierClipMorph& Node) override;
-	virtual void HideGizmoClipMorph() override;
+	virtual void HideGizmoClipMorph(bool bClearGraphSelectionSet = true) override;
 	virtual void ShowGizmoClipMesh(UCustomizableObjectNode& Node, FTransform* Transform, const UEdGraphPin& MeshPin) override;
 	virtual void HideGizmoClipMesh() override;
 	virtual void ShowGizmoLight(ULightComponent& SelectedLight) override;
diff --git a/Engine/Plugins/Mutable/Source/CustomizableObjectEditor/Private/MuCOE/Nodes/CustomizableObjectNodeMeshClipMorph.cpp b/Engine/Plugins/Mutable/Source/CustomizableObjectEditor/Private/MuCOE/Nodes/CustomizableObjectNodeMeshClipMorph.cpp
--- a/Engine/Plugins/Mutable/Source/CustomizableObjectEditor/Private/MuCOE/Nodes/CustomizableObjectNodeMeshClipMorph.cpp
+++ b/Engine/Plugins/Mutable/Source/CustomizableObjectEditor/Private/MuCOE/Nodes/CustomizableObjectNodeMeshClipMorph.cpp
@@ -5,8 +5,11 @@
 #include "MuCO/CustomizableObjectCustomVersion.h"
 #include "MuCOE/EdGraphSchema_CustomizableObject.h"
 #include "MuCOE/CustomizableObjectEditor.h"
+#include "MuCOE/CustomizableObjectGraph.h"
 #include "MuCOE/UnrealEditorPortabilityHelpers.h"
 #include "MuCOE/Nodes/CustomizableObjectNodeObject.h"
+#include "Engine/SkeletalMesh.h"
+#include "MuCOE/CustomizableObjectEditorUtilities.h"
 
 class UCustomizableObjectNodeRemapPins;
 class UEdGraphPin;
@@ -87,6 +90,91 @@ void UCustomizableObjectNodeModifierClipMorph::ChangeStartOffsetTransform()
 }
 
 
+void UCustomizableObjectNodeModifierClipMorph::InvertNormals()
+{
+	if (bLocalStartOffset)
+	{
+		StartOffset.Z *= -1.0f;
+		StartOffset.X *= -1.0f;
+	}
+
+	Normal *= -1.0f;
+}
+
+
+void UCustomizableObjectNodeModifierClipMorph::UpdateOriginAndNormal()
+{
+	// Get all the bone names of the skeletal mesh
+	if (TStrongObjectPtr<USkeletalMesh> SkeletalMesh = GetReferenceSkeletalMesh())
+	{
+		for (int32 BoneIndex = 0; BoneIndex < SkeletalMesh->GetRefSkeleton().GetRawBoneNum(); ++BoneIndex)
+		{
+			const FName RefSkeletonBoneName = SkeletalMesh->GetRefSkeleton().GetBoneName(BoneIndex);
+			
+			// Is the bone set in our property found in the array of bones of the reference skeletal mesh skeleton?
+			if (RefSkeletonBoneName == BoneName)
+			{
+				FVector Location = FVector::ZeroVector;
+
+				const TArray<FTransform>& ReferenceSkeletonBoneArray = SkeletalMesh->GetRefSkeleton().GetRefBonePose();
+				int32 ParentIndex = SkeletalMesh->GetRefSkeleton().FindBoneIndex(BoneName);
+
+				FVector ChildLocation = FVector::ForwardVector;
+
+				for (int32 i = 0; i < SkeletalMesh->GetRefSkeleton().GetRawBoneNum(); ++i)
+				{
+					if (SkeletalMesh->GetRefSkeleton().GetParentIndex(i) == ParentIndex)
+					{
+						ChildLocation = ReferenceSkeletonBoneArray[i].TransformPosition(FVector::ZeroVector);
+						break;
+					}
+				}
+			
+				while (ParentIndex >= 0)
+				{
+					Location = ReferenceSkeletonBoneArray[ParentIndex].TransformPosition(Location);
+					ChildLocation = ReferenceSkeletonBoneArray[ParentIndex].TransformPosition(ChildLocation);
+					ParentIndex = SkeletalMesh->GetRefSkeleton().GetParentIndex(ParentIndex);
+				}
+
+				const FVector Direction = (ChildLocation - Location).GetSafeNormal();
+
+				// Update the properties of the node
+				Origin = Location;
+				Normal = Direction;
+
+				// Invert the normal if declared by the property
+				if (bInvertNormal)
+				{
+					InvertNormals();
+				}
+				
+				break;
+			}
+		}
+	}
+}
+
+
+TStrongObjectPtr<USkeletalMesh> UCustomizableObjectNodeModifierClipMorph::GetReferenceSkeletalMesh() const
+{
+	TStrongObjectPtr<USkeletalMesh> SkeletalMesh = nullptr;
+
+	if (ReferenceSkeletonComponent.IsNone())
+	{
+		return SkeletalMesh;
+	}
+	
+	if (UCustomizableObject* CustomizableObject = Cast<UCustomizableObject>(GetCustomizableObjectGraph()->GetOuter()))
+	{
+		const FName ReferenceComponentName = ReferenceSkeletonComponent;
+		SkeletalMesh = TStrongObjectPtr(CustomizableObject->GetComponentMeshReferenceSkeletalMesh(ReferenceComponentName));
+	}
+
+	return SkeletalMesh;
+}
+
+
 UEdGraphPin* UCustomizableObjectNodeModifierClipMorph::GetOutputPin() const
 {
 	return FindPin(TEXT("Modifier"));
@@ -95,14 +183,56 @@ UEdGraphPin* UCustomizableObjectNodeModifierClipMorph::GetOutputPin() const
 
 void UCustomizableObjectNodeModifierClipMorph::PostEditChangeProperty(FPropertyChangedEvent& PropertyChangedEvent)
 {
-	const FName PropertyName = PropertyChangedEvent.GetPropertyName();
+	TSharedPtr<FCustomizableObjectEditor> Editor = StaticCastSharedPtr<FCustomizableObjectEditor>(GetGraphEditor());
 
-	if (PropertyName == "bLocalStartOffset")
+	// Hide the gizmo in cases where incomplete data is to be found in the node
+	if (Editor)
 	{
-		ChangeStartOffsetTransform();
+		// if no name for the bone has been provided hide the gizmo
+		if (BoneName.IsNone())
+		{
+			Editor->HideGizmoClipMorph(false);
+			return;
+		}
+
+		// Hide gizmo if the component is empty
+		if (ReferenceSkeletonComponent.IsNone())
+		{
+			Editor->HideGizmoClipMorph(false);
+			return;
+		}
+		
+		// Hide gizmo if the component name can not resolve to a skeletal mesh
+		if (const TStrongObjectPtr<USkeletalMesh> SkeletalMesh = GetReferenceSkeletalMesh(); !SkeletalMesh)
+		{
+			Editor->HideGizmoClipMorph(false);
+			return;
+		}
 	}
 
-	if (TSharedPtr<FCustomizableObjectEditor> Editor = StaticCastSharedPtr<FCustomizableObjectEditor>(GetGraphEditor()))
+	// Apply the changes provoked by the change of one of the properties of the node
+	{
+		const FName PropertyName = PropertyChangedEvent.GetPropertyName();
+		
+		if (PropertyName == GET_MEMBER_NAME_CHECKED(UCustomizableObjectNodeModifierClipMorph, bLocalStartOffset))
+		{
+			ChangeStartOffsetTransform();
+		}
+		// Handle a change in the name of the bone being used
+		else if (PropertyName == GET_MEMBER_NAME_CHECKED(UCustomizableObjectNodeModifierClipMorph, BoneName))
+		{
+			UpdateOriginAndNormal();
+		}
+		// Handle the inversion of the normals
+		else if (PropertyName == GET_MEMBER_NAME_CHECKED(UCustomizableObjectNodeModifierClipMorph, bInvertNormal))
+		{
+			InvertNormals();
+		}
+	}
+	
+	
+	// Update the viewport with the data we changed when updating the properties of the node
+	if (Editor)
 	{
 		Editor->ShowGizmoClipMorph(*this);
 	}
@@ -172,9 +302,7 @@ void UCustomizableObjectNodeModifierClipMorph::PinConnectionListChanged(UEdGraph
 void UCustomizableObjectNodeModifierClipMorph::BackwardsCompatibleFixup(int32 CustomizableObjectCustomVersion)
 {
 	Super::BackwardsCompatibleFixup(CustomizableObjectCustomVersion);
-
-	const UEdGraphSchema_CustomizableObject* Schema = GetDefault<UEdGraphSchema_CustomizableObject>();
-
+	
 	if (CustomizableObjectCustomVersion == FCustomizableObjectCustomVersion::UnifyRequiredTags)
 	{
 		RequiredTags = Tags_DEPRECATED;
diff --git a/Engine/Plugins/Mutable/Source/CustomizableObjectEditor/Private/MuCOE/Nodes/CustomizableObjectNodeModifierBaseDetails.h b/Engine/Plugins/Mutable/Source/CustomizableObjectEditor/Private/MuCOE/Nodes/CustomizableObjectNodeModifierBaseDetails.h
--- a/Engine/Plugins/Mutable/Source/CustomizableObjectEditor/Private/MuCOE/Nodes/CustomizableObjectNodeModifierBaseDetails.h
+++ b/Engine/Plugins/Mutable/Source/CustomizableObjectEditor/Private/MuCOE/Nodes/CustomizableObjectNodeModifierBaseDetails.h
@@ -6,11 +6,7 @@
 #include "MuCOE/SMutableTagListWidget.h"
 #include "MuCOE/SMutableSearchComboBox.h"
 #include "IDetailCustomization.h"
-#include "Types/SlateEnums.h"
-#include "Widgets/Input/SComboButton.h"
-#include "Widgets/Input/SEditableTextBox.h"
 #include "Widgets/Views/SListView.h"
-#include "Widgets/Views/ITableRow.h"
 
 namespace ESelectInfo { enum Type : int; }
 
diff --git a/Engine/Plugins/Mutable/Source/CustomizableObjectEditor/Private/MuCOE/Nodes/CustomizableObjectNodeModifierClipMorph.h b/Engine/Plugins/Mutable/Source/CustomizableObjectEditor/Private/MuCOE/Nodes/CustomizableObjectNodeModifierClipMorph.h
--- a/Engine/Plugins/Mutable/Source/CustomizableObjectEditor/Private/MuCOE/Nodes/CustomizableObjectNodeModifierClipMorph.h
+++ b/Engine/Plugins/Mutable/Source/CustomizableObjectEditor/Private/MuCOE/Nodes/CustomizableObjectNodeModifierClipMorph.h
@@ -22,9 +22,6 @@ class CUSTOMIZABLEOBJECTEDITOR_API UCustomizableObjectNodeModifierClipMorph : pu
 public:
 	GENERATED_BODY()
 
-	UPROPERTY(EditAnywhere, Category = MeshToClipAndMorph)
-	FName BoneName;
-
 	UPROPERTY()
 	TArray<FString> Tags_DEPRECATED;
 
@@ -33,6 +30,12 @@ public:
 	UPROPERTY(EditAnywhere, Category = MeshToClipAndMorph)
 	FName ReferenceSkeletonComponent;
 
+	UPROPERTY(EditAnywhere, Category = MeshToClipAndMorph)
+	FName BoneName;
+
+	UPROPERTY(EditAnywhere, Category = MeshClipParameters, meta = (DisplayName = "Invert normal direction", ToolTip = "Flag to invert the normal direction"))
+	bool bInvertNormal;
+
 	UPROPERTY(EditAnywhere, Category = MeshClipParameters, meta = (DisplayName="Morph Start Offset", ToolTip="Offset from the origin of the selected bone to the actual start of the morph."))
 	FVector StartOffset;
 
@@ -54,8 +57,6 @@ public:
 	UPROPERTY(EditAnywhere, Category = MeshClipParameters, meta = (DisplayName = "Morph Curve Control", ToolTip = "Controls the morph curve shape. A value of 1 is linear, less than 1 is concave and greater than 1 convex."))
 	float Exponent;
 
-	UPROPERTY(EditAnywhere, Category = MeshClipParameters, meta = (DisplayName = "Invert normal direction", ToolTip = "Flag to invert the normal direction"))
-	bool bInvertNormal;
 
 	UPROPERTY()
 	FVector Origin;
@@ -75,9 +76,7 @@ private:
 	bool bOldOffset_DEPRECATED;
 
 public:
-
-	bool bUpdateViewportWidget;
-
+	
 	UCustomizableObjectNodeModifierClipMorph();
 
 	// UObject interface.
@@ -99,7 +98,25 @@ public:
 	FVector GetOriginWithOffset() const;
 	void FindLocalAxes(FVector& Right, FVector& Up, FVector& Forward) const;
 
-	// Change StartOffset from World to Local of the other way around
+	/**
+	 * Change StartOffset from World to Local of the other way around
+	 */
 	void ChangeStartOffsetTransform();
-};
 
+	/**
+	 * Invert the normals of the Clipping plane
+	 */
+	void InvertNormals();
+	
+	/**
+	 * Update the origin and normal based on the BoneName property
+	 */
+	void UpdateOriginAndNormal();
+
+	/**
+	 * Get the reference skeletal mesh based on the ReferenceSkeletonComponent property
+	 * @return The skeletal mesh component used as reference for the ReferenceSkeletonComponent property.
+	 * Null if the ReferenceSkeletonComponent is not valid
+	 */
+	TStrongObjectPtr<USkeletalMesh> GetReferenceSkeletalMesh() const;
+};
diff --git a/Engine/Plugins/Mutable/Source/CustomizableObjectEditor/Private/MuCOE/Nodes/CustomizableObjectNodeModifierClipMorphDetails.cpp b/Engine/Plugins/Mutable/Source/CustomizableObjectEditor/Private/MuCOE/Nodes/CustomizableObjectNodeModifierClipMorphDetails.cpp
--- a/Engine/Plugins/Mutable/Source/CustomizableObjectEditor/Private/MuCOE/Nodes/CustomizableObjectNodeModifierClipMorphDetails.cpp
+++ b/Engine/Plugins/Mutable/Source/CustomizableObjectEditor/Private/MuCOE/Nodes/CustomizableObjectNodeModifierClipMorphDetails.cpp
@@ -6,13 +6,12 @@
 #include "Engine/SkeletalMesh.h"
 #include "IDetailsView.h"
 #include "MuCOE/CustomizableObjectEditorUtilities.h"
-#include "MuCOE/CustomizableObjectGraph.h"
 #include "MuCOE/GraphTraversal.h"
 #include "MuCOE/Nodes/CustomizableObjectNodeModifierClipMorph.h"
 #include "MuCOE/UnrealEditorPortabilityHelpers.h"
 #include "PropertyCustomizationHelpers.h"
 #include "MuCO/CustomizableObjectPrivate.h"
-#include "Widgets/Input/SCheckBox.h"
+#include "MuCOE/CustomizableObjectEditor.h"
 #include "Widgets/Input/STextComboBox.h"
 
 class UCustomizableObjectNodeObject;
@@ -37,288 +36,141 @@ void FCustomizableObjectNodeModifierClipMorphDetails::CustomizeDetails(IDetailLa
 	TSharedPtr<const IDetailsView> DetailsView = DetailBuilder.GetDetailsViewSharedPtr();
 	if (DetailsView && DetailsView->GetSelectedObjects().Num())
 	{
-		Node = Cast<UCustomizableObjectNodeModifierClipMorph>(DetailsView->GetSelectedObjects()[0].Get());
+		Node = TStrongObjectPtr(Cast<UCustomizableObjectNodeModifierClipMorph>(DetailsView->GetSelectedObjects()[0].Get()));
 	}
 
-	IDetailCategoryBuilder& BlocksCategory = DetailBuilder.EditCategory("MeshToClipAndMorph");
-	DetailBuilder.HideProperty("BoneName");
-
-	IDetailCategoryBuilder& MeshClipParametersCategory = DetailBuilder.EditCategory("MeshClipParameters");
-	DetailBuilder.HideProperty("bInvertNormal");
-
-	TSharedPtr<IPropertyHandle> ReferenceSkeletonIndexPropertyHandle = DetailBuilder.GetProperty("ReferenceSkeletonComponent");
-	const FSimpleDelegate OnReferenceSkeletonComponentChangedDelegate = 
-			FSimpleDelegate::CreateSP(this, &FCustomizableObjectNodeModifierClipMorphDetails::OnReferenceSkeletonComponentChanged);
-	ReferenceSkeletonIndexPropertyHandle->SetOnPropertyValueChanged(OnReferenceSkeletonComponentChangedDelegate);
-
-	if (Node)
-	{
-		SkeletalMesh = nullptr;
-		TSharedPtr<FString> BoneToSelect;
-
-		UCustomizableObject* CustomizableObject = Cast<UCustomizableObject>(Node->GetCustomizableObjectGraph()->GetOuter());
-
-		if (CustomizableObject)
-		{
-			const FName ReferenceComponentName = Node->ReferenceSkeletonComponent;
-			SkeletalMesh = CustomizableObject->GetComponentMeshReferenceSkeletalMesh(ReferenceComponentName);
-		}
-
-		if (SkeletalMesh)
-		{
-			BoneComboOptions.Empty();
-				
-			for (int32 i = 0; i < SkeletalMesh->GetRefSkeleton().GetRawBoneNum(); ++i)
-			{
-				FName BoneName = SkeletalMesh->GetRefSkeleton().GetBoneName(i);
-				BoneComboOptions.Add(MakeShareable(new FString(BoneName.ToString())));
-
-				if (BoneName == Node->BoneName)
-				{
-					BoneToSelect = BoneComboOptions.Last();
-				}
-			}
-
-            BoneComboOptions.Sort(CompareNames);
-
-			// Add them to the parent combo box
-			TSharedRef<IPropertyHandle> BoneProperty = DetailBuilder.GetProperty("BoneName");
-
-			BlocksCategory.AddCustomRow(LOCTEXT("ClipMorphDetails_BoneName", "Bone Name"))
-			[
-				SNew(SProperty, BoneProperty)
-				.ShouldDisplayName(false)
-				.CustomWidget()
-				[
-					SNew(SBorder)
-					.BorderImage(UE_MUTABLE_GET_BRUSH("NoBorder"))
-					.Padding(FMargin(0.0f, 0.0f, 10.0f, 0.0f))
-					.HAlign(EHorizontalAlignment::HAlign_Fill)
-					[
-						SNew(SHorizontalBox)
-						+SHorizontalBox::Slot().HAlign(HAlign_Left).VAlign(VAlign_Center)
-						[
-							SNew(STextBlock)
-							.Text( LOCTEXT("ClipMorphDetails_BoneName", "Bone Name"))
-							.Font(IDetailLayoutBuilder::GetDetailFont())
-						]
-						+SHorizontalBox::Slot().HAlign(HAlign_Fill)
-						[
-							SNew(STextComboBox)
-							.OptionsSource(&BoneComboOptions)
-							.InitiallySelectedItem(BoneToSelect)
-							.OnSelectionChanged(this, &FCustomizableObjectNodeModifierClipMorphDetails::OnBoneComboBoxSelectionChanged, BoneProperty)
-							.Font(IDetailLayoutBuilder::GetDetailFont())
-						]
-					]
-				]
-			];
-
-			TSharedRef<IPropertyHandle> InvertPlaneProperty = DetailBuilder.GetProperty("bInvertNormal");
-			MeshClipParametersCategory.AddCustomRow(LOCTEXT("ClipMorphDetails_PlaneNormal", "Invert plane normal"))
-			[
-				SNew(SProperty, BoneProperty)
-				.ShouldDisplayName(false)
-				.CustomWidget()
-				[
-					SNew(SBorder)
-					.BorderImage(UE_MUTABLE_GET_BRUSH("NoBorder"))
-					.Padding(FMargin(0.0f, 0.0f, 10.0f, 0.0f))
-					.HAlign(EHorizontalAlignment::HAlign_Fill)
-					[
-						SNew(SHorizontalBox)
-						+SHorizontalBox::Slot().HAlign(HAlign_Left).VAlign(VAlign_Center)
-						[
-							SNew(STextBlock)
-							.Text( LOCTEXT("ClipMorphDetails_PlaneNormal", "Invert plane normal"))
-						]
-						+SHorizontalBox::Slot().HAlign(HAlign_Left)
-						[
-							SNew(SCheckBox)
-							.OnCheckStateChanged(this, &FCustomizableObjectNodeModifierClipMorphDetails::OnInvertNormalCheckboxChanged, InvertPlaneProperty)
-							.IsChecked(this, &FCustomizableObjectNodeModifierClipMorphDetails::GetInvertNormalCheckBoxState)
-							.ToolTipText(LOCTEXT("ClipMorphDetails_InvertNormal_Tooltip", "Invert normal direction of the clip plane"))
-						]
-					]
-				]
-			];
-		}
-	}
-	else
+	// Invalid node : Early out
+	if (!Node)
 	{
+		IDetailCategoryBuilder& BlocksCategory = DetailBuilder.EditCategory("MeshToClipAndMorph");
 		BlocksCategory.AddCustomRow(LOCTEXT("Node", "Node"))
 		[
 			SNew(STextBlock)
 			.Text(LOCTEXT("ClipMorphDetails_InvertNormal_NodeNotFound", "Node not found"))
 		];
+		return;
+	}
+
+	// Handle changes in the string provided to the "Component Name" property
+	const TSharedPtr<IPropertyHandle> ReferenceSkeletonIndexPropertyHandle = DetailBuilder.GetProperty(GET_MEMBER_NAME_CHECKED(UCustomizableObjectNodeModifierClipMorph, ReferenceSkeletonComponent));
+	const FSimpleDelegate OnReferenceSkeletonComponentChangedDelegate = 
+			FSimpleDelegate::CreateSP(this, &FCustomizableObjectNodeModifierClipMorphDetails::OnReferenceSkeletonComponentChanged);
+	ReferenceSkeletonIndexPropertyHandle->SetOnPropertyValueChanged(OnReferenceSkeletonComponentChangedDelegate);
+	
+	// Handle cases where the name of the bone get changed but not due to our custom UI (for example on reset)
+	const TSharedRef<IPropertyHandle> BoneProperty = DetailBuilder.GetProperty(GET_MEMBER_NAME_CHECKED(UCustomizableObjectNodeModifierClipMorph, BoneName));
+	const FSimpleDelegate OnBoneNamePropertyResetHandle =
+		FSimpleDelegate::CreateSP(this, &FCustomizableObjectNodeModifierClipMorphDetails::OnBoneNameReset);
+	BoneProperty->SetOnPropertyResetToDefault(OnBoneNamePropertyResetHandle);
+	
+	// Pre-fill internal objects based on the current context
+	CacheSkeletalMesh();
+	if (SkeletalMesh)
+	{
+		CacheSkeletalMeshBoneNames();
 	}
+	
+	// Override the display of the BoneProperty so it uses a dropdown instead of a simple string entry box. the node willhandle the update of
+	// the property
+	DetailBuilder.EditDefaultProperty(BoneProperty)
+	->CustomWidget()
+	.Visibility(TAttribute<EVisibility>(this, &FCustomizableObjectNodeModifierClipMorphDetails::ShouldBoneDropdownBeVisible))
+	.NameContent()
+	[
+		SNew(STextBlock)
+		.Text(LOCTEXT("ClipMorphDetails_BoneName", "Bone Name"))
+		.Font(IDetailLayoutBuilder::GetDetailFont())
+	]
+	.ValueContent()
+	[
+		SAssignNew(BoneSelectionComboBox,STextComboBox)
+		.OptionsSource(&BoneComboOptions)
+		.InitiallySelectedItem(InitiallySelectedBone)
+		.OnSelectionChanged(this, &FCustomizableObjectNodeModifierClipMorphDetails::OnBoneComboBoxSelectionChanged, BoneProperty)
+		.Font(IDetailLayoutBuilder::GetDetailFont())
+	];
 }
 
 
-FVector FindBoneLocation(int32 BoneIndex, USkeletalMesh* SkeletalMesh)
+void FCustomizableObjectNodeModifierClipMorphDetails::CacheSkeletalMesh()
 {
-	const TArray<FTransform>& BoneArray = SkeletalMesh->GetRefSkeleton().GetRefBonePose();
-
-	int32 ParentIndex = BoneIndex;
-	FVector Location = FVector::ZeroVector;
-
-	while (ParentIndex >= 0)
+	if (!Node)
 	{
-		Location = BoneArray[ParentIndex].TransformPosition(Location);
-		ParentIndex = SkeletalMesh->GetRefSkeleton().GetParentIndex(ParentIndex);
+		return;
 	}
-
-	return Location;
+	
+	SkeletalMesh = Node->GetReferenceSkeletalMesh();
 }
 
 
-void FCustomizableObjectNodeModifierClipMorphDetails::OnBoneComboBoxSelectionChanged(TSharedPtr<FString> Selection, ESelectInfo::Type SelectInfo, TSharedRef<IPropertyHandle> BoneProperty)
+void FCustomizableObjectNodeModifierClipMorphDetails::CacheSkeletalMeshBoneNames()
 {
-	for (int OptionIndex = 0; OptionIndex < BoneComboOptions.Num(); ++OptionIndex)
+	BoneComboOptions.Empty();
+	InitiallySelectedBone.Reset();
+	
+	if (BoneSelectionComboBox)
 	{
-		if (BoneComboOptions[OptionIndex] == Selection)
-		{
-			//TArray<FVector> AffectedBoneLocations;
-
-			//if (SkeletalMesh)
-			//{
-			//	const TArray<FTransform>& BoneArray = SkeletalMesh->RefSkeleton.GetRefBonePose();
-			//	int32 BoneIndex = SkeletalMesh->RefSkeleton.FindBoneIndex(FName(**Selection));
-			//	
-			//	if (BoneIndex >= 0)
-			//	{
-			//		// Get tranform from skeleton root to selected bone
-			//		int32 ParentIndex = BoneIndex;
-			//		FTransform RootToBoneTransform = FTransform::Identity;
-
-			//		while (ParentIndex >= 0)
-			//		{
-			//			RootToBoneTransform = RootToBoneTransform * BoneArray[ParentIndex];
-			//			ParentIndex = SkeletalMesh->RefSkeleton.GetParentIndex(ParentIndex);
-			//		}
-
-			//		// Now that we have the transform from root to the selected bone, go down from this bone to the end of the bone chain hanging from it building a bounding box
-			//		int32 CurrentBone = BoneIndex;
-
-			//		FVector SelectedBoneLocationWithTransform = RootToBoneTransform.TransformPosition(FVector::ZeroVector);
-			//		FVector SelectedBoneLocationFromSelection = BoneLocations[OptionIndex];
-			//		check(FVector::DistSquared(SelectedBoneLocationWithTransform, SelectedBoneLocationFromSelection) < 0.01f);
-
-			//		while (CurrentBone >= 0)
-			//		{
-			//			FVector CurrentBoneLocation = RootToBoneTransform.TransformPosition(FVector::ZeroVector);
-			//			AffectedBoneLocations.Add(CurrentBoneLocation);
-
-			//			check(FVector::DistSquared(CurrentBoneLocation, FindBoneLocation(CurrentBone, SkeletalMesh)) < 0.01f);
-
-			//			int32 NextBone = -1;
-
-			//			for (int32 i = 0; i < SkeletalMesh->RefSkeleton.GetRawBoneNum(); ++i)
-			//			{
-			//				if (SkeletalMesh->RefSkeleton.GetParentIndex(i) == CurrentBone)
-			//				{
-			//					NextBone = i;
-			//					break;
-			//				}
-			//			}
-
-			//			CurrentBone = NextBone;
-
-			//			if (CurrentBone >= 0)
-			//			{
-			//				RootToBoneTransform = BoneArray[CurrentBone] * RootToBoneTransform;
-			//			}
-			//		}
-			//	}
-			//}
-
-			FVector Location = FVector::ZeroVector;
-			FVector Direction = FVector::ForwardVector;
-
-			if (SkeletalMesh)
-			{
-				const TArray<FTransform>& BoneArray = SkeletalMesh->GetRefSkeleton().GetRefBonePose();
-				int32 ParentIndex = SkeletalMesh->GetRefSkeleton().FindBoneIndex(FName(**Selection));
-
-				FVector ChildLocation = FVector::ForwardVector;
-
-				for (int32 i = 0; i < SkeletalMesh->GetRefSkeleton().GetRawBoneNum(); ++i)
-				{
-					if (SkeletalMesh->GetRefSkeleton().GetParentIndex(i) == ParentIndex)
-					{
-						ChildLocation = BoneArray[i].TransformPosition(FVector::ZeroVector);
-						break;
-					}
-				}
-
-				//Direction = ChildLocation;
-
-				while (ParentIndex >= 0)
-				{
-					Location = BoneArray[ParentIndex].TransformPosition(Location);
-					ChildLocation = BoneArray[ParentIndex].TransformPosition(ChildLocation);
-					//Direction = BoneArray[ParentIndex].TransformVector(Direction);
-					ParentIndex = SkeletalMesh->GetRefSkeleton().GetParentIndex(ParentIndex);
-				}
+		// Ensure the combo box displayed element is reset when changing the data being pointed at by "InitiallySelectedBone"
+		BoneSelectionComboBox->ClearSelection();
+	}
 
-				Direction = (ChildLocation - Location).GetSafeNormal();
-			}
-			 
-			if (Node)
+	if (SkeletalMesh)
+	{
+		for (int32 i = 0; i < SkeletalMesh->GetRefSkeleton().GetRawBoneNum(); ++i)
+		{
+			const FName BoneName = SkeletalMesh->GetRefSkeleton().GetBoneName(i);
+			BoneComboOptions.Add(MakeShareable(new FString(BoneName.ToString())));
+	
+			if (BoneName == Node->BoneName)
 			{
-				Node->Origin = Location;
-				Node->Normal = Direction;
+				InitiallySelectedBone = BoneComboOptions.Last();
 			}
-
-			// Set the bone property after Node Origin and Normal update, otherwise the 
-			// viewport gizmo will be constructed with the old values.
-			BoneProperty->SetValue(*BoneComboOptions[OptionIndex].Get());
-			
-			return;
 		}
+	
+		BoneComboOptions.Sort(CompareNames);
 	}
-
-	BoneProperty->SetValue(FString());
 }
 
 
-void FCustomizableObjectNodeModifierClipMorphDetails::OnInvertNormalCheckboxChanged(ECheckBoxState CheckBoxState, TSharedRef<IPropertyHandle> InvertPlaneProperty)
+void FCustomizableObjectNodeModifierClipMorphDetails::OnBoneNameReset()
 {
-	if (Node == nullptr)
+	if (BoneSelectionComboBox)
 	{
-		return;
+		// Ensure the combo box displayed element is reset when changing the data being pointed at by "InitiallySelectedBone"
+		BoneSelectionComboBox->ClearSelection();
 	}
-
-	if (Node->bLocalStartOffset)
-	{
-		Node->StartOffset.Z *= -1.0f;
-		Node->StartOffset.X *= -1.0f;
-	}
-
-	Node->Normal *= -1.0f;
-
-	InvertPlaneProperty->SetValue(CheckBoxState == ECheckBoxState::Checked);
 }
 
 
-ECheckBoxState FCustomizableObjectNodeModifierClipMorphDetails::GetInvertNormalCheckBoxState() const
+void FCustomizableObjectNodeModifierClipMorphDetails::OnBoneComboBoxSelectionChanged(TSharedPtr<FString> InSelectedBoneName, ESelectInfo::Type InSelectInfo, TSharedRef<IPropertyHandle> BoneProperty)
 {
-	if (Node == nullptr)
+	// Update the property of the node
+	if (InSelectedBoneName.IsValid())
 	{
-		return ECheckBoxState::Unchecked;
+		BoneProperty->SetValue(*InSelectedBoneName);
+	}
+	else
+	{
+		BoneProperty->ResetToDefault();
 	}
-
-	return Node->bInvertNormal ? ECheckBoxState::Checked : ECheckBoxState::Unchecked;
 }
 
 
 void FCustomizableObjectNodeModifierClipMorphDetails::OnReferenceSkeletonComponentChanged()
 {
-	if (Node)
+	CacheSkeletalMesh();
+	CacheSkeletalMeshBoneNames();
+}
+
+
+EVisibility FCustomizableObjectNodeModifierClipMorphDetails::ShouldBoneDropdownBeVisible() const
+{
+	if (BoneComboOptions.Num() > 0)
 	{
-		Node->BoneName = FName();
+		return EVisibility::Visible;
 	}
-	DetailBuilderPtr->ForceRefreshDetails();
+
+	return EVisibility::Collapsed;
 }
 
 
diff --git a/Engine/Plugins/Mutable/Source/CustomizableObjectEditor/Private/MuCOE/Nodes/CustomizableObjectNodeModifierClipMorphDetails.h b/Engine/Plugins/Mutable/Source/CustomizableObjectEditor/Private/MuCOE/Nodes/CustomizableObjectNodeModifierClipMorphDetails.h
--- a/Engine/Plugins/Mutable/Source/CustomizableObjectEditor/Private/MuCOE/Nodes/CustomizableObjectNodeModifierClipMorphDetails.h
+++ b/Engine/Plugins/Mutable/Source/CustomizableObjectEditor/Private/MuCOE/Nodes/CustomizableObjectNodeModifierClipMorphDetails.h
@@ -4,15 +4,16 @@
 
 #include "CustomizableObjectNodeModifierBaseDetails.h"
 #include "IDetailCustomization.h"
-#include "UObject/ObjectMacros.h"
 
 enum class ECheckBoxState : uint8;
 namespace ESelectInfo { enum Type : int; }
 
+class STextComboBox;
 class FString;
 class IDetailLayoutBuilder;
 class IPropertyHandle;
-
+class USkeletalMesh;
+class UCustomizableObjectNodeModifierClipMorph;
 
 class FCustomizableObjectNodeModifierClipMorphDetails : public FCustomizableObjectNodeModifierBaseDetails
 {
@@ -24,17 +25,47 @@ public:
 	virtual void CustomizeDetails( IDetailLayoutBuilder& DetailBuilder ) override;
 
 private:
-	class UCustomizableObjectNodeModifierClipMorph* Node;
-	class IDetailLayoutBuilder* DetailBuilderPtr = nullptr;
-	TArray< TSharedPtr<FString> > BoneComboOptions;
 
-	UPROPERTY()
-	class USkeletalMesh* SkeletalMesh;
+	/** Store the Skeletal Mesh of the current node for later access. */
+	void CacheSkeletalMesh();
+	
+	/** If a skeletal mesh has been cached this method will cache the name of the bones part of the skeleton of the cached skeletal mesh*/
+	void CacheSkeletalMeshBoneNames();
 
-	void OnBoneComboBoxSelectionChanged(TSharedPtr<FString> Selection, ESelectInfo::Type SelectInfo, TSharedRef<IPropertyHandle> BoneProperty);
+	/** Method invoked each time the BoneName property gets reset from the UI */
+	void OnBoneNameReset();
 
-	void OnInvertNormalCheckboxChanged(ECheckBoxState CheckBoxState, TSharedRef<IPropertyHandle> InvertPlaneProperty);
-	ECheckBoxState GetInvertNormalCheckBoxState() const;
+	/**
+	 * Method invoked each time the selected boneName in the exposed combobox is changed by the user
+	 * @param InSelectedBoneName The new selected bone name
+	 * @param InSelectInfo The type of selection (combobox selection)
+	 * @param BoneProperty The property to be updated with the string found in InSelection
+	 */
+	void OnBoneComboBoxSelectionChanged(TSharedPtr<FString> InSelectedBoneName, ESelectInfo::Type InSelectInfo, TSharedRef<IPropertyHandle> BoneProperty);
 
+	/** Method invoked each time the component name property gets its value updated */
 	void OnReferenceSkeletonComponentChanged();
+
+	/**
+	 * Determines if the bone dropdown should be visible or not
+	 * @return The visibility state to use for the bone name dropdown.
+	 */
+	EVisibility ShouldBoneDropdownBeVisible() const;
+	
+	/** The node whose data we are representing and updating */
+	TStrongObjectPtr<UCustomizableObjectNodeModifierClipMorph> Node = nullptr;
+	
+	/** The combo box that we use to display a list of bone names for the current Skeletal mesh */
+	TSharedPtr<STextComboBox> BoneSelectionComboBox = nullptr;
+	
+	IDetailLayoutBuilder* DetailBuilderPtr = nullptr;
+
+	/** The name of the bone we have initially selected */
+	TSharedPtr<FString> InitiallySelectedBone;
+	
+	/** The name of all bones we could select for the skeletal mesh skeleton of the current component */
+	TArray<TSharedPtr<FString>> BoneComboOptions;
+
+	/** The skeletal mesh of the component the node is modifying. It will be used to get an array of bones. */
+	TStrongObjectPtr<USkeletalMesh> SkeletalMesh = nullptr;
 };
diff --git a/Engine/Plugins/Mutable/Source/CustomizableObjectEditor/Public/MuCOE/ICustomizableObjectInstanceEditor.h b/Engine/Plugins/Mutable/Source/CustomizableObjectEditor/Public/MuCOE/ICustomizableObjectInstanceEditor.h
--- a/Engine/Plugins/Mutable/Source/CustomizableObjectEditor/Public/MuCOE/ICustomizableObjectInstanceEditor.h
+++ b/Engine/Plugins/Mutable/Source/CustomizableObjectEditor/Public/MuCOE/ICustomizableObjectInstanceEditor.h
@@ -72,8 +72,9 @@ public:
 	virtual void ShowGizmoClipMorph(UCustomizableObjectNodeModifierClipMorph& ClipPlainNode) {}
 
 	/** Hide the clip morph plane gizmo of the NodeClipMorph.
-	 * Synchronizes all editor widgets so that the selection and the widgets is always consistent. */
-	virtual void HideGizmoClipMorph() {}
+	 * @param bClearGraphSelectionSet If set to true will deselect all nodes in the graph.
+	 */
+	virtual void HideGizmoClipMorph(bool bClearGraphSelectionSet = true) {}
 	
 	/** Show the clipping mesh gizmo from a node that uses a transformable bounding mesh. */
 	virtual void ShowGizmoClipMesh(UCustomizableObjectNode& Node, FTransform* Transform, const UEdGraphPin& MeshPin) {}
