[mutable] Fixed invalid access to editor mutable module when running a standalone server.
- Now access to the mutable editor module is forbidden when performed running a standalone server.
- Removed redundant search for the CO tree root Customizable Object.

diff --git a/Engine/Plugins/Mutable/Source/CustomizableObject/Internal/MuCO/ICustomizableObjectEditorModule.h b/Engine/Plugins/Mutable/Source/CustomizableObject/Internal/MuCO/ICustomizableObjectEditorModule.h
--- a/Engine/Plugins/Mutable/Source/CustomizableObject/Internal/MuCO/ICustomizableObjectEditorModule.h
+++ b/Engine/Plugins/Mutable/Source/CustomizableObject/Internal/MuCO/ICustomizableObjectEditorModule.h
@@ -46,13 +46,21 @@ public:
 		{
 			return nullptr;
 		}
+
+		// If this is a dedicated server no access to the editor module should be allowed.
+		if (IsRunningDedicatedServer())
+		{
+			return nullptr;
+		}
 		
 		return FModuleManager::LoadModulePtr<ICustomizableObjectEditorModule>(MODULE_NAME_COE);
 	}
 	
 	static ICustomizableObjectEditorModule& GetChecked()
 	{
-		check(!IsRunningGame())		// This module is editor-only. DO NOT try to access it during gameplay
+		// This module is editor-only. DO NOT try to access it during gameplay
+		checkf(!IsRunningGame(), TEXT("Unable to access \"ICustomizableObjectEditorModule\" module. Running standalone game."))
+		checkf(!IsRunningDedicatedServer(), TEXT("Unable to access \"ICustomizableObjectEditorModule\" module. Running dedicated server."))
 		return FModuleManager::LoadModuleChecked<ICustomizableObjectEditorModule>(MODULE_NAME_COE);
 	}
 
diff --git a/Engine/Plugins/Mutable/Source/CustomizableObject/Private/MuCO/CustomizableObject.cpp b/Engine/Plugins/Mutable/Source/CustomizableObject/Private/MuCO/CustomizableObject.cpp
--- a/Engine/Plugins/Mutable/Source/CustomizableObject/Private/MuCO/CustomizableObject.cpp
+++ b/Engine/Plugins/Mutable/Source/CustomizableObject/Private/MuCO/CustomizableObject.cpp
@@ -969,20 +969,20 @@ void UCustomizableObjectPrivate::AddUncompiledCOWarning(const FString& Additiona
 	System->AddUncompiledCOWarning(*GetPublic(), &AdditionalLoggingInfo);
 }
 
+
 USkeletalMesh* UCustomizableObject::GetComponentMeshReferenceSkeletalMesh(const FName& ComponentName) const
 {
+	// Try to get the data from the CO graph.
 #if WITH_EDITORONLY_DATA
-	if (!IsRunningGame())
+	if (const ICustomizableObjectEditorModule* Module = ICustomizableObjectEditorModule::Get())
 	{
-		if (const ICustomizableObjectEditorModule* Module = ICustomizableObjectEditorModule::Get())
-		{
-			return Module->GetReferenceSkeletalMesh(*this, ComponentName);
-		}
-		
-		return nullptr;
+		return Module->GetReferenceSkeletalMesh(*this, ComponentName);
 	}
 #endif
 
+	// If WITH_EDITORONLY_DATA is false Look instead for the ReferenceSkeletalMesh in the compiled resources
+	// TODO UE-282479  Make this a separated option so we can access the compiled data even in the editor execution as
+	// currently the graph search will be always performed if the editor only data is accessible.
 	if (const UModelResources* ModelResources = Private->GetModelResources())
 	{
 		int32 ObjectComponentIndex = ModelResources->ComponentNamesPerObjectComponent.IndexOfByKey(ComponentName);
@@ -996,6 +996,7 @@ USkeletalMesh* UCustomizableObject::GetComponentMeshReferenceSkeletalMesh(const
 	return nullptr;
 }
 
+
 int32 UCustomizableObject::FindState( const FString& Name ) const
 {
 	return GetPrivate()->FindState(Name);
diff --git a/Engine/Plugins/Mutable/Source/CustomizableObjectEditor/Private/MuCOE/CustomizableObjectEditorModule.cpp b/Engine/Plugins/Mutable/Source/CustomizableObjectEditor/Private/MuCOE/CustomizableObjectEditorModule.cpp
--- a/Engine/Plugins/Mutable/Source/CustomizableObjectEditor/Private/MuCOE/CustomizableObjectEditorModule.cpp
+++ b/Engine/Plugins/Mutable/Source/CustomizableObjectEditor/Private/MuCOE/CustomizableObjectEditorModule.cpp
@@ -843,16 +843,17 @@ void FCustomizableObjectEditorModule::BakeCustomizableObjectInstance(UCustomizab
 
 USkeletalMesh* FCustomizableObjectEditorModule::GetReferenceSkeletalMesh(const UCustomizableObject& Object, const FName& ComponentName) const
 {
-	UCustomizableObject* RootObject = GetRootObject(const_cast<UCustomizableObject*>(&Object));
-
 	TSet<UCustomizableObject*> Objects;
-	GetAllObjectsInGraph(RootObject, Objects);
+	GetAllObjectsInGraph(const_cast<UCustomizableObject*>(&Object), Objects);
 	
 	for (const UCustomizableObject* CurrentObject : Objects)
 	{
-		for (UEdGraphNode* Node : CurrentObject->GetPrivate()->GetSource()->Nodes)
+		TObjectPtr<UEdGraph> Source = CurrentObject->GetPrivate()->GetSource();
+		check(Source)
+		
+		for (TObjectPtr<UEdGraphNode> Node : CurrentObject->GetPrivate()->GetSource()->Nodes)
 		{
-			if (UCustomizableObjectNodeComponentMesh* NodeComponentMesh = Cast<UCustomizableObjectNodeComponentMesh>(Node))
+			if (const UCustomizableObjectNodeComponentMesh* NodeComponentMesh = Cast<UCustomizableObjectNodeComponentMesh>(Node))
 			{
 				if (NodeComponentMesh->GetComponentName() == ComponentName)
 				{
