[Mutable] When an update is marked as optimized components with the correct skeletal mesh will not be updated in the UpdateSkeletalMesh and FinishUpdateGlobal functions.
Change the default value of CVarEnableUpdateOptimization to true.
Add mesh streaming info to the Descriptor and Descriptore Hash.

#rnx

diff --git a/Engine/Plugins/Mutable/Source/CustomizableObject/Internal/MuCO/DescriptorHash.h b/Engine/Plugins/Mutable/Source/CustomizableObject/Internal/MuCO/DescriptorHash.h
--- a/Engine/Plugins/Mutable/Source/CustomizableObject/Internal/MuCO/DescriptorHash.h
+++ b/Engine/Plugins/Mutable/Source/CustomizableObject/Internal/MuCO/DescriptorHash.h
@@ -25,6 +25,9 @@ private:
 	uint32 Hash = 0;
 
 public:
+
+	bool bStreamingEnabled = false;
+
 	// MinLOD based on user input and quality settings. First LOD to generate.
 	TMap<FName, uint8> MinLODs;
 
diff --git a/Engine/Plugins/Mutable/Source/CustomizableObject/Private/MuCO/CustomizableObjectSystem.cpp b/Engine/Plugins/Mutable/Source/CustomizableObject/Private/MuCO/CustomizableObjectSystem.cpp
--- a/Engine/Plugins/Mutable/Source/CustomizableObject/Private/MuCO/CustomizableObjectSystem.cpp
+++ b/Engine/Plugins/Mutable/Source/CustomizableObject/Private/MuCO/CustomizableObjectSystem.cpp
@@ -112,7 +112,7 @@ TAutoConsoleVariable<bool> CVarEnableMeshCache(
 
 TAutoConsoleVariable<bool> CVarEnableUpdateOptimization(
 	TEXT("mutable.EnableUpdateOptimization"),
-	false,
+	true,
 	TEXT("Enable or disable update optimization when no changes are made to the parent component."));
 
 TAutoConsoleVariable<bool> CVarEnableRealTimeMorphTargets(
@@ -313,6 +313,7 @@ FUpdateContextPrivate::FUpdateContextPrivate(UCustomizableObjectInstance& InInst
 	Object = InObject;
 
 	CapturedDescriptor = Descriptor;
+	CapturedDescriptor.bStreamingEnabled = IsStreamingEnabled(*InObject);
 	InObject->GetPrivate()->ApplyStateForcedValuesToParameters(CapturedDescriptor);
 
 	CapturedDescriptorHash = FDescriptorHash(Descriptor);
@@ -1135,6 +1136,7 @@ static FAutoConsoleVariable CVarDescriptorDebugPrint(
 
 void FinishUpdateGlobal(const TSharedRef<FUpdateContextPrivate>& Context)
 {
+	MUTABLE_CPUPROFILER_SCOPE(FinishUpdateGlobal);
 	check(IsInGameThread())
 
 	UCustomizableObjectInstance* Instance = Context->Instance.Get();
@@ -1194,26 +1196,49 @@ void FinishUpdateGlobal(const TSharedRef<FUpdateContextPrivate>& Context)
 
 	if (UCustomizableObjectSystem::IsUpdateResultValid(Context->UpdateResult))
 	{
+		MUTABLE_CPUPROFILER_SCOPE(FinishUpdateGlobal_Callbacks);
+
 		// Call CustomizableObjectInstanceUsages updated callbacks.
-		for (TObjectIterator<UCustomizableObjectInstanceUsage> It; It; ++It) // Since iterating objects is expensive, for now CustomizableObjectInstanceUsage does not have a FinishUpdate function.
+		if (Context->bOptimizedUpdate)
 		{
-			UCustomizableObjectInstanceUsage* InstanceUsage = *It;
-			if (!IsValid(InstanceUsage))
+			for (TWeakObjectPtr<UCustomizableObjectInstanceUsage>& WeakUsage : Context->AttachedParentUpdated)
 			{
-				continue;
-			}
+				TStrongObjectPtr<UCustomizableObjectInstanceUsage> InstanceUsage = WeakUsage.Pin();
+				if (!IsValid(InstanceUsage.Get()))
+				{
+					return;
+				}
 
 #if WITH_EDITOR
-			if (It->GetPrivate()->IsNetMode(NM_DedicatedServer))
-			{
-				continue;
+				if (InstanceUsage->GetPrivate()->IsNetMode(NM_DedicatedServer))
+				{
+					continue;
+				}
+#endif
+				InstanceUsage->GetPrivate()->Callbacks();
 			}
+		}
+		else
+		{
+			for (TObjectIterator<UCustomizableObjectInstanceUsage> It; It; ++It) // Since iterating objects is expensive, for now CustomizableObjectInstanceUsage does not have a FinishUpdate function.
+			{
+				UCustomizableObjectInstanceUsage* InstanceUsage = *It;
+				if (!IsValid(InstanceUsage))
+				{
+					continue;
+				}
+
+#if WITH_EDITOR
+				if (It->GetPrivate()->IsNetMode(NM_DedicatedServer))
+				{
+					continue;
+				}
 #endif
 
-			if (InstanceUsage->GetCustomizableObjectInstance() == Instance &&
-				(!Context->bOptimizedUpdate || Context->AttachedParentUpdated.Find(InstanceUsage)))
-			{
-				InstanceUsage->GetPrivate()->Callbacks();
+				if (InstanceUsage->GetCustomizableObjectInstance() == Instance)
+				{
+					InstanceUsage->GetPrivate()->Callbacks();
+				}
 			}
 		}
 	}
@@ -1402,9 +1427,14 @@ void UpdateSkeletalMesh(const TSharedRef<FUpdateContextPrivate>& Context)
 			bAttachedParentUpdated = true;
 		}
 
+		// Skip further checks since the instance did not change and the component has the correct SkeletalMesh.
+		if (!bAttachedParentUpdated && Context->bOptimizedUpdate)
+		{
+			continue;
+		}
 
 		TArray<TObjectPtr<UMaterialInterface>> NewOverridenMaterials;
-		
+
 		const bool bIsTransientMesh = SkeletalMesh ? SkeletalMesh->HasAllFlags(EObjectFlags::RF_Transient) : false;
 		const bool bIsUsingMeshCache = CustomizableObject->bEnableMeshCache && UCustomizableObjectSystem::IsMeshCacheEnabled();
 		const bool bUseOverrideMaterials = !bIsTransientMesh || bIsUsingMeshCache;
@@ -1415,7 +1445,7 @@ void UpdateSkeletalMesh(const TSharedRef<FUpdateContextPrivate>& Context)
 				NewOverridenMaterials = ComponentData->OverrideMaterials;
 			}
 		}
-		
+
 		if (Parent->OverrideMaterials != NewOverridenMaterials)
 		{
 			// Before setting the new override materials clear the old ones set in the Skeletal Mesh Component
@@ -1431,7 +1461,7 @@ void UpdateSkeletalMesh(const TSharedRef<FUpdateContextPrivate>& Context)
 			UE_LOG(LogMutable, Log, TEXT("Skeletal Mesh Component override materials have been updated."));
 			bAttachedParentUpdated = true;
 		}
-		
+
 
 		UMaterialInterface* OverlayMaterial = nullptr;
 		if (bUseOverrideMaterials)
diff --git a/Engine/Plugins/Mutable/Source/CustomizableObject/Private/MuCO/DescriptorHash.cpp b/Engine/Plugins/Mutable/Source/CustomizableObject/Private/MuCO/DescriptorHash.cpp
--- a/Engine/Plugins/Mutable/Source/CustomizableObject/Private/MuCO/DescriptorHash.cpp
+++ b/Engine/Plugins/Mutable/Source/CustomizableObject/Private/MuCO/DescriptorHash.cpp
@@ -57,6 +57,7 @@ FDescriptorHash::FDescriptorHash(const FCustomizableObjectInstanceDescriptor& De
 	Hash = HashCombine(Hash, GetTypeHash(Descriptor.State));
 	Hash = HashCombine(Hash, GetTypeHash(Descriptor.GetBuildParameterRelevancy()));
 	
+	bStreamingEnabled = Descriptor.bStreamingEnabled;
 	MinLODs = Descriptor.MinLOD;
 	QualitySettingMinLODs = Descriptor.QualitySettingMinLODs;
 	FirstRequestedLOD = Descriptor.FirstRequestedLOD;
@@ -70,6 +71,11 @@ bool FDescriptorHash::IsSubset(const FDescriptorHash& Other) const
 		return false;
 	}
 
+	if (bStreamingEnabled && Other.bStreamingEnabled)
+	{
+		return true;
+	}
+
 	for (const TTuple<FName, uint8>& Pair : MinLODs)
 	{
 		const uint8* Result = Other.MinLODs.Find(Pair.Key);
diff --git a/Engine/Plugins/Mutable/Source/CustomizableObject/Public/MuCO/CustomizableObjectInstanceDescriptor.h b/Engine/Plugins/Mutable/Source/CustomizableObject/Public/MuCO/CustomizableObjectInstanceDescriptor.h
--- a/Engine/Plugins/Mutable/Source/CustomizableObject/Public/MuCO/CustomizableObjectInstanceDescriptor.h
+++ b/Engine/Plugins/Mutable/Source/CustomizableObject/Public/MuCO/CustomizableObjectInstanceDescriptor.h
@@ -327,6 +327,8 @@ struct FCustomizableObjectInstanceDescriptor
 	/** If this is set to true, when updating the instance an additional step will be performed to calculate the list of instance parameters that are relevant for the current parameter values. */
 	bool bBuildParameterRelevancy = false;
 
+	bool bStreamingEnabled = false;
+
 	/** These are the LODs Mutable can generate based on user input and quality settings, they MUST NOT be used in an update (Mutable thread). */
 	TMap<FName, uint8> MinLOD;
 
diff --git a/Engine/Plugins/Mutable/Source/MutableRuntime/Private/MuR/OpImageDisplace.h b/Engine/Plugins/Mutable/Source/MutableRuntime/Private/MuR/OpImageDisplace.h
--- a/Engine/Plugins/Mutable/Source/MutableRuntime/Private/MuR/OpImageDisplace.h
+++ b/Engine/Plugins/Mutable/Source/MutableRuntime/Private/MuR/OpImageDisplace.h
@@ -7,6 +7,7 @@
 #include "MuR/Platform.h"
 #include "MuR/ConvertData.h"
 #include "Async/ParallelFor.h"
+#include "MuR/ParallelExecutionUtils.h"
 
 namespace mu
 {
@@ -66,7 +67,7 @@ namespace mu
 			}
 
 			// TODO: Replace with a single parallel for.
-			const auto ProcessMip = [MaskImage, NextMask, ResultImage](int32 MipIndex)
+			ParallelExecutionUtils::InvokeBatchParallelFor(MipCount, [MaskImage, NextMask, ResultImage](int32 MipIndex)
 				{
 					uint8* ResultData = ResultImage->GetMipData(MipIndex);
 					const uint8* MaskData = MaskImage->GetMipData(MipIndex);
@@ -77,7 +78,7 @@ namespace mu
 					const int32 NumRowsPerBatch = FMath::Min(FMath::DivideAndRoundUp((1 << 14), MipSize.X), MipSize.Y);
 					const int32 NumBatches = FMath::DivideAndRoundUp<int32>(MipSize.Y, NumRowsPerBatch);
 
-					const auto ProcessBatch = [
+					ParallelExecutionUtils::InvokeBatchParallelFor(NumBatches, [
 						MaskData, NextMaskData, MipSize, ResultData, NumBatches, NumRowsPerBatch
 					] (int32 BatchIndex)
 						{
@@ -134,26 +135,8 @@ namespace mu
 								}
 							}
 
-						};
-
-						if (NumBatches == 1)
-						{
-							ProcessBatch(0);
-						}
-						else
-						{
-							ParallelFor(NumBatches, ProcessBatch);
-						}
-				};
-
-			if (MipCount == 1)
-			{
-				ProcessMip(0);
-			}
-			else
-			{
-				ParallelFor(MipCount, ProcessMip);
-			}
+						});
+				});
 		}
 	}
 
