[Mutable] Improve bytecode encoding of switch operations when many branches are repeated.

#rnx

diff --git a/Engine/Plugins/Mutable/Source/CustomizableObject/Internal/MuCO/CustomizableObjectPrivate.h b/Engine/Plugins/Mutable/Source/CustomizableObject/Internal/MuCO/CustomizableObjectPrivate.h
--- a/Engine/Plugins/Mutable/Source/CustomizableObject/Internal/MuCO/CustomizableObjectPrivate.h
+++ b/Engine/Plugins/Mutable/Source/CustomizableObject/Internal/MuCO/CustomizableObjectPrivate.h
@@ -1624,6 +1624,8 @@ enum class ECustomizableObjectVersions : int32
 	
 	MovePhysicsBodiesToRoms,
 
+	SwitchOpCompactRepresentation,
+
 	// -----<new versions can be added above this line>--------
 	LastCustomizableObjectVersion
 };
diff --git a/Engine/Plugins/Mutable/Source/CustomizableObjectEditor/Private/MuCOE/SMutableCodeViewer.cpp b/Engine/Plugins/Mutable/Source/CustomizableObjectEditor/Private/MuCOE/SMutableCodeViewer.cpp
--- a/Engine/Plugins/Mutable/Source/CustomizableObjectEditor/Private/MuCOE/SMutableCodeViewer.cpp
+++ b/Engine/Plugins/Mutable/Source/CustomizableObjectEditor/Private/MuCOE/SMutableCodeViewer.cpp
@@ -1464,22 +1464,45 @@ void SMutableCodeViewer::GenerateElementRecursive(const int32& InStateIndex, mu:
 		OpData += sizeof(mu::OP::ADDRESS);
 		AddOpFunc(DefAddress, TEXT("def "));
 
-		uint32 CaseCount;
-		FMemory::Memcpy(&CaseCount, OpData, sizeof(uint32));
-		OpData += sizeof(uint32);
+		mu::OP::FSwitchCaseDescriptor CaseDesc;
+		FMemory::Memcpy(&CaseDesc, OpData, sizeof(mu::OP::FSwitchCaseDescriptor));
+		OpData += sizeof(mu::OP::FSwitchCaseDescriptor);
 
-		for (uint32 C = 0; C < CaseCount; ++C)
+		if (!CaseDesc.bUseRanges)
 		{
-			int32 Condition;
-			FMemory::Memcpy(&Condition, OpData, sizeof(int32));
-			OpData += sizeof(int32);
+			for (uint32 C = 0; C < CaseDesc.Count; ++C)
+			{
+				int32 Condition;
+				FMemory::Memcpy(&Condition, OpData, sizeof(int32));
+				OpData += sizeof(int32);
 
-			mu::OP::ADDRESS At;
-			FMemory::Memcpy(&At, OpData, sizeof(mu::OP::ADDRESS));
-			OpData += sizeof(mu::OP::ADDRESS);
-			
-			FString Caption = FString::Printf(TEXT("case %d "), Condition);
-			AddOpFunc(At, Caption);
+				mu::OP::ADDRESS At;
+				FMemory::Memcpy(&At, OpData, sizeof(mu::OP::ADDRESS));
+				OpData += sizeof(mu::OP::ADDRESS);
+
+				FString Caption = FString::Printf(TEXT("case %d "), Condition);
+				AddOpFunc(At, Caption);
+			}
+		}
+		else
+		{
+			for (uint32 C = 0; C < CaseDesc.Count; ++C)
+			{
+				int32 ConditionStart;
+				FMemory::Memcpy(&ConditionStart, OpData, sizeof(int32));
+				OpData += sizeof(int32);
+
+				uint32 RangeSize;
+				FMemory::Memcpy(&RangeSize, OpData, sizeof(uint32));
+				OpData += sizeof(uint32);
+
+				mu::OP::ADDRESS At;
+				FMemory::Memcpy(&At, OpData, sizeof(mu::OP::ADDRESS));
+				OpData += sizeof(mu::OP::ADDRESS);
+
+				FString Caption = FString::Printf(TEXT("case %d to %d "), ConditionStart, ConditionStart+RangeSize-1);
+				AddOpFunc(At, Caption);
+			}
 		}
 
 		break;
diff --git a/Engine/Plugins/Mutable/Source/MutableRuntime/Internal/MuR/CodeVisitor.h b/Engine/Plugins/Mutable/Source/MutableRuntime/Internal/MuR/CodeVisitor.h
--- a/Engine/Plugins/Mutable/Source/MutableRuntime/Internal/MuR/CodeVisitor.h
+++ b/Engine/Plugins/Mutable/Source/MutableRuntime/Internal/MuR/CodeVisitor.h
@@ -302,42 +302,69 @@ namespace mu
 				const uint8* Data = Program.GetOpArgsPointer(Address);
 				
 				OP::ADDRESS VarAddress;
-				FMemory::Memcpy( &VarAddress, Data, sizeof(OP::ADDRESS));
+				FMemory::Memcpy(&VarAddress, Data, sizeof(OP::ADDRESS));
 				Data += sizeof(OP::ADDRESS);
 
                 if (VarAddress)
                 {
 					OP::ADDRESS DefAddress;
-					FMemory::Memcpy( &DefAddress, Data, sizeof(OP::ADDRESS));
+					FMemory::Memcpy(&DefAddress, Data, sizeof(OP::ADDRESS));
 					Data += sizeof(OP::ADDRESS);
 
-					uint32 CaseCount;
-					FMemory::Memcpy( &CaseCount, Data, sizeof(uint32));
-					Data += sizeof(uint32);
+					OP::FSwitchCaseDescriptor CaseDesc;
+					FMemory::Memcpy(&CaseDesc, Data, sizeof(OP::FSwitchCaseDescriptor));
+					Data += sizeof(OP::FSwitchCaseDescriptor);
 
                     PARENT::RecurseWithCurrentState( VarAddress );
 
-                    int32 var = System->BuildInt( Model, Params, VarAddress );
+                    int32 Var = System->BuildInt( Model, Params, VarAddress );
 
-					OP::ADDRESS valueAt = DefAddress;
-					for (uint32 CaseIndex = 0; CaseIndex < CaseCount; ++CaseIndex)
-					{
-						int32 Condition;
-						FMemory::Memcpy( &Condition, Data, sizeof(int32));		
-						Data += sizeof(int32);
-
-						OP::ADDRESS At;
-						FMemory::Memcpy( &At, Data, sizeof(OP::ADDRESS));
-						Data += sizeof(OP::ADDRESS);
+					OP::ADDRESS ValueAt = DefAddress;
 
-						if (At && var == (int32)Condition)
+					if (!CaseDesc.bUseRanges)
+					{
+						for (uint32 CaseIndex = 0; CaseIndex < CaseDesc.Count; ++CaseIndex)
+						{
+							int32 Condition;
+							FMemory::Memcpy(&Condition, Data, sizeof(int32));
+							Data += sizeof(int32);
+
+							OP::ADDRESS CaseAt;
+							FMemory::Memcpy(&CaseAt, Data, sizeof(OP::ADDRESS));
+							Data += sizeof(OP::ADDRESS);
+
+							if (CaseAt && Var == (int32)Condition)
+							{
+								ValueAt = CaseAt;
+								break;
+							}
+						}
+					}
+					else
+					{
+						for (uint32 C = 0; C < CaseDesc.Count; ++C)
 						{
-							valueAt = At;
-							break;	
+							int32 ConditionStart;
+							FMemory::Memcpy(&ConditionStart, Data, sizeof(int32));
+							Data += sizeof(int32);
+
+							uint32 RangeSize;
+							FMemory::Memcpy(&RangeSize, Data, sizeof(uint32));
+							Data += sizeof(uint32);
+
+							OP::ADDRESS CaseAt;
+							FMemory::Memcpy(&CaseAt, Data, sizeof(OP::ADDRESS));
+							Data += sizeof(OP::ADDRESS);
+
+							if (CaseAt && Var >= ConditionStart && Var < int32(ConditionStart + RangeSize))
+							{
+								ValueAt = CaseAt;
+								break;
+							}
 						}
 					}
 
-					PARENT::RecurseWithCurrentState( valueAt );
+					PARENT::RecurseWithCurrentState(ValueAt);
                 }
 
                 break;
diff --git a/Engine/Plugins/Mutable/Source/MutableRuntime/Internal/MuR/Operations.h b/Engine/Plugins/Mutable/Source/MutableRuntime/Internal/MuR/Operations.h
--- a/Engine/Plugins/Mutable/Source/MutableRuntime/Internal/MuR/Operations.h
+++ b/Engine/Plugins/Mutable/Source/MutableRuntime/Internal/MuR/Operations.h
@@ -440,6 +440,12 @@ namespace mu
 			int8 ForceLoad;
 		};
 
+		struct FSwitchCaseDescriptor
+		{
+			uint32 Count      : 31;
+			uint32 bUseRanges : 1;
+		};
+
         //-------------------------------------------------------------------------------------
         struct BoolEqualScalarConstArgs
         {
@@ -1065,6 +1071,9 @@ namespace mu
 
     ENUM_CLASS_FLAGS(OP::MeshAddMetadataArgs::EnumFlags);
 
+	static_assert(sizeof(OP::FSwitchCaseDescriptor) == sizeof(uint32));
+
+
 	//! Types of data handled by the Mutable runtime.
 	enum class EDataType : uint8
 	{
@@ -1108,9 +1117,6 @@ namespace mu
 		return GetOpDesc(type).type;
     }
 
-    //! Utility function to apply a function to all operation references to other operations.
-	MUTABLERUNTIME_API extern void ForEachReference( OP& op, const TFunctionRef<void(OP::ADDRESS*)> );
-
     //! Utility function to apply a function to all operation references to other operations.
 	MUTABLERUNTIME_API extern void ForEachReference( const struct FProgram& program, OP::ADDRESS at, const TFunctionRef<void(OP::ADDRESS)> );
 
diff --git a/Engine/Plugins/Mutable/Source/MutableRuntime/Private/MuR/CodeRunner.cpp b/Engine/Plugins/Mutable/Source/MutableRuntime/Private/MuR/CodeRunner.cpp
--- a/Engine/Plugins/Mutable/Source/MutableRuntime/Private/MuR/CodeRunner.cpp
+++ b/Engine/Plugins/Mutable/Source/MutableRuntime/Private/MuR/CodeRunner.cpp
@@ -365,21 +365,21 @@ namespace mu
 	{
 		const FProgram& Program = Model->GetPrivate()->Program;
 
-		EOpType type = Program.GetOpType(Item.At);
+		EOpType Type = Program.GetOpType(Item.At);
 
-		const uint8* data = Program.GetOpArgsPointer(Item.At);
+		const uint8* Data = Program.GetOpArgsPointer(Item.At);
 
 		OP::ADDRESS VarAddress;
-		FMemory::Memcpy(&VarAddress, data, sizeof(OP::ADDRESS));
-		data += sizeof(OP::ADDRESS);
+		FMemory::Memcpy(&VarAddress, Data, sizeof(OP::ADDRESS));
+		Data += sizeof(OP::ADDRESS);
 
 		OP::ADDRESS DefAddress;
-		FMemory::Memcpy(&DefAddress, data, sizeof(OP::ADDRESS));
-		data += sizeof(OP::ADDRESS);
+		FMemory::Memcpy(&DefAddress, Data, sizeof(OP::ADDRESS));
+		Data += sizeof(OP::ADDRESS);
 
-		uint32 CaseCount;
-		FMemory::Memcpy(&CaseCount, data, sizeof(uint32));
-		data += sizeof(uint32);
+		OP::FSwitchCaseDescriptor CaseDesc;
+		FMemory::Memcpy(&CaseDesc, Data, sizeof(OP::FSwitchCaseDescriptor));
+		Data += sizeof(OP::FSwitchCaseDescriptor);
 
 		switch (Item.Stage)
 		{
@@ -391,7 +391,7 @@ namespace mu
 			}
 			else
 			{
-				switch (GetOpDataType(type))
+				switch (GetOpDataType(Type))
 				{
 				case EDataType::Bool:			StoreBool(Item, false); break;
 				case EDataType::Int:			StoreInt(Item, 0); break;
@@ -415,56 +415,84 @@ namespace mu
 		case 1:
 		{
 			// Get the variable result
-			int32 var = LoadInt(FCacheAddress(VarAddress, Item));
+			int32 Var = LoadInt(FCacheAddress(VarAddress, Item));
 
-			OP::ADDRESS valueAt = DefAddress;
-			for (uint32 C = 0; C < CaseCount; ++C)
+			OP::ADDRESS ValueAt = DefAddress;
+
+			if (!CaseDesc.bUseRanges)
 			{
-				int32 Condition;
-				FMemory::Memcpy(&Condition, data, sizeof(int32));
-				data += sizeof(int32);
+				for (uint32 C = 0; C < CaseDesc.Count; ++C)
+				{
+					int32 Condition;
+					FMemory::Memcpy(&Condition, Data, sizeof(int32));
+					Data += sizeof(int32);
 
-				OP::ADDRESS At;
-				FMemory::Memcpy(&At, data, sizeof(OP::ADDRESS));
-				data += sizeof(OP::ADDRESS);
+					OP::ADDRESS CaseAt;
+					FMemory::Memcpy(&CaseAt, Data, sizeof(OP::ADDRESS));
+					Data += sizeof(OP::ADDRESS);
 
-				if (At && var == (int32)Condition)
+					if (CaseAt && Var == (int32)Condition)
+					{
+						ValueAt = CaseAt;
+						break;
+					}
+				}
+			}
+			else
+			{
+				for (uint32 C = 0; C < CaseDesc.Count; ++C)
 				{
-					valueAt = At;
-					break; 
+					int32 ConditionStart;
+					FMemory::Memcpy(&ConditionStart, Data, sizeof(int32));
+					Data += sizeof(int32);
+
+					uint32 RangeSize;
+					FMemory::Memcpy(&RangeSize, Data, sizeof(uint32));
+					Data += sizeof(uint32);
+
+					OP::ADDRESS CaseAt;
+					FMemory::Memcpy(&CaseAt, Data, sizeof(OP::ADDRESS));
+					Data += sizeof(OP::ADDRESS);
+
+					if (CaseAt && Var >= ConditionStart && Var < int32(ConditionStart + RangeSize))
+					{
+						ValueAt = CaseAt;
+						break;
+					}
 				}
+
 			}
 
             // Schedule the end of this instruction if necessary
-            AddOp( FScheduledOp( Item.At, Item, 2, valueAt ),
-				   FScheduledOp( valueAt, Item ) );
+            AddOp(FScheduledOp(Item.At, Item, 2, ValueAt),
+				FScheduledOp(ValueAt, Item));
 
             break;
         }
 
         case 2:
         {
-			OP::ADDRESS resultAt = OP::ADDRESS(Item.CustomState);
+			OP::ADDRESS ResultAt = OP::ADDRESS(Item.CustomState);
 
             // Store the final result
-            FCacheAddress cat( Item );
-            FCacheAddress rat( resultAt, Item );
-            switch (GetOpDataType(type))
+            FCacheAddress ItemAddress(Item);
+            FCacheAddress ResultAddress(ResultAt, Item);
+            switch (GetOpDataType(Type))
             {
-            case EDataType::Bool:			StoreBool( cat, LoadBool(rat) ); break;
-            case EDataType::Int:			StoreInt( cat, LoadInt(rat) ); break;
-            case EDataType::Scalar:			StoreScalar( cat, LoadScalar(rat) ); break;
-            case EDataType::String:			StoreString( cat, LoadString( rat ) ); break;
-            case EDataType::Color:			StoreColor( cat, LoadColor( rat ) ); break;
-            case EDataType::Projector:		StoreProjector( cat, LoadProjector(rat) ); break;
-			case EDataType::Mesh:			StoreMesh( cat, LoadMesh(rat) ); break;
-            case EDataType::Image:			StoreImage( cat, LoadImage(rat) ); break;
-            case EDataType::Layout:			StoreLayout( cat, LoadLayout(rat) ); break;
-            case EDataType::Instance:		StoreInstance( cat, LoadInstance(rat) ); break;
-			case EDataType::ExtensionData:	StoreExtensionData( cat, LoadExtensionData(rat) ); break;
+            case EDataType::Bool:			StoreBool(ItemAddress, LoadBool(ResultAddress)); break;
+            case EDataType::Int:			StoreInt(ItemAddress, LoadInt(ResultAddress)); break;
+            case EDataType::Scalar:			StoreScalar(ItemAddress, LoadScalar(ResultAddress)); break;
+            case EDataType::String:			StoreString(ItemAddress, LoadString(ResultAddress)); break;
+            case EDataType::Color:			StoreColor(ItemAddress, LoadColor(ResultAddress)); break;
+            case EDataType::Projector:		StoreProjector(ItemAddress, LoadProjector(ResultAddress)); break;
+			case EDataType::Mesh:			StoreMesh(ItemAddress, LoadMesh(ResultAddress)); break;
+            case EDataType::Image:			StoreImage(ItemAddress, LoadImage(ResultAddress)); break;
+            case EDataType::Layout:			StoreLayout(ItemAddress, LoadLayout(ResultAddress)); break;
+            case EDataType::Instance:		StoreInstance(ItemAddress, LoadInstance(ResultAddress)); break;
+			case EDataType::ExtensionData:	StoreExtensionData(ItemAddress, LoadExtensionData(ResultAddress)); break;
             default:
                 // Not implemented
-                check( false );
+                check(false);
             }
 
             break;
@@ -6496,10 +6524,10 @@ namespace mu
 			FMemory::Memcpy(&DefAddress, Data, sizeof(OP::ADDRESS));
 			Data += sizeof(OP::ADDRESS);
 
-			uint32 CaseCount;
-			FMemory::Memcpy(&CaseCount, Data, sizeof(uint32));
-			Data += sizeof(uint32);
-	
+			OP::FSwitchCaseDescriptor CaseDesc;
+			FMemory::Memcpy(&CaseDesc, Data, sizeof(OP::FSwitchCaseDescriptor));
+			Data += sizeof(OP::FSwitchCaseDescriptor);
+
 			switch (Item.Stage)
 			{
 			case 0:
@@ -6525,21 +6553,49 @@ namespace mu
 				int32 Var = LoadInt(FCacheAddress(VarAddress, Item.ExecutionIndex, Item.ExecutionOptions, FScheduledOp::EType::Full));
 
 				OP::ADDRESS ValueAt = DefAddress;
-				for (uint32 C = 0; C < CaseCount; ++C)
+
+				if (!CaseDesc.bUseRanges)
 				{
-					int32 Condition;
-					FMemory::Memcpy(&Condition, Data, sizeof(int32));
-					Data += sizeof(int32);
-					
-					OP::ADDRESS At;
-					FMemory::Memcpy(&At, Data, sizeof(OP::ADDRESS));
-					Data += sizeof(OP::ADDRESS);
+					for (uint32 C = 0; C < CaseDesc.Count; ++C)
+					{
+						int32 Condition;
+						FMemory::Memcpy(&Condition, Data, sizeof(int32));
+						Data += sizeof(int32);
+
+						OP::ADDRESS CaseAt;
+						FMemory::Memcpy(&CaseAt, Data, sizeof(OP::ADDRESS));
+						Data += sizeof(OP::ADDRESS);
 
-					if (At && Var == (int32)Condition)
+						if (CaseAt && Var == (int32)Condition)
+						{
+							ValueAt = CaseAt;
+							break;
+						}
+					}
+				}
+				else
+				{
+					for (uint32 C = 0; C < CaseDesc.Count; ++C)
 					{
-						ValueAt = At;
-						break;
+						int32 ConditionStart;
+						FMemory::Memcpy(&ConditionStart, Data, sizeof(int32));
+						Data += sizeof(int32);
+
+						uint32 RangeSize;
+						FMemory::Memcpy(&RangeSize, Data, sizeof(uint32));
+						Data += sizeof(int32);
+
+						OP::ADDRESS CaseAt;
+						FMemory::Memcpy(&CaseAt, Data, sizeof(OP::ADDRESS));
+						Data += sizeof(OP::ADDRESS);
+
+						if (CaseAt && Var >= ConditionStart && Var < int32(ConditionStart + RangeSize))
+						{
+							ValueAt = CaseAt;
+							break;
+						}
 					}
+
 				}
 
 				if (ValueAt)
diff --git a/Engine/Plugins/Mutable/Source/MutableRuntime/Private/MuR/Operations.cpp b/Engine/Plugins/Mutable/Source/MutableRuntime/Private/MuR/Operations.cpp
--- a/Engine/Plugins/Mutable/Source/MutableRuntime/Private/MuR/Operations.cpp
+++ b/Engine/Plugins/Mutable/Source/MutableRuntime/Private/MuR/Operations.cpp
@@ -213,34 +213,56 @@ namespace mu
         case EOpType::IN_SWITCH:
 		case EOpType::ED_SWITCH:
         {
-			const uint8_t* data = program.GetOpArgsPointer(at);
+			const uint8* Data = program.GetOpArgsPointer(at);
 			
 			OP::ADDRESS VarAddress;
-			FMemory::Memcpy( &VarAddress, data, sizeof(OP::ADDRESS));
-			data += sizeof(OP::ADDRESS);
+			FMemory::Memcpy(&VarAddress, Data, sizeof(OP::ADDRESS));
+			Data += sizeof(OP::ADDRESS);
 
 			OP::ADDRESS DefAddress;
-			FMemory::Memcpy( &DefAddress, data, sizeof(OP::ADDRESS));
-			data += sizeof(OP::ADDRESS);
+			FMemory::Memcpy(&DefAddress, Data, sizeof(OP::ADDRESS));
+			Data += sizeof(OP::ADDRESS);
 
-			uint32 CaseCount;
-			FMemory::Memcpy( &CaseCount, data, sizeof(uint32));
-			data += sizeof(uint32);
+			OP::FSwitchCaseDescriptor CaseDesc;
+			FMemory::Memcpy(&CaseDesc, Data, sizeof(OP::FSwitchCaseDescriptor));
+			Data += sizeof(OP::FSwitchCaseDescriptor);
 
 			f(VarAddress);
 			f(DefAddress);
 
-			for ( uint32 C = 0; C < CaseCount; ++C )
-			{	
-				//int32 Condition;
-				//FMemory::Memcpy( &Condition, data, sizeof(int32_t));
-				data += sizeof(int32);
-				
-				OP::ADDRESS At;
-				FMemory::Memcpy( &At, data, sizeof(OP::ADDRESS) );
-				data += sizeof(OP::ADDRESS);
-
-				f(At);
+			if (!CaseDesc.bUseRanges)
+			{
+				for (uint32 C = 0; C < CaseDesc.Count; ++C)
+				{
+					//int32 Condition;
+					//FMemory::Memcpy( &Condition, data, sizeof(int32));
+					Data += sizeof(int32);
+
+					OP::ADDRESS CaseAt;
+					FMemory::Memcpy(&CaseAt, Data, sizeof(OP::ADDRESS));
+					Data += sizeof(OP::ADDRESS);
+
+					f(CaseAt);
+				}
+			}
+			else
+			{
+				for (uint32 C = 0; C < CaseDesc.Count; ++C)
+				{
+					//int32 ConditionStart;
+					//FMemory::Memcpy( &ConditionStart, data, sizeof(int32));
+					Data += sizeof(int32);
+
+					//uint32 RangeSize;
+					//FMemory::Memcpy( &RangeSize, data, sizeof(uint32));
+					Data += sizeof(int32);
+
+					OP::ADDRESS CaseAt;
+					FMemory::Memcpy(&CaseAt, Data, sizeof(OP::ADDRESS));
+					Data += sizeof(OP::ADDRESS);
+
+					f(CaseAt);
+				}
 			}
 
             break;
diff --git a/Engine/Plugins/Mutable/Source/MutableTools/Private/MuT/ASTOpSwitch.cpp b/Engine/Plugins/Mutable/Source/MutableTools/Private/MuT/ASTOpSwitch.cpp
--- a/Engine/Plugins/Mutable/Source/MutableTools/Private/MuT/ASTOpSwitch.cpp
+++ b/Engine/Plugins/Mutable/Source/MutableTools/Private/MuT/ASTOpSwitch.cpp
@@ -30,29 +30,35 @@ namespace mu
 	}
 
 
-	bool ASTOpSwitch::IsEqual(const ASTOp& otherUntyped) const
+	bool ASTOpSwitch::IsEqual(const ASTOp& OtherUntyped) const
 	{
-		if (GetOpType()==otherUntyped.GetOpType())
+		if (GetOpType() == OtherUntyped.GetOpType())
 		{
-			const ASTOpSwitch* other = static_cast<const ASTOpSwitch*>(&otherUntyped);
-			return Type == other->Type && Variable == other->Variable &&
-				Cases == other->Cases && Default == other->Default;
+			const ASTOpSwitch* Other = static_cast<const ASTOpSwitch*>(&OtherUntyped);
+
+			return Type     == Other->Type     && 
+				   Variable == Other->Variable &&
+				   Cases    == Other->Cases    && 
+				   Default  == Other->Default;
 		}
+
 		return false;
 	}
 
 
-	mu::Ptr<ASTOp> ASTOpSwitch::Clone(MapChildFuncRef mapChild) const
+	mu::Ptr<ASTOp> ASTOpSwitch::Clone(MapChildFuncRef MapChild) const
 	{
-		Ptr<ASTOpSwitch> n = new ASTOpSwitch();
-		n->Type = Type;
-		n->Variable = mapChild(Variable.child());
-		n->Default = mapChild(Default.child());
-		for (const FCase& c : Cases)
+		Ptr<ASTOpSwitch> New = new ASTOpSwitch();
+
+		New->Type = Type;
+		New->Variable = MapChild(Variable.child());
+		New->Default = MapChild(Default.child());
+		for (const FCase& Case : Cases)
 		{
-			n->Cases.Emplace(c.Condition, n, mapChild(c.Branch.child()));
+			New->Cases.Emplace(Case.Condition, New, MapChild(Case.Branch.child()));
 		}
-		return n;
+
+		return New;
 	}
 
 
@@ -81,46 +87,61 @@ namespace mu
 
 	uint64 ASTOpSwitch::Hash() const
 	{
-		uint64 res = std::hash<uint64>()(uint64(Type));
-		for (const FCase& c : Cases)
+		uint64 Result = std::hash<uint64>()(uint64(Type));
+		for (const FCase& Case : Cases)
 		{
-			hash_combine(res, c.Condition);
-			hash_combine(res, c.Branch.child().get());
+			hash_combine(Result, Case.Condition);
+			hash_combine(Result, Case.Branch.child().get());
 		}
-		return res;
+
+		return Result;
 	}
 
 
 	mu::Ptr<ASTOp> ASTOpSwitch::GetFirstValidValue()
 	{
-		for (int32 i = 0; i < Cases.Num(); ++i)
+		for (const FCase& Case : Cases)
 		{
-			if (Cases[i].Branch)
+			if (Case.Branch)
 			{
-				return Cases[i].Branch.child();
+				return Case.Branch.child();
 			}
 		}
+
 		return nullptr;
 	}
 
 
-	bool ASTOpSwitch::IsCompatibleWith(const ASTOpSwitch* other) const
+	bool ASTOpSwitch::IsCompatibleWith(const ASTOpSwitch* Other) const
 	{
-		if (!other) return false;
-		if (Variable.child() != other->Variable.child()) return false;
-		if (Cases.Num() != other->Cases.Num()) return false;
-		for (const FCase& c : Cases)
+		if (!Other)
+		{
+			return false;
+		}
+
+		if (Variable.child() != Other->Variable.child())
+		{
+			return false;
+		}
+
+		if (Cases.Num() != Other->Cases.Num())
 		{
-			bool found = false;
-			for (const FCase& o : other->Cases)
+			return false;
+		}
+
+		for (const FCase& Case : Cases)
+		{
+			bool bFound = false;
+			for (const FCase& OtherCase : Other->Cases)
 			{
-				if (c.Condition == o.Condition)
+				if (Case.Condition == OtherCase.Condition)
 				{
-					found = true;
+					bFound = true;
 					break;
 				}
 			}
-			if (!found)
+
+			if (!bFound)
 			{
 				return false;
 			}
@@ -132,11 +153,11 @@ namespace mu
 
 	mu::Ptr<ASTOp> ASTOpSwitch::FindBranch(int32 Condition) const
 	{
-		for (const FCase& c : Cases)
+		for (const FCase& Case : Cases)
 		{
-			if (c.Condition == Condition)
+			if (Case.Condition == Condition)
 			{
-				return c.Branch.child();
+				return Case.Branch.child();
 			}
 		}
 
@@ -144,38 +165,103 @@ namespace mu
 	}
 
 
-	void ASTOpSwitch::ForEachChild(const TFunctionRef<void(ASTChild&)> f)
+	void ASTOpSwitch::ForEachChild(const TFunctionRef<void(ASTChild&)> Func)
 	{
-		f(Variable);
-		f(Default);
-		for (FCase& cas : Cases)
+		Func(Variable);
+		Func(Default);
+
+		for (FCase& Case : Cases)
 		{
-			f(cas.Branch);
+			Func(Case.Branch);
 		}
 	}
 
 
-	void ASTOpSwitch::Link(FProgram& program, FLinkerOptions*)
+	void ASTOpSwitch::Link(FProgram& Program, FLinkerOptions*)
 	{
 		// Already linked?
 		if (!linkedAddress)
 		{
-			linkedAddress = (OP::ADDRESS)program.OpAddress.Num();
-			program.OpAddress.Add((uint32_t)program.ByteCode.Num());
+			linkedAddress = (OP::ADDRESS)Program.OpAddress.Num();
+			Program.OpAddress.Add(Program.ByteCode.Num());
 
 			OP::ADDRESS VarAddress = Variable ? Variable->linkedAddress : 0;
 			OP::ADDRESS DefAddress = Default ? Default->linkedAddress : 0;
 
-			AppendCode(program.ByteCode, Type);
-			AppendCode(program.ByteCode, VarAddress);
-			AppendCode(program.ByteCode, DefAddress);
-			AppendCode(program.ByteCode, (uint32_t)Cases.Num());
+			AppendCode(Program.ByteCode, Type);
+			AppendCode(Program.ByteCode, VarAddress);
+			AppendCode(Program.ByteCode, DefAddress);
 
-			for (const FCase& Case : Cases)
+			const int32 CaseCount = Cases.Num();
+			int32 CasesInRange = 0;
+			for (int32 CaseIndex = 0; CaseIndex < CaseCount; ++CaseIndex)
 			{
-				OP::ADDRESS CaseBranchAddress = Case.Branch ? Case.Branch->linkedAddress : 0;
-				AppendCode(program.ByteCode, Case.Condition);
-				AppendCode(program.ByteCode, CaseBranchAddress);
+				for (; CaseIndex < CaseCount; ++CaseIndex)
+				{
+					int32 CandidateCaseIndex = CaseIndex + 1;
+
+					if (Cases.IsValidIndex(CandidateCaseIndex) &&
+						Cases[CandidateCaseIndex].Branch == Cases[CandidateCaseIndex - 1].Branch &&
+						Cases[CandidateCaseIndex].Condition == Cases[CandidateCaseIndex - 1].Condition + 1)
+					{
+						++CasesInRange;
+					}
+					else
+					{
+						break;
+					}
+				}				
+			}
+
+			bool bUseRanges = CasesInRange >= CaseCount / 2;
+
+			if (!bUseRanges)
+			{
+				OP::FSwitchCaseDescriptor CaseDesc;
+				CaseDesc.Count = Cases.Num();
+				CaseDesc.bUseRanges = false;
+				AppendCode(Program.ByteCode, CaseDesc);
+
+				for (const FCase& Case : Cases)
+				{
+					OP::ADDRESS CaseBranchAddress = Case.Branch ? Case.Branch->linkedAddress : 0;
+					AppendCode(Program.ByteCode, Case.Condition);
+					AppendCode(Program.ByteCode, CaseBranchAddress);
+				}
+			}
+			else
+			{
+				OP::FSwitchCaseDescriptor CaseDesc;
+				CaseDesc.Count = Cases.Num() - CasesInRange;
+				CaseDesc.bUseRanges = true;
+
+				AppendCode(Program.ByteCode, CaseDesc);
+
+				for (int32 CaseIndex = 0; CaseIndex < CaseCount; ++CaseIndex)
+				{
+					int32 RangeStart = CaseIndex;
+					uint32 RangeSize = 1;
+					for (; CaseIndex < CaseCount; ++CaseIndex)
+					{
+						int32 CandidateCaseIndex = CaseIndex + 1;
+						if (Cases.IsValidIndex(CandidateCaseIndex) &&
+							Cases[CandidateCaseIndex].Branch == Cases[CandidateCaseIndex - 1].Branch &&
+							Cases[CandidateCaseIndex].Condition == Cases[CandidateCaseIndex - 1].Condition + 1)
+						{
+							++RangeSize;
+						}
+						else
+						{
+							break;
+						}
+					}
+
+					const FCase& FirstCase = Cases[RangeStart];
+					OP::ADDRESS CaseBranchAddress = FirstCase.Branch ? FirstCase.Branch->linkedAddress : 0;
+					AppendCode(Program.ByteCode, FirstCase.Condition);
+					AppendCode(Program.ByteCode, RangeSize);
+					AppendCode(Program.ByteCode, CaseBranchAddress);
+				}
 			}
 		}
 	}
