[Mutable] Minor API improvements.

Actor:
- Added ACustomizableSkeletalMeshActor::GetCustomizableObjectInstance.
- Added ACustomizableSkeletalMeshActor::GetSkeletalMeshComponent.

Override Materials:
- Deprecated UCustomizableObjectInstance::GetOverrideMaterials(int32).
- Added UCustomizableObjectInstance::GetSkeletalMeshComponentOverrideMaterials(FName).

diff --git a/Engine/Plugins/Mutable/Source/CustomizableObject/Private/MuCO/CustomizableObjectInstance.cpp b/Engine/Plugins/Mutable/Source/CustomizableObject/Private/MuCO/CustomizableObjectInstance.cpp
--- a/Engine/Plugins/Mutable/Source/CustomizableObject/Private/MuCO/CustomizableObjectInstance.cpp
+++ b/Engine/Plugins/Mutable/Source/CustomizableObject/Private/MuCO/CustomizableObjectInstance.cpp
@@ -5535,23 +5535,22 @@ FCustomizableObjectInstanceDescriptor& UCustomizableInstancePrivate::GetDescript
 }
 
 
-const TArray<TObjectPtr<UMaterialInterface>>* UCustomizableObjectInstance::GetOverrideMaterials(FCustomizableObjectComponentIndex ComponentIndex) const
+TArray<UMaterialInterface*> UCustomizableObjectInstance::GetSkeletalMeshComponentOverrideMaterials(const FName& ComponentName) const
 {
-	UCustomizableObject* Object = GetCustomizableObject();
-	if (!Object)
+	FCustomizableInstanceComponentData* ComponentData = PrivateData->GetComponentData(ComponentName);
+	if (!ComponentData)
 	{
-		return nullptr;
+		return {};
 	}
+	
+	TArray<UMaterialInterface*> Result;
 
-	const UModelResources* ModelResources = Object->GetPrivate()->GetModelResources();
-	if (!ModelResources)
+	for (const TObjectPtr<UMaterialInterface>& OverrideMaterial : ComponentData->OverrideMaterials)
 	{
-		return nullptr;
+		Result.Add(OverrideMaterial);
 	}
-
-	FName ComponentName = ModelResources->ComponentNamesPerObjectComponent[ComponentIndex.GetValue()];
-	FCustomizableInstanceComponentData* ComponentData = PrivateData->GetComponentData(ComponentName);
-	return ComponentData ? &ComponentData->OverrideMaterials : nullptr;
+	
+	return Result;
 }
 
 
diff --git a/Engine/Plugins/Mutable/Source/CustomizableObject/Private/MuCO/CustomizableSkeletalMeshActor.cpp b/Engine/Plugins/Mutable/Source/CustomizableObject/Private/MuCO/CustomizableSkeletalMeshActor.cpp
--- a/Engine/Plugins/Mutable/Source/CustomizableObject/Private/MuCO/CustomizableSkeletalMeshActor.cpp
+++ b/Engine/Plugins/Mutable/Source/CustomizableObject/Private/MuCO/CustomizableSkeletalMeshActor.cpp
@@ -20,7 +20,7 @@
 #define LOCTEXT_NAMESPACE "CustomizableObject"
 
 
-UCustomizableObjectInstance* ACustomizableSkeletalMeshActor::GetInstance()
+UCustomizableObjectInstance* ACustomizableSkeletalMeshActor::GetCustomizableObjectInstance()
 {
 	for (const UCustomizableSkeletalComponent* Component : CustomizableSkeletalComponents)
 	{
@@ -43,7 +43,7 @@ ACustomizableSkeletalMeshActor::ACustomizableSkeletalMeshActor(const FObjectInit
 	{
 		UCustomizableSkeletalComponent* CustomizableSkeletalComponent = CreateDefaultSubobject<UCustomizableSkeletalComponent>(TEXT("CustomizableSkeletalComponent0"));
 		CustomizableSkeletalComponents.Add(CustomizableSkeletalComponent);
-		if (USkeletalMeshComponent* SkeletalMeshComp = GetSkeletalMeshComponent()) 
+		if (USkeletalMeshComponent* SkeletalMeshComp = Super::GetSkeletalMeshComponent()) 
 		{
 			SkeletalMeshComponents.Add(SkeletalMeshComp);
 			CustomizableSkeletalComponents[0]->AttachToComponent(SkeletalMeshComp, FAttachmentTransformRules::KeepRelativeTransform);
@@ -52,6 +52,22 @@ ACustomizableSkeletalMeshActor::ACustomizableSkeletalMeshActor(const FObjectInit
 }
 
 
+USkeletalMeshComponent* ACustomizableSkeletalMeshActor::GetSkeletalMeshComponent(const FName& ComponentName)
+{
+	for (int32 ComponentIndex = 0; ComponentIndex < CustomizableSkeletalComponents.Num(); ComponentIndex++)
+	{
+		const UCustomizableSkeletalComponent* Component = CustomizableSkeletalComponents[ComponentIndex];
+		
+		if (Component->GetComponentName() == ComponentName)
+		{
+			return SkeletalMeshComponents[ComponentIndex];
+		}
+	}
+
+	return nullptr; 
+}
+
+
 void UCustomizableSkeletalMeshActorPrivate::Init(UCustomizableObjectInstance* Instance)
 {
 	UCustomizableObject* Object = Instance->GetCustomizableObject();
@@ -126,7 +142,7 @@ void ACustomizableSkeletalMeshActor::EnableDebugMaterial(bool bEnableDebugMateri
 	bRemoveDebugMaterial = bDebugMaterialEnabled && !bEnableDebugMaterial;
 	bDebugMaterialEnabled = bEnableDebugMaterial;
 
-	if (UCustomizableObjectInstance* COInstance = GetInstance())
+	if (UCustomizableObjectInstance* COInstance = GetCustomizableObjectInstance())
 	{
 		//Bind Instance Update delegate to Actor
 		COInstance->UpdatedDelegate.AddUniqueDynamic(this, &ACustomizableSkeletalMeshActor::SwitchComponentsMaterials);
@@ -158,7 +174,7 @@ void ACustomizableSkeletalMeshActor::SwitchComponentsMaterials(UCustomizableObje
 
 	if (bDebugMaterialEnabled || bRemoveDebugMaterial)
 	{
-		UCustomizableObjectInstance* COInstance = GetInstance();
+		UCustomizableObjectInstance* COInstance = GetCustomizableObjectInstance();
 
 		if (!COInstance)
 		{
@@ -178,15 +194,17 @@ void ACustomizableSkeletalMeshActor::SwitchComponentsMaterials(UCustomizableObje
 			}
 			else // Remove debugmaterial
 			{
+				const FName ComponentName = CustomizableSkeletalComponents[ObjectComponentIndex]->GetComponentName();
+				
 				// check if original materials already overriden
-				const TArray<TObjectPtr<UMaterialInterface>>* OverrideMaterials = COInstance->GetOverrideMaterials(FCustomizableObjectComponentIndex(ObjectComponentIndex));
+				const TArray<UMaterialInterface*> OverrideMaterials = COInstance->GetSkeletalMeshComponentOverrideMaterials(ComponentName);
 				const bool bUseOverrideMaterials = COInstance->GetCustomizableObject()->bEnableMeshCache && CVarEnableMeshCache.GetValueOnAnyThread();
 
-				if (bUseOverrideMaterials && OverrideMaterials->Num() > 0)
+				if (bUseOverrideMaterials && OverrideMaterials.Num() > 0)
 				{
-					for (int32 MatIndex = 0; MatIndex < OverrideMaterials->Num(); ++MatIndex)
+					for (int32 MatIndex = 0; MatIndex < OverrideMaterials.Num(); ++MatIndex)
 					{
-						SkeletalMeshComponents[ObjectComponentIndex]->SetMaterial(MatIndex, (*OverrideMaterials)[MatIndex]);
+						SkeletalMeshComponents[ObjectComponentIndex]->SetMaterial(MatIndex, OverrideMaterials[MatIndex]);
 					}
 				}
 				else
diff --git a/Engine/Plugins/Mutable/Source/CustomizableObject/Public/MuCO/CustomizableObjectInstance.h b/Engine/Plugins/Mutable/Source/CustomizableObject/Public/MuCO/CustomizableObjectInstance.h
--- a/Engine/Plugins/Mutable/Source/CustomizableObject/Public/MuCO/CustomizableObjectInstance.h
+++ b/Engine/Plugins/Mutable/Source/CustomizableObject/Public/MuCO/CustomizableObjectInstance.h
@@ -702,7 +702,8 @@ public:
 	/** If enabled, low-priority textures will generate resident mipmaps too. */
 	UE_API void SetForceGenerateResidentMips(bool bForceGenerateResidentMips);
 
-	UE_API const TArray<TObjectPtr<UMaterialInterface>>* GetOverrideMaterials(FCustomizableObjectComponentIndex ComponentIndex) const;
+	UFUNCTION(BlueprintCallable, Category = CustomizableObjectInstance)
+	CUSTOMIZABLEOBJECT_API TArray<UMaterialInterface*> GetSkeletalMeshComponentOverrideMaterials(const FName& ComponentName) const;
 	
 	// The following methods should only be used in an LOD management class
 	UE_API void SetIsBeingUsedByComponentInPlay(bool bIsUsedByComponent);
diff --git a/Engine/Plugins/Mutable/Source/CustomizableObject/Public/MuCO/CustomizableSkeletalMeshActor.h b/Engine/Plugins/Mutable/Source/CustomizableObject/Public/MuCO/CustomizableSkeletalMeshActor.h
--- a/Engine/Plugins/Mutable/Source/CustomizableObject/Public/MuCO/CustomizableSkeletalMeshActor.h
+++ b/Engine/Plugins/Mutable/Source/CustomizableObject/Public/MuCO/CustomizableSkeletalMeshActor.h
@@ -23,6 +23,12 @@ class ACustomizableSkeletalMeshActor : public ASkeletalMeshActor
 
 public:
 	UE_API ACustomizableSkeletalMeshActor(const FObjectInitializer& Initializer);
+
+	UFUNCTION(BlueprintCallable, Category = "CustomizableSkeletalMeshActor")
+	UE_API UCustomizableObjectInstance* GetCustomizableObjectInstance();
+	
+	UFUNCTION(BlueprintCallable, Category = "CustomizableSkeletalMeshActor")
+	UE_API USkeletalMeshComponent* GetSkeletalMeshComponent(const FName& ComponentName);
 	
 	UE_DEPRECATED(5.6, "Plase replace with custom code.")
 	UFUNCTION(BlueprintCallable, Category = CustomizableSkeletalMeshActor, meta=(DeprecatedFunction))
@@ -42,9 +48,6 @@ private:
 	UFUNCTION()
 	UE_API void SwitchComponentsMaterials(UCustomizableObjectInstance* Instance);
 
-	/** @return the instance represented in all the actor components. */
-	UE_API UCustomizableObjectInstance* GetInstance();
-
 	/** TODO: There are as many as components in the CO (not in the COI). */
 	UPROPERTY(VisibleAnywhere, Category = CustomizableSkeletalMeshActor)
 	TArray<TObjectPtr<UCustomizableSkeletalComponent>> CustomizableSkeletalComponents;
