[Mutable] Fix crash after DDC read request failure.

#rnx

diff --git a/Engine/Plugins/Mutable/Source/CustomizableObject/Internal/MuCO/UnrealMutableModelDiskStreamer.h b/Engine/Plugins/Mutable/Source/CustomizableObject/Internal/MuCO/UnrealMutableModelDiskStreamer.h
--- a/Engine/Plugins/Mutable/Source/CustomizableObject/Internal/MuCO/UnrealMutableModelDiskStreamer.h
+++ b/Engine/Plugins/Mutable/Source/CustomizableObject/Internal/MuCO/UnrealMutableModelDiskStreamer.h
@@ -33,6 +33,15 @@ private:
 	FArchive& m_ar;
 };
 
+#if WITH_EDITOR
+struct FDDCReadRequest
+{
+	FDDCReadRequest();
+
+	UE::DerivedData::FRequestOwner RequestOwner;
+	bool bSuccess = false;
+};
+#endif
 
 // Implementation of a mutable streamer using bulk storage.
 class FUnrealMutableModelBulkReader : public mu::FModelReader
@@ -77,7 +86,7 @@ protected:
 		TSharedPtr<FAsyncFileCallBack> FileCallback;
 
 #if WITH_EDITORONLY_DATA
-		TSharedPtr<UE::DerivedData::FRequestOwner> DDCReadRequest;
+		TSharedPtr<FDDCReadRequest> DDCReadRequest;
 #endif
 	};
 	
diff --git a/Engine/Plugins/Mutable/Source/CustomizableObject/Private/MuCO/UnrealMutableModelDiskStreamer.cpp b/Engine/Plugins/Mutable/Source/CustomizableObject/Private/MuCO/UnrealMutableModelDiskStreamer.cpp
--- a/Engine/Plugins/Mutable/Source/CustomizableObject/Private/MuCO/UnrealMutableModelDiskStreamer.cpp
+++ b/Engine/Plugins/Mutable/Source/CustomizableObject/Private/MuCO/UnrealMutableModelDiskStreamer.cpp
@@ -373,13 +373,13 @@ mu::FModelReader::FOperationID FUnrealMutableModelBulkReader::BeginReadBlock(con
 			Request.Policy = PolicyBuilder.Build();
 
 			FReadRequest ReadRequest;
-			ReadRequest.DDCReadRequest = MakeShared<FRequestOwner>(EPriority::High);
+			ReadRequest.DDCReadRequest = MakeShared<FDDCReadRequest>();
 
-			GetCache().Get(MakeArrayView(&Request, 1), *ReadRequest.DDCReadRequest.Get(),
-				[CompletionCallbackCapture = *CompletionCallback, Block, ResourceId, Buffer, Size](FCacheGetResponse&& Response)
+			GetCache().Get(MakeArrayView(&Request, 1), ReadRequest.DDCReadRequest->RequestOwner,
+				[CompletionCallbackCapture = *CompletionCallback, Request = ReadRequest.DDCReadRequest, Block, ResourceId, Buffer, Size](FCacheGetResponse&& Response)
 				{
-					bool bSuccess = Response.Status == EStatus::Ok;
-					if (ensure(bSuccess))
+					const bool bSuccess = Response.Status == EStatus::Ok;
+					if (bSuccess)
 					{
 						const FCompressedBuffer& CompressedBuffer = Response.Record.GetValue(ResourceId).GetData();
 						if (ensure(!CompressedBuffer.IsNull()))
@@ -394,6 +394,8 @@ mu::FModelReader::FOperationID FUnrealMutableModelBulkReader::BeginReadBlock(con
 							{
 								check(CompressedBuffer.TryDecompressTo(MakeMemoryView(Buffer, Size)));
 							}
+
+							Request->bSuccess = true;
 						}
 					}
 
@@ -479,7 +481,7 @@ bool FUnrealMutableModelBulkReader::IsReadCompleted(mu::FModelReader::FOperation
 #if WITH_EDITORONLY_DATA
 			else if (ReadRequest->DDCReadRequest)
 			{
-				return ReadRequest->DDCReadRequest->Poll(); // ReadLock
+				return ReadRequest->DDCReadRequest->RequestOwner.Poll(); // ReadLock
 			}
 #endif
 		}
@@ -518,11 +520,10 @@ bool FUnrealMutableModelBulkReader::EndRead(mu::FModelReader::FOperationID Opera
 #if WITH_EDITORONLY_DATA
 			else if (ReadRequest->DDCReadRequest)
 			{
-				ReadRequest->DDCReadRequest->Wait();
+				ReadRequest->DDCReadRequest->RequestOwner.Wait();
 				
-				// bSuccess TODO Pere: find a way to know if the request completed successfully
 				bCompleted = true;
-				bResult = true;
+				bResult = ReadRequest->DDCReadRequest->bSuccess;
 			}
 #endif
 			if (!bCompleted)
@@ -624,5 +625,10 @@ void FUnrealMutableModelBulkWriterCook::CloseWriteFile()
 }
 
 
+FDDCReadRequest::FDDCReadRequest() :
+RequestOwner(UE::DerivedData::EPriority::High)
+{
+}
+
 #endif
 
