[Mutable] Optimize MeshFormatBuffer for same format but different number of components.
Fix incorrect initial format on some meshes.

#rnx

diff --git a/Engine/Plugins/Mutable/Source/CustomizableObjectEditor/Private/MuCOE/GenerateMutableSource/GenerateMutableSourceMesh.cpp b/Engine/Plugins/Mutable/Source/CustomizableObjectEditor/Private/MuCOE/GenerateMutableSource/GenerateMutableSourceMesh.cpp
--- a/Engine/Plugins/Mutable/Source/CustomizableObjectEditor/Private/MuCOE/GenerateMutableSource/GenerateMutableSourceMesh.cpp
+++ b/Engine/Plugins/Mutable/Source/CustomizableObjectEditor/Private/MuCOE/GenerateMutableSource/GenerateMutableSourceMesh.cpp
@@ -2483,7 +2483,7 @@ TSharedPtr<mu::FMesh> ConvertSkeletalMeshToMutable(
 		{
 			MeshData.bHasVertexColors = SkeletalMesh->GetHasVertexColors();
 			MeshData.NumTexCoordChannels = ImportedModel->LODModels[Source.LODIndex].NumTexCoords;
-			MeshData.MaxBoneIndexTypeSizeBytes = ImportedModel->LODModels[Source.LODIndex].RequiredBones.Num() > 256 ? 2 : 1;
+			MeshData.MaxBoneIndexTypeSizeBytes = MutableMesh->BoneMap.Num() > 256 ? 2 : 1;
 			MeshData.MaxNumBonesPerVertex = ImportedModel->LODModels[Source.LODIndex].GetMaxBoneInfluences();
 			MeshData.bHasRealTimeMorphs = Context.Options.bRealTimeMorphTargetsEnabled && SkeletalMesh->GetMorphTargets().Num();
 			MeshData.bHasClothing = Context.Options.bClothingEnabled && ImportedModel->LODModels[Source.LODIndex].HasClothData();
diff --git a/Engine/Plugins/Mutable/Source/MutableRuntime/Private/MuR/OpMeshFormat.cpp b/Engine/Plugins/Mutable/Source/MutableRuntime/Private/MuR/OpMeshFormat.cpp
--- a/Engine/Plugins/Mutable/Source/MutableRuntime/Private/MuR/OpMeshFormat.cpp
+++ b/Engine/Plugins/Mutable/Source/MutableRuntime/Private/MuR/OpMeshFormat.cpp
@@ -271,6 +271,17 @@ namespace mu
 						SourceBuf += SourceElemSize;
 					}
 				}
+				else if (ResultChannel.Format == SourceFormat && ResultComponents > SourceComponents)
+				{
+					const int32 ResultFormatSize = GetMeshFormatData(ResultChannel.Format).SizeInBytes;
+					for (int32 v = 0; v < SourceElementCount; ++v)
+					{
+						FMemory::Memzero(ResultBuf, ResultFormatSize * ResultComponents);
+						FMemory::Memcpy(ResultBuf, SourceBuf, ResultFormatSize * SourceComponents);
+						ResultBuf += ResultElemSize;
+						SourceBuf += SourceElemSize;
+					}
+				}
 				else if (ResultChannel.Format == EMeshBufferFormat::PackedDir8_W_TangentSign
 					||
 					ResultChannel.Format == EMeshBufferFormat::PackedDirS8_W_TangentSign)
@@ -617,11 +628,6 @@ namespace mu
 							BoneIndexFormat = EMeshBufferFormat::UInt16;
 							InOutFormatBufferSet.UpdateOffsets(ResultBuf);
 						}
-						else if (HighestBoneIndexInUse > 0x7f && BoneIndexFormat == EMeshBufferFormat::UInt8)
-						{
-							BoneIndexFormat = EMeshBufferFormat::UInt16;
-							InOutFormatBufferSet.UpdateOffsets(ResultBuf);
-						}
 					}
 				}
 			}
