[Mutable] Simplify DoComponentNeedUpdate.

#rnx

diff --git a/Engine/Plugins/Mutable/Source/CustomizableObject/Internal/MuCO/CustomizableObjectInstancePrivate.h b/Engine/Plugins/Mutable/Source/CustomizableObject/Internal/MuCO/CustomizableObjectInstancePrivate.h
--- a/Engine/Plugins/Mutable/Source/CustomizableObject/Internal/MuCO/CustomizableObjectInstancePrivate.h
+++ b/Engine/Plugins/Mutable/Source/CustomizableObject/Internal/MuCO/CustomizableObjectInstancePrivate.h
@@ -324,9 +324,6 @@ private:
 
 	UE_API bool DoComponentsNeedUpdate(UCustomizableObjectInstance* CustomizableObjectInstance, const TSharedRef<FUpdateContextPrivate>& OperationData, bool& bOutEmptyMesh);
 
-	// TODO: Why is this a method?
-	UE_API void SetLastMeshId(FCustomizableObjectComponentIndex ObjectComponentIndex, int32 LODIndex, mu::FResourceID MeshId);
-
 public:
 	UE_API bool LoadParametersFromProfile(int32 ProfileIndex);
 	
diff --git a/Engine/Plugins/Mutable/Source/CustomizableObject/Private/MuCO/CustomizableObjectInstance.cpp b/Engine/Plugins/Mutable/Source/CustomizableObject/Private/MuCO/CustomizableObjectInstance.cpp
--- a/Engine/Plugins/Mutable/Source/CustomizableObject/Private/MuCO/CustomizableObjectInstance.cpp
+++ b/Engine/Plugins/Mutable/Source/CustomizableObject/Private/MuCO/CustomizableObjectInstance.cpp
@@ -203,20 +203,6 @@ UTexture2D* UCustomizableInstancePrivate::CreateTexture(const FString& TextureNa
 }
 
 
-void UCustomizableInstancePrivate::SetLastMeshId(FCustomizableObjectComponentIndex ObjectComponentIndex, int32 LODIndex, mu::FResourceID MeshId)
-{
-	FCustomizableInstanceComponentData* ComponentData = GetComponentData(ObjectComponentIndex);
-	if (ComponentData && ComponentData->LastMeshIdPerLOD.IsValidIndex(LODIndex))
-	{
-		ComponentData->LastMeshIdPerLOD[LODIndex] = MeshId;
-	}
-	else
-	{
-		check(false);
-	}
-}
-
-
 void UCustomizableInstancePrivate::InvalidateGeneratedData()
 {
 	SkeletalMeshStatus = ESkeletalMeshStatus::NotGenerated;
@@ -1975,52 +1961,12 @@ bool UCustomizableInstancePrivate::DoComponentsNeedUpdate(UCustomizableObjectIns
 	{
 		return false;
 	}
-
-	const UModelResources& ModelResources = *CustomizableObject->GetPrivate()->GetModelResources();
-
+	
+	bHasInvalidMesh = false;
+	bool bNeedsUpdate = false;
+	
 	const int32 NumInstanceComponents = OperationData->InstanceUpdateData.Components.Num();
-
-	// To be indexed with instance component index
-	TArray<bool> ComponentWithMesh;
-	ComponentWithMesh.Init(false, NumInstanceComponents);
-
-	TArray<mu::FResourceID> MeshIDs;
-	MeshIDs.Init(MAX_uint64, NumInstanceComponents * MAX_MESH_LOD_COUNT);
-
-	// Gather the Mesh Ids of all components, and validate the integrity of the meshes to generate. 
-	for (int32 InstanceComponentIndex = 0; InstanceComponentIndex < NumInstanceComponents; ++InstanceComponentIndex)
-	{
-		const FInstanceUpdateData::FComponent& Component = OperationData->InstanceUpdateData.Components[InstanceComponentIndex];
-		const FName ComponentName = OperationData->ComponentNames[Component.Id.GetValue()];
-
-		for (int32 LODIndex = OperationData->GetFirstRequestedLOD()[ComponentName]; LODIndex < Component.LODCount; ++LODIndex)
-		{
-			const FInstanceUpdateData::FLOD& LOD = OperationData->InstanceUpdateData.LODs[Component.FirstLOD + LODIndex];
-
-			if (!LOD.Mesh)
-			{
-				continue;
-			}
-
-			if (LOD.SurfaceCount == 0 && !LOD.Mesh->IsReference())
-			{
-				continue;
-			}
-
-			// Unreal does not support empty sections.
-			if (!LOD.Mesh->IsReference() && LOD.Mesh->GetVertexCount() == 0) // else
-			{
-				UE_LOG(LogMutable, Error, TEXT("Failed to generate SkeletalMesh for CO Instance [%s]. CO [%s] has invalid geometry for LOD [%d] Component [%d]."),
-					*Public->GetName(), *CustomizableObject->GetName(),
-					LODIndex, NumInstanceComponents);
-				bHasInvalidMesh = true;
-				continue;
-			}
-
-			ComponentWithMesh[InstanceComponentIndex] = true;
-			MeshIDs[InstanceComponentIndex * MAX_MESH_LOD_COUNT + LODIndex] = LOD.MeshID;
-		}
-	}
+	bNeedsUpdate |= NumInstanceComponents != SkeletalMeshes.Num();
 
 	// Find which components need an update
 	OperationData->MeshChangedPerInstanceComponent.Init(false, NumInstanceComponents);
@@ -2029,122 +1975,54 @@ bool UCustomizableInstancePrivate::DoComponentsNeedUpdate(UCustomizableObjectIns
 	{
 		const FInstanceUpdateData::FComponent& Component = OperationData->InstanceUpdateData.Components[InstanceComponentIndex];
 		const FCustomizableObjectComponentIndex ObjectComponentIndex = Component.Id;
-		if (!ObjectComponentIndex.IsValid())
-		{
-			continue;
-		}
+		const FName ComponentName = OperationData->ComponentNames[ObjectComponentIndex.GetValue()];
+		
+		bool bComponentWithMesh = false;
 
-		const FName ComponentName = ModelResources.ComponentNamesPerObjectComponent[ObjectComponentIndex.GetValue()];
+		const int32 FirstLOD = OperationData->bStreamMeshLODs ?
+			OperationData->FirstLODAvailable[ComponentName] :
+			OperationData->GetFirstRequestedLOD()[ComponentName];
 
-		if (OperationData->bUseMeshCache)
+		// Look for invalid meshes
+		for (int32 LODIndex = FirstLOD; LODIndex < Component.LODCount; ++LODIndex)
 		{
-			USkeletalMesh* CachedMesh = CustomizableObject->GetPrivate()->MeshCache.Get(*OperationData->GetMeshDescriptors(ObjectComponentIndex));
-			if (CachedMesh)
+			const FInstanceUpdateData::FLOD& LOD = OperationData->InstanceUpdateData.LODs[Component.FirstLOD + LODIndex];
+
+			if (!LOD.Mesh)
 			{
-				TObjectPtr<USkeletalMesh>* SkeletalMesh = SkeletalMeshes.Find(ComponentName);
-				const bool bMeshNeedsUpdate = !SkeletalMesh || (*SkeletalMesh != CachedMesh);
-				OperationData->MeshChangedPerInstanceComponent[InstanceComponentIndex] = bMeshNeedsUpdate;
-				ComponentWithMesh[InstanceComponentIndex] = true;
+				bHasInvalidMesh |= bComponentWithMesh; // Mesh will be invalid if the mu::mesh is missing on higher LODs.
 				continue;
 			}
-		}
 
-		// Components with mesh must have valid geometry at CurrentMaxLOD
-
-		TObjectPtr<USkeletalMesh>* SkeletalMeshPtr = SkeletalMeshes.Find(ComponentName);
-		const bool bHadSkeletalMesh = SkeletalMeshPtr && *SkeletalMeshPtr;
-
-		if (Component.LODCount == 0)
-		{
-			// We don't have a mesh in the component, so it has changed if we had one before.
-			OperationData->MeshChangedPerInstanceComponent[InstanceComponentIndex] = bHadSkeletalMesh;
-			continue;
-		}
-
-		const FInstanceUpdateData::FLOD& LOD = OperationData->InstanceUpdateData.LODs[Component.FirstLOD];
-		bool bIsReferenced = LOD.Mesh && LOD.Mesh->IsReference();
-		if (!bIsReferenced)
-		{
-			if (ComponentWithMesh[InstanceComponentIndex] && MeshIDs[InstanceComponentIndex * MAX_MESH_LOD_COUNT + OperationData->NumLODsAvailable[ComponentName] - 1] == MAX_uint64)
+			if (!LOD.Mesh->IsReference() && LOD.Mesh->GetVertexCount() == 0)
 			{
-				UE_LOG(LogMutable, Error, TEXT("Failed to generate SkeletalMesh for CO Instance [%s]. CO [%s] is missing geometry for LOD [%d] Object Component [%d]."),
+				UE_LOG(LogMutable, Error, TEXT("Failed to generate SkeletalMesh for CO Instance [%s]. CO [%s] has invalid geometry for LOD [%d] Component [%d]."),
 					*Public->GetName(), *CustomizableObject->GetName(),
-					OperationData->NumLODsAvailable[ComponentName] - 1, ObjectComponentIndex.GetValue());
-				bHasInvalidMesh = true;
-				continue;
+					LODIndex, InstanceComponentIndex);
+				bHasInvalidMesh = true; // Unreal does not support empty LODs.
 			}
-		}
-
-		// If the component wasn't there and now is there, we need to update it.
-		OperationData->MeshChangedPerInstanceComponent[InstanceComponentIndex] = !bHadSkeletalMesh && LOD.Mesh && (LOD.Mesh->GetFaceCount()>0);
 
-		const FCustomizableInstanceComponentData* ComponentData = GetComponentData(ObjectComponentIndex);
-		if (!ComponentData) // Could be nullptr if the component has not been generated.
-		{
-			continue;
+			check(LOD.SurfaceCount > 0);
+			bComponentWithMesh = true;
 		}
 
-		// Update if MeshIDs are different
-		const int32 ComponentOffset = InstanceComponentIndex * MAX_MESH_LOD_COUNT;
-		for (int32 MeshIndex = 0; !OperationData->MeshChangedPerInstanceComponent[InstanceComponentIndex] && MeshIndex < MAX_MESH_LOD_COUNT; ++MeshIndex)
-		{
-			OperationData->MeshChangedPerInstanceComponent[InstanceComponentIndex] = MeshIDs[ComponentOffset + MeshIndex] != ComponentData->LastMeshIdPerLOD[MeshIndex];
-		}
+		bool bMeshChanged = true;
 
-		// If the component SkeletalMesh was build with scalability settings different than the current, we need to update.
-		// The mismatch can happen when the scalability setting is changed and an update is requested with a MinLOD lower 
-		// than the new FirstLODAvailable. This is an artifact of our LOD management system setting MinLOD.
-		if (!OperationData->bStreamMeshLODs &&
-			!OperationData->MeshChangedPerInstanceComponent[InstanceComponentIndex] &&
-			bHadSkeletalMesh) 
+		// Compare Mesh IDs 
+		const TArray<mu::FResourceID>* MeshIDs = OperationData->GetMeshDescriptors(ObjectComponentIndex);
+		if (ensure(MeshIDs))
 		{
-			const int32 FirstRequestedLOD = OperationData->GetFirstRequestedLOD()[ComponentName];
-			for (int32 LODIndex = OperationData->FirstLODAvailable[ComponentName]; LODIndex < FirstRequestedLOD; ++LODIndex)
+			if (const FCustomizableInstanceComponentData* ComponentData = GetComponentData(ObjectComponentIndex)) // Could be nullptr if the component has not been generated.
 			{
-				TIndirectArray<FSkeletalMeshLODRenderData>& LODRenderData = (*SkeletalMeshPtr)->GetResourceForRendering()->LODRenderData;
-				if (LODRenderData[LODIndex].RenderSections.Num() != LODRenderData[FirstRequestedLOD].RenderSections.Num())
-				{
-					OperationData->MeshChangedPerInstanceComponent[InstanceComponentIndex] = true;
-					break;
-				}
+				bMeshChanged = *MeshIDs != ComponentData->LastMeshIdPerLOD;
 			}
 		}
 
+		OperationData->MeshChangedPerInstanceComponent[InstanceComponentIndex] = bMeshChanged;
+		bNeedsUpdate |= bMeshChanged;
 	}
 
-	bool bChanged =
-		OperationData->MeshChangedPerInstanceComponent.Num() != SkeletalMeshes.Num()
-		||
-		OperationData->MeshChangedPerInstanceComponent.Find(true) != INDEX_NONE;
-
-	// It also changed if we removed a component that we did have before
-	if (!bChanged)
-	{
-		for (const TPair<FName, TObjectPtr<USkeletalMesh>>& OldMesh : SkeletalMeshes)
-		{
-			bool bFound = false;
-			for (int32 InstanceComponentIndex = 0; InstanceComponentIndex < NumInstanceComponents; ++InstanceComponentIndex)
-			{
-				const FInstanceUpdateData::FComponent& Component = OperationData->InstanceUpdateData.Components[InstanceComponentIndex];
-				const FCustomizableObjectComponentIndex ObjectComponentIndex = Component.Id;
-				const FName ComponentName = ModelResources.ComponentNamesPerObjectComponent[ObjectComponentIndex.GetValue()];
-
-				if (ComponentName == OldMesh.Key)
-				{
-					bFound = true;
-					break;
-				}
-			}
-
-			if (!bFound)
-			{
-				bChanged = true;
-				break;
-			}
-		}
-	}
-
-	return !bHasInvalidMesh && bChanged;
+	return !bHasInvalidMesh && bNeedsUpdate;
 }
 
 
@@ -2345,33 +2223,31 @@ bool UCustomizableInstancePrivate::UpdateSkeletalMesh_PostBeginUpdate0(UCustomiz
 			continue;
 		}
 		
-		if (OperationData->bUseMeshCache)
+		TArray<mu::FResourceID>* MeshDescriptors = OperationData->GetMeshDescriptors(ObjectComponentIndex);
+		check(MeshDescriptors);
+		check(MeshDescriptors->Num() == MAX_MESH_LOD_COUNT);
+		ComponentData->LastMeshIdPerLOD = *MeshDescriptors;
+
+		uint32 MeshHash = INDEX_NONE;
+		for (int32 LODIndex = 0; LODIndex < Component->LODCount && LODIndex < MAX_MESH_LOD_COUNT; ++LODIndex)
 		{
-			TArray<mu::FResourceID>* MeshDescriptors = OperationData->GetMeshDescriptors(ObjectComponentIndex);
-			if (!MeshDescriptors)
-			{
-				continue;
-			}
+			MeshHash = HashCombine(MeshHash, GetTypeHash((*MeshDescriptors)[LODIndex]));
+		}
 
+		if (OperationData->bUseMeshCache)
+		{
 			if (USkeletalMesh* CachedMesh = CustomizableObject->GetPrivate()->MeshCache.Get(*MeshDescriptors))
 			{
-				check(MeshDescriptors->Num() == MAX_MESH_LOD_COUNT);
-				ComponentData->LastMeshIdPerLOD = *MeshDescriptors;
 				SkeletalMeshes.Add(ComponentName, CachedMesh);
 				continue;
 			}
 		}
 
-		// Reset last mesh IDs.
-		ComponentData->LastMeshIdPerLOD.Init(MAX_uint64, MAX_MESH_LOD_COUNT);
-
 		// We need the first valid mesh. get it from the component, considering that some LOSs may have been skipped.
 		TSharedPtr<const mu::FMesh> ComponentMesh;
-		uint32 MeshHash = INDEX_NONE;
-		for (int32 FirstValidLODIndex = Component->FirstLOD; FirstValidLODIndex < OperationData->InstanceUpdateData.LODs.Num() && !ComponentMesh; ++FirstValidLODIndex)
+		for (int32 LODIndex = 0; LODIndex < Component->LODCount && !ComponentMesh; ++LODIndex)
 		{
-			ComponentMesh = OperationData->InstanceUpdateData.LODs[FirstValidLODIndex].Mesh;
-			MeshHash = HashCombine(MeshHash,GetTypeHash(OperationData->InstanceUpdateData.LODs[FirstValidLODIndex].MeshID));
+			ComponentMesh = OperationData->InstanceUpdateData.LODs[Component->FirstLOD + LODIndex].Mesh;
 		}
 		
 		if (!ComponentMesh)
@@ -4768,18 +4644,6 @@ bool UCustomizableInstancePrivate::BuildOrCopyRenderData(const TSharedRef<FUpdat
 	const UModelResources& ModelResources = CustomizableObject->GetPrivate()->GetModelResourcesChecked();
 	const FName ComponentName = ModelResources.ComponentNamesPerObjectComponent[Component->Id.GetValue()];
 
-	for (int32 LODIndex = OperationData->GetFirstRequestedLOD()[ComponentName]; LODIndex < Component->LODCount; ++LODIndex)
-	{
-		const FInstanceUpdateData::FLOD& LOD = OperationData->InstanceUpdateData.LODs[Component->FirstLOD+LODIndex];
-
-		if (!LOD.Mesh || LOD.SurfaceCount == 0)
-		{
-			continue;	
-		}
-
-		SetLastMeshId(Component->Id, LODIndex, LOD.MeshID);
-	}
-	
 	const int32 FirstLOD = OperationData->bStreamMeshLODs ?
 		OperationData->FirstLODAvailable[ComponentName] :
 		OperationData->GetFirstRequestedLOD()[ComponentName];
diff --git a/Engine/Plugins/Mutable/Source/CustomizableObject/Private/MuCO/CustomizableObjectSystem.cpp b/Engine/Plugins/Mutable/Source/CustomizableObject/Private/MuCO/CustomizableObjectSystem.cpp
--- a/Engine/Plugins/Mutable/Source/CustomizableObject/Private/MuCO/CustomizableObjectSystem.cpp
+++ b/Engine/Plugins/Mutable/Source/CustomizableObject/Private/MuCO/CustomizableObjectSystem.cpp
@@ -3841,35 +3841,37 @@ namespace impl
 		
 		FixLODs(Operation);
 		
-		if (Operation->bUseMeshCache)
+		Operation->InitMeshDescriptors(Operation->NumObjectComponents);
+
+		const TMap<FName, uint8>& NumLODsAvailable = Operation->NumLODsAvailable;
+
+		for (int32 InstanceComponentIndex = 0; InstanceComponentIndex < Operation->NumInstanceComponents; ++InstanceComponentIndex)
 		{
-			Operation->InitMeshDescriptors(Operation->NumObjectComponents);
+			FCustomizableObjectComponentIndex ObjectComponentIndex(Operation->MutableInstance->GetComponentId(InstanceComponentIndex));
+			TArray<mu::FResourceID>* MeshId = Operation->GetMeshDescriptors(ObjectComponentIndex);
+			check(MeshId);
 
-			const TMap<FName, uint8>& NumLODsAvailable = Operation->NumLODsAvailable;
-		
-			for (int32 InstanceComponentIndex = 0; InstanceComponentIndex < Operation->NumInstanceComponents; ++InstanceComponentIndex)
+			if (!MeshId)
 			{
-				FCustomizableObjectComponentIndex ObjectComponentIndex(Operation->MutableInstance->GetComponentId(InstanceComponentIndex));
-				TArray<mu::FResourceID>* MeshId = Operation->GetMeshDescriptors(ObjectComponentIndex);
-				if (!MeshId)
-				{
-					continue;
-				}
+				continue;
+			}
 
-				MeshId->Init(MAX_uint64, MAX_MESH_LOD_COUNT);
+			MeshId->Init(MAX_uint64, MAX_MESH_LOD_COUNT);
 
-				const FName& ComponentName = Operation->ComponentNames[ObjectComponentIndex.GetValue()];
+			const FName& ComponentName = Operation->ComponentNames[ObjectComponentIndex.GetValue()];
 
-				const int32 FirstLOD = Operation->bStreamMeshLODs ?
-					Operation->FirstLODAvailable[ComponentName] :
-					Operation->GetFirstRequestedLOD()[ComponentName];
-	
-				for (int32 LODIndex = FirstLOD; LODIndex < NumLODsAvailable[ComponentName]; ++LODIndex)
-				{
-					(*MeshId)[LODIndex] = Operation->MutableInstance->GetMeshId(InstanceComponentIndex, LODIndex);
-				}
+			const int32 FirstLOD = Operation->bStreamMeshLODs ?
+				Operation->FirstLODAvailable[ComponentName] :
+				Operation->GetFirstRequestedLOD()[ComponentName];
+
+			for (int32 LODIndex = FirstLOD; LODIndex < NumLODsAvailable[ComponentName]; ++LODIndex)
+			{
+				(*MeshId)[LODIndex] = Operation->MutableInstance->GetMeshId(InstanceComponentIndex, LODIndex);
 			}
-			
+		}
+
+		if (Operation->bUseMeshCache)
+		{
 			for (const TArray<mu::FResourceID>& MeshId : Operation->GetMeshDescriptors())
 			{
 				if (USkeletalMesh* CachedMesh = ObjectPrivate->MeshCache.Get(MeshId))
