[Mutable] Fix incorrect check in DoComponentsNeedUpdate.

#rnx

diff --git a/Engine/Plugins/Mutable/Source/CustomizableObject/Private/MuCO/CustomizableObjectInstance.cpp b/Engine/Plugins/Mutable/Source/CustomizableObject/Private/MuCO/CustomizableObjectInstance.cpp
--- a/Engine/Plugins/Mutable/Source/CustomizableObject/Private/MuCO/CustomizableObjectInstance.cpp
+++ b/Engine/Plugins/Mutable/Source/CustomizableObject/Private/MuCO/CustomizableObjectInstance.cpp
@@ -1994,6 +1994,11 @@ bool UCustomizableInstancePrivate::DoComponentsNeedUpdate(UCustomizableObjectIns
 				continue;
 			}
 
+			if (LOD.SurfaceCount == 0 && !LOD.Mesh->IsReference())
+			{
+				continue;
+			}
+
 			if (!LOD.Mesh->IsReference() && LOD.Mesh->GetVertexCount() == 0)
 			{
 				UE_LOG(LogMutable, Error, TEXT("Failed to generate SkeletalMesh for CO Instance [%s]. CO [%s] has invalid geometry for LOD [%d] Component [%d]."),
