[Mutable] Optimize MeshRemoveVerticesWithCullSet

#rnx

diff --git a/Engine/Plugins/Mutable/Source/MutableRuntime/Private/MuR/OpMeshRemove.cpp b/Engine/Plugins/Mutable/Source/MutableRuntime/Private/MuR/OpMeshRemove.cpp
--- a/Engine/Plugins/Mutable/Source/MutableRuntime/Private/MuR/OpMeshRemove.cpp
+++ b/Engine/Plugins/Mutable/Source/MutableRuntime/Private/MuR/OpMeshRemove.cpp
@@ -4,6 +4,7 @@
 
 #include "MuR/MeshPrivate.h"
 #include "MuR/MutableTrace.h"
+#include "MuR/ParallelExecutionUtils.h"
 #include "MuR/Platform.h"
 
 
@@ -131,212 +132,7 @@ namespace mu
         return -1;
     }
 
-    // TODO: Remove once the new implementation is proven to be faster for all cases and does not present 
-    // any problem. For now keep it as a reference.
-	//void MeshRemoveVerticesWithMap( FMesh* Result, const TBitArray<>& RemovedVertices)
-	//{
-    //    uint32 RemovedVertexCount = RemovedVertices.Num();
-
-	//	int32 FirstFreeVertex = 0;
-	//	int32 RemovedIndices = 0;
-
-    //    // Rebuild index buffers
-
-    //    // Map from source vertex index, to new vertex index for used vertices.
-    //    // These are indices as in the index buffer, not the absolute vertex index as in the
-    //    // vertexbuffer EMeshBufferSemantic::VertexIndex buffers.
-	//	TArray<int32> UsedVertexMap;
-    //    TBitArray<> UsedFaces;
-	//	UsedVertexMap.Init(-1, Result->GetVertexCount());
-    //    UsedFaces.Init(false, Result->GetIndexCount() / 3);  
-    //    {
-    //        if ( Result->GetIndexBuffers().GetElementSize(0)==4 )
-    //        {
-    //            MeshBufferIteratorConst<EMeshBufferFormat::UInt32,uint32,1> itSource( Result->GetIndexBuffers(), EMeshBufferSemantic::VertexIndex );
-    //            MeshBufferIterator<EMeshBufferFormat::UInt32,uint32,1> itDest( Result->GetIndexBuffers(), EMeshBufferSemantic::VertexIndex );
-
-    //            int32 IndexCount = Result->GetIndexCount();
-	//			check(IndexCount%3==0);
-    //            for ( int32 f=0; f<IndexCount/3; ++f )
-    //            {
-    //                uint32 sourceIndices[3];
-    //                sourceIndices[0] = (*itSource)[0];
-    //                sourceIndices[1] = (*itSource)[1];
-    //                sourceIndices[2] = (*itSource)[2];
-
-	//				check(sourceIndices[0] < RemovedVertexCount);
-	//				check(sourceIndices[1] < RemovedVertexCount);
-	//				check(sourceIndices[2] < RemovedVertexCount);
-
-	//				int32 RemoveCount =  (RemovedVertices[ sourceIndices[0] ]
-    //                        + RemovedVertices[ sourceIndices[1] ]
-    //                        + RemovedVertices[ sourceIndices[2] ]
-    //                        );
-	//				bool bFaceRemoved = RemoveCount > 0;
-    //                if (!bFaceRemoved)
-    //                {
-    //                    UsedFaces[f] = true;
-
-    //                    for (int32 i=0;i<3;++i)
-    //                    {
-    //                        uint32 sourceIndex = sourceIndices[i];
-
-    //                        if (UsedVertexMap[ sourceIndex ] < 0 )
-    //                        {
-	//							UsedVertexMap[ sourceIndex ] = FirstFreeVertex;
-    //                            FirstFreeVertex++;
-    //                        }
-
-    //                        uint32 destIndex = UsedVertexMap[ sourceIndex ];
-    //                        *(uint32*)itDest.ptr() = destIndex;
-
-    //                        itDest++;
-    //                    }
-    //                }
-
-    //                itSource+=3;
-    //            }
-
-    //            RemovedIndices = itSource - itDest;
-    //        }
-
-    //        else if ( Result->GetIndexBuffers().GetElementSize(0)==2 )
-    //        {
-    //            MeshBufferIteratorConst<EMeshBufferFormat::UInt16,uint16,1> itSource( Result->GetIndexBuffers(), EMeshBufferSemantic::VertexIndex );
-    //            MeshBufferIterator<EMeshBufferFormat::UInt16,uint16,1> itDest( Result->GetIndexBuffers(), EMeshBufferSemantic::VertexIndex );
-
-    //            int32 indexCount = Result->GetIndexCount();
-    //            for ( int32 f=0; f<indexCount/3; ++f )
-    //            {
-    //                uint16 sourceIndices[3];
-    //                sourceIndices[0] = (*itSource)[0];
-    //                sourceIndices[1] = (*itSource)[1];
-    //                sourceIndices[2] = (*itSource)[2];
-
-	//				check(sourceIndices[0] < RemovedVertexCount);
-	//				check(sourceIndices[1] < RemovedVertexCount);
-	//				check(sourceIndices[2] < RemovedVertexCount);
-
-	//				int32 RemoveCount =  (RemovedVertices[ sourceIndices[0] ]
-    //                        + RemovedVertices[ sourceIndices[1] ]
-    //                        + RemovedVertices[ sourceIndices[2] ]
-    //                        );
-	//				bool bFaceRemoved = RemoveCount > 0;
-
-    //                if (!bFaceRemoved)
-    //                {
-    //                    UsedFaces[f] = true;
-    //                    for (int32 i = 0; i < 3; ++i)
-    //                    {
-    //                        uint16 sourceIndex = sourceIndices[i];
-
-    //                        if (UsedVertexMap[ sourceIndex ] < 0 )
-    //                        {
-	//							UsedVertexMap[ sourceIndex ] = FirstFreeVertex;
-    //                            FirstFreeVertex++;
-    //                        }
-
-    //                        uint16 destIndex = (uint16)UsedVertexMap[ sourceIndex ];
-    //                        *(uint16*)itDest.ptr() = destIndex;
-
-    //                        itDest++;
-    //                    }
-    //                }
-
-    //                itSource+=3;
-    //            }
-
-    //            RemovedIndices = itSource - itDest;
-    //        }
-
-    //        else
-    //        {
-    //            // Index buffer format case not implemented
-    //            check( false );
-    //        }
-
-    //        check( RemovedIndices%3==0 );
-
-    //        int32 FaceCount = Result->GetFaceCount();
-    //        Result->GetIndexBuffers().SetElementCount( FaceCount*3-RemovedIndices );
-    //    }
-
-
-    //    // Rebuild the vertex buffers
-
-	//	// If we had implicit indices, make them explicit or relative to keep them valid
-	//	if (RemovedIndices && Result->AreVertexIdsImplicit())
-	//	{
-	//		Result->MakeVertexIdsRelative();
-	//	}
-
-	//	// The temp array is necessary because if the vertex buffer is not sorted according to the index buffer we cannot do it in-place
-	//	// This happens with some mesh import options.
-	//	TArray<uint8> Temp;
-	//	for ( int32 b=0; b<Result->GetVertexBuffers().GetBufferCount(); ++b )
-    //    {
-    //        int32 elemSize = Result->GetVertexBuffers().GetElementSize( b );
-    //        const uint8* SourceData = Result->GetVertexBuffers().GetBufferData( b );
-
-	//		Temp.SetNumUninitialized(FirstFreeVertex*elemSize,EAllowShrinking::No);
-    //        uint8* DestData = Temp.GetData();
-
-    //        for ( int32 v=0; v<Result->GetVertexCount(); ++v )
-    //        {
-    //            int32 span = 0;
-    //            for ( int32 s=0; v+s<Result->GetVertexCount(); ++s )
-    //            {
-    //                if (UsedVertexMap[v+s]>=0 )
-    //                {
-    //                    if (span==0)
-    //                    {
-    //                        ++span;
-    //                    }
-    //                    else
-    //                    {
-    //                        if (UsedVertexMap[v+s] == UsedVertexMap[v+s-1]+1 )
-    //                        {
-    //                            ++span;
-    //                        }
-    //                        else
-    //                        {
-    //                            break;
-    //                        }
-    //                    }
-    //                }
-    //                else
-    //                {
-    //                    break;
-    //                }
-    //            }
-
-
-    //            if (span>0)
-    //            {
-    //                FMemory::Memcpy( DestData+elemSize*UsedVertexMap[v], SourceData+elemSize*v, elemSize*span );
-
-    //                v += span-1;
-    //            }
-    //        }
-
-	//		// Copy from temp buffer to final vertex buffer
-	//		FMemory::Memcpy( Result->GetVertexBuffers().GetBufferData(b), DestData, FirstFreeVertex*elemSize);
-    //    }
-    //    
-    //    Result->GetVertexBuffers().SetElementCount(FirstFreeVertex);
-
-    //    //Temp fix, transform the vertex map to a bitset.
-    //    TBitArray<> UsedVertices;
-    //    UsedVertices.SetNum(RemovedVertexCount, false);
-
-    //    for (uint32 I = 0; I < RemovedVertexCount; ++I)
-    //    {
-    //        UsedVertices[I] = RemovedVertices[I] == 0;
-    //    }
-
-    //    MeshRemoveRecreateSurface(Result, UsedVertices, UsedFaces);
-    //}
-
+	
     void MeshRemoveVerticesWithCullSet(FMesh* Result, const TBitArray<>& VerticesToCull, bool bRemoveIfAllVerticesCulled)
     { 
         MUTABLE_CPUPROFILER_SCOPE(MeshRemoveVerticesWithCullSet);
@@ -346,79 +142,132 @@ namespace mu
         const int32 NumFaces = Result->GetFaceCount();
         const int32 NumVertices = Result->GetVertexCount();
 
-		TBitArray<> UsedVertices;
-		UsedVertices.SetNum(NumVertices, false);
-	
-		TBitArray<> UsedFaces;
-		UsedFaces.SetNum(NumFaces, false);
-
-        int32 NumUsedFaces = 0;
 		const uint32 IndexTypeSize = IndicesBegin.GetElementSize();
+		check(IndexTypeSize == 4 || IndexTypeSize == 2);
 
-        if (IndexTypeSize == 4)
-        {
-            for (int32 FaceIndex = 0; FaceIndex < NumFaces; ++FaceIndex)
-            {
-                const uint32* FaceIndicesData = reinterpret_cast<uint32*>((IndicesBegin + FaceIndex*3).ptr());
+		static constexpr int32 MaxNumBatches = 5;
+		TStaticArray<TBitArray<>, MaxNumBatches> ParallelUsedVertices;
+		TStaticArray<TBitArray<>, MaxNumBatches> ParallelUsedFaces;
 
-				bool bRemoved = false;
+		const TBitArray<>& UsedVertices = ParallelUsedVertices[0];
+		const TBitArray<>& UsedFaces = ParallelUsedFaces[0];
+		
+		ParallelUsedFaces[0].Reserve(NumFaces);
 
-				if (bRemoveIfAllVerticesCulled)
+		{
+			const bool bUseBatches = NumFaces > (1 << 13);
+			const int32 NumBatches = bUseBatches ? MaxNumBatches : 1;
+			const int32 NumFacesPerBatch = FMath::DivideAndRoundUp<int32>(NumFaces, NumBatches);
+
+			ParallelExecutionUtils::InvokeBatchParallelFor(NumBatches,
+				[
+					IndicesBegin,
+					IndexTypeSize,
+					&VerticesToCull,
+					&ParallelUsedVertices,
+					&ParallelUsedFaces,
+					NumFaces,
+					NumFacesPerBatch,
+					NumVertices,
+					bRemoveIfAllVerticesCulled
+				](int32 BatchId)
+			{
+				const int32 BatchBeginFaceIndex = BatchId * NumFacesPerBatch;
+				const int32 BatchEndFaceIndex = FMath::Min(BatchBeginFaceIndex + NumFacesPerBatch, NumFaces);
+				const int32 BatchNumFaces = BatchEndFaceIndex - BatchBeginFaceIndex;
+
+				TBitArray<>& UsedFaces = ParallelUsedFaces[BatchId];
+				UsedFaces.SetNum(BatchNumFaces, false);
+
+				TBitArray<>& UsedVertices = ParallelUsedVertices[BatchId];
+				UsedVertices.SetNum(NumVertices, false);
+
+				if (IndexTypeSize == 4)
 				{
-					const bool bAllVertsRemoved =
-						VerticesToCull[FaceIndicesData[0]] &
-						VerticesToCull[FaceIndicesData[1]] &
-						VerticesToCull[FaceIndicesData[2]];
-					bRemoved = bAllVertsRemoved;
+					const uint32* FaceIndicesData = reinterpret_cast<uint32*>((IndicesBegin + BatchBeginFaceIndex * 3).ptr());
+					for (int32 FaceIndex = 0; FaceIndex < BatchNumFaces; ++FaceIndex)
+					{
+						bool bRemoved = false;
+
+						if (bRemoveIfAllVerticesCulled)
+						{
+							const bool bAllVertsRemoved =
+								VerticesToCull[FaceIndicesData[0]] &
+								VerticesToCull[FaceIndicesData[1]] &
+								VerticesToCull[FaceIndicesData[2]];
+							bRemoved = bAllVertsRemoved;
+						}
+						else
+						{
+							const bool bOneVertRemoved =
+								VerticesToCull[FaceIndicesData[0]] |
+								VerticesToCull[FaceIndicesData[1]] |
+								VerticesToCull[FaceIndicesData[2]];
+							bRemoved = bOneVertRemoved;
+						}
+
+						if (!bRemoved)
+						{
+							UsedFaces[FaceIndex] = true;
+
+							UsedVertices[FaceIndicesData[0]] = true;
+							UsedVertices[FaceIndicesData[1]] = true;
+							UsedVertices[FaceIndicesData[2]] = true;
+						}
+
+						FaceIndicesData += 3;
+					}
 				}
-				else
+				else if (IndexTypeSize == 2)
 				{
-					const bool bOneVertRemoved =
-						VerticesToCull[FaceIndicesData[0]] |
-						VerticesToCull[FaceIndicesData[1]] |
-						VerticesToCull[FaceIndicesData[2]];
-					bRemoved = bOneVertRemoved;
+					const uint16* FaceIndicesData = reinterpret_cast<uint16*>((IndicesBegin + BatchBeginFaceIndex * 3).ptr());
+
+					for (int32 FaceIndex = 0; FaceIndex < BatchNumFaces; ++FaceIndex)
+					{
+						bool bRemoved = false;
+
+						if (bRemoveIfAllVerticesCulled)
+						{
+							const bool bAllVertsRemoved =
+								VerticesToCull[FaceIndicesData[0]] &
+								VerticesToCull[FaceIndicesData[1]] &
+								VerticesToCull[FaceIndicesData[2]];
+							bRemoved = bAllVertsRemoved;
+						}
+						else
+						{
+							const bool bOneVertRemoved =
+								VerticesToCull[FaceIndicesData[0]] |
+								VerticesToCull[FaceIndicesData[1]] |
+								VerticesToCull[FaceIndicesData[2]];
+							bRemoved = bOneVertRemoved;
+						}
+
+						if (!bRemoved)
+						{
+							UsedFaces[FaceIndex] = true;
+
+							UsedVertices[FaceIndicesData[0]] = true;
+							UsedVertices[FaceIndicesData[1]] = true;
+							UsedVertices[FaceIndicesData[2]] = true;
+						}
+
+						FaceIndicesData += 3;
+					}
 				}
+			});
 
-                if (!bRemoved)
-                {
-                    ++NumUsedFaces;
-                    UsedFaces[FaceIndex] = true;
-
-                    UsedVertices[FaceIndicesData[0]] = true;
-                    UsedVertices[FaceIndicesData[1]] = true;
-                    UsedVertices[FaceIndicesData[2]] = true;
-                }
-            }
-        }
-        else if (IndexTypeSize == 2)
-        {
-            for (int32 FaceIndex = 0; FaceIndex < NumFaces; ++FaceIndex)
-            {
-                const uint16* FaceIndicesData = reinterpret_cast<uint16*>((IndicesBegin + FaceIndex*3).ptr());
-
-                const bool bAllVertsRemoved = 
-                        VerticesToCull[FaceIndicesData[0]] & 
-                        VerticesToCull[FaceIndicesData[1]] & 
-                        VerticesToCull[FaceIndicesData[2]];
+			for (int32 Index = 1; Index < NumBatches; ++Index)
+			{
+				ParallelUsedFaces[0].AddRange(ParallelUsedFaces[Index], ParallelUsedFaces[Index].Num(), 0);
+				ParallelUsedVertices[0].CombineWithBitwiseOR(ParallelUsedVertices[Index], EBitwiseOperatorFlags::MaxSize);
+			}
+		}
 
-                if (!bAllVertsRemoved)
-                {
-                    ++NumUsedFaces;
-                    UsedFaces[FaceIndex] = true;
+		const bool bHasFaces = UsedFaces.Find(true) != INDEX_NONE;
+		const bool bHasRemovedFaces = UsedFaces.Find(false) != INDEX_NONE;
 
-                    UsedVertices[FaceIndicesData[0]] = true;
-                    UsedVertices[FaceIndicesData[1]] = true;
-                    UsedVertices[FaceIndicesData[2]] = true;
-                }
-            }
-        }
-        else
-        {
-            check(false);
-        }
- 
-		if (NumUsedFaces < NumFaces && Result->AreVertexIdsImplicit())
+		if (bHasFaces && bHasRemovedFaces && Result->AreVertexIdsImplicit())
 		{
 			Result->MakeVertexIdsRelative();
 		}
@@ -436,7 +285,7 @@ namespace mu
 
 		// Compute vertices indices remap
         int32 NumVerticesRemaining = 0;
-		if (NumUsedFaces > 0) 
+		if (bHasFaces)
         {
 		    int32 LastFreeVertexIndex = 0;
             for (int32 VertexIndex = UsedVertices.FindFrom(true, 0); VertexIndex >= 0;)
@@ -504,66 +353,82 @@ namespace mu
 		Result->GetVertexBuffers().SetElementCount(NumVerticesRemaining);
 
 		int32 LastFreeFaceIndex = 0;
-        if (NumUsedFaces > 0)
-        {
-            for (int32 FaceIndex = UsedFaces.FindFrom(true, 0); FaceIndex >= 0;)
-            {		
-                const int32 UsedSpanStart = FaceIndex;
-                FaceIndex = UsedFaces.FindFrom(false, FaceIndex);
-                FaceIndex = FaceIndex >= 0 ? FaceIndex : NumFaces;
-                const int32 UsedSpanEnd = FaceIndex;
-
-                FaceIndex = UsedFaces.FindFrom(true, FaceIndex);
-
-                const int32 UsedSpanSize = UsedSpanEnd - UsedSpanStart;
-        
-                if (LastFreeFaceIndex != UsedSpanStart)
-                {
-                    FMemory::Memmove(
-                            (IndicesBegin + LastFreeFaceIndex*3).ptr(), 
-                            (IndicesBegin + UsedSpanStart*3).ptr(), 
-                            UsedSpanSize*IndexTypeSize*3);
-                }
+		if (bHasFaces)
+		{
+			// Move Indices
+			for (int32 FaceIndex = UsedFaces.FindFrom(true, 0); FaceIndex >= 0;)
+			{
+				const int32 UsedSpanStart = FaceIndex;
+				FaceIndex = UsedFaces.FindFrom(false, FaceIndex);
+				FaceIndex = FaceIndex >= 0 ? FaceIndex : NumFaces;
+				const int32 UsedSpanEnd = FaceIndex;
 
-                // Remap vertices
-                if (IndexTypeSize == 4)
-                {
-                    for (int32 I = LastFreeFaceIndex; I < LastFreeFaceIndex + UsedSpanSize; ++I)
-                    {
-                        uint32* FaceIndicesData = reinterpret_cast<uint32*>((IndicesBegin + I*3).ptr());
+				FaceIndex = UsedFaces.FindFrom(true, FaceIndex);
 
-                        check(UsedVerticesMap[FaceIndicesData[0]] >= 0);
-                        check(UsedVerticesMap[FaceIndicesData[1]] >= 0);
-                        check(UsedVerticesMap[FaceIndicesData[2]] >= 0);
+				const int32 UsedSpanSize = UsedSpanEnd - UsedSpanStart;
+				check(UsedSpanSize > 0);
 
-                        FaceIndicesData[0] = UsedVerticesMap[FaceIndicesData[0]];
-                        FaceIndicesData[1] = UsedVerticesMap[FaceIndicesData[1]];
-                        FaceIndicesData[2] = UsedVerticesMap[FaceIndicesData[2]];
-                    }
-                }
-                else if (IndexTypeSize == 2)
-                {	
-                    for (int32 I = LastFreeFaceIndex; I < LastFreeFaceIndex + UsedSpanSize; ++I)
-                    {
-                        uint16* FaceIndicesData = reinterpret_cast<uint16*>((IndicesBegin + I*3).ptr());
-
-                        check(UsedVerticesMap[FaceIndicesData[0]] >= 0);
-                        check(UsedVerticesMap[FaceIndicesData[1]] >= 0);
-                        check(UsedVerticesMap[FaceIndicesData[2]] >= 0);
-                        
-                        FaceIndicesData[0] = static_cast<uint16>(UsedVerticesMap[FaceIndicesData[0]]);
-                        FaceIndicesData[1] = static_cast<uint16>(UsedVerticesMap[FaceIndicesData[1]]);
-                        FaceIndicesData[2] = static_cast<uint16>(UsedVerticesMap[FaceIndicesData[2]]);
-                    }
-                }
-                else
-                {
-                    check(false);
-                }
+				if (LastFreeFaceIndex != UsedSpanStart)
+				{
+					FMemory::Memmove(
+						IndicesBegin.ptr() + LastFreeFaceIndex * IndexTypeSize * 3,
+						IndicesBegin.ptr() + UsedSpanStart * IndexTypeSize * 3,
+						UsedSpanSize * IndexTypeSize * 3);
+				}
 
-                LastFreeFaceIndex += UsedSpanSize;
-            }
-        }
+				LastFreeFaceIndex += UsedSpanSize;
+			}
+
+			constexpr int32 NumFacesPerBatch = (1 << 13);
+			const int32 NumBatches = FMath::DivideAndRoundUp<int32>(LastFreeFaceIndex, NumFacesPerBatch);
+
+			// Remap Vertices
+			ParallelExecutionUtils::InvokeBatchParallelFor(NumBatches,
+				[
+					IndicesBegin,
+					IndexTypeSize,
+					&UsedVerticesMap,
+					LastFreeFaceIndex,
+					NumFacesPerBatch
+				](int32 BatchId)
+				{
+					const int32 BatchBeginFaceIndex = BatchId * NumFacesPerBatch;
+					const int32 BatchEndFaceIndex = FMath::Min(BatchBeginFaceIndex + NumFacesPerBatch, LastFreeFaceIndex);
+
+					if (IndexTypeSize == 4)
+					{
+						uint32* FaceIndicesData = reinterpret_cast<uint32*>((IndicesBegin + BatchBeginFaceIndex * 3).ptr());
+						for (int32 FaceIndex = BatchBeginFaceIndex; FaceIndex < BatchEndFaceIndex; ++FaceIndex)
+						{
+							check(UsedVerticesMap[FaceIndicesData[0]] >= 0);
+							check(UsedVerticesMap[FaceIndicesData[1]] >= 0);
+							check(UsedVerticesMap[FaceIndicesData[2]] >= 0);
+
+							FaceIndicesData[0] = UsedVerticesMap[FaceIndicesData[0]];
+							FaceIndicesData[1] = UsedVerticesMap[FaceIndicesData[1]];
+							FaceIndicesData[2] = UsedVerticesMap[FaceIndicesData[2]];
+
+							FaceIndicesData += 3;
+						}
+					}
+					else if (IndexTypeSize == 2)
+					{
+						uint16* FaceIndicesData = reinterpret_cast<uint16*>((IndicesBegin + BatchBeginFaceIndex).ptr());
+						for (int32 FaceIndex = BatchBeginFaceIndex; FaceIndex < BatchEndFaceIndex; ++FaceIndex)
+						{
+							check(UsedVerticesMap[FaceIndicesData[0]] >= 0);
+							check(UsedVerticesMap[FaceIndicesData[1]] >= 0);
+							check(UsedVerticesMap[FaceIndicesData[2]] >= 0);
+
+							FaceIndicesData[0] = static_cast<uint16>(UsedVerticesMap[FaceIndicesData[0]]);
+							FaceIndicesData[1] = static_cast<uint16>(UsedVerticesMap[FaceIndicesData[1]]);
+							FaceIndicesData[2] = static_cast<uint16>(UsedVerticesMap[FaceIndicesData[2]]);
+
+							FaceIndicesData += 3;
+						}
+					}
+				});
+		}
 
 		check(LastFreeFaceIndex <= NumFaces);
 
