[Mutable] Fix cases of surfaces being shared when they should not.
Improve detection of surfaces that can be shared between LODs.
From now on, all surfaces will have a valid shared surface id, even only used by a single LOD.

#rnx

diff --git a/Engine/Plugins/Mutable/Source/CustomizableObject/Internal/MuCO/CustomizableObjectSystemPrivate.h b/Engine/Plugins/Mutable/Source/CustomizableObject/Internal/MuCO/CustomizableObjectSystemPrivate.h
--- a/Engine/Plugins/Mutable/Source/CustomizableObject/Internal/MuCO/CustomizableObjectSystemPrivate.h
+++ b/Engine/Plugins/Mutable/Source/CustomizableObject/Internal/MuCO/CustomizableObjectSystemPrivate.h
@@ -474,6 +474,7 @@ struct FInstanceUpdateData
 
 		/** Id of the surface in the mutable core instance. */
 		uint32 SurfaceId = 0;
+		uint32 SharedSurfaceId = 0;
 
 		/** Id of the metadata associated with this surface. */ 
 		uint32 SurfaceMetadataId = 0;
diff --git a/Engine/Plugins/Mutable/Source/CustomizableObject/Private/MuCO/CustomizableObjectInstance.cpp b/Engine/Plugins/Mutable/Source/CustomizableObject/Private/MuCO/CustomizableObjectInstance.cpp
--- a/Engine/Plugins/Mutable/Source/CustomizableObject/Private/MuCO/CustomizableObjectInstance.cpp
+++ b/Engine/Plugins/Mutable/Source/CustomizableObject/Private/MuCO/CustomizableObjectInstance.cpp
@@ -5840,7 +5840,7 @@ void UCustomizableInstancePrivate::BuildMaterials(const TSharedRef<FUpdateContex
 				}
 
 				// Reuse MaterialSlot from the previous LOD.
-				if (const int32 MaterialIndex = SurfaceIdToMaterialIndex.Find(Surface.SurfaceId); MaterialIndex != INDEX_NONE)
+				if (const int32 MaterialIndex = SurfaceIdToMaterialIndex.Find(Surface.SharedSurfaceId); MaterialIndex != INDEX_NONE)
 				{
 					if (!bUseOverrideMaterialsOnly)
 					{
@@ -5862,7 +5862,7 @@ void UCustomizableInstancePrivate::BuildMaterials(const TSharedRef<FUpdateContex
 				}
 
 				// This section will require a new slot
-				SurfaceIdToMaterialIndex.Add(Surface.SurfaceId);
+				SurfaceIdToMaterialIndex.Add(Surface.SharedSurfaceId);
 
 				// Add and set up the material data for this slot
 				const int32 MaterialSlotIndex = Materials.Num();
@@ -5872,7 +5872,7 @@ void UCustomizableInstancePrivate::BuildMaterials(const TSharedRef<FUpdateContex
 				uint32 UsedSurfaceMetadataId = Surface.SurfaceMetadataId;
 				
 				// If the surface metadata is invalid, check if any of the mesh fragments has metadata. 
-				// For now use the fisrt found, an aggregate may be needed. 
+				// For now use the first found, an aggregate may be needed. 
 				if (Surface.SurfaceMetadataId == 0 && LOD.Mesh)
 				{
 					int32 MeshSurfaceIndex = LOD.Mesh->Surfaces.IndexOfByPredicate([SurfaceId = Surface.SurfaceId](const mu::FMeshSurface& Surface)
diff --git a/Engine/Plugins/Mutable/Source/CustomizableObject/Private/MuCO/CustomizableObjectSystem.cpp b/Engine/Plugins/Mutable/Source/CustomizableObject/Private/MuCO/CustomizableObjectSystem.cpp
--- a/Engine/Plugins/Mutable/Source/CustomizableObject/Private/MuCO/CustomizableObjectSystem.cpp
+++ b/Engine/Plugins/Mutable/Source/CustomizableObject/Private/MuCO/CustomizableObjectSystem.cpp
@@ -2543,7 +2543,7 @@ namespace impl
 
 		const mu::FInstance* MutableInstance = OperationData->MutableInstance.Get();
 
-		TArray<int32> SurfacesSharedId;
+		TArray<uint32> SurfacesSharedId;
 
 		const TSharedRef<TArray<FGetImageData>> GetImagesData = MakeShared<TArray<FGetImageData>>();
 		
@@ -2575,30 +2575,27 @@ namespace impl
 						OperationData->InstanceUpdateData.Surfaces.Push({});
 						FInstanceUpdateData::FSurface& Surface = OperationData->InstanceUpdateData.Surfaces.Last();
 						++LOD.SurfaceCount;
-
-						// Now Surface.MaterialIndex is decoded from a parameter at the end of this if()
-						Surface.SurfaceId = SurfaceId;
-						Surface.SurfaceMetadataId = SurfaceMetadataId;
-
-						const int32 SharedSurfaceId = MutableInstance->GetSharedSurfaceId(InstanceComponentIndex, MutableLODIndex, InstanceSurfaceIndex);
+						
+						const uint32 SharedSurfaceId = MutableInstance->GetSharedSurfaceId(InstanceComponentIndex, MutableLODIndex, InstanceSurfaceIndex);
 						const int32 SharedSurfaceIndex = SurfacesSharedId.Find(SharedSurfaceId);
-
 						SurfacesSharedId.Add(SharedSurfaceId);
 
-						if (SharedSurfaceId != INDEX_NONE)
+						if (SharedSurfaceIndex != INDEX_NONE)
 						{
-							if (SharedSurfaceIndex >= 0)
-							{
-								Surface = OperationData->InstanceUpdateData.Surfaces[SharedSurfaceIndex];
-								return;
-							}
+							Surface = OperationData->InstanceUpdateData.Surfaces[SharedSurfaceIndex];
+							return;
+						}
+
+						// New Surface.MaterialIndex is decoded from a parameter at the end of this if()
+						Surface.SurfaceId = SurfaceId;
+						Surface.SurfaceMetadataId = SurfaceMetadataId;
+						Surface.SharedSurfaceId = SharedSurfaceId;
 
-							// Find the first LOD where this surface can be found
-							MutableInstance->FindBaseSurfaceBySharedId(InstanceComponentIndex, SharedSurfaceId, BaseSurfaceIndex, BaseLODIndex);
+						// Find the first LOD where this surface can be found
+						MutableInstance->FindBaseSurfaceBySharedId(InstanceComponentIndex, Surface.SharedSurfaceId, BaseSurfaceIndex, BaseLODIndex);
 
-							Surface.SurfaceId = MutableInstance->GetSurfaceId(InstanceComponentIndex, BaseLODIndex, BaseSurfaceIndex);
-							Surface.SurfaceMetadataId = MutableInstance->GetSurfaceCustomId(InstanceComponentIndex, BaseLODIndex, BaseSurfaceIndex);
-						}
+						Surface.SurfaceId = MutableInstance->GetSurfaceId(InstanceComponentIndex, BaseLODIndex, BaseSurfaceIndex);
+						Surface.SurfaceMetadataId = MutableInstance->GetSurfaceCustomId(InstanceComponentIndex, BaseLODIndex, BaseSurfaceIndex);
 
 						// Vectors
 						Surface.FirstVector = OperationData->InstanceUpdateData.Vectors.Num();
diff --git a/Engine/Plugins/Mutable/Source/CustomizableObjectEditor/Private/MuCOE/CustomizableObjectCompiler.cpp b/Engine/Plugins/Mutable/Source/CustomizableObjectEditor/Private/MuCOE/CustomizableObjectCompiler.cpp
--- a/Engine/Plugins/Mutable/Source/CustomizableObjectEditor/Private/MuCOE/CustomizableObjectCompiler.cpp
+++ b/Engine/Plugins/Mutable/Source/CustomizableObjectEditor/Private/MuCOE/CustomizableObjectCompiler.cpp
@@ -782,8 +782,6 @@ mu::Ptr<mu::NodeObject> GenerateMutableRoot(
 	mu::Ptr<mu::NodeObject> MutableRoot = GenerateMutableSource(RootNodeObject->OutputPin(), GenerationContext);
 	UE_LOG(LogMutable, Verbose, TEXT("PROFILE: [ %16.8f ] GenerateMutableSource end."), FPlatformTime::Seconds());
 
-	GenerationContext.GenerateSharedSurfacesUniqueIds();
-
 	// Generate ReferenceSkeletalMeshes data;
 	PopulateReferenceSkeletalMeshesData(GenerationContext);
 
diff --git a/Engine/Plugins/Mutable/Source/CustomizableObjectEditor/Private/MuCOE/GenerateMutableSource/GenerateMutableSource.cpp b/Engine/Plugins/Mutable/Source/CustomizableObjectEditor/Private/MuCOE/GenerateMutableSource/GenerateMutableSource.cpp
--- a/Engine/Plugins/Mutable/Source/CustomizableObjectEditor/Private/MuCOE/GenerateMutableSource/GenerateMutableSource.cpp
+++ b/Engine/Plugins/Mutable/Source/CustomizableObjectEditor/Private/MuCOE/GenerateMutableSource/GenerateMutableSource.cpp
@@ -340,7 +340,13 @@ const FGuid FMutableGraphGenerationContext::GetNodeIdUnique(const UCustomizableO
 		return NodeID;
 	}
 
-	ArrayResult->Add(Node);
+	ArrayResult->AddUnique(Node);
+
+	if (ArrayResult->Num() == 1)
+	{
+		return NodeID;
+	}
+
 	return FGuid::NewGuid();
 }
 
@@ -363,58 +369,6 @@ const FGuid FMutableGraphGenerationContext::GetNodeIdUnchecked(const UCustomizab
 }
 
 
-void FMutableGraphGenerationContext::GenerateSharedSurfacesUniqueIds()
-{
-	int32 UniqueId = 0;
-
-	TArray<TArray<FSharedSurface>> NodeToSharedSurfaces;
-	SharedSurfaceIds.GenerateValueArray(NodeToSharedSurfaces);
-
-	TArray<bool> VisitedSurfaces;
-	for (TArray<FSharedSurface>& SharedSurfaces : NodeToSharedSurfaces)
-	{
-		const int32 NumSurfaces = SharedSurfaces.Num();
-		VisitedSurfaces.Init(false, NumSurfaces);
-		
-		// Iterate all surfaces for a given NodeMaterial and set the same SharedSurfaceId to those that are equal.
-		for (int32 SurfaceIndex = 0; SurfaceIndex < NumSurfaces; ++SurfaceIndex)
-		{
-			if (VisitedSurfaces[SurfaceIndex])
-			{
-				continue;
-			}
-
-			FSharedSurface& CurrentSharedSurface = SharedSurfaces[SurfaceIndex];
-			CurrentSharedSurface.NodeSurfaceNew->SharedSurfaceId = UniqueId;
-			VisitedSurfaces[SurfaceIndex] = true;
-
-			for (int32 AuxSurfaceIndex = SurfaceIndex; AuxSurfaceIndex < NumSurfaces && !CurrentSharedSurface.bMakeUnique; ++AuxSurfaceIndex)
-			{
-				if (VisitedSurfaces[AuxSurfaceIndex])
-				{
-					continue;
-				}
-
-				if (SharedSurfaces[AuxSurfaceIndex].NodeModifierIDs != CurrentSharedSurface.NodeModifierIDs)
-				{
-					continue;
-				}
-
-				if (SharedSurfaces[AuxSurfaceIndex].MacroContext != CurrentSharedSurface.MacroContext)
-				{
-					continue;
-				}
-
-				SharedSurfaces[AuxSurfaceIndex].NodeSurfaceNew->SharedSurfaceId = UniqueId;
-				VisitedSurfaces[AuxSurfaceIndex] = true;
-			}
-
-			++UniqueId;
-		}
-	}
-}
-
-
 FMutableComponentInfo* FMutableGraphGenerationContext::GetCurrentComponentInfo()
 {
 	return CompilationContext.GetComponentInfo(CurrentMeshComponent);
@@ -1955,19 +1909,6 @@ void FMutableComponentInfo::AccumulateBonesToRemovePerLOD(const TArray<FLODReduc
 
 #undef LOCTEXT_NAMESPACE
 
-FMutableGraphGenerationContext::FSharedSurface::FSharedSurface(uint8 InLOD, const mu::Ptr<mu::NodeSurfaceNew>& InNodeSurfaceNew, TArray<const UCustomizableObjectNodeMacroInstance*> CurrentMacroContext)
-{
-	LOD = InLOD;
-	NodeSurfaceNew = InNodeSurfaceNew;
-	MacroContext = CurrentMacroContext;
-}
-
-
-bool FMutableGraphGenerationContext::FSharedSurface::operator==(const FSharedSurface& o) const
-{
-	return NodeModifierIDs == o.NodeModifierIDs;
-}
-
 
 FGeneratedMutableDataTableKey::FGeneratedMutableDataTableKey(FString TableName, FName VersionColumn, const TArray<FTableNodeCompilationFilter>& CompilationFilterOptions)
 	: TableName(TableName), VersionColumn(VersionColumn), CompilationFilterOptions(CompilationFilterOptions)
@@ -1994,3 +1935,25 @@ uint32 GetTypeHash(const FGeneratedMutableDataTableKey& Key)
 	return Hash;
 }
 
+
+FSharedSurfaceKey::FSharedSurfaceKey(const UCustomizableObjectNodeMaterialBase* NodeMaterial, const TArray<const UCustomizableObjectNodeMacroInstance*>& MacroContext)
+{
+	check(NodeMaterial);
+
+	MaterialGuid = NodeMaterial->NodeGuid;
+
+	MacroContextGuid = FGuid();
+	for (const UCustomizableObjectNodeMacroInstance* Macro : MacroContext)
+	{
+		if (ensure(Macro))
+		{
+			MacroContextGuid = FGuid::Combine(MacroContextGuid, Macro->NodeGuid);
+		}
+	}
+}
+
+
+uint32 GetTypeHash(const FSharedSurfaceKey& Key)
+{
+	return HashCombine(GetTypeHash(Key.MaterialGuid), GetTypeHash(Key.MacroContextGuid));
+}
diff --git a/Engine/Plugins/Mutable/Source/CustomizableObjectEditor/Private/MuCOE/GenerateMutableSource/GenerateMutableSource.h b/Engine/Plugins/Mutable/Source/CustomizableObjectEditor/Private/MuCOE/GenerateMutableSource/GenerateMutableSource.h
--- a/Engine/Plugins/Mutable/Source/CustomizableObjectEditor/Private/MuCOE/GenerateMutableSource/GenerateMutableSource.h
+++ b/Engine/Plugins/Mutable/Source/CustomizableObjectEditor/Private/MuCOE/GenerateMutableSource/GenerateMutableSource.h
@@ -551,6 +551,18 @@ struct FGeneratedMutableDataTableKey
 
 uint32 GetTypeHash(const FGeneratedMutableDataTableKey& Key);
 
+struct FSharedSurfaceKey
+{
+	FSharedSurfaceKey(const UCustomizableObjectNodeMaterialBase* NodeMaterial, const TArray<const UCustomizableObjectNodeMacroInstance*>& CurrentMacroContext);
+
+	bool operator==(const FSharedSurfaceKey& o) const = default;
+
+	FGuid MaterialGuid;
+	FGuid MacroContextGuid;
+};
+
+uint32 GetTypeHash(const FSharedSurfaceKey& Key);
+
 
 /** This structure stores the information that is used only during the graph generation stage of the CustomizableObject compilation. 
 * Data needed for node graph processing goes here.
@@ -746,10 +758,6 @@ struct FMutableGraphGenerationContext
 	*/
 	const FGuid GetNodeIdUnchecked(const UCustomizableObjectNode* Node);
 
-	/** Generates shared surface IDs for all surface nodes. If one or more nodes are equal, they will use the same SharedSurfaceId */
-	void GenerateSharedSurfacesUniqueIds();
-
-
 	/** Get the reference skeletal mesh associated to the current mesh component being generated */
 	FMutableComponentInfo* GetCurrentComponentInfo();
 
@@ -848,22 +856,8 @@ struct FMutableGraphGenerationContext
 	// Current material parameter id to find the corresponding column in a mutable table
 	FString CurrentMaterialTableParameterId;
 
-	struct FSharedSurface
-	{
-		FSharedSurface(uint8 InLOD, const mu::Ptr<mu::NodeSurfaceNew>& InNodeSurfaceNew, TArray<const UCustomizableObjectNodeMacroInstance*> CurrentMacroContext);
-
-		bool operator==(const FSharedSurface& o) const;
-
-		uint8 LOD = 0;
-		mu::Ptr<mu::NodeSurfaceNew> NodeSurfaceNew;
-
-		bool bMakeUnique = false;
-		TArray<SIZE_T> NodeModifierIDs;
-		TArray<const UCustomizableObjectNodeMacroInstance*> MacroContext;
-	};
-
 	// Material to SharedSurfaceId
-	TMap<UCustomizableObjectNodeMaterialBase*, TArray<FSharedSurface>> SharedSurfaceIds;
+	TMap<FSharedSurfaceKey, FGuid> SurfaceGuids;
 
 	/** Resource Data constants */
 	TMap<uint32, int32> StreamedResourceIndices;
diff --git a/Engine/Plugins/Mutable/Source/CustomizableObjectEditor/Private/MuCOE/GenerateMutableSource/GenerateMutableSourceModifier.cpp b/Engine/Plugins/Mutable/Source/CustomizableObjectEditor/Private/MuCOE/GenerateMutableSource/GenerateMutableSourceModifier.cpp
--- a/Engine/Plugins/Mutable/Source/CustomizableObjectEditor/Private/MuCOE/GenerateMutableSource/GenerateMutableSourceModifier.cpp
+++ b/Engine/Plugins/Mutable/Source/CustomizableObjectEditor/Private/MuCOE/GenerateMutableSource/GenerateMutableSourceModifier.cpp
@@ -247,6 +247,7 @@ mu::Ptr<mu::NodeModifier> GenerateMutableSourceModifier(const UEdGraphPin * Pin,
 
 		mu::Ptr<mu::NodeModifierSurfaceEdit> SurfNode = new mu::NodeModifierSurfaceEdit();
 		Result = SurfNode;
+		SurfNode->ModifierGuid = GenerationContext.GetNodeIdUnique(Node);
 
 		// TODO: This was used in the non-modifier version for group projectors. It may affect the "drop projection from LOD" feature.
 		const int32 LOD = Node->IsAffectedByLOD() ? GenerationContext.CurrentLOD : 0;
diff --git a/Engine/Plugins/Mutable/Source/CustomizableObjectEditor/Private/MuCOE/GenerateMutableSource/GenerateMutableSourceSurface.cpp b/Engine/Plugins/Mutable/Source/CustomizableObjectEditor/Private/MuCOE/GenerateMutableSource/GenerateMutableSourceSurface.cpp
--- a/Engine/Plugins/Mutable/Source/CustomizableObjectEditor/Private/MuCOE/GenerateMutableSource/GenerateMutableSourceSurface.cpp
+++ b/Engine/Plugins/Mutable/Source/CustomizableObjectEditor/Private/MuCOE/GenerateMutableSource/GenerateMutableSourceSurface.cpp
@@ -124,11 +124,13 @@ mu::Ptr<mu::NodeSurface> GenerateMutableSourceSurface(const UEdGraphPin * Pin, F
 		mu::Ptr<mu::NodeSurfaceNew> SurfNode = new mu::NodeSurfaceNew();
 		Result = SurfNode;
 
-		// Add to the list of surfaces that could be reused between LODs for this NodeMaterial.
-		// TODO(Max)
-		TArray<FMutableGraphGenerationContext::FSharedSurface>& SharedSurfaces = GenerationContext.SharedSurfaceIds.FindOrAdd(TypedNodeMat, {});
-		FMutableGraphGenerationContext::FSharedSurface& SharedSurface = SharedSurfaces.Add_GetRef(FMutableGraphGenerationContext::FSharedSurface(GenerationContext.CurrentLOD, SurfNode, GenerationContext.MacroNodesStack));
-		SharedSurface.bMakeUnique = !TypedNodeMat->IsReuseMaterialBetweenLODs();
+		FGuid& SurfaceGuid = GenerationContext.SurfaceGuids.FindOrAdd(FSharedSurfaceKey(TypedNodeMat, GenerationContext.MacroNodesStack));
+		if (!SurfaceGuid.IsValid())
+		{
+			SurfaceGuid = GenerationContext.GetNodeIdUnique(TypedNodeMat);
+		}
+
+		SurfNode->SurfaceGuid = SurfaceGuid;
 
 		GenerationContext.CurrentReferencedMaterialIndex = GenerationContext.ReferencedMaterials.AddUnique(Material);
 
diff --git a/Engine/Plugins/Mutable/Source/CustomizableObjectEditor/Private/MuCOE/Nodes/CustomizableObjectNodeCopyMaterial.cpp b/Engine/Plugins/Mutable/Source/CustomizableObjectEditor/Private/MuCOE/Nodes/CustomizableObjectNodeCopyMaterial.cpp
--- a/Engine/Plugins/Mutable/Source/CustomizableObjectEditor/Private/MuCOE/Nodes/CustomizableObjectNodeCopyMaterial.cpp
+++ b/Engine/Plugins/Mutable/Source/CustomizableObjectEditor/Private/MuCOE/Nodes/CustomizableObjectNodeCopyMaterial.cpp
@@ -262,19 +262,6 @@ UMaterialInterface* UCustomizableObjectNodeCopyMaterial::GetMaterial() const
 }
 
 
-bool UCustomizableObjectNodeCopyMaterial::IsReuseMaterialBetweenLODs() const
-{
-	if (UCustomizableObjectNodeMaterial* NodeMaterial = GetMaterialNode())
-	{
-		return NodeMaterial->IsReuseMaterialBetweenLODs();
-	}
-	else
-	{
-		return false;
-	}
-}
-
-
 TArray<FString> UCustomizableObjectNodeCopyMaterial::GetEnableTags(TArray<const UCustomizableObjectNodeMacroInstance*>* MacroContext)
 {
 	if (UCustomizableObjectNodeMaterial* NodeMaterial = GetMaterialNode())
diff --git a/Engine/Plugins/Mutable/Source/CustomizableObjectEditor/Private/MuCOE/Nodes/CustomizableObjectNodeCopyMaterial.h b/Engine/Plugins/Mutable/Source/CustomizableObjectEditor/Private/MuCOE/Nodes/CustomizableObjectNodeCopyMaterial.h
--- a/Engine/Plugins/Mutable/Source/CustomizableObjectEditor/Private/MuCOE/Nodes/CustomizableObjectNodeCopyMaterial.h
+++ b/Engine/Plugins/Mutable/Source/CustomizableObjectEditor/Private/MuCOE/Nodes/CustomizableObjectNodeCopyMaterial.h
@@ -36,7 +36,6 @@ public:
 
 	// UCustomizableObjectNodeMaterialBase interface
 	virtual UMaterialInterface* GetMaterial() const override;
-	virtual bool IsReuseMaterialBetweenLODs() const override;
 	virtual UEdGraphPin* GetMaterialAssetPin() const override;
 	virtual int32 GetNumParameters(EMaterialParameterType Type) const override;
 	virtual FNodeMaterialParameterId GetParameterId(EMaterialParameterType Type, int32 ParameterIndex) const override;
diff --git a/Engine/Plugins/Mutable/Source/CustomizableObjectEditor/Private/MuCOE/Nodes/CustomizableObjectNodeMaterial.cpp b/Engine/Plugins/Mutable/Source/CustomizableObjectEditor/Private/MuCOE/Nodes/CustomizableObjectNodeMaterial.cpp
--- a/Engine/Plugins/Mutable/Source/CustomizableObjectEditor/Private/MuCOE/Nodes/CustomizableObjectNodeMaterial.cpp
+++ b/Engine/Plugins/Mutable/Source/CustomizableObjectEditor/Private/MuCOE/Nodes/CustomizableObjectNodeMaterial.cpp
@@ -1257,12 +1257,6 @@ UMaterialInterface* UCustomizableObjectNodeMaterial::GetMaterial() const
 }
 
 
-bool UCustomizableObjectNodeMaterial::IsReuseMaterialBetweenLODs() const
-{
-	return bReuseMaterialBetweenLODs;	
-}
-
-
 TArray<FString> UCustomizableObjectNodeMaterial::GetEnableTags(TArray<const UCustomizableObjectNodeMacroInstance*>* MacroContext)
 {
 	// Getting tags from linked pin
diff --git a/Engine/Plugins/Mutable/Source/CustomizableObjectEditor/Private/MuCOE/Nodes/CustomizableObjectNodeMaterial.h b/Engine/Plugins/Mutable/Source/CustomizableObjectEditor/Private/MuCOE/Nodes/CustomizableObjectNodeMaterial.h
--- a/Engine/Plugins/Mutable/Source/CustomizableObjectEditor/Private/MuCOE/Nodes/CustomizableObjectNodeMaterial.h
+++ b/Engine/Plugins/Mutable/Source/CustomizableObjectEditor/Private/MuCOE/Nodes/CustomizableObjectNodeMaterial.h
@@ -137,7 +137,6 @@ public:
 	virtual TArray<UCustomizableObjectLayout*> GetLayouts() const override;
 	virtual UEdGraphPin* OutputPin() const override;
 	virtual UMaterialInterface* GetMaterial() const override;
-	virtual bool IsReuseMaterialBetweenLODs() const override;
 	virtual int32 GetNumParameters(EMaterialParameterType Type) const override;
 	virtual FNodeMaterialParameterId GetParameterId(EMaterialParameterType Type, int32 ParameterIndex) const override;
 	virtual FName GetParameterName(EMaterialParameterType Type, int32 ParameterIndex) const override;
@@ -190,9 +189,6 @@ public:
 	FName MeshComponentName_DEPRECATED;
 
 private:
-	/** Materials will be reused between LODs when possible unless explicitly disabled. */
-	UPROPERTY(EditAnywhere, Category = CustomizableObject)
-	bool bReuseMaterialBetweenLODs = true;
 
 	/** Last static or skeletal mesh connected. Used to remove the callback once disconnected. */
 	TWeakObjectPtr<UCustomizableObjectNode> LastMeshNodeConnected;
diff --git a/Engine/Plugins/Mutable/Source/CustomizableObjectEditor/Private/MuCOE/Nodes/CustomizableObjectNodeMaterialBase.h b/Engine/Plugins/Mutable/Source/CustomizableObjectEditor/Private/MuCOE/Nodes/CustomizableObjectNodeMaterialBase.h
--- a/Engine/Plugins/Mutable/Source/CustomizableObjectEditor/Private/MuCOE/Nodes/CustomizableObjectNodeMaterialBase.h
+++ b/Engine/Plugins/Mutable/Source/CustomizableObjectEditor/Private/MuCOE/Nodes/CustomizableObjectNodeMaterialBase.h
@@ -66,8 +66,6 @@ public:
 	
 	virtual UMaterialInterface* GetMaterial() const PURE_VIRTUAL(UCustomizableObjectNodeMaterial::GetMaterial, return {}; );
 	
-	virtual bool IsReuseMaterialBetweenLODs() const PURE_VIRTUAL(UCustomizableObjectNodeMaterial::IsReuseMaterialBetweenLODs, return {}; );
-		
 	virtual UEdGraphPin* GetMeshPin() const PURE_VIRTUAL(UCustomizableObjectNodeMaterial::GetMeshPin, return {}; );
 
 	virtual UEdGraphPin* GetMaterialAssetPin() const PURE_VIRTUAL(UCustomizableObjectNodeMaterial::GetMaterialAssetPin, return {}; );
diff --git a/Engine/Plugins/Mutable/Source/MutableRuntime/Private/MuR/Instance.cpp b/Engine/Plugins/Mutable/Source/MutableRuntime/Private/MuR/Instance.cpp
--- a/Engine/Plugins/Mutable/Source/MutableRuntime/Private/MuR/Instance.cpp
+++ b/Engine/Plugins/Mutable/Source/MutableRuntime/Private/MuR/Instance.cpp
@@ -158,7 +158,7 @@ namespace mu
 
 	
 	//---------------------------------------------------------------------------------------------
-	void FInstance::FindBaseSurfaceBySharedId(int32 CompIndex, int32 SharedId, int32& OutSurfaceIndex, int32& OutLODIndex) const
+	void FInstance::FindBaseSurfaceBySharedId(int32 CompIndex, uint32 SharedId, int32& OutSurfaceIndex, int32& OutLODIndex) const
 	{
 		if (m_pD->Components.IsValidIndex(CompIndex))
 		{
@@ -184,7 +184,7 @@ namespace mu
 
 
 	//---------------------------------------------------------------------------------------------
-	int32 FInstance::GetSharedSurfaceId(int32 ComponentIndex, int32 LODIndex, int32 SurfaceIndex) const
+	uint32 FInstance::GetSharedSurfaceId(int32 ComponentIndex, int32 LODIndex, int32 SurfaceIndex) const
 	{
 		if (m_pD->Components.IsValidIndex(ComponentIndex) &&
 			m_pD->Components[ComponentIndex].LODs.IsValidIndex(LODIndex) &&
diff --git a/Engine/Plugins/Mutable/Source/MutableRuntime/Public/MuR/Instance.h b/Engine/Plugins/Mutable/Source/MutableRuntime/Public/MuR/Instance.h
--- a/Engine/Plugins/Mutable/Source/MutableRuntime/Public/MuR/Instance.h
+++ b/Engine/Plugins/Mutable/Source/MutableRuntime/Public/MuR/Instance.h
@@ -83,13 +83,13 @@ namespace mu
 		//! \param SharedSurfaceId - Id of the surface to look for (as returned by GetSharedSurfaceId).
 		//! \param OutSurfaceIndex - Index of the surface in the OutLODIndex lod. 
 		//! \param OutLODIndex - Index of the first LOD where the surface can be found. 
-		UE_API void FindBaseSurfaceBySharedId(int32 ComponentIndex, int32 SharedId, int32& OutSurfaceIndex, int32& OutLODIndex) const;
+		UE_API void FindBaseSurfaceBySharedId(int32 ComponentIndex, uint32 SharedId, int32& OutSurfaceIndex, int32& OutLODIndex) const;
 
 		//! Get an id that can be used to find the same surface on other LODs
 		//! \param ComponentIndex Index of the component, from 0 to GetComponentCount()-1
 		//! \param lod Index of the level of detail, from 0 to GetLODCount()-1
 		//! \param SurfaceIndex Index of the surface, from 0 to GetSurfaceCount()-1
-		UE_API int32 GetSharedSurfaceId(int32 ComponentIndex, int32 LODIndex, int32 SurfaceIndex) const;
+		UE_API uint32 GetSharedSurfaceId(int32 ComponentIndex, int32 LODIndex, int32 SurfaceIndex) const;
 
         //! Get an optional, opaque application-defined identifier for this surface. The meaning of
         //! this ID depends on each application, and it is specified when creating the source data
diff --git a/Engine/Plugins/Mutable/Source/MutableTools/Private/MuT/CodeGenerator.cpp b/Engine/Plugins/Mutable/Source/MutableTools/Private/MuT/CodeGenerator.cpp
--- a/Engine/Plugins/Mutable/Source/MutableTools/Private/MuT/CodeGenerator.cpp
+++ b/Engine/Plugins/Mutable/Source/MutableTools/Private/MuT/CodeGenerator.cpp
@@ -980,36 +980,26 @@ namespace mu
 
 		// The options depenend on the shared surface being generated, so we need to add the previous lod dependency
 		FMeshOptionsTask MeshOptionsTask = LocalPipe.Launch(TEXT("MutableSurfaceMeshOptions"),
-			[MeshDynamicOptions, SharedSurfaceId=SurfaceNode->SharedSurfaceId, this]() mutable
-			{
-				// This assumes that the lods are processed in order. It checks it this way because some platforms may have empty lods at the top.
-				bool bIsBaseForSharedSurface = SharedSurfaceId != INDEX_NONE;
-				if (bIsBaseForSharedSurface)
-				{
-					UE::TUniqueLock Lock(SharedMeshOptions.Mutex);
-					bIsBaseForSharedSurface = !SharedMeshOptions.Map.Contains(SharedSurfaceId);
-				}
-
-				// If this is true, we will reuse the surface properties from a higher LOD, se we can skip the generation of material properties and images.
-				const bool bShareSurface = SharedSurfaceId != INDEX_NONE && !bIsBaseForSharedSurface;
+			[MeshDynamicOptions, BaseSurfaceGuid = SurfaceNode->SurfaceGuid, this]() mutable
+			{	
+				// This assumes that the lods are processed in order.
 
-				const FMeshGenerationResult* SharedMeshResults = nullptr;
-				if (bShareSurface)
-				{
-					UE::TUniqueLock Lock(SharedMeshOptions.Mutex);
+				UE::TUniqueLock Lock(BaseMeshOptions.Mutex);
 
-					// Do we have the surface we need to share it with?
-					SharedMeshResults = SharedMeshOptions.Map.Find(SharedSurfaceId);
-					check(SharedMeshResults);
+				check(BaseSurfaceGuid.IsValid());
+				const FMeshGenerationResult* BaseMeshResult = BaseMeshOptions.Map.Find(BaseSurfaceGuid);
 
+				// If this is true, we will reuse the surface properties from a higher LOD, so we can skip the generation of material properties and images.
+				if (BaseMeshResult)
+				{
 					// Override the layouts with the ones from the surface we share
-					MeshDynamicOptions.OverrideLayouts = SharedMeshResults->GeneratedLayouts;
+					MeshDynamicOptions.OverrideLayouts = BaseMeshResult->GeneratedLayouts;
 				}
 
 				// Ensure UV islands remain within their main layout block on lower LODs to avoid unexpected reordering 
 				// of the layout blocks when reusing a surface between LODs. Used to fix small displacements on vertices
 				// that may cause them to fall on a different block.
-				MeshDynamicOptions.bClampUVIslands = bShareSurface;
+				MeshDynamicOptions.bClampUVIslands = BaseMeshResult != nullptr;
 
 				return MeshDynamicOptions;
 			},
@@ -1021,7 +1011,7 @@ namespace mu
 		// Apply the modifier for the post-normal operations stage.
 		MeshTask = ApplyMeshModifiers(Modifiers,
 			MeshStaticOptions, MeshOptionsTask,
-			MeshTask, SurfaceNode->SharedSurfaceId, SurfaceNode->GetMessageContext(), nullptr);
+			MeshTask, SurfaceNode->SurfaceGuid, SurfaceNode->GetMessageContext(), nullptr);
 
 		FSurfaceTask SurfaceTask = LocalPipe.Launch(TEXT("MutableSurface"),
 			[
@@ -1042,10 +1032,27 @@ namespace mu
 				TArray<FGeneratedLayout> SurfaceReferenceLayouts;
 				TArray<Ptr<ASTOp>> SurfaceLayoutOps;
 
+				{	// This assumes that the lods are processed in order.
+
+					UE::TUniqueLock Lock(BaseMeshOptions.Mutex);
+
+					// Store the base mesh layouts to reuse them in higher LODs.
+					check(SurfaceNode->SurfaceGuid.IsValid());
+					if (!BaseMeshOptions.Map.Contains(SurfaceNode->SurfaceGuid))
+					{
+						FMeshGenerationResult BaseMeshResults;
+						BaseMeshResults.GeneratedLayouts = MeshResults.GeneratedLayouts;
+						BaseMeshOptions.Map.Add(SurfaceNode->SurfaceGuid, MoveTemp(BaseMeshResults));
+					}
+				}
+
+				FGuid SharedSurfaceGuid = SurfaceNode->SurfaceGuid;
+
 				int32 MaxLayoutNum = MeshResults.GeneratedLayouts.Num();
 				for (const FMeshGenerationResult::FExtraLayouts& ExtraLayoutData : MeshResults.ExtraMeshLayouts)
 				{
 					MaxLayoutNum = FMath::Max(MaxLayoutNum, ExtraLayoutData.GeneratedLayouts.Num());
+					SharedSurfaceGuid = FGuid::Combine(SharedSurfaceGuid, ExtraLayoutData.EditGuid);
 				}
 
 				SurfaceReferenceLayouts.SetNum(MaxLayoutNum);
@@ -1056,28 +1063,18 @@ namespace mu
 
 				// Scope for access control to shared data
 				bool bIsBaseForSharedSurface = false;
-				bool bShareSurface = false;
-				TArray<Ptr<ASTOp>> SharedResultLayoutOps;
-				{
-					UE::TUniqueLock Lock(SharedMeshOptions.Mutex);
-
-					bIsBaseForSharedSurface = SurfaceNode->SharedSurfaceId!=INDEX_NONE 
-						&& 
-						!SharedMeshOptions.Map.Contains(SurfaceNode->SharedSurfaceId);
-
-					// If this is true, we will reuse the surface properties from a higher LOD, se we can skip the generation of material properties and images.
-					bShareSurface = SurfaceNode->SharedSurfaceId != INDEX_NONE && !bIsBaseForSharedSurface;
+				const FSharedSurfaceResultOptions::FSharedSurfaceResult* SharedSurfaceResults = nullptr;
 
-					const FMeshGenerationResult* SharedMeshResults = nullptr;
-					if (bShareSurface)
-					{
-						// Do we have the surface we need to share it with?
-						SharedMeshResults = SharedMeshOptions.Map.Find(SurfaceNode->SharedSurfaceId);
-						check(SharedMeshResults);
+				{
+					UE::TUniqueLock Lock(SharedSurfaceOptions.Mutex);
 
-						SharedResultLayoutOps = SharedMeshResults->LayoutOps;
+					FSharedSurfaceResultOptions::FSharedSurfaceResultKey SharedResultKey;
+					SharedResultKey.BaseSurfaceGuid = SurfaceNode->SurfaceGuid;
+					SharedResultKey.CombinedGuid = SharedSurfaceGuid;
 
-					}
+					// Do we have the surface we need to share it with?
+					SharedSurfaceResults = SharedSurfaceOptions.Map.Find(SharedResultKey);
+					bIsBaseForSharedSurface = SharedSurfaceResults == nullptr;
 
 					// Add layouts form the base mesh.	
 					for (int32 LayoutIndex = 0; LayoutIndex < MeshResults.GeneratedLayouts.Num(); ++LayoutIndex)
@@ -1089,15 +1086,15 @@ namespace mu
 
 						SurfaceReferenceLayouts[LayoutIndex] = MeshResults.GeneratedLayouts[LayoutIndex];
 
-						bool bSharedHasThisLayout = SharedMeshResults
+						bool bSharedHasThisLayout = SharedSurfaceResults
 							&&
-							SharedMeshResults->LayoutOps.IsValidIndex(LayoutIndex)
+							SharedSurfaceResults->LayoutOps.IsValidIndex(LayoutIndex)
 							&&
-							SharedMeshResults->LayoutOps[LayoutIndex];
+							SharedSurfaceResults->LayoutOps[LayoutIndex];
 
 						if (bSharedHasThisLayout)
 						{
-							SurfaceLayoutOps[LayoutIndex] = SharedMeshResults->LayoutOps[LayoutIndex];
+							SurfaceLayoutOps[LayoutIndex] = SharedSurfaceResults->LayoutOps[LayoutIndex];
 						}
 						else
 						{
@@ -1141,12 +1138,12 @@ namespace mu
 							LayoutFromExtension[LayoutIndex] = bLayoutSetByThisExtension;
 						}
 
-						if (bShareSurface)
+						if (SharedSurfaceResults)
 						{
 							if (!SurfaceLayoutOps[LayoutIndex] && bLayoutSetByThisExtension)
 							{
-								check(SharedResultLayoutOps.IsValidIndex(LayoutIndex));
-								SurfaceLayoutOps[LayoutIndex] = SharedResultLayoutOps[LayoutIndex];
+								check(SharedSurfaceResults->LayoutOps.IsValidIndex(LayoutIndex));
+								SurfaceLayoutOps[LayoutIndex] = SharedSurfaceResults->LayoutOps[LayoutIndex];
 							}
 						}
 						else
@@ -1198,7 +1195,7 @@ namespace mu
 					}
 
 					// Add layout packing instructions
-					if (!bShareSurface)
+					if (!SharedSurfaceResults)
 					{
 						// Make sure we removed unnecessary blocks
 						Ptr<ASTOpLayoutFromMesh> ExtractOp = new ASTOpLayoutFromMesh();
@@ -1242,7 +1239,7 @@ namespace mu
 
 				// Create the expression for each texture, if we are not reusing the surface from another LOD.
 				//------------------------------------------------------------------------
-				if (!bShareSurface)
+				if (!SharedSurfaceResults)
 				{
 					for (int32 ImageIndex = 0; ImageIndex < SurfaceNode->Images.Num(); ++ImageIndex)
 					{
@@ -1674,14 +1671,42 @@ namespace mu
 
 				FSurfaceGenerationResult SurfaceResult;
 				SurfaceResult.SurfaceOp = LastSurfOp;
+				SurfaceResult.SharedSurfaceGuid = SharedSurfaceGuid;
 
 				// If we are going to share this surface properties, remember it.
 				if (bIsBaseForSharedSurface)
 				{
-					UE::TUniqueLock Lock(SharedMeshOptions.Mutex);
+					{
+						FSharedSurfaceResultOptions::FSharedSurfaceResult SharedResult;
+						SharedResult.LayoutOps = MeshResults.LayoutOps;
 
-					check(!SharedMeshOptions.Map.Contains(SurfaceNode->SharedSurfaceId));
-					SharedMeshOptions.Map.Add(SurfaceNode->SharedSurfaceId, MeshResults);
+						UE::TUniqueLock Lock(SharedSurfaceOptions.Mutex);
+
+						FSharedSurfaceResultOptions::FSharedSurfaceResultKey SharedResultKey;
+						SharedResultKey.BaseSurfaceGuid = SurfaceNode->SurfaceGuid;
+						SharedResultKey.CombinedGuid = SharedSurfaceGuid;
+						
+						check(!SharedSurfaceOptions.Map.Contains(SharedResultKey));
+						SharedSurfaceOptions.Map.Add(SharedResultKey, SharedResult);
+					}
+
+					{
+						UE::TUniqueLock Lock(ExtraLayoutsOptions.Mutex);
+
+						FExtraLayoutsResultOptions::FExtraLayoutsKey ExtraMeshLayoutKey;
+						ExtraMeshLayoutKey.BaseSurfaceGuid = SurfaceNode->SurfaceGuid;
+						
+						for (const FMeshGenerationResult::FExtraLayouts& ExtraMeshLayout : MeshResults.ExtraMeshLayouts)
+						{
+							ExtraMeshLayoutKey.EditGuid = ExtraMeshLayout.EditGuid;
+
+							// This assumes that the lods are processed in order. It checks it this way because some platforms may have empty lods at the top.
+							if (!ExtraLayoutsOptions.Map.Contains(ExtraMeshLayoutKey))
+							{
+								ExtraLayoutsOptions.Map.Add(ExtraMeshLayoutKey, ExtraMeshLayout);
+							}
+						}
+					}
 				}
 
 				return SurfaceResult;
@@ -1777,7 +1802,7 @@ namespace mu
 					SurfaceOp->value = SurfaceGenerationResult.SurfaceOp;
 					SurfaceOp->id = surfaceID;
 					SurfaceOp->ExternalId = SurfaceData.Node->ExternalId;
-					SurfaceOp->SharedSurfaceId = SurfaceData.Node->SharedSurfaceId;
+					SurfaceOp->SharedSurfaceId = GetTypeHash(SurfaceGenerationResult.SharedSurfaceGuid);
 
 					Ptr<ASTOp> SurfaceConditionOp = SurfaceData.FinalCondition;
 
@@ -2153,7 +2178,7 @@ namespace mu
 		const FMeshGenerationStaticOptions& StaticOptions, 
 		FMeshOptionsTask Options,
 		FMeshTask BaseTask,
-		int32 SharedSurfaceId,
+		FGuid BaseSurfaceGuid,
 		const void* ErrorContext,
 		const NodeMeshConstant* OriginalMeshNode)
 	{
@@ -2188,31 +2213,28 @@ namespace mu
 					MergedMeshStaticOptions.ModifiersToIgnore.Add(Modifier);
 
 					FMeshOptionsTask MergedMeshOptionsTask = UE::Tasks::Launch(TEXT("MutableMergedMeshOptions"),
-						[Options, SharedSurfaceId, EditIndex, MergedMeshStaticOptions, this]() mutable
-						{
+						[Options, BaseSurfaceGuid, EditGuid = Edit->ModifierGuid, MergedMeshStaticOptions, this]() mutable
+						{	
+							// This assumes that the lods are processed in order.
+
 							FMeshGenerationDynamicOptions Result = Options.GetResult();
 							Result.bEnsureAllVerticesHaveLayoutBlock = false;
 
-							UE::TUniqueLock Lock(SharedMeshOptions.Mutex);
-
-							// This assumes that the lods are processed in order. It checks it this way because some platforms may have empty lods at the top.
-							const bool bIsBaseForSharedSurface = SharedSurfaceId != INDEX_NONE && !SharedMeshOptions.Map.Contains(SharedSurfaceId);
-
-							// If this is true, we will reuse the surface properties from a higher LOD, se we can skip the generation of material properties and images.
-							const bool bShareSurface = SharedSurfaceId != INDEX_NONE && !bIsBaseForSharedSurface;
+							UE::TUniqueLock Lock(ExtraLayoutsOptions.Mutex);
+							
+							check(BaseSurfaceGuid.IsValid() && EditGuid.IsValid());
+							FExtraLayoutsResultOptions::FExtraLayoutsKey ExtraLayoutKey;
+							ExtraLayoutKey.BaseSurfaceGuid = BaseSurfaceGuid;
+							ExtraLayoutKey.EditGuid = EditGuid;
 
-							const FMeshGenerationResult* SharedMeshResults = nullptr;
-							if (bShareSurface)
+							// Reuse layouts from previous LODs if possible.
+							if (const FMeshGenerationResult::FExtraLayouts* ExtraLayouts = ExtraLayoutsOptions.Map.Find(ExtraLayoutKey))
 							{
-								// Do we have the surface we need to share it with?
-								SharedMeshResults = SharedMeshOptions.Map.Find(SharedSurfaceId);
-								check(SharedMeshResults);
-
-								if (SharedMeshResults)
-								{
-									check(SharedMeshResults->ExtraMeshLayouts.IsValidIndex(EditIndex));
-									Result.OverrideLayouts = SharedMeshResults->ExtraMeshLayouts[EditIndex].GeneratedLayouts;
-								}
+								Result.OverrideLayouts = ExtraLayouts->GeneratedLayouts;
+							}
+							else
+							{
+								Result.OverrideLayouts.Empty();
 							}
 
 							return Result;
@@ -2232,7 +2254,7 @@ namespace mu
 							{
 								if (BaseMeshResult.GeneratedLayouts.Num() != AddResults.GeneratedLayouts.Num())
 								{
-									// When extending a mesh section the added mesh sectio will use the layout strategy of the base one
+									// When extending a mesh section the added mesh section will use the layout strategy of the base one
 									FString Msg = FString::Printf(TEXT("Extended mesh section layout count is differenta than the mesh being extended."));
 									ErrorLog->Add(Msg, ELMT_INFO, Edit->GetMessageContext(), ErrorContext);
 								}
@@ -2286,10 +2308,10 @@ namespace mu
 					constexpr bool bModifiersForBeforeOperations = false;
 					GetModifiersFor(StaticOptions.ComponentId, ModifierOptions.ActiveTags, bModifiersForBeforeOperations, ChildModifiers);
 
-					FMeshTask AddWithModifiersTask = ApplyMeshModifiers(ChildModifiers, ModifierOptions, Options, AddTask, SharedSurfaceId, ErrorContext, nullptr);
+					FMeshTask AddWithModifiersTask = ApplyMeshModifiers(ChildModifiers, ModifierOptions, Options, AddTask, BaseSurfaceGuid, ErrorContext, nullptr);
 
 					LastMeshTask = UE::Tasks::Launch(TEXT("MutableMeshMergeModifier"),
-						[AddBaseTask, AddWithModifiersTask, LastMeshTask, Modifier, EditIndex]() mutable
+						[AddBaseTask, AddWithModifiersTask, LastMeshTask, Modifier, EditIndex, EditGuid = Edit->ModifierGuid]() mutable
 						{
 							FMeshGenerationResult AddResults = AddBaseTask.GetResult();
 							FMeshGenerationResult AddFinalResults = AddWithModifiersTask.GetResult();
@@ -2299,6 +2321,7 @@ namespace mu
 							Data.GeneratedLayouts = AddResults.GeneratedLayouts;
 							Data.Condition = Modifier.FinalCondition;
 							Data.MeshFragment = AddFinalResults.MeshOp;
+							Data.EditGuid = EditGuid;
 							if (LastMeshResults.ExtraMeshLayouts.Num() <= EditIndex)
 							{
 								LastMeshResults.ExtraMeshLayouts.SetNum(EditIndex+1);
diff --git a/Engine/Plugins/Mutable/Source/MutableTools/Private/MuT/CodeGenerator.h b/Engine/Plugins/Mutable/Source/MutableTools/Private/MuT/CodeGenerator.h
--- a/Engine/Plugins/Mutable/Source/MutableTools/Private/MuT/CodeGenerator.h
+++ b/Engine/Plugins/Mutable/Source/MutableTools/Private/MuT/CodeGenerator.h
@@ -609,21 +609,84 @@ namespace mu
 
 		using FMeshOptionsTask = UE::Tasks::TTask<CodeGenerator::FMeshGenerationDynamicOptions>;
 
-		struct FSharedMeshOptions
+		struct FBaseMeshResultOptions
 		{
-			/** Use this mutex to control acces to the SharedMeshOptionsMap attribute from any thread.
+			/** Use this mutex to control access to the BaseMeshOptions attribute from any thread.
 			* The task dependency will make sure that the actual order for data insert and retrieval is correct, but unrelated data can be added while
 			* unrelated data is being retrieved so we need the mutex.
 			*/
 			UE::FMutex Mutex;
 
-			/** Store the mesh generation data for surfaces that we intend to share across LODs.
-			* The key is the SharedSurfaceId.
+			/** Store the mesh generation data for edits that we intend to share across LODs.
+			* The key is the combination of the SharedSurfaceGuid of the surface.
 			*/
-			TMap<int32, FMeshGenerationResult> Map;
+			TMap<FGuid, FMeshGenerationResult> Map;
 		};
+		FBaseMeshResultOptions BaseMeshOptions;
 
-		FSharedMeshOptions SharedMeshOptions;
+		struct FExtraLayoutsResultOptions
+		{
+			/** Use this mutex to control access to the ExtraLayoutsOptions attribute from any thread.
+			* The task dependency will make sure that the actual order for data insert and retrieval is correct, but unrelated data can be added while
+			* unrelated data is being retrieved so we need the mutex.
+			*/
+			UE::FMutex Mutex;
+
+			struct FExtraLayoutsKey
+			{
+				FGuid BaseSurfaceGuid;
+				FGuid EditGuid;
+
+				friend FORCEINLINE uint32 GetTypeHash(const FExtraLayoutsKey& InKey)
+				{
+					return HashCombineFast(GetTypeHash(InKey.BaseSurfaceGuid), GetTypeHash(InKey.EditGuid));
+				}
+
+				FORCEINLINE bool operator==(const FExtraLayoutsKey& Other) const = default;
+			};
+
+			/** Store the mesh generation data for edits that we intend to share across LODs.
+			* The key is the combination of the BaseSurfaceGuid and the ModifierGuid.
+			*/
+			TMap<FExtraLayoutsKey, FMeshGenerationResult::FExtraLayouts> Map;
+		};
+
+		FExtraLayoutsResultOptions ExtraLayoutsOptions;
+
+		struct FSharedSurfaceResultOptions
+		{
+			/** Use this mutex to control access to the SharedSurfaceOptions attribute from any thread.
+			* The task dependency will make sure that the actual order for data insert and retrieval is correct, but unrelated data can be added while
+			* unrelated data is being retrieved so we need the mutex.
+			*/
+			UE::FMutex Mutex;
+
+			struct FSharedSurfaceResult
+			{
+				TArray<Ptr<ASTOp>> LayoutOps;
+			};
+
+			struct FSharedSurfaceResultKey
+			{
+				FGuid BaseSurfaceGuid;
+				FGuid CombinedGuid;
+
+				friend FORCEINLINE uint32 GetTypeHash(const FSharedSurfaceResultKey& InKey)
+				{
+					return HashCombineFast(GetTypeHash(InKey.BaseSurfaceGuid), GetTypeHash(InKey.CombinedGuid));
+				}
+
+				FORCEINLINE bool operator==(const FSharedSurfaceResultKey& Other) const = default;
+			};
+
+			/** Store the generated layout operations that we intend to share across LODs.
+			* The key is the combination of the BaseSurfaceGuid and the combination of all ModifierGuids 
+			* affecting this surface.
+			*/
+			TMap<FSharedSurfaceResultKey, FSharedSurfaceResult> Map;
+		};
+
+		FSharedSurfaceResultOptions SharedSurfaceOptions;
 
 		/** This struct contains the generation state that can only be accessed when generating constant layouts.
 		* It should only be accessed after locking the included mutex.
@@ -868,7 +931,7 @@ namespace mu
 			const FMeshGenerationStaticOptions&,
 			FMeshOptionsTask DynamicOptionsTask,
 			FMeshTask BaseTask,
-			int32 SharedSurfaceId,
+			FGuid SharedSurfaceId,
 			const void* ErrorContext,
 			const NodeMeshConstant* OriginalMeshNode);
 
diff --git a/Engine/Plugins/Mutable/Source/MutableTools/Private/MuT/CodeGenerator_FirstPass.h b/Engine/Plugins/Mutable/Source/MutableTools/Private/MuT/CodeGenerator_FirstPass.h
--- a/Engine/Plugins/Mutable/Source/MutableTools/Private/MuT/CodeGenerator_FirstPass.h
+++ b/Engine/Plugins/Mutable/Source/MutableTools/Private/MuT/CodeGenerator_FirstPass.h
@@ -63,12 +63,14 @@ namespace mu
 
 		struct FExtraLayouts
 		{
+			FGuid EditGuid;
+
 			/** Source node layouts to use with these extra mesh. They don't have block ids. */
 			TArray<FGeneratedLayout> GeneratedLayouts;
 			Ptr<ASTOp> Condition;
 			Ptr<ASTOp> MeshFragment;
 		};
-		TArray< FExtraLayouts > ExtraMeshLayouts;
+		TArray<FExtraLayouts> ExtraMeshLayouts;
 	};
 
 	using FMeshTask = UE::Tasks::TTask<FMeshGenerationResult>;
@@ -76,6 +78,7 @@ namespace mu
 	struct FSurfaceGenerationResult
 	{
 		Ptr<ASTOp> SurfaceOp;
+		FGuid SharedSurfaceGuid;
 	};
 
 	using FSurfaceTask = UE::Tasks::TTask<FSurfaceGenerationResult>;
diff --git a/Engine/Plugins/Mutable/Source/MutableTools/Private/MuT/CodeGenerator_Mesh.cpp b/Engine/Plugins/Mutable/Source/MutableTools/Private/MuT/CodeGenerator_Mesh.cpp
--- a/Engine/Plugins/Mutable/Source/MutableTools/Private/MuT/CodeGenerator_Mesh.cpp
+++ b/Engine/Plugins/Mutable/Source/MutableTools/Private/MuT/CodeGenerator_Mesh.cpp
@@ -895,7 +895,7 @@ namespace mu
 		this->GetModifiersFor(StaticOptions.ComponentId, StaticOptions.ActiveTags, bModifiersForBeforeOperations, Modifiers);
 
 		// This task does the necessary processing of the mesh when it is available
-		FMeshTask ApplyModifiersTask = ApplyMeshModifiers(Modifiers, StaticOptions, InOptions, MeshProcessingTask, INDEX_NONE, Node->GetMessageContext(), Node);
+		FMeshTask ApplyModifiersTask = ApplyMeshModifiers(Modifiers, StaticOptions, InOptions, MeshProcessingTask, FGuid(), Node->GetMessageContext(), Node);
 
 		return ApplyModifiersTask;
     }
diff --git a/Engine/Plugins/Mutable/Source/MutableTools/Public/MuT/NodeModifierSurfaceEdit.h b/Engine/Plugins/Mutable/Source/MutableTools/Public/MuT/NodeModifierSurfaceEdit.h
--- a/Engine/Plugins/Mutable/Source/MutableTools/Public/MuT/NodeModifierSurfaceEdit.h
+++ b/Engine/Plugins/Mutable/Source/MutableTools/Public/MuT/NodeModifierSurfaceEdit.h
@@ -72,6 +72,9 @@ namespace mu
 		/** Factor of the morph to apply. */
 		Ptr<NodeScalar> MorphFactor;
 
+		/** Source modifier guid, used to share resources between LODs and generate a deterministic surface guid. */
+		FGuid ModifierGuid;
+
 	public:
 
 		NodeModifierSurfaceEdit() 
diff --git a/Engine/Plugins/Mutable/Source/MutableTools/Public/MuT/NodeSurfaceNew.h b/Engine/Plugins/Mutable/Source/MutableTools/Public/MuT/NodeSurfaceNew.h
--- a/Engine/Plugins/Mutable/Source/MutableTools/Public/MuT/NodeSurfaceNew.h
+++ b/Engine/Plugins/Mutable/Source/MutableTools/Public/MuT/NodeSurfaceNew.h
@@ -12,6 +12,8 @@
 #include "MuT/NodeString.h"
 #include "MuT/NodeColour.h"
 
+#include "Misc/Guid.h"
+
 namespace mu
 {
 
@@ -27,8 +29,8 @@ namespace mu
 		*/
 		uint32 ExternalId = 0;
 
-		/** Add an id that will be used to identify the same surface in other LODs. */
-		int32 SharedSurfaceId = INDEX_NONE;
+		/** Source surface Guid, used to share layouts between LODs and generate a deterministic SharedSurfaceId. */
+		FGuid SurfaceGuid;
 
 		Ptr<NodeMesh> Mesh;
 
