[Mutable] Added support for RGBA textures to the SampleTexture node.

#rnx

diff --git a/Engine/Plugins/Mutable/Source/MutableTools/Private/MuT/CodeGenerator_Color.cpp b/Engine/Plugins/Mutable/Source/MutableTools/Private/MuT/CodeGenerator_Color.cpp
--- a/Engine/Plugins/Mutable/Source/MutableTools/Private/MuT/CodeGenerator_Color.cpp
+++ b/Engine/Plugins/Mutable/Source/MutableTools/Private/MuT/CodeGenerator_Color.cpp
@@ -334,7 +334,7 @@ namespace mu
 			// This argument is required
 			base = GenerateMissingImageCode(TEXT("Sample image"), EImageFormat::RGB_UByte, Typed->GetMessageContext(), ImageOptions);
 		}
-		base = GenerateImageFormat(base, EImageFormat::RGB_UByte);
+		base = GenerateImageFormat(base, EImageFormat::RGBA_UByte);
 		op->Image = base;
 
 		FScalarGenerationResult ChildResult;
