[Mutable] Added pointer check to avoid crashes when a node still points to deleted skeletal meshes.

#rnx

diff --git a/Engine/Plugins/Mutable/Source/CustomizableObjectEditor/Private/MuCOE/GenerateMutableSource/GenerateMutableSourceMesh.cpp b/Engine/Plugins/Mutable/Source/CustomizableObjectEditor/Private/MuCOE/GenerateMutableSource/GenerateMutableSourceMesh.cpp
--- a/Engine/Plugins/Mutable/Source/CustomizableObjectEditor/Private/MuCOE/GenerateMutableSource/GenerateMutableSourceMesh.cpp
+++ b/Engine/Plugins/Mutable/Source/CustomizableObjectEditor/Private/MuCOE/GenerateMutableSource/GenerateMutableSourceMesh.cpp
@@ -3775,170 +3775,181 @@ mu::Ptr<mu::NodeMesh> GenerateMutableSourceMesh(const UEdGraphPin* Pin,
 		}
 		else
 		{
-			int32 LODIndexConnected = -1; // LOD which the pin is connected to
-			int32 SectionIndexConnected = -1;
-			
-			int32 LODIndex = -1;
-			int32 SectionIndex = -1;
-
-			{
-				int32 LayoutIndex;
-				TypedNodeSkel->GetPinSection(*Pin, LODIndexConnected, SectionIndexConnected, LayoutIndex);
-			}
-
 			// TODO: Avoid loading the mesh here.
 			USkeletalMesh* SkeletalMesh = GenerationContext.LoadObject(TypedNodeSkel->SkeletalMesh);
-			GetLODAndSectionForAutomaticLODs(GenerationContext, *Node, *SkeletalMesh, LODIndexConnected, SectionIndexConnected, LODIndex, SectionIndex, bOnlyConnectedLOD);
-			
-			// First process the mesh tags that are going to make the mesh unique and affect whether it's repeated in 
-			// the mesh cache or not
-			FString MeshUniqueTags;
-			FString AnimBPAssetTag;
 
-			if (!TypedNodeSkel->AnimInstance.IsNull())
+			if (SkeletalMesh)
 			{
-				FName SlotIndex = TypedNodeSkel->AnimBlueprintSlotName;
-				const int32 AnimInstanceIndex = GenerationContext.AnimBPAssets.AddUnique(TypedNodeSkel->AnimInstance);
+				int32 LODIndexConnected = -1; // LOD which the pin is connected to
+				int32 SectionIndexConnected = -1;
 
-				AnimBPAssetTag = GenerateAnimationInstanceTag(AnimInstanceIndex, SlotIndex);
-				MeshUniqueTags += AnimBPAssetTag;
-			}
+				int32 LODIndex = -1;
+				int32 SectionIndex = -1;
 
-			TArray<FString> ArrayAnimBPTags;
+				{
+					int32 LayoutIndex;
+					TypedNodeSkel->GetPinSection(*Pin, LODIndexConnected, SectionIndexConnected, LayoutIndex);
+				}
 
-			for (const FGameplayTag& GamePlayTag : TypedNodeSkel->AnimationGameplayTags)
-			{
-				const FString AnimBPTag = GenerateGameplayTag(GamePlayTag.ToString());
-				ArrayAnimBPTags.Add(AnimBPTag);
-				MeshUniqueTags += AnimBPTag;
-			}
 
-			TArray<FCustomizableObjectStreameableResourceId> StreamedResources;
+				GetLODAndSectionForAutomaticLODs(GenerationContext, *Node, *SkeletalMesh, LODIndexConnected, SectionIndexConnected, LODIndex, SectionIndex, bOnlyConnectedLOD);
 
-			if (GenerationContext.CompilationContext.Object->GetPrivate()->IsAssetUserDataMergeEnabled())
-			{
-				const TArray<UAssetUserData*>* AssetUserDataArray = SkeletalMesh->GetAssetUserDataArray();
+				// First process the mesh tags that are going to make the mesh unique and affect whether it's repeated in 
+				// the mesh cache or not
+				FString MeshUniqueTags;
+				FString AnimBPAssetTag;
+
+				if (!TypedNodeSkel->AnimInstance.IsNull())
+				{
+					FName SlotIndex = TypedNodeSkel->AnimBlueprintSlotName;
+					const int32 AnimInstanceIndex = GenerationContext.AnimBPAssets.AddUnique(TypedNodeSkel->AnimInstance);
+
+					AnimBPAssetTag = GenerateAnimationInstanceTag(AnimInstanceIndex, SlotIndex);
+					MeshUniqueTags += AnimBPAssetTag;
+				}
+
+				TArray<FString> ArrayAnimBPTags;
 
-				if (AssetUserDataArray)
+				for (const FGameplayTag& GamePlayTag : TypedNodeSkel->AnimationGameplayTags)
 				{
-					for (UAssetUserData* AssetUserData : *AssetUserDataArray)
+					const FString AnimBPTag = GenerateGameplayTag(GamePlayTag.ToString());
+					ArrayAnimBPTags.Add(AnimBPTag);
+					MeshUniqueTags += AnimBPTag;
+				}
+
+				TArray<FCustomizableObjectStreameableResourceId> StreamedResources;
+
+				if (GenerationContext.CompilationContext.Object->GetPrivate()->IsAssetUserDataMergeEnabled())
+				{
+					const TArray<UAssetUserData*>* AssetUserDataArray = SkeletalMesh->GetAssetUserDataArray();
+
+					if (AssetUserDataArray)
 					{
-						if (!AssetUserData)
+						for (UAssetUserData* AssetUserData : *AssetUserDataArray)
 						{
-							continue;
-						}
+							if (!AssetUserData)
+							{
+								continue;
+							}
 
-						const int32 ResourceIndex = GenerationContext.AddStreamedResource(*AssetUserData);
-						if (ResourceIndex >= 0)
-						{	
-							FCustomizableObjectStreameableResourceId ResourceId;
-							ResourceId.Id = (uint32)ResourceIndex;
-							ResourceId.Type = static_cast<uint8>(FCustomizableObjectStreameableResourceId::EType::AssetUserData);
+							const int32 ResourceIndex = GenerationContext.AddStreamedResource(*AssetUserData);
+							if (ResourceIndex >= 0)
+							{
+								FCustomizableObjectStreameableResourceId ResourceId;
+								ResourceId.Id = (uint32)ResourceIndex;
+								ResourceId.Type = static_cast<uint8>(FCustomizableObjectStreameableResourceId::EType::AssetUserData);
 
-							StreamedResources.Add(ResourceId);
-						}
+								StreamedResources.Add(ResourceId);
+							}
 
-						MeshUniqueTags += AssetUserData->GetPathName();
+							MeshUniqueTags += AssetUserData->GetPathName();
+						}
 					}
 				}
-			}
 
-			FSkeletalMeshModel* ImportedModel = SkeletalMesh->GetImportedModel();
-			
-			FMutableSourceMeshData Source = BaseMeshData;
-			Source.Mesh = SkeletalMesh;
-			Source.AnimInstance = TypedNodeSkel->AnimInstance;
+				FSkeletalMeshModel* ImportedModel = SkeletalMesh->GetImportedModel();
 
-			TSharedPtr<mu::FMesh> MutableMesh = GenerateMutableSkeletalMesh(Source, LODIndexConnected, SectionIndexConnected,
-																LODIndex, SectionIndex, MeshUniqueTags, GenerationContext, TypedNodeSkel);
-			if (MutableMesh)
-			{
-				MeshNode->Value = MutableMesh;
+				FMutableSourceMeshData Source = BaseMeshData;
+				Source.Mesh = SkeletalMesh;
+				Source.AnimInstance = TypedNodeSkel->AnimInstance;
 
-				if (!TypedNodeSkel->AnimInstance.IsNull())
+				TSharedPtr<mu::FMesh> MutableMesh = GenerateMutableSkeletalMesh(Source, LODIndexConnected, SectionIndexConnected,
+					LODIndex, SectionIndex, MeshUniqueTags, GenerationContext, TypedNodeSkel);
+				if (MutableMesh)
 				{
-					AddTagToMutableMeshUnique(*MutableMesh, AnimBPAssetTag);
-				}
+					MeshNode->Value = MutableMesh;
 
-				for (const FString& GamePlayTag : ArrayAnimBPTags)
-				{
-					AddTagToMutableMeshUnique(*MutableMesh, GamePlayTag);
-				}
+					if (!TypedNodeSkel->AnimInstance.IsNull())
+					{
+						AddTagToMutableMeshUnique(*MutableMesh, AnimBPAssetTag);
+					}
 
-				for (FCustomizableObjectStreameableResourceId ResourceId : StreamedResources)
-				{
-					MutableMesh->AddStreamedResource(BitCast<uint64>(ResourceId));
-				}
+					for (const FString& GamePlayTag : ArrayAnimBPTags)
+					{
+						AddTagToMutableMeshUnique(*MutableMesh, GamePlayTag);
+					}
 
-				AddSocketTagsToMesh(SkeletalMesh, MutableMesh, GenerationContext);
+					for (FCustomizableObjectStreameableResourceId ResourceId : StreamedResources)
+					{
+						MutableMesh->AddStreamedResource(BitCast<uint64>(ResourceId));
+					}
 
-				if (UCustomizableObjectSystem::GetInstance()->IsMutableAnimInfoDebuggingEnabled())
-				{
-					FString MeshPath;
-					SkeletalMesh->GetOuter()->GetPathName(nullptr, MeshPath);
-					FString MeshTag = FString("__MeshPath:") + MeshPath;
-					AddTagToMutableMeshUnique(*MutableMesh, MeshTag);
-				}
-			}
+					AddSocketTagsToMesh(SkeletalMesh, MutableMesh, GenerationContext);
 
-			// Layouts
-			{
-				// When using Automatic From Mesh all LODs share the same base layout, hence we use LODIndexConnected (as the base layout) instead of the LODIndex.
-				const int32 LODIndexLayout = GenerationContext.CurrentAutoLODStrategy == ECustomizableObjectAutomaticLODStrategy::AutomaticFromMesh ?
-					LODIndexConnected :
-					LODIndex;
-				const int32 SectionIndexLayout = GenerationContext.CurrentAutoLODStrategy == ECustomizableObjectAutomaticLODStrategy::AutomaticFromMesh ?
-					SectionIndexConnected :
-					SectionIndex;
-					
-				int32 NumLayouts = 0;
-				if (ImportedModel && ImportedModel->LODModels.IsValidIndex(LODIndexLayout))
-				{
-					NumLayouts = ImportedModel->LODModels[LODIndexLayout].NumTexCoords;
+					if (UCustomizableObjectSystem::GetInstance()->IsMutableAnimInfoDebuggingEnabled())
+					{
+						FString MeshPath;
+						SkeletalMesh->GetOuter()->GetPathName(nullptr, MeshPath);
+						FString MeshTag = FString("__MeshPath:") + MeshPath;
+						AddTagToMutableMeshUnique(*MutableMesh, MeshTag);
+					}
 				}
-				MeshNode->Layouts.SetNum(NumLayouts);
-				
-				const FLayoutGenerationFlags& LayoutFlags = GenerationContext.LayoutGenerationFlags.Last();
 
-				const TArray<UCustomizableObjectLayout*>& Layouts = TypedNodeSkel->GetLayouts(*Pin);
-				for (int32 LayoutIndex = 0; LayoutIndex < NumLayouts; ++LayoutIndex)
+				// Layouts
 				{
-					if (!LayoutFlags.TexturePinModes.IsValidIndex(LayoutIndex) ||
-						LayoutFlags.TexturePinModes[LayoutIndex] != EPinMode::Mutable)
+					// When using Automatic From Mesh all LODs share the same base layout, hence we use LODIndexConnected (as the base layout) instead of the LODIndex.
+					const int32 LODIndexLayout = GenerationContext.CurrentAutoLODStrategy == ECustomizableObjectAutomaticLODStrategy::AutomaticFromMesh ?
+						LODIndexConnected :
+						LODIndex;
+					const int32 SectionIndexLayout = GenerationContext.CurrentAutoLODStrategy == ECustomizableObjectAutomaticLODStrategy::AutomaticFromMesh ?
+						SectionIndexConnected :
+						SectionIndex;
+
+					int32 NumLayouts = 0;
+					if (ImportedModel && ImportedModel->LODModels.IsValidIndex(LODIndexLayout))
 					{
-						MeshNode->Layouts[LayoutIndex] = CreateDefaultLayout();
-					
-						// Keep packing strategy if possible, Overlay can be valid with EPinMode == Passthrough	
-						if (Layouts.IsValidIndex(LayoutIndex) && Layouts[LayoutIndex])
-						{
-							MeshNode->Layouts[LayoutIndex]->Strategy = ConvertLayoutStrategy(Layouts[LayoutIndex]->PackingStrategy);
-						}
-
-						// Ignore layout
-						continue;
+						NumLayouts = ImportedModel->LODModels[LODIndexLayout].NumTexCoords;
 					}
+					MeshNode->Layouts.SetNum(NumLayouts);
+
+					const FLayoutGenerationFlags& LayoutFlags = GenerationContext.LayoutGenerationFlags.Last();
 
-					const UCustomizableObjectLayout* Layout = Layouts.IsValidIndex(LayoutIndex) ? Layouts[LayoutIndex] : nullptr;
-					if (ensure(Layout))
+					const TArray<UCustomizableObjectLayout*>& Layouts = TypedNodeSkel->GetLayouts(*Pin);
+					for (int32 LayoutIndex = 0; LayoutIndex < NumLayouts; ++LayoutIndex)
 					{
-						bool bWasEmpty = false;
-						mu::Ptr<mu::NodeLayout> LayoutNode = CreateMutableLayoutNode(GenerationContext, Layout, bLinkedToExtendMaterial, bWasEmpty); // TODO PERE: Figure out
-						LayoutNode->SetMessageContext(Node);
+						if (!LayoutFlags.TexturePinModes.IsValidIndex(LayoutIndex) ||
+							LayoutFlags.TexturePinModes[LayoutIndex] != EPinMode::Mutable)
+						{
+							MeshNode->Layouts[LayoutIndex] = CreateDefaultLayout();
+
+							// Keep packing strategy if possible, Overlay can be valid with EPinMode == Passthrough	
+							if (Layouts.IsValidIndex(LayoutIndex) && Layouts[LayoutIndex])
+							{
+								MeshNode->Layouts[LayoutIndex]->Strategy = ConvertLayoutStrategy(Layouts[LayoutIndex]->PackingStrategy);
+							}
+
+							// Ignore layout
+							continue;
+						}
+
+						const UCustomizableObjectLayout* Layout = Layouts.IsValidIndex(LayoutIndex) ? Layouts[LayoutIndex] : nullptr;
+						if (ensure(Layout))
+						{
+							bool bWasEmpty = false;
+							mu::Ptr<mu::NodeLayout> LayoutNode = CreateMutableLayoutNode(GenerationContext, Layout, bLinkedToExtendMaterial, bWasEmpty); // TODO PERE: Figure out
+							LayoutNode->SetMessageContext(Node);
 
-						MeshNode->Layouts[LayoutIndex] = LayoutNode;
+							MeshNode->Layouts[LayoutIndex] = LayoutNode;
+						}
 					}
 				}
-			}
 
 
-			const FString MeshName = TypedNodeSkel->SkeletalMesh.GetLongPackageName().ToLower();
-			MeshNode->SourceDataDescriptor.SourceId = CityHash32(reinterpret_cast<const char*>(*MeshName), MeshName.Len() * sizeof(FString::ElementType));
-			MeshNode->SourceDataDescriptor.OptionalMaxLODSize = 0;
+				const FString MeshName = TypedNodeSkel->SkeletalMesh.GetLongPackageName().ToLower();
+				MeshNode->SourceDataDescriptor.SourceId = CityHash32(reinterpret_cast<const char*>(*MeshName), MeshName.Len() * sizeof(FString::ElementType));
+				MeshNode->SourceDataDescriptor.OptionalMaxLODSize = 0;
 
-			// Applying Mesh Morph Nodes
-			if (GenerationContext.MeshMorphStack.Num())
+				// Applying Mesh Morph Nodes
+				if (GenerationContext.MeshMorphStack.Num())
+				{
+					MorphResult = GenerateMorphMesh(Pin, GenerationContext.MeshMorphStack, 0, Result, GenerationContext, bOnlyConnectedLOD);
+				}
+			}
+			else
 			{
-				MorphResult = GenerateMorphMesh(Pin, GenerationContext.MeshMorphStack, 0, Result, GenerationContext, bOnlyConnectedLOD);
+				FText Path = TypedNodeSkel->SkeletalMesh.GetLongPackageFName().IsNone() ? FText::FromString("Unknown") : FText::FromName(TypedNodeSkel->SkeletalMesh.GetLongPackageFName());
+				FText Msg = FText::Format(LOCTEXT("SkeletalMeshNotLoaded", "Failed to load the skeletal mesh set in the node.[Asset path: {0}]"), Path);
+				GenerationContext.Log(Msg, Node);
 			}
 		}
 	}
