[Mutable] PR #11864 Fix UpdateSkeletalMeshAsync missing parameters.

diff --git a/Engine/Plugins/Mutable/Source/CustomizableObject/Private/MuCO/CustomizableObjectInstanceUsage.cpp b/Engine/Plugins/Mutable/Source/CustomizableObject/Private/MuCO/CustomizableObjectInstanceUsage.cpp
--- a/Engine/Plugins/Mutable/Source/CustomizableObject/Private/MuCO/CustomizableObjectInstanceUsage.cpp
+++ b/Engine/Plugins/Mutable/Source/CustomizableObject/Private/MuCO/CustomizableObjectInstanceUsage.cpp
@@ -271,13 +271,13 @@ USkeletalMesh* UCustomizableObjectInstanceUsagePrivate::GetAttachedSkeletalMesh(
 }
 
 
-void UCustomizableObjectInstanceUsage::UpdateSkeletalMeshAsync(bool bNeverSkipUpdate)
+void UCustomizableObjectInstanceUsage::UpdateSkeletalMeshAsync(bool bNeverSkipUpdate, bool bIgnoreCloseDist, bool bForceHighPriority)
 {
 	UCustomizableObjectInstance* CustomizableObjectInstance = GetCustomizableObjectInstance();
 
 	if (CustomizableObjectInstance)
 	{
-		CustomizableObjectInstance->UpdateSkeletalMeshAsync(false, false);
+		CustomizableObjectInstance->UpdateSkeletalMeshAsync(bIgnoreCloseDist, bForceHighPriority);
 	}
 }
 
@@ -288,7 +288,7 @@ void UCustomizableObjectInstanceUsage::UpdateSkeletalMeshAsyncResult(FInstanceUp
 
 	if (CustomizableObjectInstance)
 	{
-		CustomizableObjectInstance->UpdateSkeletalMeshAsyncResult(Callback, false, false);
+		CustomizableObjectInstance->UpdateSkeletalMeshAsyncResult(Callback, bIgnoreCloseDist, bForceHighPriority);
 	}
 }
 
diff --git a/Engine/Plugins/Mutable/Source/CustomizableObject/Public/MuCO/CustomizableObjectInstanceUsage.h b/Engine/Plugins/Mutable/Source/CustomizableObject/Public/MuCO/CustomizableObjectInstanceUsage.h
--- a/Engine/Plugins/Mutable/Source/CustomizableObject/Public/MuCO/CustomizableObjectInstanceUsage.h
+++ b/Engine/Plugins/Mutable/Source/CustomizableObject/Public/MuCO/CustomizableObjectInstanceUsage.h
@@ -64,7 +64,7 @@ public:
 	
 	/** Update Skeletal Mesh asynchronously. */
 	UFUNCTION(BlueprintCallable, Category = CustomizableObject)
-	UE_API void UpdateSkeletalMeshAsync(bool bNeverSkipUpdate = false);
+	UE_API void UpdateSkeletalMeshAsync(bool bNeverSkipUpdate = false, bool bIgnoreCloseDist = false, bool bForceHighPriority = false);
 
 	/** Update Skeletal Mesh asynchronously. Callback will be called once the update finishes, even if it fails. */
 	UFUNCTION(BlueprintCallable, Category = CustomizableObjectInstance)
