[Mutable] PR #11859 Added Instance to FUpdateContext

diff --git a/Engine/Plugins/Mutable/Source/CustomizableObject/Private/MuCO/CustomizableObjectSystem.cpp b/Engine/Plugins/Mutable/Source/CustomizableObject/Private/MuCO/CustomizableObjectSystem.cpp
--- a/Engine/Plugins/Mutable/Source/CustomizableObject/Private/MuCO/CustomizableObjectSystem.cpp
+++ b/Engine/Plugins/Mutable/Source/CustomizableObject/Private/MuCO/CustomizableObjectSystem.cpp
@@ -1245,6 +1245,7 @@ void FinishUpdateGlobal(const TSharedRef<FUpdateContextPrivate>& Context)
 
 	FUpdateContext ContextPublic;
 	ContextPublic.UpdateResult = Context->UpdateResult;
+	ContextPublic.Instance = Instance;
 
 #if WITH_EDITORONLY_DATA
 	// Restore the overriden system settings (bake operation)
diff --git a/Engine/Plugins/Mutable/Source/CustomizableObject/Public/MuCO/CustomizableObjectInstance.h b/Engine/Plugins/Mutable/Source/CustomizableObject/Public/MuCO/CustomizableObjectInstance.h
--- a/Engine/Plugins/Mutable/Source/CustomizableObject/Public/MuCO/CustomizableObjectInstance.h
+++ b/Engine/Plugins/Mutable/Source/CustomizableObject/Public/MuCO/CustomizableObjectInstance.h
@@ -218,8 +218,11 @@ struct FUpdateContext
 {	
 	GENERATED_BODY()
 
-	UPROPERTY(BlueprintReadOnly, Category=UpdateResult)
+	UPROPERTY(BlueprintReadOnly, Category = UpdateResult)
 	EUpdateResult UpdateResult = EUpdateResult::Success;
+
+	UPROPERTY(BlueprintReadOnly, Category = UpdateResult)
+	TObjectPtr<UCustomizableObjectInstance> Instance;
 };
 
 
