[Mutable] Add recompile and gather references for editor selected Customizable Objects as a contextual menu option.

#rnx

diff --git a/Engine/Plugins/Mutable/Source/CustomizableObjectEditor/Private/MuCOE/AssetDefinition_CustomizableObject.cpp b/Engine/Plugins/Mutable/Source/CustomizableObjectEditor/Private/MuCOE/AssetDefinition_CustomizableObject.cpp
--- a/Engine/Plugins/Mutable/Source/CustomizableObjectEditor/Private/MuCOE/AssetDefinition_CustomizableObject.cpp
+++ b/Engine/Plugins/Mutable/Source/CustomizableObjectEditor/Private/MuCOE/AssetDefinition_CustomizableObject.cpp
@@ -174,7 +174,8 @@ namespace MenuExtension_CustomizableObject
 		RCO_All,
 		RCO_AllRootObjects,
 		RCO_Selected,
-		RCO_InMemory
+		RCO_InMemory,
+		RCO_SelectedGatherReferences,
 	};
 	
 	void RecompileObjects(const FToolMenuContext& InMenuContext, ERecompileCO RecompileType)
@@ -193,7 +194,7 @@ namespace MenuExtension_CustomizableObject
 		TArray<TSharedRef<FCompilationRequest>> CompilationRequests;
 
 
-		if (RecompileType == ERecompileCO::RCO_Selected)
+		if (RecompileType == ERecompileCO::RCO_Selected || RecompileType == ERecompileCO::RCO_SelectedGatherReferences)
 		{
 			const UContentBrowserAssetContextMenuContext* Context = UContentBrowserAssetContextMenuContext::FindContextWithAssets(InMenuContext);
 			check(Context);
@@ -207,6 +208,12 @@ namespace MenuExtension_CustomizableObject
 				{
 					TSharedRef<FCompilationRequest> NewRequest = MakeShared<FCompilationRequest>(*Object);
 					NewRequest->bSilentCompilation = true;
+					
+					if (RecompileType == ERecompileCO::RCO_SelectedGatherReferences)
+					{
+						NewRequest->Options.bGatherReferences = true;
+					}
+
 					CompilationRequests.Add(NewRequest);
 				}
 			}
@@ -277,6 +284,18 @@ namespace MenuExtension_CustomizableObject
 			UIAction.ExecuteAction = FToolMenuExecuteAction::CreateStatic(&RecompileObjects,ERecompileCO::RCO_Selected);
 			Section.AddMenuEntry("CustomizableObject_ExecuteCompileSelected", Label, ToolTip, Icon, UIAction);
 		}
+	
+		// Recompile Selected Gather References.
+		{
+			const TAttribute<FText> Label = NSLOCTEXT("UAssetDefinition_CustomizableObject", "CustomizableObject_RecompileSelectedGatherReferences", "Selected and Gather references");
+			const TAttribute<FText> ToolTip = NSLOCTEXT("UAssetDefinition_CustomizableObject", "CustomizableObject_RecompileSelectedGatherReferencesTooltip", "Recompile selected Customizable Object and gather references.");
+			const FSlateIcon Icon = FSlateIcon();
+			
+			FToolUIAction UIAction;
+			UIAction.ExecuteAction = FToolMenuExecuteAction::CreateStatic(&RecompileObjects,ERecompileCO::RCO_SelectedGatherReferences);
+			Section.AddMenuEntry("CustomizableObject_ExecuteCompileSelectedGatherReferences", Label, ToolTip, Icon, UIAction);
+		}
+
 
 		// Recompile All
 		{
