[Mutable] Fix image layer combine op missing blend image.

diff --git a/Engine/Plugins/Mutable/Source/MutableRuntime/Private/MuR/OpImageLayer.cpp b/Engine/Plugins/Mutable/Source/MutableRuntime/Private/MuR/OpImageLayer.cpp
--- a/Engine/Plugins/Mutable/Source/MutableRuntime/Private/MuR/OpImageLayer.cpp
+++ b/Engine/Plugins/Mutable/Source/MutableRuntime/Private/MuR/OpImageLayer.cpp
@@ -2253,8 +2253,10 @@ namespace OpImageLayerInternal
 		BatchArgs.LODEnd             = NumLODs;
 		BatchArgs.Result             = ResultImage;
 		BatchArgs.Base               = BaseImage;
+		BatchArgs.Blend              = BlendImage;
 		BatchArgs.ResultBytesPerElem = NC;
 		BatchArgs.BaseBytesPerElem   = NC;
+		BatchArgs.BlendBytesPerElem  = NC;
 
 		const int32 NumBatches = OpImageLayerInternal::GetOpLayerNumBatchesLODRange(BatchArgs);
 
