[Mutable] Optimize OpRasterMesh parallel for batching.

#rnx

diff --git a/Engine/Plugins/Mutable/Source/MutableRuntime/Private/MuR/OpImageRasterMesh.h b/Engine/Plugins/Mutable/Source/MutableRuntime/Private/MuR/OpImageRasterMesh.h
--- a/Engine/Plugins/Mutable/Source/MutableRuntime/Private/MuR/OpImageRasterMesh.h
+++ b/Engine/Plugins/Mutable/Source/MutableRuntime/Private/MuR/OpImageRasterMesh.h
@@ -2,7 +2,7 @@
 
 #pragma once
 
-#include "Async/ParallelFor.h"
+#include "MuR/ParallelExecutionUtils.h"
 #include "MuR/ImagePrivate.h"
 #include "MuR/Platform.h"
 #include "MuR/MeshPrivate.h"
@@ -75,45 +75,63 @@ namespace mu
 		}
 
 		// Get the indices
-		int32 faceCount = pMesh->GetFaceCount();
-		TArray<int32> indices;
-		indices.SetNumZeroed(faceCount * 3);
+		const int32 NumFaces = pMesh->GetFaceCount();
+		
+		TArrayView<const int32> IndicesView;
+		TArray<int32> FormatedIndices;
 
-		UntypedMeshBufferIteratorConst indIt( pMesh->GetIndexBuffers(), EMeshBufferSemantic::VertexIndex, 0 );
-		for ( int32 i=0; i<faceCount*3; ++i )
+		UntypedMeshBufferIteratorConst indIt(pMesh->GetIndexBuffers(), EMeshBufferSemantic::VertexIndex, 0);
+		if (indIt.GetFormat() == EMeshBufferFormat::UInt32)
 		{
-            uint32_t index=0;
-			ConvertData( 0, &index, EMeshBufferFormat::UInt32, indIt.ptr(), indIt.GetFormat() );
+			IndicesView = TArrayView<const int32>(reinterpret_cast<const int32*>(indIt.ptr()), NumFaces * 3);
+		}
+		else
+		{
+			FormatedIndices.SetNumZeroed(NumFaces * 3);
+
+			for (int32 i = 0; i < NumFaces * 3; ++i)
+			{
+				uint32_t index = 0;
+				ConvertData(0, &index, EMeshBufferFormat::UInt32, indIt.ptr(), indIt.GetFormat());
+
+				FormatedIndices[i] = index;
+				++indIt;
+			}
 
-			indices[i] = index;
-			++indIt;
+			IndicesView = TArrayView<const int32>(FormatedIndices);
 		}
 
         UntypedMeshBufferIteratorConst bloIt( pMesh->GetVertexBuffers(), EMeshBufferSemantic::LayoutBlock, LayoutIndex );
 
+		constexpr int32 NumBatchElems = 1 << 5; // Small batches perform better.
+		const int32 NumBatches = FMath::DivideAndRoundUp<int32>(NumFaces, NumBatchElems);
+
         if (BlockId== FLayoutBlock::InvalidBlockId || bloIt.GetElementSize()==0 )
 		{
 			// Raster all the faces
             WhitePixelProcessor pixelProc;
 			const TArrayView<uint8> ImageData = pImage->DataStorage.GetLOD(0);
 
-			//for ( int f=0; f<faceCount; ++f )
-			const auto& ProcessFace = [
-				vertices, indices, ImageData, sizeX, sizeY, pixelProc
-			] (int32 f)
-			{
-				constexpr int32 NumInterpolators = 1;
-				Triangle<NumInterpolators>(ImageData.GetData(), ImageData.Num(),
-					sizeX, sizeY,
-					1,
-					vertices[indices[f * 3 + 0]],
-					vertices[indices[f * 3 + 1]],
-					vertices[indices[f * 3 + 2]],
-					pixelProc,
-					false);
-			};
-
-			ParallelFor(faceCount, ProcessFace);
+			ParallelExecutionUtils::InvokeBatchParallelFor(NumBatches, [
+				&vertices, IndicesView, ImageData, sizeX, sizeY, pixelProc, NumFaces, NumBatchElems
+			] (int32 BatchId)
+				{
+					const int32 BatchBeginFaceIndex = BatchId * NumBatchElems;
+					const int32 BatchEndFaceIndex = FMath::Min(BatchBeginFaceIndex + NumBatchElems, NumFaces);
+
+					for (int32 FaceIndex = BatchBeginFaceIndex; FaceIndex < BatchEndFaceIndex; ++FaceIndex)
+					{
+						constexpr int32 NumInterpolators = 1;
+						Triangle<NumInterpolators>(ImageData.GetData(), ImageData.Num(),
+							sizeX, sizeY,
+							1,
+							vertices[IndicesView[FaceIndex * 3 + 0]],
+							vertices[IndicesView[FaceIndex * 3 + 1]],
+							vertices[IndicesView[FaceIndex * 3 + 2]],
+							pixelProc,
+							false);
+					}
+				});
 		}
 		else
 		{
@@ -155,26 +173,30 @@ namespace mu
 
 			const TArrayView<uint8> ImageData = pImage->DataStorage.GetLOD(0); 
 
-			const auto& ProcessFace = [
-				vertices, indices, VertexBlockIds, BlockId, ImageData, sizeX, sizeY, pixelProc
-			] (int32 f)
-			{
-				// TODO: Select faces outside for loop?
-				if (VertexBlockIds[indices[f * 3 + 0]] == BlockId)
+			ParallelExecutionUtils::InvokeBatchParallelFor(NumBatches, [
+				&vertices, IndicesView, VertexBlockIds, BlockId, ImageData, sizeX, sizeY, pixelProc, NumFaces, NumBatchElems
+			] (int32 BatchId)
 				{
-					constexpr int32 NumInterpolators = 1;
-					Triangle<NumInterpolators>(ImageData.GetData(), ImageData.Num(),
-						sizeX, sizeY,
-						1,
-						vertices[indices[f * 3 + 0]],
-						vertices[indices[f * 3 + 1]],
-						vertices[indices[f * 3 + 2]],
-						pixelProc,
-						false);
-				}
-			};
-
-			ParallelFor(faceCount, ProcessFace);
+					const int32 BatchBeginFaceIndex = BatchId * NumBatchElems;
+					const int32 BatchEndFaceIndex = FMath::Min(BatchBeginFaceIndex + NumBatchElems, NumFaces);
+
+					for (int32 FaceIndex = BatchBeginFaceIndex; FaceIndex < BatchEndFaceIndex; ++FaceIndex)
+					{
+						// TODO: Select faces outside for loop?
+						if (VertexBlockIds[IndicesView[FaceIndex * 3 + 0]] == BlockId)
+						{
+							constexpr int32 NumInterpolators = 1;
+							Triangle<NumInterpolators>(ImageData.GetData(), ImageData.Num(),
+								sizeX, sizeY,
+								1,
+								vertices[IndicesView[FaceIndex * 3 + 0]],
+								vertices[IndicesView[FaceIndex * 3 + 1]],
+								vertices[IndicesView[FaceIndex * 3 + 2]],
+								pixelProc,
+								false);
+						}
+					}
+				});
 		}
 
 	}
