[mutable] Restore layout block mask support, necessary for automatic UVs basde on UV islands.

#rnx

diff --git a/Engine/Plugins/Mutable/Source/CustomizableObject/Internal/MuCO/CustomizableObjectPrivate.h b/Engine/Plugins/Mutable/Source/CustomizableObject/Internal/MuCO/CustomizableObjectPrivate.h
--- a/Engine/Plugins/Mutable/Source/CustomizableObject/Internal/MuCO/CustomizableObjectPrivate.h
+++ b/Engine/Plugins/Mutable/Source/CustomizableObject/Internal/MuCO/CustomizableObjectPrivate.h
@@ -1626,6 +1626,8 @@ enum class ECustomizableObjectVersions : int32
 
 	SwitchOpCompactRepresentation,
 
+	AddedBlockMasksToRuntimeLayouts,
+
 	// -----<new versions can be added above this line>--------
 	LastCustomizableObjectVersion
 };
diff --git a/Engine/Plugins/Mutable/Source/MutableRuntime/Internal/MuR/OpMeshPrepareLayout.h b/Engine/Plugins/Mutable/Source/MutableRuntime/Internal/MuR/OpMeshPrepareLayout.h
--- a/Engine/Plugins/Mutable/Source/MutableRuntime/Internal/MuR/OpMeshPrepareLayout.h
+++ b/Engine/Plugins/Mutable/Source/MutableRuntime/Internal/MuR/OpMeshPrepareLayout.h
@@ -2,14 +2,19 @@
 
 #pragma once
 
-#include "MuR/Ptr.h"
 #include "HAL/Platform.h"
+#include "Containers/Array.h"
+#include "Math/Vector.h"
 
 namespace mu
 {
 	class FMesh;
 	class FLayout;
 
+	/** Modify the mesh to use the given layout. 
+	* It will be added to the mesh and additional mesh buffers created for the layout block ids. 
+	* It may also modify the UVs to clamp them to the blocks if necessary.
+	*/
 	MUTABLERUNTIME_API void MeshPrepareLayout(
 		FMesh& Mesh,
 		const FLayout& InLayout,
@@ -20,4 +25,36 @@ namespace mu
 		bool bUseAbsoluteBlockIds
 	);
 
+
+	/** Store the result of a layout vertex classify operation. */
+	struct MUTABLERUNTIME_API FMutableLayoutClassifyResult
+	{
+		/** The tex coords of every mesh vertex. They could be modified during the classify operation due to forced block clamping. */
+		TArray<FVector2f> TexCoords;
+
+		/** A block index for every mesh vertex. 
+		* These are indices in the array of blocks of the given FLayout object to use to classify.
+		* It could be NullBlockId if it didn't belong to a block. 
+		*/
+		TArray<uint16> LayoutData;
+
+		/** Id of unassigned blocks. */
+		static constexpr uint16 NullBlockId = TNumericLimits<uint16>::Max();
+	};
+
+
+	/** Assign a layout block to every vertex in the mesh if possible. 
+	* If a vertex cannot be assigned, it will remain as NullBlockId.
+	* Vertex UVs for the given layout will be copied and may be modified in the output struct if they needed to be clamped to fit in a block,
+	* but the mesh won't be modified by this method.
+	*/
+	MUTABLERUNTIME_API bool ClassifyMeshVerticesInLayout(
+		const FMesh& Mesh,
+		const FLayout& InLayout,
+		int32 LayoutChannel,
+		bool bNormalizeUVs,
+		bool bClampUVIslands,
+		FMutableLayoutClassifyResult& OutResult
+	);
+
 }
diff --git a/Engine/Plugins/Mutable/Source/MutableRuntime/Private/MuR/Layout.cpp b/Engine/Plugins/Mutable/Source/MutableRuntime/Private/MuR/Layout.cpp
--- a/Engine/Plugins/Mutable/Source/MutableRuntime/Private/MuR/Layout.cpp
+++ b/Engine/Plugins/Mutable/Source/MutableRuntime/Private/MuR/Layout.cpp
@@ -6,6 +6,7 @@
 #include "HAL/LowLevelMemTracker.h"
 #include "Math/IntPoint.h"
 #include "MuR/MutableMath.h"
+#include "MuR/Image.h"
 #include "MuR/SerialisationPrivate.h"
 
 
@@ -39,6 +40,7 @@ namespace mu {
 		Result->Blocks = Blocks;
 		Result->Strategy = Strategy;
 		Result->ReductionMethod = ReductionMethod;
+		Result->Masks = Masks;
 		return Result;
 	}
 
@@ -123,18 +125,17 @@ namespace mu {
 	void FLayout::Serialise(FOutputArchive& arch) const
 	{
 		arch << Size;
-		arch << Blocks;
-
 		arch << MaxSize;
 		arch << uint32(Strategy);
 		arch << uint32(ReductionMethod);
+		arch << Blocks;
+		arch << Masks;
 	}
 
 	
 	void FLayout::Unserialise(FInputArchive& arch)
 	{
 		arch >> Size;
-		arch >> Blocks;
 		arch >> MaxSize;
 
 		uint32 Temp;
@@ -143,6 +144,9 @@ namespace mu {
 
 		arch >> Temp;
 		ReductionMethod = EReductionMethod(Temp);
+
+		arch >> Blocks;
+		arch >> Masks;
 	}
 
 
diff --git a/Engine/Plugins/Mutable/Source/MutableRuntime/Private/MuR/OpMeshPrepareLayout.cpp b/Engine/Plugins/Mutable/Source/MutableRuntime/Private/MuR/OpMeshPrepareLayout.cpp
--- a/Engine/Plugins/Mutable/Source/MutableRuntime/Private/MuR/OpMeshPrepareLayout.cpp
+++ b/Engine/Plugins/Mutable/Source/MutableRuntime/Private/MuR/OpMeshPrepareLayout.cpp
@@ -43,32 +43,28 @@ namespace mu
 		}
 
 	};
-
-
-	void MeshPrepareLayout(
-		FMesh& Mesh,
+	
+	bool ClassifyMeshVerticesInLayout(
+		const FMesh& Mesh,
 		const FLayout& InLayout,
 		int32 LayoutChannel,
 		bool bNormalizeUVs,
 		bool bClampUVIslands,
-		bool bEnsureAllVerticesHaveLayoutBlock,
-		bool bUseAbsoluteBlockIds
+		FMutableLayoutClassifyResult& OutResult
 	)
 	{
-		MUTABLE_CPUPROFILER_SCOPE(MeshPrepareLayout);
+		MUTABLE_CPUPROFILER_SCOPE(ClassifyMeshVerticesInLayout);
 
 		if (Mesh.GetVertexCount() == 0)
 		{
-			return;
+			return false;
 		}
 
-		TSharedPtr<FLayout> Layout = InLayout.Clone();
+		const FLayout* Layout = &InLayout;
 
 		// The layout must have block ids.
 		check(Layout->Blocks.IsEmpty() || Layout->Blocks[0].Id != FLayoutBlock::InvalidBlockId);
 
-		Mesh.AddLayout(Layout);
-
 		const int32 NumVertices = Mesh.GetVertexCount();
 		const int32 NumBlocks = Layout->GetBlockCount();
 
@@ -85,46 +81,31 @@ namespace mu
 		TArray2D<int32> GridBlockBlockId;
 		GridBlockBlockId.Init(MAX_uint16, WorkingGrid.X, WorkingGrid.Y);
 
-		// 
-		TArray<box<FVector2f>> BlockRects;
-		BlockRects.SetNumUninitialized(NumBlocks);
-
 		// Create an array of block index per cell
 		TArray<int32> OverlappingBlocks;
+		for (int32 BlockIndex = 0; BlockIndex < NumBlocks; ++BlockIndex)
 		{
-			MUTABLE_CPUPROFILER_SCOPE(CreateGrid);
-
-			for (int32 BlockIndex = 0; BlockIndex < NumBlocks; ++BlockIndex)
-			{
-				// TODO
-				bool bBlockHasMask = false; // GeneratedLayout.Source->Blocks[BlockIndex].Mask != nullptr;
-
-				// Fill the block rect
-				FIntVector2 Min = Layout->Blocks[BlockIndex].Min;
-				FIntVector2 Size = Layout->Blocks[BlockIndex].Size;
+			bool bBlockHasMask = Layout->Blocks[BlockIndex].MaskIndex != FLayoutBlock::InvalidMaskIndex;
 
-				box<FVector2f>& BlockRect = BlockRects[BlockIndex];
-				BlockRect.min[0] = float(Min.X) / float(LayoutGrid.X);
-				BlockRect.min[1] = float(Min.Y) / float(LayoutGrid.Y);
-				BlockRect.size[0] = float(Size.X) / float(LayoutGrid.X);
-				BlockRect.size[1] = float(Size.Y) / float(LayoutGrid.Y);
+			// Fill the block rect
+			const FIntVector2& Min = Layout->Blocks[BlockIndex].Min;
+			const FIntVector2& Size = Layout->Blocks[BlockIndex].Size;
 
-				// Fill the block index per cell array
-				// Ignore the block in this stage if it has a mask, because blocks with masks will very likely overlap other blocks
-				if (!bBlockHasMask)
+			// Fill the block index per cell array
+			// Ignore the block in this stage if it has a mask, because blocks with masks will very likely overlap other blocks
+			if (!bBlockHasMask)
+			{
+				for (uint16 Y = Min.Y; Y < Min.Y + Size.Y; ++Y)
 				{
-					for (uint16 Y = Min.Y; Y < Min.Y + Size.Y; ++Y)
+					for (uint16 X = Min.X; X < Min.X + Size.X; ++X)
 					{
-						for (uint16 X = Min.X; X < Min.X + Size.X; ++X)
+						if (GridBlockBlockId.Get(X, Y) == MAX_uint16)
 						{
-							if (GridBlockBlockId.Get(X, Y) == MAX_uint16)
-							{
-								GridBlockBlockId.Set(X, Y, BlockIndex);
-							}
-							else
-							{
-								OverlappingBlocks.AddUnique(BlockIndex);
-							}
+							GridBlockBlockId.Set(X, Y, BlockIndex);
+						}
+						else
+						{
+							OverlappingBlocks.AddUnique(BlockIndex);
 						}
 					}
 				}
@@ -147,30 +128,24 @@ namespace mu
 		{
 			// This is actually fine when using shared materials across LODs
 			UE_LOG(LogMutableCore, Log, TEXT("Trying to generate layout for missing UV channel %d. No layout blocks generated."), LayoutChannel);
-			return;
+			return false;
 		}
 
-		check(TexCoordsBufferIndex >= 0);
-		check(TexCoordsChannelIndex >= 0);
-
 		const FMeshBufferChannel& TexCoordsChannel = Mesh.VertexBuffers.Buffers[TexCoordsBufferIndex].Channels[TexCoordsChannelIndex];
 		check(TexCoordsChannel.Semantic == EMeshBufferSemantic::TexCoords);
 
-		uint8* TexCoordData = Mesh.GetVertexBuffers().GetBufferData(TexCoordsBufferIndex);
-		int32 TexElemSize = Mesh.GetVertexBuffers().GetElementSize(TexCoordsBufferIndex);
-		int32 channelOffset = TexCoordsChannel.Offset;
-		TexCoordData += channelOffset;
+		const uint8* TexCoordData = Mesh.GetVertexBuffers().GetBufferData(TexCoordsBufferIndex);
+		int32 TexCoordElemSize = Mesh.GetVertexBuffers().GetElementSize(TexCoordsBufferIndex);
+		int32 TexCoordChannelOffset = TexCoordsChannel.Offset;
+		TexCoordData += TexCoordChannelOffset;
 
 		// Get a copy of the UVs as FVector2f to work with them. 
-		TArray<FVector2f> TexCoords;
+		TArray<FVector2f>& TexCoords = OutResult.TexCoords;
 		{
 			MUTABLE_CPUPROFILER_SCOPE(CopyUVs);
 
 			TexCoords.SetNumUninitialized(NumVertices);
 
-			bool bNonNormalizedUVs = false;
-			const bool bIsOverlayLayout = Layout->GetLayoutPackingStrategy() == mu::EPackStrategy::Overlay;
-
 			const uint8* pVertices = TexCoordData;
 			for (int32 VertexIndex = 0; VertexIndex < NumVertices; ++VertexIndex)
 			{
@@ -185,20 +160,31 @@ namespace mu
 					UV = FVector2f(float(pUV[0]), float(pUV[1]));
 				}
 
+				pVertices += TexCoordElemSize;
+			}
+		}
+
+		const bool bIsOverlayLayout = Layout->GetLayoutPackingStrategy() == mu::EPackStrategy::Overlay;
+		if (bNormalizeUVs && !bIsOverlayLayout)
+		{
+			MUTABLE_CPUPROFILER_SCOPE(NormalizeUVs);
+			bool bNonNormalizedUVs = false;
+
+			for (int32 VertexIndex = 0; VertexIndex < NumVertices; ++VertexIndex)
+			{
+				FVector2f& UV = TexCoords[VertexIndex];
+
 				// Check that UVs are normalized. If not, clamp the values and throw a warning.
-				if (bNormalizeUVs && !bIsOverlayLayout
-					&& (UV[0] < 0.f || UV[0] > 1.f || UV[1] < 0.f || UV[1] > 1.f))
+				if (UV[0] < 0.f || UV[0] > 1.f || UV[1] < 0.f || UV[1] > 1.f)
 				{
 					UV[0] = FMath::Clamp(UV[0], 0.f, 1.f);
 					UV[1] = FMath::Clamp(UV[1], 0.f, 1.f);
 					bNonNormalizedUVs = true;
 				}
-
-				pVertices += TexElemSize;
 			}
 
 			// Mutable does not support non-normalized UVs
-			if (bNonNormalizedUVs && !bIsOverlayLayout)
+			if (bNonNormalizedUVs)
 			{
 				UE_LOG(LogMutableCore, Warning, TEXT("Source mesh has non-normalized UVs index %d"), LayoutChannel);
 			}
@@ -209,65 +195,86 @@ namespace mu
 		const uint32 MaxGridY = bNormalizeUVs ? MAX_uint32 : WorkingGrid.Y - 1;
 
 		// Allocate the per-vertex layout block data
-		TArray<uint16> LayoutData;
-		constexpr uint16 NullBlockId = MAX_uint16 - 1;
-		LayoutData.SetNumUninitialized(NumVertices);
+		OutResult.LayoutData.SetNumUninitialized(NumVertices);
 
 		// Assign a block to each vertex
 		{
-			MUTABLE_CPUPROFILER_SCOPE(Assign);
 
-			for (int32 VertexIndex = 0; VertexIndex < NumVertices; ++VertexIndex)
+			if (Layout->Masks.IsEmpty())
 			{
-				uint16 BlockIndex = NullBlockId;
-
-				FVector2f UV = TexCoords[VertexIndex];
+				MUTABLE_CPUPROFILER_SCOPE(Assign);
+				for (int32 VertexIndex = 0; VertexIndex < NumVertices; ++VertexIndex)
+				{
+					FVector2f& UV = TexCoords[VertexIndex];
 
-				int32 VertexWorkingGridX = FMath::Clamp(LayoutGrid.X * UV[0], 0, LayoutGrid.X - 1);
-				int32 VertexWorkingGridY = FMath::Clamp(LayoutGrid.Y * UV[1], 0, LayoutGrid.Y - 1);
+					int32 VertexWorkingGridX = FMath::Clamp(LayoutGrid.X * UV[0], 0, LayoutGrid.X - 1);
+					int32 VertexWorkingGridY = FMath::Clamp(LayoutGrid.Y * UV[1], 0, LayoutGrid.Y - 1);
+					uint32 ClampedX = FMath::Min<uint32>(MaxGridX, VertexWorkingGridX);
+					uint32 ClampedY = FMath::Min<uint32>(MaxGridY, VertexWorkingGridY);
 
-				// First: Assign the vertices to masked blocks in order
-				for (int32 CandidateBlockIndex = 0; CandidateBlockIndex < NumBlocks; ++CandidateBlockIndex)
+					OutResult.LayoutData[VertexIndex] = GridBlockBlockId.Get(ClampedX, ClampedY);
+				}
+			}
+			else
+			{
+				MUTABLE_CPUPROFILER_SCOPE(AssignWithMask);
+				// TODO Optimize (6x slower)
+				for (int32 VertexIndex = 0; VertexIndex < NumVertices; ++VertexIndex)
 				{
-					// TODO
-					const mu::FImage* Mask = nullptr; // GeneratedLayout.Source->Blocks[CandidateBlockIndex].Mask.Get();
-					if (Mask)
-					{
-						// First discard with block limits.
-						FIntVector2 Min = Layout->Blocks[CandidateBlockIndex].Min;
-						FIntVector2 Size = Layout->Blocks[CandidateBlockIndex].Size;
+					uint16 BlockIndex = FMutableLayoutClassifyResult::NullBlockId;
+
+					FVector2f& UV = TexCoords[VertexIndex];
+
+					int32 VertexWorkingGridX = FMath::Clamp(LayoutGrid.X * UV[0], 0, LayoutGrid.X - 1);
+					int32 VertexWorkingGridY = FMath::Clamp(LayoutGrid.Y * UV[1], 0, LayoutGrid.Y - 1);
 
-						bool bInBlock =
-							(VertexWorkingGridX >= Min.X && VertexWorkingGridX < Min.X + Size.X)
-							&&
-							(VertexWorkingGridY >= Min.Y && VertexWorkingGridY < Min.Y + Size.Y);
+					// First: Assign the vertices to masked blocks in order
+					for (int32 CandidateBlockIndex = 0; CandidateBlockIndex < NumBlocks; ++CandidateBlockIndex)
+					{
+						const mu::FImage* Mask = nullptr;
+						int32 MaskIndex = Layout->Blocks[CandidateBlockIndex].MaskIndex;
+						if (Layout->Masks.IsValidIndex(MaskIndex))
+						{
+							Mask = Layout->Masks[MaskIndex].Get();
+						}
 
-						if (bInBlock)
+						if (Mask)
 						{
-							// TODO: This always clamps the UVs
-							FVector2f SampleUV;
-							SampleUV.X = FMath::Fmod(UV.X, 1.0);
-							SampleUV.Y = FMath::Fmod(UV.Y, 1.0);
+							// First discard with block limits.
+							FIntVector2 Min = Layout->Blocks[CandidateBlockIndex].Min;
+							FIntVector2 Size = Layout->Blocks[CandidateBlockIndex].Size;
+
+							bool bInBlock =
+								(VertexWorkingGridX >= Min.X && VertexWorkingGridX < Min.X + Size.X)
+								&&
+								(VertexWorkingGridY >= Min.Y && VertexWorkingGridY < Min.Y + Size.Y);
 
-							FVector4f MaskValue = Mask->Sample(SampleUV);
-							if (MaskValue.X > 0.5f)
+							if (bInBlock)
 							{
-								BlockIndex = CandidateBlockIndex;
-								break;
+								FVector2f SampleUV;
+								SampleUV.X = FMath::Fmod(UV.X, 1.0);
+								SampleUV.Y = FMath::Fmod(UV.Y, 1.0);
+
+								FVector4f MaskValue = Mask->Sample(SampleUV);
+								if (MaskValue.X > 0.5f)
+								{
+									BlockIndex = CandidateBlockIndex;
+									break;
+								}
 							}
 						}
 					}
-				}
 
-				// Second: Assign to non-masked blocks if not assigned yet
-				if (BlockIndex == NullBlockId)
-				{
-					uint32 ClampedX = FMath::Min<uint32>(MaxGridX, FMath::Max<uint32>(0, VertexWorkingGridX));
-					uint32 ClampedY = FMath::Min<uint32>(MaxGridY, FMath::Max<uint32>(0, VertexWorkingGridY));
-					BlockIndex = GridBlockBlockId.Get(ClampedX, ClampedY);
-				}
+					// Second: Assign to non-masked blocks if not assigned yet
+					if (BlockIndex == FMutableLayoutClassifyResult::NullBlockId)
+					{
+						uint32 ClampedX = FMath::Min<uint32>(MaxGridX, FMath::Max<uint32>(0, VertexWorkingGridX));
+						uint32 ClampedY = FMath::Min<uint32>(MaxGridY, FMath::Max<uint32>(0, VertexWorkingGridY));
+						BlockIndex = GridBlockBlockId.Get(ClampedX, ClampedY);
+					}
 
-				LayoutData[VertexIndex] = BlockIndex;
+					OutResult.LayoutData[VertexIndex] = BlockIndex;
+				}
 			}
 		}
 
@@ -279,24 +286,43 @@ namespace mu
 		{
 			MUTABLE_CPUPROFILER_SCOPE(DetectConflictiveTriangles);
 
-			UntypedMeshBufferIteratorConst ItIndices(Mesh.GetIndexBuffers(), EMeshBufferSemantic::VertexIndex);
+			UntypedMeshBufferIteratorConst IndicesBegin(Mesh.GetIndexBuffers(), EMeshBufferSemantic::VertexIndex);
+
+			const EMeshBufferFormat IndicesFormat = IndicesBegin.GetFormat();
+			check(IndicesBegin.GetFormat() == EMeshBufferFormat::UInt32 || IndicesBegin.GetFormat() == EMeshBufferFormat::UInt16);
 
-			for (int32 TriangleIndex = 0; TriangleIndex < NumTriangles; ++TriangleIndex)
+			if (IndicesBegin.GetFormat() == EMeshBufferFormat::UInt32)
 			{
-				uint32 Index0 = ItIndices.GetAsUINT32();
-				++ItIndices;
-				uint32 Index1 = ItIndices.GetAsUINT32();
-				++ItIndices;
-				uint32 Index2 = ItIndices.GetAsUINT32();
-				++ItIndices;
-
-				const uint16 BlockIndexV0 = LayoutData[Index0];
-				const uint16 BlockIndexV1 = LayoutData[Index1];
-				const uint16 BlockIndexV2 = LayoutData[Index2];
-
-				if (BlockIndexV0 != BlockIndexV1 || BlockIndexV0 != BlockIndexV2)
+				const uint32* Indices = reinterpret_cast<const uint32*>(IndicesBegin.ptr());
+				for (int32 TriangleIndex = 0; TriangleIndex < NumTriangles; ++TriangleIndex)
+				{
+					const uint16 BlockIndexV0 = OutResult.LayoutData[Indices[0]];
+					const uint16 BlockIndexV1 = OutResult.LayoutData[Indices[1]];
+					const uint16 BlockIndexV2 = OutResult.LayoutData[Indices[2]];
+
+					if (BlockIndexV0 != BlockIndexV1 || BlockIndexV0 != BlockIndexV2)
+					{
+						ConflictiveTriangles.Add(TriangleIndex);
+					}
+
+					Indices += 3;
+				}
+			}
+			else if (IndicesBegin.GetFormat() == EMeshBufferFormat::UInt16)
+			{
+				const uint16* Indices = reinterpret_cast<const uint16*>(IndicesBegin.ptr());
+				for (int32 TriangleIndex = 0; TriangleIndex < NumTriangles; ++TriangleIndex)
 				{
-					ConflictiveTriangles.Add(TriangleIndex);
+					const uint16 BlockIndexV0 = OutResult.LayoutData[Indices[0]];
+					const uint16 BlockIndexV1 = OutResult.LayoutData[Indices[1]];
+					const uint16 BlockIndexV2 = OutResult.LayoutData[Indices[2]];
+
+					if (BlockIndexV0 != BlockIndexV1 || BlockIndexV0 != BlockIndexV2)
+					{
+						ConflictiveTriangles.Add(TriangleIndex);
+					}
+
+					Indices += 3;
 				}
 			}
 		}
@@ -314,13 +340,10 @@ namespace mu
 			TMultiMap<int32, uint32> VertexToFaceMap;
 
 			// Find Unique Vertices
-			if (bClampUVIslands)
-			{
-				VertexToFaceMap.Reserve(NumTriangles * 3);
-				Triangles.SetNumUninitialized(NumTriangles);
+			VertexToFaceMap.Reserve(NumTriangles * 3);
+			Triangles.SetNumUninitialized(NumTriangles);
 
-				MeshCreateCollapsedVertexMap(&Mesh, CollapsedVertices);
-			}
+			MeshCreateCollapsedVertexMap(&Mesh, CollapsedVertices);
 
 			UntypedMeshBufferIteratorConst ItIndices(Mesh.GetIndexBuffers(), EMeshBufferSemantic::VertexIndex);
 			{
@@ -335,9 +358,9 @@ namespace mu
 					uint32 Index2 = ItIndices.GetAsUINT32();
 					++ItIndices;
 
-					const uint16 BlockIndexV0 = LayoutData[Index0];
-					const uint16 BlockIndexV1 = LayoutData[Index1];
-					const uint16 BlockIndexV2 = LayoutData[Index2];
+					const uint16 BlockIndexV0 = OutResult.LayoutData[Index0];
+					const uint16 BlockIndexV1 = OutResult.LayoutData[Index1];
+					const uint16 BlockIndexV2 = OutResult.LayoutData[Index2];
 
 					FTriangleInfo& Triangle = Triangles[TriangleIndex];
 
@@ -433,7 +456,7 @@ namespace mu
 							FVector2f& UV = TexCoords[UVIndex];
 							UV[0] = FMath::Clamp(UV[0], MinX, MaxX);
 							UV[1] = FMath::Clamp(UV[1], MinY, MaxY);
-							LayoutData[UVIndex] = BlockIndex;
+							OutResult.LayoutData[UVIndex] = BlockIndex;
 						}
 
 						OtherTriangle.bUVsFixed = true;
@@ -466,6 +489,47 @@ namespace mu
 		//	}
 		//}
 
+		return true;
+	}
+
+
+	void MeshPrepareLayout(
+		FMesh& Mesh,
+		const FLayout& InLayout,
+		int32 LayoutChannel,
+		bool bNormalizeUVs,
+		bool bClampUVIslands,
+		bool bEnsureAllVerticesHaveLayoutBlock,
+		bool bUseAbsoluteBlockIds
+	)
+	{
+		MUTABLE_CPUPROFILER_SCOPE(MeshPrepareLayout);
+
+		if (Mesh.GetVertexCount() == 0)
+		{
+			return;
+		}
+
+		TSharedPtr<FLayout> Layout = InLayout.Clone();
+
+		// The layout must have block ids.
+		check(Layout->Blocks.IsEmpty() || Layout->Blocks[0].Id != FLayoutBlock::InvalidBlockId);
+
+		Mesh.AddLayout(Layout);
+
+		const int32 NumVertices = Mesh.GetVertexCount();
+		const int32 NumBlocks = Layout->GetBlockCount();
+
+		FMutableLayoutClassifyResult ClassifyResult;
+		if (!ClassifyMeshVerticesInLayout(Mesh, *Layout, LayoutChannel,
+			bNormalizeUVs,
+			bClampUVIslands,
+			ClassifyResult))
+		{
+			// Classify mesh vertices failed. Missing texture coords?
+			return;
+		}
+
 		{
 			MUTABLE_CPUPROFILER_SCOPE(CreateBuffers);
 
@@ -498,11 +562,31 @@ namespace mu
 			{
 				MUTABLE_CPUPROFILER_SCOPE(LayoutAndHomogenize);
 
+				// 
+				const  FIntPoint& LayoutGrid = Layout->GetGridSize();
+				TArray<box<FVector2f>> BlockRects;
+				BlockRects.SetNumUninitialized(NumBlocks);
+				for (int32 BlockIndex = 0; BlockIndex < NumBlocks; ++BlockIndex)
+				{
+					bool bBlockHasMask = Layout->Blocks[BlockIndex].MaskIndex != FLayoutBlock::InvalidMaskIndex;
+
+					// Fill the block rect
+					const FIntVector2& Min = Layout->Blocks[BlockIndex].Min;
+					const FIntVector2& Size = Layout->Blocks[BlockIndex].Size;
+
+					box<FVector2f>& BlockRect = BlockRects[BlockIndex];
+					BlockRect.min[0] = float(Min.X) / float(LayoutGrid.X);
+					BlockRect.min[1] = float(Min.Y) / float(LayoutGrid.Y);
+					BlockRect.size[0] = float(Size.X) / float(LayoutGrid.X);
+					BlockRect.size[1] = float(Size.Y) / float(LayoutGrid.Y);
+				}
+
+
 				for (int32 VertexIndex = 0; VertexIndex < NumVertices; ++VertexIndex)
 				{
-					FVector2f& UV = TexCoords[VertexIndex];
+					FVector2f& UV = ClassifyResult.TexCoords[VertexIndex];
 
-					uint16 LayoutBlockIndex = LayoutData[VertexIndex];
+					uint16 LayoutBlockIndex = ClassifyResult.LayoutData[VertexIndex];
 					if (Layout->Blocks.IsValidIndex(LayoutBlockIndex))
 					{
 						uint64 LayoutBlockId = Layout->Blocks[LayoutBlockIndex].Id;
@@ -541,10 +625,31 @@ namespace mu
 			{
 				MUTABLE_CPUPROFILER_SCOPE(CopyUVs);
 
+				// Get the information about the texture coordinates channel
+				int32 TexCoordsBufferIndex = -1;
+				int32 TexCoordsChannelIndex = -1;
+				Mesh.GetVertexBuffers().FindChannel(EMeshBufferSemantic::TexCoords, LayoutChannel, &TexCoordsBufferIndex, &TexCoordsChannelIndex);
+
+				if (TexCoordsBufferIndex < 0 || TexCoordsChannelIndex < 0)
+				{
+					// This is actually fine when using shared materials across LODs
+					UE_LOG(LogMutableCore, Log, TEXT("Trying to generate layout for missing UV channel %d. No layout blocks generated."), LayoutChannel);
+					return;
+				}
+
+				const FMeshBufferChannel& TexCoordsChannel = Mesh.VertexBuffers.Buffers[TexCoordsBufferIndex].Channels[TexCoordsChannelIndex];
+				check(TexCoordsChannel.Semantic == EMeshBufferSemantic::TexCoords);
+
+				uint8* TexCoordData = Mesh.GetVertexBuffers().GetBufferData(TexCoordsBufferIndex);
+				int32 TexCoordElemSize = Mesh.GetVertexBuffers().GetElementSize(TexCoordsBufferIndex);
+				int32 TexCoordChannelOffset = TexCoordsChannel.Offset;
+				TexCoordData += TexCoordChannelOffset;
+
+
 				// Copy UVs
-				if (TexCoordsChannel.Format == EMeshBufferFormat::Float32 && TexElemSize==sizeof(float)*2)
+				if (TexCoordsChannel.Format == EMeshBufferFormat::Float32 && TexCoordElemSize ==sizeof(float)*2)
 				{
-					FMemory::Memcpy(TexCoordData, TexCoords.GetData(), sizeof(float) * 2 * NumVertices);
+					FMemory::Memcpy(TexCoordData, ClassifyResult.TexCoords.GetData(), sizeof(float) * 2 * NumVertices);
 				}
 				else
 				{
@@ -552,7 +657,7 @@ namespace mu
 
 					for (int32 VertexIndex = 0; VertexIndex < NumVertices; ++VertexIndex)
 					{
-						FVector2f& UV = TexCoords[VertexIndex];
+						FVector2f& UV = ClassifyResult.TexCoords[VertexIndex];
 
 						if (TexCoordsChannel.Format == EMeshBufferFormat::Float32)
 						{
@@ -566,7 +671,7 @@ namespace mu
 							pUV[1] = FFloat16(UV[1]);
 						}
 
-						TexCoordData += TexElemSize;
+						TexCoordData += TexCoordElemSize;
 					}
 				}
 			}
diff --git a/Engine/Plugins/Mutable/Source/MutableRuntime/Public/MuR/Layout.h b/Engine/Plugins/Mutable/Source/MutableRuntime/Public/MuR/Layout.h
--- a/Engine/Plugins/Mutable/Source/MutableRuntime/Public/MuR/Layout.h
+++ b/Engine/Plugins/Mutable/Source/MutableRuntime/Public/MuR/Layout.h
@@ -34,6 +34,7 @@ namespace mu
 	struct FLayoutBlock
 	{
 		static constexpr uint64 InvalidBlockId = TNumericLimits<uint64>::Max();
+		static constexpr uint64 InvalidMaskIndex = TNumericLimits<uint16>::Max();
 
 		FIntVector2 Min = { 0, 0 };
 		FIntVector2 Size = { 0, 0 };
@@ -45,13 +46,21 @@ namespace mu
 		int32 Priority;
 
 		//! Value to control the method to reduce the block
-		uint32 bReduceBothAxes : 1;
+		uint16 bReduceBothAxes : 1;
 
 		//! Value to control if a block has to be reduced by two in an unitary reduction strategy
-		uint32 bReduceByTwo : 1;
+		uint16 bReduceByTwo : 1;
 
 		/** Explicit padding to prevent uninitialized memory in this POD. */
-		uint32 UnusedPadding : 30;
+		uint16 UnusedPadding : 14;
+
+		/** Index of the block mask used to classify the vertices and decide if they belong to the block.
+		* The image is stored in the Masks array of the FLayout object. 
+		* It is interpreted as covering the space of the entire layout grid.
+		* If the mask is null, all the vertices inside the block rect belong to the block.
+		* The image may be shared among several blocks and also among several layouts.
+		*/
+		uint16 MaskIndex = InvalidMaskIndex;
 
 		/** */
 		FLayoutBlock(FIntVector2 InMin = {}, FIntVector2 InSize = {})
@@ -78,7 +87,8 @@ namespace mu
 				(Size == o.Size) &&
 				(Priority == o.Priority) &&
 				(bReduceBothAxes == o.bReduceBothAxes) &&
-				(bReduceByTwo == o.bReduceByTwo);
+				(bReduceByTwo == o.bReduceByTwo) &&
+				(MaskIndex==o.MaskIndex);
 		}
 	};
 
@@ -100,14 +110,20 @@ namespace mu
 		*/
 		FIntVector2 MaxSize = FIntVector2(0, 0);
 
-		//!
-		TArray<FLayoutBlock> Blocks;
-
-		//! Packing strategy
+		/** */
 		EPackStrategy Strategy = EPackStrategy::Resizeable;
 
+		/** */
 		EReductionMethod ReductionMethod = EReductionMethod::Halve;
 
+		//!
+		TArray<FLayoutBlock> Blocks;
+
+		/** List of masks that may be used by the blocks for finer-grade selection of vertices. 
+		* The masks may be shared by multiple blocks and also across layouts. 
+		*/
+		TArray<TSharedPtr<const FImage>> Masks;
+
 	public:
 
 		//-----------------------------------------------------------------------------------------
diff --git a/Engine/Plugins/Mutable/Source/MutableTools/Private/MuT/CodeGenerator.cpp b/Engine/Plugins/Mutable/Source/MutableTools/Private/MuT/CodeGenerator.cpp
--- a/Engine/Plugins/Mutable/Source/MutableTools/Private/MuT/CodeGenerator.cpp
+++ b/Engine/Plugins/Mutable/Source/MutableTools/Private/MuT/CodeGenerator.cpp
@@ -416,28 +416,7 @@ namespace mu
 			return *CachedLayout;
 		}
 
-		TSharedPtr<FLayout> GeneratedLayout = MakeShared<FLayout>();
-		GeneratedLayout->Size = SourceLayout->Size;
-		GeneratedLayout->MaxSize = SourceLayout->MaxSize;
-		GeneratedLayout->Strategy = SourceLayout->Strategy;
-		GeneratedLayout->ReductionMethod = SourceLayout->ReductionMethod;
-
-		const int32 BlockCount = SourceLayout->Blocks.Num();
-		GeneratedLayout->Blocks.SetNum(BlockCount);
-		for (int32 BlockIndex = 0; BlockIndex < BlockCount; ++BlockIndex)
-		{
-			const FSourceLayoutBlock& From = SourceLayout->Blocks[BlockIndex];
-			FLayoutBlock& To = GeneratedLayout->Blocks[BlockIndex];
-			To.Min = From.Min;
-			To.Size = From.Size;
-			To.Priority = From.Priority;
-			To.bReduceBothAxes = From.bReduceBothAxes;
-			To.bReduceByTwo = From.bReduceByTwo;
-
-			// Assign unique ids to each layout block
-			uint64 Id = uint64(MeshIDPrefix) << 32 | uint64(BlockIndex);
-			To.Id = Id;
-		}
+		TSharedPtr<FLayout> GeneratedLayout = SourceLayout->BuildRuntimeLayout(MeshIDPrefix);
 
 		check(GeneratedLayout->Blocks.IsEmpty() || GeneratedLayout->Blocks[0].Id != FLayoutBlock::InvalidBlockId);
 		GenerateLayoutConstantState.GeneratedLayouts.Add({ SourceLayout,MeshIDPrefix }, GeneratedLayout);
@@ -996,10 +975,9 @@ namespace mu
 					MeshDynamicOptions.OverrideLayouts = BaseMeshResult->GeneratedLayouts;
 				}
 
-				// Ensure UV islands remain within their main layout block on lower LODs to avoid unexpected reordering 
-				// of the layout blocks when reusing a surface between LODs. Used to fix small displacements on vertices
+				// Ensure UV islands remain within their main layout block to fix small displacements on vertices
 				// that may cause them to fall on a different block.
-				MeshDynamicOptions.bClampUVIslands = BaseMeshResult != nullptr;
+				MeshDynamicOptions.bClampUVIslands = true; // bShareSurface;
 
 				return MeshDynamicOptions;
 			},
diff --git a/Engine/Plugins/Mutable/Source/MutableTools/Private/MuT/CodeGenerator.h b/Engine/Plugins/Mutable/Source/MutableTools/Private/MuT/CodeGenerator.h
--- a/Engine/Plugins/Mutable/Source/MutableTools/Private/MuT/CodeGenerator.h
+++ b/Engine/Plugins/Mutable/Source/MutableTools/Private/MuT/CodeGenerator.h
@@ -588,7 +588,7 @@ namespace mu
 			* layout blocks will be clamped to fit the one with more vertices belonging to the UV island.
 			* Mainly used to keep consistent layouts when reusing textures between LODs. 
 			*/
-			bool bClampUVIslands = false;
+			bool bClampUVIslands = true;
 
 			/** If true, UVs will be normalized. Normalize UVs should be done in cases where we operate with Images and Layouts */
 			bool bNormalizeUVs = false;
diff --git a/Engine/Plugins/Mutable/Source/MutableTools/Private/MuT/NodeLayout.cpp b/Engine/Plugins/Mutable/Source/MutableTools/Private/MuT/NodeLayout.cpp
--- a/Engine/Plugins/Mutable/Source/MutableTools/Private/MuT/NodeLayout.cpp
+++ b/Engine/Plugins/Mutable/Source/MutableTools/Private/MuT/NodeLayout.cpp
@@ -9,6 +9,7 @@
 #include "MuR/MeshPrivate.h"
 #include "MuR/MutableMath.h"
 #include "MuR/OpImageGrow.h"
+#include "MuR/OpMeshPrepareLayout.h"
 #include "Math/IntPoint.h"
 #include "Misc/AssertionMacros.h"
 
@@ -19,6 +20,41 @@ namespace mu
 	FNodeType NodeLayout::StaticType = FNodeType(Node::EType::Layout, Node::GetStaticType() );
 
 
+	TSharedPtr<FLayout> NodeLayout::BuildRuntimeLayout(uint32 BlockIDPrefix) const
+	{
+		TSharedPtr<FLayout> GeneratedLayout = MakeShared<FLayout>();
+		GeneratedLayout->Size = Size;
+		GeneratedLayout->MaxSize = MaxSize;
+		GeneratedLayout->Strategy = Strategy;
+		GeneratedLayout->ReductionMethod = ReductionMethod;
+
+		const int32 BlockCount = Blocks.Num();
+		GeneratedLayout->Blocks.SetNum(BlockCount);
+		for (int32 BlockIndex = 0; BlockIndex < BlockCount; ++BlockIndex)
+		{
+			const FSourceLayoutBlock& From = Blocks[BlockIndex];
+			FLayoutBlock& To = GeneratedLayout->Blocks[BlockIndex];
+			To.Min = From.Min;
+			To.Size = From.Size;
+			To.Priority = From.Priority;
+			To.bReduceBothAxes = From.bReduceBothAxes;
+			To.bReduceByTwo = From.bReduceByTwo;
+
+			// TODO: Optimize the format to RLE 1-bit
+			if (From.Mask)
+			{
+				To.MaskIndex = GeneratedLayout->Masks.AddUnique(From.Mask);
+			}
+
+			// Assign unique ids to each layout block
+			uint64 Id = uint64(BlockIDPrefix) << 32 | uint64(BlockIndex);
+			To.Id = Id;
+		}
+
+		return GeneratedLayout;
+	}
+
+
 	void NodeLayout::GenerateLayoutBlocks(const TSharedPtr<FMesh>& Mesh, int32 LayoutIndex )
 	{
 		if (!(Mesh && LayoutIndex >= 0 && Size.X > 0 && Size.Y > 0))
@@ -26,12 +62,31 @@ namespace mu
 			return;
 		}
 
+		int32 VertexCount = Mesh->VertexBuffers.GetElementCount();
+		
+		TArray<bool> VerticesAlreadyClassified;
+		{
+			VerticesAlreadyClassified.SetNumZeroed(VertexCount);
+
+			TSharedPtr<FLayout> RuntimeLayout = BuildRuntimeLayout(0);
+			FMutableLayoutClassifyResult ClassifyResult;
+			bool bNormalizeUVs = false;
+			bool bClampUVIslands = true;
+			if (ClassifyMeshVerticesInLayout(*Mesh, *RuntimeLayout, LayoutIndex, bNormalizeUVs, bClampUVIslands, ClassifyResult))
+			{
+				for (int32 VertexIndex = 0; VertexIndex < VertexCount; ++VertexIndex)
+				{
+					VerticesAlreadyClassified[VertexIndex] = ClassifyResult.LayoutData[VertexIndex] != FMutableLayoutClassifyResult::NullBlockId;
+				}
+			}
+		}
+
 		int32 IndexCount = Mesh->GetIndexCount();
 		TArray< FVector2f > UVs;
 
 		// Extract all the triangle edges
 		{
-			UVs.SetNumUninitialized(IndexCount * 2);
+			UVs.Reserve(IndexCount * 2);
 
 			UntypedMeshBufferIteratorConst IndexIt(Mesh->GetIndexBuffers(), EMeshBufferSemantic::VertexIndex, 0);
 			UntypedMeshBufferIteratorConst TexIt(Mesh->GetVertexBuffers(), EMeshBufferSemantic::TexCoords, LayoutIndex);
@@ -46,18 +101,25 @@ namespace mu
 				uint32 VertexIndex3 = IndexIt.GetAsUINT32();
 				IndexIt++;
 
+				// If any vertex belongs to a user-provided block, ignore the triangle: it will be assigned to the 
+				// block owning the mask.
+				if (VerticesAlreadyClassified[VertexIndex1] || VerticesAlreadyClassified[VertexIndex2] || VerticesAlreadyClassified[VertexIndex3])
+				{
+					continue;
+				}
+
 				FVector2f UV1 = (TexIt + VertexIndex1).GetAsVec2f();
 				FVector2f UV2 = (TexIt + VertexIndex2).GetAsVec2f();
 				FVector2f UV3 = (TexIt + VertexIndex3).GetAsVec2f();
 
-				UVs[TriangleIndex * 6 + 0] = UV1;
-				UVs[TriangleIndex * 6 + 1] = UV2;
+				UVs.Add( UV1 );
+				UVs.Add( UV2 );
 
-				UVs[TriangleIndex * 6 + 2] = UV2;
-				UVs[TriangleIndex * 6 + 3] = UV3;
+				UVs.Add( UV2 );
+				UVs.Add( UV3 );
 
-				UVs[TriangleIndex * 6 + 4] = UV3;
-				UVs[TriangleIndex * 6 + 5] = UV1;
+				UVs.Add( UV3 );
+				UVs.Add( UV1 );
 			}
 		}
 			
@@ -74,7 +136,7 @@ namespace mu
 		}
 
 		// Generate blocks by iterating all the edges
-		int32 EdgeCount = IndexCount;
+		int32 EdgeCount = UVs.Num()/2;
 		for (int32 EdgeIndex = 0; EdgeIndex < EdgeCount; ++EdgeIndex)
 		{
 			FVector2f APosition = UVs[EdgeIndex * 2];
@@ -101,6 +163,12 @@ namespace mu
 			
 				for (int32 BlockIndex = 0; BlockIndex < BlockRects.Num(); ++BlockIndex)
 				{
+					// ignore blocks with masks here.
+					if (Blocks.IsValidIndex(BlockIndex) && Blocks[BlockIndex].Mask)
+					{
+						continue;
+					}
+
 					if (BlockRects[BlockIndex].Contains(AGrid) || BlockRects[BlockIndex].Contains(BGrid))
 					{
 						bIsContained = true;
diff --git a/Engine/Plugins/Mutable/Source/MutableTools/Public/MuT/NodeLayout.h b/Engine/Plugins/Mutable/Source/MutableTools/Public/MuT/NodeLayout.h
--- a/Engine/Plugins/Mutable/Source/MutableTools/Public/MuT/NodeLayout.h
+++ b/Engine/Plugins/Mutable/Source/MutableTools/Public/MuT/NodeLayout.h
@@ -67,6 +67,9 @@ namespace mu
 
 		// Own interface
 
+		/** */
+		TSharedPtr<FLayout> BuildRuntimeLayout(uint32 BlockIDPrefix) const;
+
 		/** Generate the blocks of a layout using the UV of the meshes. 
 		* A list of existing blocks may be provided to specify starting blocks that shouldn't be modified unless it is mandatory because they contain partial face islands. 
 		*/
