[mutable] Updated COI editor layout.
- Hidden the "preview Settings" tab as its use case is not primary for this editor (here we want to edit the instance on the viewport and not the viewport itself)
- Fixed some typos
- Removed no longer necessary includes

#rnx

diff --git a/Engine/Plugins/Mutable/Source/CustomizableObjectEditor/Private/MuCOE/CustomizableInstanceDetails.cpp b/Engine/Plugins/Mutable/Source/CustomizableObjectEditor/Private/MuCOE/CustomizableInstanceDetails.cpp
--- a/Engine/Plugins/Mutable/Source/CustomizableObjectEditor/Private/MuCOE/CustomizableInstanceDetails.cpp
+++ b/Engine/Plugins/Mutable/Source/CustomizableObjectEditor/Private/MuCOE/CustomizableInstanceDetails.cpp
@@ -12,14 +12,12 @@
 #include "IDetailGroup.h"
 #include "Misc/TransactionObjectEvent.h"
 #include "ScopedTransaction.h"
-#include "Serialization/BufferArchive.h"
 #include "Toolkits/ToolkitManager.h"
 #include "Math/Transform.h"
 #include "PropertyCustomizationHelpers.h"
 #include "MuCO/CustomizableObjectInstancePrivate.h"
 #include "MuCO/CustomizableObject.h"
 #include "MuCO/CustomizableObjectInstance.h"
-#include "MuCO/CustomizableObjectSystem.h"
 #include "MuCOE/CustomizableObjectEditor.h"
 #include "MuCOE/UnrealEditorPortabilityHelpers.h"
 #include "AssetThumbnail.h"
@@ -31,7 +29,6 @@
 #include "Widgets/Colors/SColorPicker.h"
 #include "Widgets/Input/SButton.h"
 #include "Widgets/Input/SCheckBox.h"
-#include "Widgets/Input/SSearchBox.h"
 #include "Widgets/Input/SSpinBox.h"
 #include "Widgets/Input/SSlider.h"
 #include "Widgets/Input/STextComboBox.h"
@@ -227,17 +224,21 @@ void FCustomizableInstanceDetails::CustomizeDetails(const TSharedPtr<IDetailLayo
 		.OnCheckStateChanged(this, &FCustomizableInstanceDetails::OnUseUIThumbnailsSelectionChanged)
 	];
 
-
-	TSharedPtr<ICustomizableObjectInstanceEditor> Editor = GetEditorChecked();
-	if (UCustomizableObjectEditorProperties* EditorProperties = Editor->GetEditorProperties())
+	// Gameplay filters property
 	{
-		IDetailPropertyRow* FilterPropertyRow = VisibilitySettingsCategory.AddExternalObjectProperty({ EditorProperties }, FName("Filter"));
-		FilterPropertyRow->GetPropertyHandle()->SetOnChildPropertyValueChanged(FSimpleDelegate::CreateSP(this, &FCustomizableInstanceDetails::Refresh));
-
-		Filter = EditorProperties->Filter.GameplayTagsFilter;
-		FilterType = EditorProperties->Filter.GameplayTagsFilterType;
+		TSharedPtr<ICustomizableObjectInstanceEditor> Editor = GetEditorChecked();
+		if (UCustomizableObjectEditorProperties* EditorProperties = Editor->GetEditorProperties())
+		{
+			IDetailPropertyRow* FilterPropertyRow = VisibilitySettingsCategory.AddExternalObjectProperty({ EditorProperties }, GET_MEMBER_NAME_CHECKED(UCustomizableObjectEditorProperties, Filter));
+			check(FilterPropertyRow);
+			FilterPropertyRow->GetPropertyHandle()->SetOnChildPropertyValueChanged(FSimpleDelegate::CreateSP(this, &FCustomizableInstanceDetails::Refresh));
+		
+			Filter = EditorProperties->Filter.GameplayTagsFilter;
+			FilterType = EditorProperties->Filter.GameplayTagsFilterType;
+		}
 	}
 
+
 	// Copy, Paste and Reset FParameters
 	ParametersCategory.AddCustomRow(LOCTEXT("CustomizableInstanceDetails_CopyPasteResetButtons", "Copy Paste Reset"))
 	[
@@ -270,7 +271,7 @@ void FCustomizableInstanceDetails::CustomizeDetails(const TSharedPtr<IDetailLayo
 		[
 			SNew(SButton)
 			.HAlign(HAlign_Center)
-			.Text(LOCTEXT("Reset_Integer_Paramaters", "Reset parameters"))
+			.Text(LOCTEXT("Reset_Integer_Parameters", "Reset parameters"))
 			.OnClicked(this, &FCustomizableInstanceDetails::OnResetAllParameters)
 			.IsEnabled(CustomInstance->HasAnyParameters())
 			.ToolTipText(FText::FromString(FString("Clear the preview Instance parameters")))
@@ -278,13 +279,11 @@ void FCustomizableInstanceDetails::CustomizeDetails(const TSharedPtr<IDetailLayo
 	];
 
 	// FParameters Widgets
-	bool bHiddenParamsRuntime = GenerateParametersView(ParametersCategory);
-
-	if (bHiddenParamsRuntime)
+	if (bool bHiddenParamsRuntime = GenerateParametersView(ParametersCategory))
 	{
-		FText HiddenParamsRuntimeMessage = LOCTEXT("CustomizableInstanceDetails_HiddemParamsRuntime", "Parameters are hidden due to their Runtime type. \nUncheck the Only Runtime checkbox to see them.");
+		FText HiddenParamsRuntimeMessage = LOCTEXT("CustomizableInstanceDetails_HiddenParamsRuntime", "Parameters are hidden due to their Runtime type. \nUncheck the Only Runtime checkbox to see them.");
 
-		ParametersCategory.AddCustomRow(LOCTEXT("CustomizableInstanceDetails_HiddemParamsRuntimeRow", "Parameters are hidden"))
+		ParametersCategory.AddCustomRow(LOCTEXT("CustomizableInstanceDetails_HiddenParamsRuntimeRow", "Parameters are hidden"))
 		[
 			SNew(SHorizontalBox)
 				+ SHorizontalBox::Slot()
diff --git a/Engine/Plugins/Mutable/Source/CustomizableObjectEditor/Private/MuCOE/CustomizableInstanceDetails.h b/Engine/Plugins/Mutable/Source/CustomizableObjectEditor/Private/MuCOE/CustomizableInstanceDetails.h
--- a/Engine/Plugins/Mutable/Source/CustomizableObjectEditor/Private/MuCOE/CustomizableInstanceDetails.h
+++ b/Engine/Plugins/Mutable/Source/CustomizableObjectEditor/Private/MuCOE/CustomizableInstanceDetails.h
@@ -194,14 +194,14 @@ private:
 	/** True when a slider is being edited*/
 	bool bUpdatingSlider = false;
 
-	/** Array to store dynamic brushes. Neede because an image widget only stores a pointer to a Brush. */
+	/** Array to store dynamic brushes. Needed because an image widget only stores a pointer to a Brush. */
 	TArray< TSharedPtr<class FDeferredCleanupSlateBrush> > DynamicBrushes;
 
-	// Unique transaction pointer to allow transactions that start and finish in different funtion scopes.
+	// Unique transaction pointer to allow transactions that start and finish in different function scopes.
 	TUniquePtr<FScopedTransaction> Transaction;
 
 	//** Editor gameplay tags filter and filter type. Used to filter int parameter options. */
 	FGameplayTagContainer Filter;
-	EGameplayContainerMatchType FilterType;
+	EGameplayContainerMatchType FilterType = EGameplayContainerMatchType::Any;
 };
 
diff --git a/Engine/Plugins/Mutable/Source/CustomizableObjectEditor/Private/MuCOE/CustomizableObjectInstanceEditor.cpp b/Engine/Plugins/Mutable/Source/CustomizableObjectEditor/Private/MuCOE/CustomizableObjectInstanceEditor.cpp
--- a/Engine/Plugins/Mutable/Source/CustomizableObjectEditor/Private/MuCOE/CustomizableObjectInstanceEditor.cpp
+++ b/Engine/Plugins/Mutable/Source/CustomizableObjectEditor/Private/MuCOE/CustomizableObjectInstanceEditor.cpp
@@ -2,13 +2,8 @@
 
 #include "MuCOE/CustomizableObjectInstanceEditor.h"
 
-#include "AdvancedPreviewSceneModule.h"
-#include "Animation/AnimSingleNodeInstance.h"
-#include "Animation/DebugSkelMeshComponent.h"
 #include "Animation/PoseAsset.h"
-#include "AssetRegistry/AssetRegistryModule.h"
 #include "ContentBrowserModule.h"
-#include "CustomizableObjectCustomSettings.h"
 #include "CustomizableObjectEditor.h"
 #include "DetailsViewArgs.h"
 #include "FileHelpers.h"
@@ -17,7 +12,6 @@
 #include "IContentBrowserSingleton.h"
 #include "IDetailsView.h"
 #include "MuCO/CustomizableObjectInstance.h"
-#include "MuCO/CustomizableObjectSystem.h"
 #include "MuCO/CustomizableSkeletalComponent.h"
 #include "MuCOE/CustomizableObjectEditorActions.h"
 #include "MuCO/ICustomizableObjectEditorModule.h"
@@ -25,16 +19,12 @@
 #include "MuCOE/CustomizableObjectInstanceEditorActions.h"
 #include "MuCOE/CustomizableObjectPreviewScene.h"
 #include "MuCOE/SCustomizableObjectEditorTextureAnalyzer.h"
-#include "MuCOE/UnrealEditorPortabilityHelpers.h"
 #include "PropertyEditorModule.h"
 #include "SCustomizableObjectEditorAdvancedPreviewSettings.h"
 #include "SCustomizableObjectEditorViewport.h"
-#include "Selection.h"
 #include "Widgets/Docking/SDockTab.h"
-#include "Widgets/Notifications/SNotificationList.h"
 #include "WorkspaceMenuStructure.h"
 #include "WorkspaceMenuStructureModule.h"
-#include "Animation/SkeletalMeshActor.h"
 #include "MuCO/CustomizableObjectInstancePrivate.h"
 
 class FAdvancedPreviewScene;
@@ -283,37 +273,31 @@ void FCustomizableObjectInstanceEditor::InitCustomizableObjectInstanceEditor( co
 	bOnlyRelevantParameters = InCustomizableObjectInstance->GetPrivate()->bShowOnlyRelevantParameters;
 	bOnlyRuntimeParameters = InCustomizableObjectInstance->GetPrivate()->bShowOnlyRuntimeParameters;
 	
-	FAdvancedPreviewSceneModule& AdvancedPreviewSceneModule = FModuleManager::LoadModuleChecked<FAdvancedPreviewSceneModule>("AdvancedPreviewScene");
-
 	TSharedPtr<FAdvancedPreviewScene> AdvancedPreviewScene = StaticCastSharedPtr<FAdvancedPreviewScene>(Viewport->GetPreviewScene());
 
-	CustomizableObjectEditorAdvancedPreviewSettings = SNew(SCustomizableObjectEditorAdvancedPreviewSettings, AdvancedPreviewScene.ToSharedRef())
+	AdvancedPreviewSettingsWidget = SNew(SCustomizableObjectEditorAdvancedPreviewSettings, AdvancedPreviewScene.ToSharedRef())
 		.CustomSettings(CustomSettings)
 		.CustomizableObjectEditor(SharedThis(this).ToWeakPtr());
 	
-	AdvancedPreviewSettingsWidget = CustomizableObjectEditorAdvancedPreviewSettings;
+	const TSharedRef<FTabManager::FLayout> StandaloneDefaultLayout = FTabManager::NewLayout( "Standalone_CustomizableObjectInstanceEditor_Layout_v2.2" );
 
-	const TSharedRef<FTabManager::FLayout> StandaloneDefaultLayout = FTabManager::NewLayout( "Standalone_CustomizableObjectInstanceEditor_Layout_v2.1" )
-	->AddArea
+	StandaloneDefaultLayout->AddArea
 	(
 		FTabManager::NewPrimaryArea()->SetOrientation(Orient_Horizontal)
 		->Split
 		(
 			FTabManager::NewStack()
-			->SetSizeCoefficient(0.8f)
+			->SetSizeCoefficient(0.7f)
+			->SetHideTabWell(true)
 			->AddTab(ViewportTabId, ETabState::OpenedTab)
 		)
 		->Split
 		(
 			FTabManager::NewStack()
-			->SetSizeCoefficient(0.2f)
+			->SetSizeCoefficient(0.3f)
 			->AddTab(InstancePropertiesTabId, ETabState::OpenedTab)
-		)
-		->Split
-		(
-			FTabManager::NewStack()
-			->SetSizeCoefficient(0.2f)
 			->AddTab(AdvancedPreviewSettingsTabId, ETabState::OpenedTab)
+			->SetForegroundTab(InstancePropertiesTabId)
 		)
 	);
 
@@ -327,7 +311,7 @@ void FCustomizableObjectInstanceEditor::InitCustomizableObjectInstanceEditor( co
 	RegenerateMenusAndToolbars();
 
 	// Clears selection highlight.
-	OnInstancePropertySelectionChanged(NULL);
+	OnInstancePropertySelectionChanged(nullptr);
 	
 	if (UCustomizableObject* CustomizableObject = CustomizableObjectInstance->GetCustomizableObject())
 	{
@@ -337,7 +321,6 @@ void FCustomizableObjectInstanceEditor::InitCustomizableObjectInstanceEditor( co
 		OnCustomizableObjectStatusChanged(FCustomizableObjectStatusTypes::EState::Loading, CustomizableObjectPrivate->Status.Get()); // Fake we are still in the loading phase.
 
 		int32 StateParameterCount = CustomizableObject->GetStateParameterCount(CustomizableObjectInstance->GetCurrentState());
-		int32 ParameterCount = CustomizableObject->GetParameterCount();
 
 		if (StateParameterCount == 0)
 		{
@@ -588,7 +571,7 @@ UCustomizableObjectEditorProperties* FCustomizableObjectInstanceEditor::GetEdito
 
 TSharedPtr<SCustomizableObjectEditorAdvancedPreviewSettings> FCustomizableObjectInstanceEditor::GetAdvancedPreviewSettings()
 {
-	return CustomizableObjectEditorAdvancedPreviewSettings;
+	return AdvancedPreviewSettingsWidget;
 }
 
 
diff --git a/Engine/Plugins/Mutable/Source/CustomizableObjectEditor/Private/MuCOE/CustomizableObjectInstanceEditor.h b/Engine/Plugins/Mutable/Source/CustomizableObjectEditor/Private/MuCOE/CustomizableObjectInstanceEditor.h
--- a/Engine/Plugins/Mutable/Source/CustomizableObjectEditor/Private/MuCOE/CustomizableObjectInstanceEditor.h
+++ b/Engine/Plugins/Mutable/Source/CustomizableObjectEditor/Private/MuCOE/CustomizableObjectInstanceEditor.h
@@ -296,7 +296,6 @@ private:
 	TSharedPtr<SCustomizableObjectEditorViewportTabBody> Viewport;
 	TSharedPtr<IDetailsView> CustomizableInstanceDetailsView;
 
-	TSharedPtr<SCustomizableObjectEditorAdvancedPreviewSettings> CustomizableObjectEditorAdvancedPreviewSettings;
 	
 	/** Level of Details Settings widget. */
 	TSharedPtr<SLevelOfDetailSettings> LevelOfDetailSettings;
@@ -326,7 +325,7 @@ private:
 	TObjectPtr<UUpdateClassWrapperClass> HelperCallback;
 
 	/** Scene preview settings widget */
-	TSharedPtr<SWidget> AdvancedPreviewSettingsWidget;
+	TSharedPtr<SCustomizableObjectEditorAdvancedPreviewSettings> AdvancedPreviewSettingsWidget;
 
 	/** Pose asset when doing drag and drop of an UPoseAsset to the viewport */
 	UPoseAsset* PoseAsset;
