[Mutable] Removed deprecated 5.6 API functions.

diff --git a/Engine/Plugins/Mutable/Config/BaseMutable.ini b/Engine/Plugins/Mutable/Config/BaseMutable.ini
--- a/Engine/Plugins/Mutable/Config/BaseMutable.ini
+++ b/Engine/Plugins/Mutable/Config/BaseMutable.ini
@@ -24,4 +24,5 @@
 +FunctionRedirects=(OldName="/Script/CustomizableObject.CustomizableObject.GetIntParameterGroupType",NewName="/Script/CustomizableObject.CustomizableObject.GetEnumParameterGroupType")
 +FunctionRedirects=(OldName="/Script/CustomizableObject.CustomizableObject.GetIntParameterOptionDataTable",NewName="/Script/CustomizableObject.CustomizableObject.GetEnumParameterValueDataTable")
 +FunctionRedirects=(OldName="/Script/CustomizableObject.CustomizableObject.GetIntParameterDefaultValue",NewName="/Script/CustomizableObject.CustomizableObject.GetEnumParameterDefaultValue")
-+PropertyRedirects=(OldName="/Script/CustomizableObjectEditor.CustomizableObjectNodeComponentPassthroughMesh.Mesh",NewName="/Script/CustomizableObjectEditor.CustomizableObjectNodeComponentPassthroughMesh.SkeletalMesh")
\ No newline at end of file
++PropertyRedirects=(OldName="/Script/CustomizableObjectEditor.CustomizableObjectNodeComponentPassthroughMesh.Mesh",NewName="/Script/CustomizableObjectEditor.CustomizableObjectNodeComponentPassthroughMesh.SkeletalMesh")
++PropertyRedirects=(OldName="/Script/CustomizableObject.CustomizableObject.bEnableAnimBpPhysicsAssetsManipualtion",NewName="/Script/CustomizableObject.CustomizableObject.bEnableAnimBpPhysicsAssetsManipulation")
\ No newline at end of file
diff --git a/Engine/Plugins/Mutable/Source/CustomizableObject/Private/MuCO/CustomizableObject.cpp b/Engine/Plugins/Mutable/Source/CustomizableObject/Private/MuCO/CustomizableObject.cpp
--- a/Engine/Plugins/Mutable/Source/CustomizableObject/Private/MuCO/CustomizableObject.cpp
+++ b/Engine/Plugins/Mutable/Source/CustomizableObject/Private/MuCO/CustomizableObject.cpp
@@ -997,9 +997,9 @@ USkeletalMesh* UCustomizableObject::GetComponentMeshReferenceSkeletalMesh(const
 }
 
 
-int32 UCustomizableObject::FindState( const FString& Name ) const
+USkeletalMesh* UCustomizableObject::GetSkeletalMeshComponentReferenceSkeletalMesh(const FName& ComponentName) const
 {
-	return GetPrivate()->FindState(Name);
+	return GetComponentMeshReferenceSkeletalMesh(ComponentName);
 }
 
 
@@ -1035,8 +1035,10 @@ FString UCustomizableObjectPrivate::GetStateName(int32 StateIndex) const
 }
 
 
-int32 UCustomizableObject::GetStateParameterCount( int32 StateIndex ) const
+int32 UCustomizableObject::GetStateParameterCount(const FString& StateName) const
 {
+	int32 StateIndex = GetPrivate()->FindState(StateName);
+
 	int32 Result = 0;
 
 	if (Private->GetModel())
@@ -1047,33 +1049,11 @@ int32 UCustomizableObject::GetStateParameterCount( int32 StateIndex ) const
 	return Result;
 }
 
-int32 UCustomizableObject::GetStateParameterIndex(int32 StateIndex, int32 ParameterIndex) const
-{
-	return GetPrivate()->GetStateParameterIndex(StateIndex, ParameterIndex);
-}
-
-
-int32 UCustomizableObject::GetStateParameterCount(const FString& StateName) const
-{
-	int32 StateIndex = GetPrivate()->FindState(StateName);
-
-	PRAGMA_DISABLE_DEPRECATION_WARNINGS
-	return GetStateParameterCount(StateIndex);
-	PRAGMA_ENABLE_DEPRECATION_WARNINGS
-}
-
 
 FString UCustomizableObject::GetStateParameterName(const FString& StateName, int32 ParameterIndex) const
 {
 	int32 StateIndex = GetPrivate()->FindState(StateName);
 
-	PRAGMA_DISABLE_DEPRECATION_WARNINGS
-	return GetStateParameterName(StateIndex, ParameterIndex);
-	PRAGMA_ENABLE_DEPRECATION_WARNINGS
-}
-
-FString UCustomizableObject::GetStateParameterName(int32 StateIndex, int32 ParameterIndex) const
-{
 	return GetParameterName(GetPrivate()->GetStateParameterIndex(StateIndex, ParameterIndex));
 }
 
@@ -1234,7 +1214,7 @@ bool UCustomizableObjectPrivate::IsPhysicsAssetMergeEnabled() const
 bool UCustomizableObjectPrivate::IsEnabledAnimBpPhysicsAssetsManipulation() const
 {
 	PRAGMA_DISABLE_DEPRECATION_WARNINGS // TODO Remove >5.6
-	return GetPublic()->bEnableAnimBpPhysicsAssetsManipualtion;
+	return GetPublic()->bEnableAnimBpPhysicsAssetsManipulation;
 	PRAGMA_ENABLE_DEPRECATION_WARNINGS
 }
 #endif
@@ -1337,12 +1317,6 @@ int32 UCustomizableObject::GetParameterCount() const
 }
 
 
-EMutableParameterType UCustomizableObject::GetParameterType(int32 ParamIndex) const
-{
-	return GetPrivate()->GetParameterType(ParamIndex);
-}
-
-
 EMutableParameterType UCustomizableObjectPrivate::GetParameterType(int32 ParamIndex) const
 {
 	if (ParamIndex >= 0 && ParamIndex < ParameterProperties.Num())
@@ -1494,12 +1468,6 @@ void UCustomizableObjectPrivate::UpdateParameterPropertiesFromModel(const TShare
 }
 
 
-int32 UCustomizableObject::GetIntParameterNumOptions(int32 ParamIndex) const
-{
-	return GetPrivate()->GetEnumParameterNumValues(ParamIndex);
-}
-
-
 int32 UCustomizableObject::GetEnumParameterNumValues(const FString& ParamName) const
 {
 	if (int32 ParamIndex = GetPrivate()->FindParameter(ParamName))
@@ -1538,18 +1506,6 @@ bool UCustomizableObject::ContainsEnumParameterValue(const FString& ParameterNam
 }
 
 
-const FString& UCustomizableObject::GetIntParameterAvailableOption(int32 ParamIndex, int32 K) const
-{
-	return GetPrivate()->GetIntParameterAvailableOption(ParamIndex, K);
-}
-
-
-int32 UCustomizableObject::FindParameter(const FString& Name) const
-{
-	return GetPrivate()->FindParameter(Name);
-}
-
-
 bool UCustomizableObject::ContainsParameter(const FString& ParameterName) const
 {
 	return GetPrivate()->FindParameter(ParameterName) != INDEX_NONE;
@@ -1581,12 +1537,6 @@ int32 UCustomizableObjectPrivate::FindParameterTyped(const FString& Name, EMutab
 }
 
 
-int32 UCustomizableObject::FindIntParameterValue(int32 ParamIndex, const FString& Value) const
-{
-	return GetPrivate()->FindIntParameterValue(ParamIndex, Value);
-}
-
-
 int32 UCustomizableObjectPrivate::FindIntParameterValue(int32 ParamIndex, const FString& Value) const
 {
 	int32 MinValueIndex = INDEX_NONE;
@@ -1616,18 +1566,6 @@ int32 UCustomizableObjectPrivate::FindIntParameterValue(int32 ParamIndex, const
 }
 
 
-FString UCustomizableObject::FindIntParameterValueName(int32 ParamIndex, int32 ParamValue) const
-{
-	return GetPrivate()->FindIntParameterValueName(ParamIndex, ParamValue);
-}
-
-
-USkeletalMesh* UCustomizableObject::GetRefSkeletalMesh(int32 ObjectComponentIndex) const
-{
-	return GetComponentMeshReferenceSkeletalMesh(FName(FString::FromInt(ObjectComponentIndex)));
-}
-
-
 FMutableParamUIMetadata UCustomizableObject::GetParameterUIMetadata(const FString& ParamName) const
 {
 	if (const UModelResources* ModelResources = Private->GetModelResources())
@@ -1666,12 +1604,6 @@ FMutableParamUIMetadata UCustomizableObject::GetEnumParameterValueUIMetadata(con
 }
 
 
-FMutableParamUIMetadata UCustomizableObject::GetIntParameterOptionUIMetadata(const FString& ParamName, const FString& OptionName) const
-{
-	return GetEnumParameterValueUIMetadata(ParamName, OptionName);
-}
-
-
 ECustomizableObjectGroupType UCustomizableObject::GetEnumParameterGroupType(const FString& ParamName) const
 {
 	const UModelResources* ModelResources = Private->GetModelResources();
@@ -1695,12 +1627,6 @@ ECustomizableObjectGroupType UCustomizableObject::GetEnumParameterGroupType(cons
 }
 
 
-ECustomizableObjectGroupType UCustomizableObject::GetIntParameterGroupType(const FString& ParamName) const
-{
-	return GetEnumParameterGroupType(ParamName);
-}
-
-
 FMutableStateUIMetadata UCustomizableObject::GetStateUIMetadata(const FString& StateName) const
 {
 	if (const UModelResources* ModelResources = Private->GetModelResources())
@@ -1736,12 +1662,6 @@ TArray<TSoftObjectPtr<UDataTable>> UCustomizableObject::GetEnumParameterValueDat
 	
 	return {};
 }
-
-
-TArray<TSoftObjectPtr<UDataTable>> UCustomizableObject::GetIntParameterOptionDataTable(const FString& ParamName, const FString& OptionName)
-{
-	return GetEnumParameterValueDataTable(ParamName, OptionName);
-}
 #endif
 
 
@@ -1785,12 +1705,6 @@ int32 UCustomizableObject::GetEnumParameterDefaultValue(const FString& InParamet
 }
 
 
-int32 UCustomizableObject::GetIntParameterDefaultValue(const FString& InParameterName) const
-{
-	return GetEnumParameterDefaultValue(InParameterName);
-}
-
-
 bool UCustomizableObject::GetBoolParameterDefaultValue(const FString& InParameterName) const
 {
 	const int32 ParameterIndex = GetPrivate()->FindParameter(InParameterName);
@@ -1856,21 +1770,6 @@ FTransform UCustomizableObject::GetTransformParameterDefaultValue(const FString&
 }
 
 
-void UCustomizableObject::GetProjectorParameterDefaultValue(const FString& InParameterName, FVector3f& OutPos,
-	FVector3f& OutDirection, FVector3f& OutUp, FVector3f& OutScale, float& OutAngle,
-	ECustomizableObjectProjectorType& OutType) const
-{
-	const FCustomizableObjectProjector Projector = GetProjectorParameterDefaultValue(InParameterName);
-		
-	OutType = Projector.ProjectionType;
-	OutPos = Projector.Position;
-	OutDirection = Projector.Direction;
-	OutUp = Projector.Up;
-	OutScale = Projector.Scale;
-	OutAngle = Projector.Angle;
-}
-
-
 FCustomizableObjectProjector UCustomizableObject::GetProjectorParameterDefaultValue(const FString& InParameterName) const 
 {
 	const int32 ParameterIndex = GetPrivate()->FindParameter(InParameterName);
@@ -1941,12 +1840,6 @@ void UCustomizableObject::Compile(const FCompileParams& Params)
 }
 
 
-bool UCustomizableObject::IsParameterMultidimensional(const int32& InParamIndex) const
-{
-	return GetPrivate()->IsParameterMultidimensional(InParamIndex);
-}
-
-
 bool UCustomizableObjectPrivate::IsParameterMultidimensional(const int32& InParamIndex) const
 {
 	check(InParamIndex != INDEX_NONE);
diff --git a/Engine/Plugins/Mutable/Source/CustomizableObject/Private/MuCO/CustomizableObjectInstance.cpp b/Engine/Plugins/Mutable/Source/CustomizableObject/Private/MuCO/CustomizableObjectInstance.cpp
--- a/Engine/Plugins/Mutable/Source/CustomizableObject/Private/MuCO/CustomizableObjectInstance.cpp
+++ b/Engine/Plugins/Mutable/Source/CustomizableObject/Private/MuCO/CustomizableObjectInstance.cpp
@@ -2389,10 +2389,6 @@ bool UCustomizableInstancePrivate::UpdateSkeletalMesh_PostBeginUpdate0(UCustomiz
 		
 		for (const TPair<const UCustomizableObjectExtension*, TArray<FInputPinDataContainer>>& Pair : ExtensionToExtensionData)
 		{
-			PRAGMA_DISABLE_DEPRECATION_WARNINGS
-			Pair.Key->OnSkeletalMeshCreated(Pair.Value, ObjectComponentIndex.GetValue(), SkeletalMesh);
-			PRAGMA_ENABLE_DEPRECATION_WARNINGS
-
 			Pair.Key->OnSkeletalMeshCreated(Pair.Value, ComponentName, SkeletalMesh);
 		}
 		
@@ -2552,17 +2548,24 @@ void UCustomizableObjectInstance::LoadDescriptor(FArchive &Ar)
 
 const FString& UCustomizableObjectInstance::GetIntParameterSelectedOption(const FString& ParamName, const int32 RangeIndex) const
 {
-	return Descriptor.GetIntParameterSelectedOption(ParamName, RangeIndex);
+	return GetEnumParameterSelectedOption(ParamName, RangeIndex);
 }
 
 
-void UCustomizableObjectInstance::SetIntParameterSelectedOption(int32 IntParamIndex, const FString& SelectedOption, const int32 RangeIndex)
+const FString& UCustomizableObjectInstance::GetEnumParameterSelectedOption(const FString& ParamName, int32 RangeIndex) const
 {
-	Descriptor.SetIntParameterSelectedOption(IntParamIndex, SelectedOption, RangeIndex);
+	return Descriptor.GetIntParameterSelectedOption(ParamName, RangeIndex);
+
 }
 
 
 void UCustomizableObjectInstance::SetIntParameterSelectedOption(const FString& ParamName, const FString& SelectedOptionName, const int32 RangeIndex)
+{
+	SetEnumParameterSelectedOption(ParamName, SelectedOptionName, RangeIndex);
+}
+
+
+void UCustomizableObjectInstance::SetEnumParameterSelectedOption(const FString& ParamName, const FString& SelectedOptionName, int32 RangeIndex)
 {
 	Descriptor.SetIntParameterSelectedOption(ParamName, SelectedOptionName, RangeIndex);
 }
@@ -2744,6 +2747,12 @@ FCustomizableObjectProjector UCustomizableObjectInstance::GetProjector(const FSt
 
 
 bool UCustomizableObjectInstance::ContainsIntParameter(const FString& ParameterName) const
+{
+	return ContainsEnumParameter(ParameterName);
+}
+
+
+bool UCustomizableObjectInstance::ContainsEnumParameter(const FString& ParameterName) const
 {
 	return Descriptor.FindTypedParameterIndex(ParameterName, EMutableParameterType::Int) != INDEX_NONE;
 }
@@ -2779,36 +2788,6 @@ bool UCustomizableObjectInstance::ContainsTransformParameter(const FString& Para
 }
 
 
-int32 UCustomizableObjectInstance::FindIntParameterNameIndex(const FString& ParamName) const
-{
-	return Descriptor.FindTypedParameterIndex(ParamName, EMutableParameterType::Int);
-}
-
-
-int32 UCustomizableObjectInstance::FindFloatParameterNameIndex(const FString& ParamName) const
-{
-	return Descriptor.FindTypedParameterIndex(ParamName, EMutableParameterType::Float);
-}
-
-
-int32 UCustomizableObjectInstance::FindBoolParameterNameIndex(const FString& ParamName) const
-{
-	return Descriptor.FindTypedParameterIndex(ParamName, EMutableParameterType::Bool);
-}
-
-
-int32 UCustomizableObjectInstance::FindVectorParameterNameIndex(const FString& ParamName) const
-{
-	return Descriptor.FindTypedParameterIndex(ParamName, EMutableParameterType::Color);
-}
-
-
-int32 UCustomizableObjectInstance::FindProjectorParameterNameIndex(const FString& ParamName) const
-{
-	return Descriptor.FindTypedParameterIndex(ParamName, EMutableParameterType::Projector);
-}
-
-
 int32 UCustomizableInstancePrivate::FindIntParameterNameIndex(const FString& ParamName) const
 {
 	return GetPublic()->Descriptor.FindTypedParameterIndex(ParamName, EMutableParameterType::Int);
@@ -2840,7 +2819,6 @@ int32 UCustomizableInstancePrivate::FindProjectorParameterNameIndex(const FStrin
 
 
 #if WITH_EDITOR
-
 void UCustomizableInstancePrivate::UpdateSkeletalMeshAsyncResult(FInstanceUpdateNativeDelegate Callback, bool bIgnoreCloseDist, bool bForceHighPriority, TSharedPtr<FMutableSystemSettingsOverrides> MutableSystemSettingsOverride)
 {
 	UCustomizableObjectSystemPrivate* SystemPrivate = UCustomizableObjectSystem::GetInstance()->GetPrivate();
@@ -5404,6 +5382,37 @@ FCustomizableObjectInstanceDescriptor& UCustomizableInstancePrivate::GetDescript
 }
 
 
+const TArray<TObjectPtr<UMaterialInterface>>* UCustomizableObjectInstance::GetOverrideMaterials(int32 ComponentIndex) const
+{
+	UCustomizableObject* Object = GetCustomizableObject();
+	if (!Object)
+	{
+		return nullptr;
+	}
+
+	const UModelResources* ModelResources = Object->GetPrivate()->GetModelResources();
+	if (!ModelResources)
+	{
+		return nullptr;
+	}
+
+	if (!ModelResources->ComponentNamesPerObjectComponent.IsValidIndex(ComponentIndex))
+	{
+		return nullptr;	
+	}
+	
+	const FName ComponentName = ModelResources->ComponentNamesPerObjectComponent[ComponentIndex];
+
+	FCustomizableInstanceComponentData* ComponentData = PrivateData->GetComponentData(ComponentName);
+	if (!ComponentData)
+	{
+		return nullptr;
+	}
+
+	return &ComponentData->OverrideMaterials;
+}
+
+
 TArray<UMaterialInterface*> UCustomizableObjectInstance::GetSkeletalMeshComponentOverrideMaterials(const FName& ComponentName) const
 {
 	FCustomizableInstanceComponentData* ComponentData = PrivateData->GetComponentData(ComponentName);
@@ -6786,6 +6795,12 @@ USkeletalMesh* UCustomizableObjectInstance::GetComponentMeshSkeletalMesh(const F
 }
 
 
+USkeletalMesh* UCustomizableObjectInstance::GetSkeletalMeshComponentSkeletalMesh(const FName& ComponentName) const
+{
+	return GetComponentMeshSkeletalMesh(ComponentName);
+}
+
+
 bool UCustomizableObjectInstance::HasAnySkeletalMesh() const
 {
 	return !GetPrivate()->SkeletalMeshes.IsEmpty();
@@ -7072,14 +7087,7 @@ TSet<UAssetUserData*> UCustomizableObjectInstance::GetMergedAssetUserData(int32
 }
 
 
-bool UCustomizableObjectInstance::CanUpdateInstance() const
-{
-	return GetPrivate()->CanUpdateInstance();
-}
-
-
 #if WITH_EDITORONLY_DATA
-
 void CalculateBonesToRemove(const FSkeletalMeshLODRenderData& LODResource, const FReferenceSkeleton& RefSkeleton, TArray<FBoneReference>& OutBonesToRemove)
 {
 	const int32 NumBones = RefSkeleton.GetNum();
diff --git a/Engine/Plugins/Mutable/Source/CustomizableObject/Private/MuCO/CustomizableObjectSystem.cpp b/Engine/Plugins/Mutable/Source/CustomizableObject/Private/MuCO/CustomizableObjectSystem.cpp
--- a/Engine/Plugins/Mutable/Source/CustomizableObject/Private/MuCO/CustomizableObjectSystem.cpp
+++ b/Engine/Plugins/Mutable/Source/CustomizableObject/Private/MuCO/CustomizableObjectSystem.cpp
@@ -4479,18 +4479,6 @@ void UCustomizableObjectSystem::SetOnlyGenerateRequestedLODsEnabled(bool bIsEnab
 }
 
 
-#if WITH_EDITOR
-UE_DEPRECATED(5.6, "This method has been deprecated.")
-void UCustomizableObjectSystem::SetImagePixelFormatOverride(const mu::FImageOperator::FImagePixelFormatFunc& InFunc)
-{
-	if (Private != nullptr)
-	{
-		Private->ImageFormatOverrideFunc = InFunc;
-	}
-}
-#endif
-
-
 void UCustomizableObjectSystem::AddUncompiledCOWarning(const UCustomizableObject& InObject, FString const* OptionalLogInfo)
 {
 	FString Msg;
diff --git a/Engine/Plugins/Mutable/Source/CustomizableObject/Public/MuCO/CustomizableObject.h b/Engine/Plugins/Mutable/Source/CustomizableObject/Public/MuCO/CustomizableObject.h
--- a/Engine/Plugins/Mutable/Source/CustomizableObject/Public/MuCO/CustomizableObject.h
+++ b/Engine/Plugins/Mutable/Source/CustomizableObject/Public/MuCO/CustomizableObject.h
@@ -281,55 +281,43 @@ private:
 	UPROPERTY()
 	FCompilationOptions_DEPRECATED CompileOptions_DEPRECATED;
 
-public:
-	UE_DEPRECATED(5.6, "Only modify this value through the Customizable ObjectEditor")
 	UPROPERTY(EditAnywhere, Category = CompileOptions)
-	bool bEnableRealTimeMorphTargets = false; // TODO >5.6 make private
+	bool bEnableRealTimeMorphTargets = false;
 
-	UE_DEPRECATED(5.6, "Only modify this value through the Customizable ObjectEditor")
 	UPROPERTY(EditAnywhere, Category = CompileOptions)
-	bool bEnableClothing = false; // TODO >5.6 make private
+	bool bEnableClothing = false;
 
-	UE_DEPRECATED(5.6, "Only modify this value through the Customizable ObjectEditor")
 	UPROPERTY(EditAnywhere, Category = CompileOptions)
-	bool bEnable16BitBoneWeights = false; // TODO >5.6 make private
+	bool bEnable16BitBoneWeights = false;
 
-	UE_DEPRECATED(5.6, "Only modify this value through the Customizable ObjectEditor")
 	UPROPERTY(EditAnywhere, Category = CompileOptions)
-	bool bEnableAltSkinWeightProfiles = false; // TODO >5.6 make private
+	bool bEnableAltSkinWeightProfiles = false;
 
-	UE_DEPRECATED(5.6, "Only modify this value through the Customizable ObjectEditor")
 	UPROPERTY(EditAnywhere, Category = CompileOptions)
-	bool bEnablePhysicsAssetMerge = false; // TODO >5.6 make private
+	bool bEnablePhysicsAssetMerge = false;
 
 	/** Experimental */
-	UE_DEPRECATED(5.6, "Only modify this value through the Customizable ObjectEditor")
 	UPROPERTY(EditAnywhere, Category = CompileOptions)
-	bool bEnableAnimBpPhysicsAssetsManipualtion = false; // TODO >5.6 make private
-
+	bool bEnableAnimBpPhysicsAssetsManipulation = false;
+	
 	// When this is enabled generated meshes will merge the AssetUserData from all of its constituent mesh parts
-	UE_DEPRECATED(5.6, "Only modify this value through the Customizable ObjectEditor")
 	UPROPERTY(EditAnywhere, Category = CompileOptions)
-	bool bEnableAssetUserDataMerge = false; // TODO >5.6 make private
-
+	bool bEnableAssetUserDataMerge = false;
+	
 	// Disabling the Table Materials parent material check will let the user use any material regardless of its parent when connecting a material from a table column to a material node.
 	// Warning! But it will not check if the table material channels connected to the Material node exist in the actual material used in the Instance, and will fail silently at runtime 
 	// when setting the value of those channels if they don't exist.
-	UE_DEPRECATED(5.6, "Only modify this value through the Customizable ObjectEditor")
 	UPROPERTY(EditAnywhere, Category = CompileOptions)
-	bool bDisableTableMaterialsParentCheck = false; // TODO >5.6 make private
+	bool bDisableTableMaterialsParentCheck = false;
 
 	// Options when compiling this customizable object (see EMutableCompileMeshType declaration for info)
-	UE_DEPRECATED(5.6, "Only modify this value through the Customizable ObjectEditor")
 	UPROPERTY(EditAnywhere, Category = CompileOptions)
-	EMutableCompileMeshType MeshCompileType = EMutableCompileMeshType::LocalAndChildren; // TODO >5.6 make private
+	EMutableCompileMeshType MeshCompileType = EMutableCompileMeshType::LocalAndChildren;
 
 	// Array of elements to use with compile option CompileType = WorkingSet
-	UE_DEPRECATED(5.6, "Only modify this value through the Customizable ObjectEditor")
 	UPROPERTY(EditAnywhere, Category = CompileOptions)
-	TArray<TSoftObjectPtr<UCustomizableObject>> WorkingSet; // TODO >5.6 make private
+	TArray<TSoftObjectPtr<UCustomizableObject>> WorkingSet;
 
-private:
 	// Editor graph
 	UPROPERTY()
 	TObjectPtr<UEdGraph> Source;
@@ -356,32 +344,16 @@ public:
 	/** Get the number of parameters available in instances of this object. */
 	UFUNCTION(BlueprintCallable, Category = CustomizableObject)
 	CUSTOMIZABLEOBJECT_API int32 GetParameterCount() const;
-
-	UE_DEPRECATED(5.6, "Use ContainsParameter instead")
-	UFUNCTION(BlueprintCallable, Category = CustomizableObject)
-	CUSTOMIZABLEOBJECT_API int32 FindParameter(const FString& Name) const;
-
+	
 	UFUNCTION(BlueprintCallable, Category = CustomizableObject)
 	CUSTOMIZABLEOBJECT_API bool ContainsParameter(const FString& ParameterName) const;
 	
-	UE_DEPRECATED(5.6, "Use ParameterName signature instead")
-	UFUNCTION(BlueprintCallable, Category = CustomizableObject)
-	CUSTOMIZABLEOBJECT_API EMutableParameterType GetParameterType(int32 ParamIndex) const;
-
 	/** Get the type of a parameter from its name. */
 	UFUNCTION(BlueprintCallable, Category = CustomizableObject)
 	CUSTOMIZABLEOBJECT_API EMutableParameterType GetParameterTypeByName(const FString& Name) const; // TODO >5.6 rename to GetParameterType (with redirector). Name -> ParameterName (create redirector)
 	
 	UFUNCTION(BlueprintCallable, Category = CustomizableObject)
 	CUSTOMIZABLEOBJECT_API const FString& GetParameterName(int32 ParamIndex) const;
-
-	UE_DEPRECATED(5.6, "Use ParameterName signature instead")
-	UFUNCTION(BlueprintCallable, Category = CustomizableObject)
-	CUSTOMIZABLEOBJECT_API int32 GetIntParameterNumOptions(int32 ParamIndex) const;
-
-	UE_DEPRECATED(5.6, "Use ParameterName signature instead")
-	UFUNCTION(BlueprintCallable, Category = CustomizableObject)
-	CUSTOMIZABLEOBJECT_API const FString& GetIntParameterAvailableOption(int32 ParamIndex, int32 K) const;
 	
 	/** If the given parameter is enum int parameter, return how many possible values an enum parameter has. Otherwise, return 0. */
 	UFUNCTION(BlueprintCallable, Category = CustomizableObject)
@@ -394,19 +366,14 @@ public:
 	/** Return true if the enum parameter contains this value a possible option. */
 	UFUNCTION(BlueprintCallable, Category = CustomizableObject)
 	CUSTOMIZABLEOBJECT_API bool ContainsEnumParameterValue(const FString& ParameterName, const FString Value) const;
-	
-	UE_DEPRECATED(5.6, "Use ParameterName signature instead")
-	CUSTOMIZABLEOBJECT_API int32 FindIntParameterValue( int32 ParamIndex, const FString& Value ) const;
-
-	UE_DEPRECATED(5.6, "Use ParameterName signature instead")
-	CUSTOMIZABLEOBJECT_API FString FindIntParameterValueName(int32 ParamIndex, int32 ParamValue) const;
 
-	UE_DEPRECATED(5.6, "Use GetComponentMeshReferenceSkeletalMesh instead")
-	CUSTOMIZABLEOBJECT_API USkeletalMesh* GetRefSkeletalMesh(int32 ObjectComponentIndex = 0) const;
+	// DEPRECATED: Use GetSkeletalMeshComponentReferenceSkeletalMesh instead.
+	UFUNCTION(BlueprintCallable, Category = CustomizableObject)
+	CUSTOMIZABLEOBJECT_API USkeletalMesh* GetComponentMeshReferenceSkeletalMesh(const FName& Name) const;
 
 	/** Given a Mesh Component name, return its reference Skeletal Mesh. */
 	UFUNCTION(BlueprintCallable, Category = CustomizableObject)
-	CUSTOMIZABLEOBJECT_API USkeletalMesh* GetComponentMeshReferenceSkeletalMesh(const FName& Name) const;
+	CUSTOMIZABLEOBJECT_API USkeletalMesh* GetSkeletalMeshComponentReferenceSkeletalMesh(const FName& ComponentName) const;
 	
 	/** Get the default value of a parameter of type Float.
 	  * @param ParameterName The name of the Float parameter to get the default value of.
@@ -419,10 +386,7 @@ public:
 	  * @return The default value of the provided parameter name. */
 	UFUNCTION(BlueprintCallable, Category = CustomizableObject)
 	CUSTOMIZABLEOBJECT_API int32 GetEnumParameterDefaultValue(const FString& ParameterName) const;
-
-	UE_DEPRECATED(5.6, "Use GetEnumParameterDefaultValue instead")
-	CUSTOMIZABLEOBJECT_API int32 GetIntParameterDefaultValue(UPARAM(DisplayName = "Parameter Name") const FString& InParameterName) const;
- 
+	
 	/** Get the default value of a parameter of type Bool.
 	  * @param ParameterName The name of the Bool parameter to get the default value of.
 	  * @return The default value of the provided parameter name. */
@@ -440,16 +404,6 @@ public:
 	  * @return The default value of the provided parameter name. */
 	UFUNCTION(BlueprintCallable, Category = CustomizableObject)
 	CUSTOMIZABLEOBJECT_API FTransform GetTransformParameterDefaultValue(const FString& ParameterName) const;
-	
-	UE_DEPRECATED(5.6, "Use the return FCustomizableObjectProjector signature instead")
-	CUSTOMIZABLEOBJECT_API void GetProjectorParameterDefaultValue(
-		UPARAM(DisplayName = "Parameter Name") const FString& InParameterName,
-		UPARAM(DisplayName = "Position") FVector3f& OutPos,
-		UPARAM(DisplayName = "Direction") FVector3f& OutDirection,
-		UPARAM(DisplayName = "Up") FVector3f& OutUp,
-		UPARAM(DisplayName = "Scale") FVector3f& OutScale,
-		UPARAM(DisplayName = "Angle") float& OutAngle,
-		UPARAM(DisplayName = "Type") ECustomizableObjectProjectorType& OutType) const;
 
 	/** Get the default value of a projector with the provided name
 	  * @param ParameterName The name of the parameter to get the default value of.
@@ -462,12 +416,9 @@ public:
 	CUSTOMIZABLEOBJECT_API UTexture* GetTextureParameterDefaultValue(const FString& InParameterName) const;
 
 	/** Get the default value of a parameter of type Skeletal Mesh.
-	  * @param InParameterName The name of the Projector parameter to get the default value of. */
+	  * @param ParameterName The name of the Projector parameter to get the default value of. */
 	CUSTOMIZABLEOBJECT_API USkeletalMesh* GetSkeletalMeshParameterDefaultValue(const FString& ParameterName) const;
 	
-	UE_DEPRECATED(5.6, "Use the ParameterName signature instead")
-	CUSTOMIZABLEOBJECT_API bool IsParameterMultidimensional(const int32& InParamIndex) const;
-	
 	/** Return true if the parameter at the index provided is multidimensional.
 	  * @param InParameterName The name of the parameter to check.
 	  * @return True if the parameter is multidimensional and false if it is not. */
@@ -502,24 +453,13 @@ public:
 	CUSTOMIZABLEOBJECT_API virtual bool IsCachedCookedPlatformDataLoaded(const ITargetPlatform* TargetPlatform) override;
 #endif
 
-	UE_DEPRECATED(5.6, "Use StateName functions instead")
-	CUSTOMIZABLEOBJECT_API int32 FindState(const FString& Name) const;
-
 	/** Return the number of object states that are defined in the CustomizableObject. */
 	UFUNCTION(BlueprintCallable, Category = CustomizableObject)
 	CUSTOMIZABLEOBJECT_API int32 GetStateCount() const;
 
-	/** Return the name of an object state from its index. */
-	UE_DEPRECATED(5.6, "Use StateName functions instead")
 	UFUNCTION(BlueprintCallable, Category = CustomizableObject)
 	CUSTOMIZABLEOBJECT_API FString GetStateName(int32 StateIndex) const;
-
-	UE_DEPRECATED(5.6, "Use StateName signature instead")
-	CUSTOMIZABLEOBJECT_API int32 GetStateParameterCount(int32 StateIndex) const;
-
-	UE_DEPRECATED(5.6, "Use StateName signature instead")
-	CUSTOMIZABLEOBJECT_API int32 GetStateParameterIndex(int32 StateIndex, int32 ParameterIndex) const;
-
+	
 	/** Return the number of parameters that are editable at runtime for a specific state. */
 	UFUNCTION(BlueprintCallable, Category = CustomizableObject)
 	CUSTOMIZABLEOBJECT_API int32 GetStateParameterCount(const FString& StateName) const;
@@ -527,9 +467,6 @@ public:
 	/** Return the name of one of the state's runtime parameters, by its index (from 0 to GetStateParameterCount - 1). */
 	UFUNCTION(BlueprintCallable, Category = CustomizableObject)
 	CUSTOMIZABLEOBJECT_API FString GetStateParameterName(const FString& StateName, int32 ParameterIndex) const;
-
-	UE_DEPRECATED(5.6, "Use StateName signature instead")
-	CUSTOMIZABLEOBJECT_API FString GetStateParameterName(int32 StateIndex, int32 ParameterIndex) const;
 	
 	/** Return the metadata associated to a parameter. */
 	UFUNCTION(BlueprintCallable, Category = CustomizableObject)
@@ -538,17 +475,11 @@ public:
 	/** Return the metadata associated to the given enum parameter value. */
 	UFUNCTION(BlueprintCallable, Category = CustomizableObject)
 	CUSTOMIZABLEOBJECT_API FMutableParamUIMetadata GetEnumParameterValueUIMetadata(const FString& ParamName, const FString& OptionName) const; // TODO GMT Rename "OptionName" -> "Value" using core redirects
-
-	UE_DEPRECATED(5.6, "Use GetEnumParameterValueUIMetadata instead")
-	CUSTOMIZABLEOBJECT_API FMutableParamUIMetadata GetIntParameterOptionUIMetadata(const FString& ParamName, const FString& OptionName) const;
 	
 	/** Returns the group type of the given integer parameter */
 	UFUNCTION(BlueprintCallable, Category = CustomizableObject)
 	CUSTOMIZABLEOBJECT_API ECustomizableObjectGroupType GetEnumParameterGroupType(const FString& ParamName) const;
 
-	UE_DEPRECATED(5.6, "Use GetEnumParameterGroupType instead")
-	CUSTOMIZABLEOBJECT_API ECustomizableObjectGroupType GetIntParameterGroupType(const FString& ParamName) const;
-
 	/** Return the metadata associated to a state. */
 	UFUNCTION(BlueprintCallable, Category = CustomizableObject)
 	CUSTOMIZABLEOBJECT_API FMutableStateUIMetadata GetStateUIMetadata(const FString& StateName) const;
@@ -557,9 +488,6 @@ public:
 	/** Return the DataTables used by the given parameter and its value (if any). */
 	UFUNCTION(BlueprintCallable, Category = CustomizableObject)
 	CUSTOMIZABLEOBJECT_API TArray<TSoftObjectPtr<UDataTable>> GetEnumParameterValueDataTable(const FString& ParamName, const FString& Value);
-
-	UE_DEPRECATED(5.6, "Use GetEnumParameterValueDataTable instead")
-	CUSTOMIZABLEOBJECT_API TArray<TSoftObjectPtr<UDataTable>> GetIntParameterOptionDataTable(const FString& ParamName, const FString& OptionName);
 #endif
 	
 private:
diff --git a/Engine/Plugins/Mutable/Source/CustomizableObject/Public/MuCO/CustomizableObjectExtension.h b/Engine/Plugins/Mutable/Source/CustomizableObject/Public/MuCO/CustomizableObjectExtension.h
--- a/Engine/Plugins/Mutable/Source/CustomizableObject/Public/MuCO/CustomizableObjectExtension.h
+++ b/Engine/Plugins/Mutable/Source/CustomizableObject/Public/MuCO/CustomizableObjectExtension.h
@@ -106,9 +106,6 @@ public:
 	 */
 	virtual void OnSkeletalMeshCreated(const TArray<FInputPinDataContainer>& InputPinData, FName ComponentName, USkeletalMesh* SkeletalMesh) const {}
 	
-	UE_DEPRECATED(5.6, "This overload of OnSkeletalMeshCreated is deprecated and will be removed. Use ComponentName instead.")
-	virtual void OnSkeletalMeshCreated(const TArray<FInputPinDataContainer>& InputPinData, int32 ObjectComponentIndex, USkeletalMesh* SkeletalMesh) const {}
-
 	/**
 	 * Note that the data registered here is completely independent of any Extension Data used in
 	 * the Customizable Object graph. Even though Extension Data and this Extension Instance Data
diff --git a/Engine/Plugins/Mutable/Source/CustomizableObject/Public/MuCO/CustomizableObjectInstance.h b/Engine/Plugins/Mutable/Source/CustomizableObject/Public/MuCO/CustomizableObjectInstance.h
--- a/Engine/Plugins/Mutable/Source/CustomizableObject/Public/MuCO/CustomizableObjectInstance.h
+++ b/Engine/Plugins/Mutable/Source/CustomizableObject/Public/MuCO/CustomizableObjectInstance.h
@@ -9,8 +9,6 @@
 
 #include "CustomizableObjectInstance.generated.h"
 
-#define UE_API CUSTOMIZABLEOBJECT_API
-
 namespace mu
 {
 	class FImage;
@@ -32,7 +30,7 @@ struct FGameplayTagContainer;
 struct FPropertyChangedEvent;
 struct FTexturePlatformData;
 struct FMutableModelImageProperties;
-struct FCustomizableObjectComponentIndex;
+
 
 /**
  * Represents what kind of saving procedure was performed to save the package
@@ -145,13 +143,6 @@ struct FBakingConfiguration
 	UPROPERTY(BlueprintReadWrite, Blueprintable, Category = CustomizableObjectInstanceBaker)
 	bool bGenerateConstantMaterialInstancesOnBake = false;
 	
-	/**
-	 * Flag that determines if we should override already exported files or if we should not. If we encounter files to override and we have not permission to override them then the baking operation will fail.
-	 */
-	UE_DEPRECATED(5.6, "The files will get automatically re-used instead")
-	UPROPERTY(meta = (DeprecatedProperty))
-	bool bAllowOverridingOfFiles_DEPRECATED = false;
-	
 	UPROPERTY(BlueprintReadWrite, Blueprintable, Category = "CustomizableObjectInstanceBaker|ResourceNameGeneration")
 	FString SkeletalMeshAssetPrefix = TEXT("SK_");
 
@@ -273,7 +264,7 @@ class UCustomizableObjectInstance : public UObject
 	friend FMutableUpdateCandidate;
 
 public:
-	UE_API UCustomizableObjectInstance();
+	CUSTOMIZABLEOBJECT_API UCustomizableObjectInstance();
 	
 	/** Broadcast when the Customizable Object Instance is updated. */
 	UPROPERTY(Transient, BlueprintAssignable, Category = CustomizableObjectInstance)
@@ -292,335 +283,331 @@ public:
 	
 	// UObject interface.
 #if WITH_EDITOR
-	UE_API virtual bool CanEditChange( const FProperty* InProperty ) const override;
-	UE_API virtual void PostTransacted(const FTransactionObjectEvent& TransactionEvent) override;
+	CUSTOMIZABLEOBJECT_API virtual bool CanEditChange( const FProperty* InProperty ) const override;
+	CUSTOMIZABLEOBJECT_API virtual void PostTransacted(const FTransactionObjectEvent& TransactionEvent) override;
 #endif //WITH_EDITOR
 
-	UE_API virtual void Serialize(FArchive& Ar) override;
-	UE_API virtual void PostLoad() override;
-	UE_API virtual void BeginDestroy() override;
-	UE_API virtual bool IsReadyForFinishDestroy() override;
-	UE_API virtual FString GetDesc() override;
-	UE_API virtual bool IsEditorOnly() const override;
-	UE_API virtual void PostInitProperties() override;
+	CUSTOMIZABLEOBJECT_API virtual void Serialize(FArchive& Ar) override;
+	CUSTOMIZABLEOBJECT_API virtual void PostLoad() override;
+	CUSTOMIZABLEOBJECT_API virtual void BeginDestroy() override;
+	CUSTOMIZABLEOBJECT_API virtual bool IsReadyForFinishDestroy() override;
+	CUSTOMIZABLEOBJECT_API virtual FString GetDesc() override;
+	CUSTOMIZABLEOBJECT_API virtual bool IsEditorOnly() const override;
+	CUSTOMIZABLEOBJECT_API virtual void PostInitProperties() override;
 
 	/** Set the CustomizableObject this instance will be generated from. 
 	  * It is usually not necessary to call this since instances are already generated from a CustomizableObject. */
 	UFUNCTION(BlueprintCallable, Category = CustomizableObjectInstance)
-	UE_API void SetObject(UCustomizableObject* InObject);
+	CUSTOMIZABLEOBJECT_API void SetObject(UCustomizableObject* InObject);
 
 	/** Get the CustomizableObject that this is an instance of. */
 	UFUNCTION(BlueprintCallable, Category = CustomizableObjectInstance)
-	UE_API UCustomizableObject* GetCustomizableObject() const;
+	CUSTOMIZABLEOBJECT_API UCustomizableObject* GetCustomizableObject() const;
 
 	/** Return true if the parameter relevancy will be updated when this instance is generated. */
 	UFUNCTION(BlueprintCallable, Category = CustomizableObjectInstance)
-	UE_API bool GetBuildParameterRelevancy() const;
+	CUSTOMIZABLEOBJECT_API bool GetBuildParameterRelevancy() const;
 
 	/** Set the flag that controls if parameter relevancy will be updated when this instance is generated. */
 	UFUNCTION(BlueprintCallable, Category = CustomizableObjectInstance)
-	UE_API void SetBuildParameterRelevancy(bool Value);
+	CUSTOMIZABLEOBJECT_API void SetBuildParameterRelevancy(bool Value);
 
 	/** Return the name of the current CustomizableObject state this is instance is set to. */
 	UFUNCTION(BlueprintCallable, Category = CustomizableObjectInstance)
-	UE_API FString GetCurrentState() const;
+	CUSTOMIZABLEOBJECT_API FString GetCurrentState() const;
 
 	/** Set the CustomizableObject state that this instance will be generated into. */
 	UFUNCTION(BlueprintCallable, Category = CustomizableObjectInstance)
-	UE_API void SetCurrentState(const FString& StateName);
+	CUSTOMIZABLEOBJECT_API void SetCurrentState(const FString& StateName);
 
-	// DEPRECATED: Use GetComponentMeshSkeletalMesh 
+	// DEPRECATED: Use GetSkeletalMeshComponentSkeletalMesh 
 	UFUNCTION(BlueprintCallable, Category = CustomizableObjectInstance)
-	UE_API USkeletalMesh* GetSkeletalMesh(int32 ComponentIndex = 0) const;
+	CUSTOMIZABLEOBJECT_API USkeletalMesh* GetSkeletalMesh(int32 ComponentIndex = 0) const;
+
+	// DEPRECATED: Use GetSkeletalMeshComponentSkeletalMesh 
+	UFUNCTION(BlueprintCallable, Category = CustomizableObjectInstance)
+	CUSTOMIZABLEOBJECT_API USkeletalMesh* GetComponentMeshSkeletalMesh(const FName& ComponentName) const;
 
 	/** Given a Mesh Component name, return its generated Skeletal Mesh. */
 	UFUNCTION(BlueprintCallable, Category = CustomizableObjectInstance)
-	UE_API USkeletalMesh* GetComponentMeshSkeletalMesh(const FName& ComponentName) const;
+	CUSTOMIZABLEOBJECT_API USkeletalMesh* GetSkeletalMeshComponentSkeletalMesh(const FName& ComponentName) const;
+	
+	// DEPRECATED: Use GetSkeletalMeshComponentOverrideMaterials 
+	CUSTOMIZABLEOBJECT_API const TArray<TObjectPtr<UMaterialInterface>>* GetOverrideMaterials(int32 ComponentIndex) const;
+
+	UFUNCTION(BlueprintCallable, Category = CustomizableObjectInstance)
+	CUSTOMIZABLEOBJECT_API TArray<UMaterialInterface*> GetSkeletalMeshComponentOverrideMaterials(const FName& ComponentName) const;
 	
 	/** Return true if a skeletal mesh has been generated for any component of this instance. */
 	UFUNCTION(BlueprintCallable, Category = CustomizableObjectInstance)
-	UE_API bool HasAnySkeletalMesh() const;
+	CUSTOMIZABLEOBJECT_API bool HasAnySkeletalMesh() const;
 	
 	/** Return true if the instance has any parameters. */
 	UFUNCTION(BlueprintCallable, Category = CustomizableObjectInstance)
-	UE_API bool HasAnyParameters() const;
+	CUSTOMIZABLEOBJECT_API bool HasAnyParameters() const;
 	
 	/** Set random values to the parameters. */
 	UFUNCTION(BlueprintCallable, Category = CustomizableObjectInstance)
-	UE_API void SetRandomValues();
+	CUSTOMIZABLEOBJECT_API void SetRandomValues();
 
 	/** Set random values to the parameters using a stream. */
 	UFUNCTION(BlueprintCallable, Category = CustomizableObjectInstance)
-	UE_API void SetRandomValuesFromStream(const FRandomStream& InStream);
+	CUSTOMIZABLEOBJECT_API void SetRandomValuesFromStream(const FRandomStream& InStream);
 
 	/** Sets a parameter to its default value */
 	UFUNCTION(BlueprintCallable, Category = CustomizableObjectInstance)
-	UE_API void SetDefaultValue(const FString& ParamName);
+	CUSTOMIZABLEOBJECT_API void SetDefaultValue(const FString& ParamName);
 
 	/** Set all parameters to their default value. */
 	UFUNCTION(BlueprintCallable, Category = CustomizableObjectInstance)
-	UE_API void SetDefaultValues();
+	CUSTOMIZABLEOBJECT_API void SetDefaultValues();
 	
 	/** Returns the AssetUserData that was gathered from all the constituent mesh parts during the last update. 
 	  * It requires that the CustomizableObject had the bEnableAssetUserDataMerge set to true during compilation. */
 	UFUNCTION(BlueprintCallable, Category = CustomizableObjectInstance)
-	UE_API TSet<UAssetUserData*> GetMergedAssetUserData(int32 ComponentIndex) const;
-
-	UE_DEPRECATED(5.6, "Use UCustomizableObject::IsLoading or UCustomizableObject::IsCompiling instead.")
-	UE_API bool CanUpdateInstance() const;
-
+	CUSTOMIZABLEOBJECT_API TSet<UAssetUserData*> GetMergedAssetUserData(int32 ComponentIndex) const;
+	
 	/** Generate the instance with the current parameters and update all the components Skeletal Meshes asynchronously. */
 	UFUNCTION(BlueprintCallable, Category = CustomizableObjectInstance)
-	UE_API void UpdateSkeletalMeshAsync(bool bIgnoreCloseDist = false, bool bForceHighPriority = false);
+	CUSTOMIZABLEOBJECT_API void UpdateSkeletalMeshAsync(bool bIgnoreCloseDist = false, bool bForceHighPriority = false);
 		
 	/** Generate the instance with the current parameters and update all the components Skeletal Meshes asynchronously.
 	  * Callback will be called once the update finishes, even if it fails. */
 	UFUNCTION(BlueprintCallable, Category = CustomizableObjectInstance)
-	UE_API void UpdateSkeletalMeshAsyncResult(FInstanceUpdateDelegate Callback, bool bIgnoreCloseDist = false, bool bForceHighPriority = false);
+	CUSTOMIZABLEOBJECT_API void UpdateSkeletalMeshAsyncResult(FInstanceUpdateDelegate Callback, bool bIgnoreCloseDist = false, bool bForceHighPriority = false);
 
-	UE_API void UpdateSkeletalMeshAsyncResult(FInstanceUpdateNativeDelegate Callback, bool bIgnoreCloseDist = false, bool bForceHighPriority = false);
+	CUSTOMIZABLEOBJECT_API void UpdateSkeletalMeshAsyncResult(FInstanceUpdateNativeDelegate Callback, bool bIgnoreCloseDist = false, bool bForceHighPriority = false);
 
 	/** Clones the instance creating a new identical transient instance. */
 	UFUNCTION(BlueprintCallable, Category = CustomizableObjectInstance)
-	UE_API UCustomizableObjectInstance* Clone();
+	CUSTOMIZABLEOBJECT_API UCustomizableObjectInstance* Clone();
 
 	/** Clones the instance creating a new identical static instance with the given Outer. */
 	UFUNCTION(BlueprintCallable, Category = CustomizableObjectInstance)
-	UE_API UCustomizableObjectInstance* CloneStatic(UObject* Outer);
+	CUSTOMIZABLEOBJECT_API UCustomizableObjectInstance* CloneStatic(UObject* Outer);
 
 	/** Copy parameters from the given Instance. */
-	UE_API void CopyParametersFromInstance(UCustomizableObjectInstance* Instance);
+	CUSTOMIZABLEOBJECT_API void CopyParametersFromInstance(UCustomizableObjectInstance* Instance);
 
 	/** Immediately destroy the Mutable Core Live Update Instance attached to this (if exists). */
-	UE_API void DestroyLiveUpdateInstance();
+	CUSTOMIZABLEOBJECT_API void DestroyLiveUpdateInstance();
 
 	/** Return true if changing the parameter would affect the Instance given its current generation. */
-	UE_API bool IsParameterRelevant(int32 ParameterIndex) const;
+	CUSTOMIZABLEOBJECT_API bool IsParameterRelevant(int32 ParameterIndex) const;
 
 	/** Return true if the given parameter has any effect in the current object state, and considering the current values of the other parameters. */
 	UFUNCTION(BlueprintCallable, Category = CustomizableObjectInstance)
-	UE_API bool IsParameterRelevant(const FString& ParamName) const;
+	CUSTOMIZABLEOBJECT_API bool IsParameterRelevant(const FString& ParamName) const;
 
 	/** Return true if the parameter has changed but the Instance has not yet been updated. */
 	UFUNCTION(BlueprintCallable, Category = CustomizableObjectInstance)
-	UE_API bool IsParameterDirty(const FString& ParamName, int32 RangeIndex = -1) const;
+	CUSTOMIZABLEOBJECT_API bool IsParameterDirty(const FString& ParamName, int32 RangeIndex = -1) const;
 	
 	/** For multidimensional parameters, return the number of dimensions that the given projector parameter supports. */
 	UFUNCTION(BlueprintCallable, Category = CustomizableObjectInstance)
-	UE_API int32 GetProjectorValueRange(const FString& ParamName) const;
+	CUSTOMIZABLEOBJECT_API int32 GetProjectorValueRange(const FString& ParamName) const;
 
 	/** For multidimensional parameters, return the number of dimensions that the given int parameter supports. */
 	UFUNCTION(BlueprintCallable, Category = CustomizableObjectInstance)
-	UE_API int32 GetIntValueRange(const FString& ParamName) const;
+	CUSTOMIZABLEOBJECT_API int32 GetIntValueRange(const FString& ParamName) const;
 
 	/** For multidimensional parameters, return the number of dimensions that the given float parameter supports. */
 	UFUNCTION(BlueprintCallable, Category = CustomizableObjectInstance)
-	UE_API int32 GetFloatValueRange(const FString& ParamName) const;
+	CUSTOMIZABLEOBJECT_API int32 GetFloatValueRange(const FString& ParamName) const;
 
 	/** For multidimensional parameters, return the number of dimensions that the given texture parameter supports. */
 	UFUNCTION(BlueprintCallable, Category = CustomizableObjectInstance)
-	UE_API int32 GetTextureValueRange(const FString& ParamName) const;
+	CUSTOMIZABLEOBJECT_API int32 GetTextureValueRange(const FString& ParamName) const;
 
-	/** Return the name of the option currently set in the given parameter. */
+	// DEPRECATED. Use GetEnumParameterSelectedOption instead.
 	UFUNCTION(BlueprintCallable, Category = CustomizableObjectInstance)
-	UE_API const FString& GetIntParameterSelectedOption(const FString& ParamName, int32 RangeIndex = -1) const;
+	CUSTOMIZABLEOBJECT_API const FString& GetIntParameterSelectedOption(const FString& ParamName, int32 RangeIndex = -1) const;
 
-	UE_DEPRECATED(5.6, "Use ParameterName signature instead")
-	UE_API void SetIntParameterSelectedOption(int32 IntParamIndex, const FString& SelectedOption, int32 RangeIndex = -1);
+	/** Return the name of the option currently set in the given parameter. */
+	UFUNCTION(BlueprintCallable, Category = CustomizableObjectInstance)
+	CUSTOMIZABLEOBJECT_API const FString& GetEnumParameterSelectedOption(const FString& ParamName, int32 RangeIndex = -1) const;
+	
+	// DEPRECATED. Use SetEnumParameterSelectedOption instead.
+	UFUNCTION(BlueprintCallable, Category = CustomizableObjectInstance)
+	CUSTOMIZABLEOBJECT_API void SetIntParameterSelectedOption(const FString& ParamName, const FString& SelectedOptionName, int32 RangeIndex = -1);
 
 	/** Set the currently selected option value for the given parameter, by parameter name and option name. */
 	UFUNCTION(BlueprintCallable, Category = CustomizableObjectInstance)
-	UE_API void SetIntParameterSelectedOption(const FString& ParamName, const FString& SelectedOptionName, int32 RangeIndex = -1);
-
+	CUSTOMIZABLEOBJECT_API void SetEnumParameterSelectedOption(const FString& ParamName, const FString& SelectedOptionName, int32 RangeIndex = -1);
+	
 	/** Gets the value of a float parameter with name "FloatParamName". */
 	UFUNCTION(BlueprintCallable, Category = CustomizableObjectInstance)
-	UE_API float GetFloatParameterSelectedOption(const FString& FloatParamName, int32 RangeIndex = -1) const;
+	CUSTOMIZABLEOBJECT_API float GetFloatParameterSelectedOption(const FString& FloatParamName, int32 RangeIndex = -1) const;
 
 	/** Sets the float value "FloatValue" of a float parameter with index "FloatParamIndex". */
 	UFUNCTION(BlueprintCallable, Category = CustomizableObjectInstance)
-	UE_API void SetFloatParameterSelectedOption(const FString& FloatParamName, float FloatValue, int32 RangeIndex = -1);
+	CUSTOMIZABLEOBJECT_API void SetFloatParameterSelectedOption(const FString& FloatParamName, float FloatValue, int32 RangeIndex = -1);
 	
 	/** Gets the value of a texture parameter with name "TextureParamName". */
 	UFUNCTION(BlueprintCallable, Category = CustomizableObjectInstance)
-	UE_API UTexture* GetTextureParameterSelectedOption(const FString& TextureParamName, int32 RangeIndex = -1) const;
+	CUSTOMIZABLEOBJECT_API UTexture* GetTextureParameterSelectedOption(const FString& TextureParamName, int32 RangeIndex = -1) const;
 
 	/** Sets the texture value "TextureValue" of a texture parameter with index "TextureParamIndex". */
 	UFUNCTION(BlueprintCallable, Category = CustomizableObjectInstance)
-	UE_API void SetTextureParameterSelectedOption(const FString& TextureParamName, UTexture* TextureValue, int32 RangeIndex = -1);
+	CUSTOMIZABLEOBJECT_API void SetTextureParameterSelectedOption(const FString& TextureParamName, UTexture* TextureValue, int32 RangeIndex = -1);
 
 	UFUNCTION(BlueprintCallable, Category = CustomizableObjectInstance)
-	UE_API USkeletalMesh* GetSkeletalMeshParameterSelectedOption(const FString& SkeletalMeshParamName, int32 RangeIndex = -1);
+	CUSTOMIZABLEOBJECT_API USkeletalMesh* GetSkeletalMeshParameterSelectedOption(const FString& SkeletalMeshParamName, int32 RangeIndex = -1);
 
 	UFUNCTION(BlueprintCallable, Category = CustomizableObjectInstance)
-	UE_API void SetSkeletalMeshParameterSelectedOption(const FString& SkeletalMeshParamName, USkeletalMesh* SkeletalMeshValue, int32 RangeIndex = -1);
+	CUSTOMIZABLEOBJECT_API void SetSkeletalMeshParameterSelectedOption(const FString& SkeletalMeshParamName, USkeletalMesh* SkeletalMeshValue, int32 RangeIndex = -1);
 	
 	/** Gets the value of a color parameter with name "ColorParamName". */
 	UFUNCTION(BlueprintCallable, Category = CustomizableObjectInstance)
-	UE_API FLinearColor GetColorParameterSelectedOption(const FString& ColorParamName) const;
+	CUSTOMIZABLEOBJECT_API FLinearColor GetColorParameterSelectedOption(const FString& ColorParamName) const;
 
 	/** Sets the color value "ColorValue" of a color parameter with index "ColorParamIndex". */
 	UFUNCTION(BlueprintCallable, Category = CustomizableObjectInstance)
-	UE_API void SetColorParameterSelectedOption(const FString& ColorParamName, const FLinearColor& ColorValue);
+	CUSTOMIZABLEOBJECT_API void SetColorParameterSelectedOption(const FString& ColorParamName, const FLinearColor& ColorValue);
 
 	/** Gets the bool value "BoolValue" of a bool parameter with name "BoolParamName". */
 	UFUNCTION(BlueprintCallable, Category = CustomizableObjectInstance)
-	UE_API bool GetBoolParameterSelectedOption(const FString& BoolParamName) const;
+	CUSTOMIZABLEOBJECT_API bool GetBoolParameterSelectedOption(const FString& BoolParamName) const;
 
 	/** Sets the bool value "BoolValue" of a bool parameter with name "BoolParamName". */
 	UFUNCTION(BlueprintCallable, Category = CustomizableObjectInstance)
-	UE_API void SetBoolParameterSelectedOption(const FString& BoolParamName, bool BoolValue);
+	CUSTOMIZABLEOBJECT_API void SetBoolParameterSelectedOption(const FString& BoolParamName, bool BoolValue);
 
 	/** Sets the vector value "VectorValue" of a vector parameter with index "VectorParamIndex". */
 	UFUNCTION(BlueprintCallable, Category = CustomizableObjectInstance)
-	UE_API void SetVectorParameterSelectedOption(const FString& VectorParamName, const FLinearColor& VectorValue);
+	CUSTOMIZABLEOBJECT_API void SetVectorParameterSelectedOption(const FString& VectorParamName, const FLinearColor& VectorValue);
 
 	/** Gets the value of a transform parameter with name "TransformParamName". */
 	UFUNCTION(BlueprintCallable, Category = CustomizableObjectInstance)
-	UE_API FTransform GetTransformParameterSelectedOption(const FString& TransformParamName) const;
+	CUSTOMIZABLEOBJECT_API FTransform GetTransformParameterSelectedOption(const FString& TransformParamName) const;
 
 	/** Sets the transform value "TransformValue" of a transform parameter with name "TransformParamName". */
 	UFUNCTION(BlueprintCallable, Category = CustomizableObjectInstance)
-	UE_API void SetTransformParameterSelectedOption(const FString& TransformParamName, const FTransform& TransformValue);
+	CUSTOMIZABLEOBJECT_API void SetTransformParameterSelectedOption(const FString& TransformParamName, const FTransform& TransformValue);
 	
 	/** Sets the projector values of a projector parameter with index "ProjectorParamIndex". */
 	UFUNCTION(BlueprintCallable, Category = CustomizableObjectInstance)
-	UE_API void SetProjectorValue(const FString& ProjectorParamName,
+	CUSTOMIZABLEOBJECT_API void SetProjectorValue(const FString& ProjectorParamName,
 		const FVector& OutPos, const FVector& OutDirection, const FVector& OutUp, const FVector& OutScale,
 		float OutAngle,
 		int32 RangeIndex = -1);
 
 	/** Set only the projector position keeping the rest of values. */
 	UFUNCTION(BlueprintCallable, Category = CustomizableObjectInstance)
-	UE_API void SetProjectorPosition(const FString& ProjectorParamName, const FVector& Pos, int32 RangeIndex = -1);
+	CUSTOMIZABLEOBJECT_API void SetProjectorPosition(const FString& ProjectorParamName, const FVector& Pos, int32 RangeIndex = -1);
 	
 	/** Set only the projector direction vector keeping the rest of values. */
 	UFUNCTION(BlueprintCallable, Category = CustomizableObjectInstance)
-	UE_API void SetProjectorDirection(const FString& ProjectorParamName, const FVector& Direction, int32 RangeIndex = -1);
+	CUSTOMIZABLEOBJECT_API void SetProjectorDirection(const FString& ProjectorParamName, const FVector& Direction, int32 RangeIndex = -1);
 	
 	/** Set only the projector up vector keeping the rest of values. */
 	UFUNCTION(BlueprintCallable, Category = CustomizableObjectInstance)
-	UE_API void SetProjectorUp(const FString& ProjectorParamName, const FVector& Up, int32 RangeIndex = -1);
+	CUSTOMIZABLEOBJECT_API void SetProjectorUp(const FString& ProjectorParamName, const FVector& Up, int32 RangeIndex = -1);
 
 	/** Set only the projector scale keeping the rest of values. */
 	UFUNCTION(BlueprintCallable, Category = CustomizableObjectInstance)
-	UE_API void SetProjectorScale(const FString& ProjectorParamName, const FVector& Scale, int32 RangeIndex = -1);
+	CUSTOMIZABLEOBJECT_API void SetProjectorScale(const FString& ProjectorParamName, const FVector& Scale, int32 RangeIndex = -1);
 
 	/** Set only the cylindrical projector angle keeping the rest of values. */
 	UFUNCTION(BlueprintCallable, Category = CustomizableObjectInstance)
-	UE_API void SetProjectorAngle(const FString& ProjectorParamName, float Angle, int32 RangeIndex = -1);
+	CUSTOMIZABLEOBJECT_API void SetProjectorAngle(const FString& ProjectorParamName, float Angle, int32 RangeIndex = -1);
 	
 	/** Get the projector values of a projector parameter with index "ProjectorParamIndex". */
 	UFUNCTION(BlueprintCallable, Category = CustomizableObjectInstance)
-	UE_API void GetProjectorValue(const FString& ProjectorParamName,
+	CUSTOMIZABLEOBJECT_API void GetProjectorValue(const FString& ProjectorParamName,
 		FVector& OutPos, FVector& OutDirection, FVector& OutUp, FVector& OutScale,
 		float& OutAngle, ECustomizableObjectProjectorType& OutType,
 		int32 RangeIndex = -1) const;
 
-	UE_API void GetProjectorValueF(const FString& ProjectorParamName,
+	CUSTOMIZABLEOBJECT_API void GetProjectorValueF(const FString& ProjectorParamName,
 		FVector3f& Pos, FVector3f& Direction, FVector3f& Up, FVector3f& Scale,
 		float& Angle, ECustomizableObjectProjectorType& Type,
 		int32 RangeIndex = -1) const;
 
 	/** Get the current projector position for the parameter with the given name. */
 	UFUNCTION(BlueprintCallable, Category = CustomizableObjectInstance)
-	UE_API FVector GetProjectorPosition(const FString & ParamName, int32 RangeIndex = -1) const;
+	CUSTOMIZABLEOBJECT_API FVector GetProjectorPosition(const FString & ParamName, int32 RangeIndex = -1) const;
 
 	/** Get the current projector direction vector for the parameter with the given name. */
 	UFUNCTION(BlueprintCallable, Category = CustomizableObjectInstance)
-	UE_API FVector GetProjectorDirection(const FString & ParamName, int32 RangeIndex = -1) const;
+	CUSTOMIZABLEOBJECT_API FVector GetProjectorDirection(const FString & ParamName, int32 RangeIndex = -1) const;
 
 	/** Get the current projector up vector for the parameter with the given name. */
 	UFUNCTION(BlueprintCallable, Category = CustomizableObjectInstance)
-	UE_API FVector GetProjectorUp(const FString & ParamName, int32 RangeIndex = -1) const;
+	CUSTOMIZABLEOBJECT_API FVector GetProjectorUp(const FString & ParamName, int32 RangeIndex = -1) const;
 
 	/** Get the current projector scale for the parameter with the given name. */
 	UFUNCTION(BlueprintCallable, Category = CustomizableObjectInstance)
-	UE_API FVector GetProjectorScale(const FString & ParamName, int32 RangeIndex = -1) const;
+	CUSTOMIZABLEOBJECT_API FVector GetProjectorScale(const FString & ParamName, int32 RangeIndex = -1) const;
 
 	/** Get the current cylindrical projector angle for the parameter with the given name. */
 	UFUNCTION(BlueprintCallable, Category = CustomizableObjectInstance)
-	UE_API float GetProjectorAngle(const FString& ParamName, int32 RangeIndex = -1) const;
+	CUSTOMIZABLEOBJECT_API float GetProjectorAngle(const FString& ParamName, int32 RangeIndex = -1) const;
 
 	/** Get the current projector type for the parameter with the given name. */
 	UFUNCTION(BlueprintCallable, Category = CustomizableObjectInstance)
-	UE_API ECustomizableObjectProjectorType GetProjectorParameterType(const FString& ParamName, int32 RangeIndex = -1) const;
+	CUSTOMIZABLEOBJECT_API ECustomizableObjectProjectorType GetProjectorParameterType(const FString& ParamName, int32 RangeIndex = -1) const;
 
 	/** Get the current projector for the parameter with the given name. */
-	UE_API FCustomizableObjectProjector GetProjector(const FString& ParamName, int32 RangeIndex) const;
-	
-	/** Return true if the Int Parameter exists. */
+	CUSTOMIZABLEOBJECT_API FCustomizableObjectProjector GetProjector(const FString& ParamName, int32 RangeIndex) const;
+
+	// DEPRECATED. Use ContainsEnumParameter instead.
 	UFUNCTION(BlueprintCallable, Category = CustomizableObjectInstance)
-	UE_API bool ContainsIntParameter(const FString& ParameterName) const;
+	CUSTOMIZABLEOBJECT_API bool ContainsIntParameter(const FString& ParameterName) const;
 
+	/** Return true if the Int Parameter exists. */
+	UFUNCTION(BlueprintCallable, Category = CustomizableObjectInstance)
+	CUSTOMIZABLEOBJECT_API bool ContainsEnumParameter(const FString& ParameterName) const;
+	
 	/** Return true if the Float Parameter exists. */
 	UFUNCTION(BlueprintCallable, Category = CustomizableObjectInstance)
-	UE_API bool ContainsFloatParameter(const FString& ParameterName) const;
+	CUSTOMIZABLEOBJECT_API bool ContainsFloatParameter(const FString& ParameterName) const;
 
 	/** Return true if the Bool Parameter exists. */
 	UFUNCTION(BlueprintCallable, Category = CustomizableObjectInstance)
-	UE_API bool ContainsBoolParameter(const FString& ParameterName) const;
+	CUSTOMIZABLEOBJECT_API bool ContainsBoolParameter(const FString& ParameterName) const;
 
 	/** Return true if the Vector Parameter exists. */
 	UFUNCTION(BlueprintCallable, Category = CustomizableObjectInstance)
-	UE_API bool ContainsVectorParameter(const FString& ParameterName) const;
+	CUSTOMIZABLEOBJECT_API bool ContainsVectorParameter(const FString& ParameterName) const;
 
 	/** Return true if the Projector Parameter exists. */
 	UFUNCTION(BlueprintCallable, Category = CustomizableObjectInstance)
-	UE_API bool ContainsProjectorParameter(const FString& ParameterName) const;
+	CUSTOMIZABLEOBJECT_API bool ContainsProjectorParameter(const FString& ParameterName) const;
 
 	/** Return true if the Transform Parameter exists. */
 	UFUNCTION(BlueprintCallable, Category = CustomizableObjectInstance)
-	UE_API bool ContainsTransformParameter(const FString& ParameterName) const;
-
-	UE_DEPRECATED(5.6, "Use ContainsIntParameter instead")
-	UFUNCTION(BlueprintCallable, Category = CustomizableObjectInstance)
-	UE_API int32 FindIntParameterNameIndex(const FString& ParamName) const;
-
-	UE_DEPRECATED(5.6, "Use ContainsFloatParameter instead")
-	UFUNCTION(BlueprintCallable, Category = CustomizableObjectInstance)
-	UE_API int32 FindFloatParameterNameIndex(const FString& ParamName) const;
-
-	UE_DEPRECATED(5.6, "Use ContainsBoolParameter instead")
-	UFUNCTION(BlueprintCallable, Category = CustomizableObjectInstance)
-	UE_API int32 FindBoolParameterNameIndex(const FString& ParamName) const;
-
-	UE_DEPRECATED(5.6, "Use ContainsVectorParameter instead")
-	UFUNCTION(BlueprintCallable, Category = CustomizableObjectInstance)
-	UE_API int32 FindVectorParameterNameIndex(const FString& ParamName) const;
-
-	UE_DEPRECATED(5.6, "Use ContainsProjectorParameter instead")
-	UFUNCTION(BlueprintCallable, Category = CustomizableObjectInstance)
-	UE_API int32 FindProjectorParameterNameIndex(const FString& ParamName) const;
+	CUSTOMIZABLEOBJECT_API bool ContainsTransformParameter(const FString& ParameterName) const;
 	
 	/** Increases the range of values of the integer with ParamName and returns the index of the new integer value, -1 otherwise.
 	  * The added value is initialized with the first integer option and is the last one of the range. */
 	UFUNCTION(BlueprintCallable, Category = CustomizableObjectInstance)
-	UE_API int32 AddValueToIntRange(const FString& ParamName);
+	CUSTOMIZABLEOBJECT_API int32 AddValueToIntRange(const FString& ParamName);
 
 	/** Increases the range of values of the float with ParamName, returns the index of the new float value, -1 otherwise.
 	  * The added value is initialized with 0.5f and is the last one of the range. */
 	UFUNCTION(BlueprintCallable, Category = CustomizableObjectInstance)
-	UE_API int32 AddValueToFloatRange(const FString& ParamName);
+	CUSTOMIZABLEOBJECT_API int32 AddValueToFloatRange(const FString& ParamName);
 
 	/** Increases the range of values of the projector with ParamName, returns the index of the new projector value, -1 otherwise.
 	  * The added value is initialized with the default projector as set up in the editor and is the last one of the range. */
 	UFUNCTION(BlueprintCallable, Category = CustomizableObjectInstance)
-	UE_API int32 AddValueToProjectorRange(const FString& ParamName);
+	CUSTOMIZABLEOBJECT_API int32 AddValueToProjectorRange(const FString& ParamName);
 
 	/** Remove the FRangeIndex element of the integer range of values from the parameter ParamName. If RangeValue is -1 removes the last of the integer range of values.
 		Returns the index of the last valid integer, -1 if no values left. */
 	UFUNCTION(BlueprintCallable, Category = CustomizableObjectInstance)
-	UE_API int32 RemoveValueFromIntRange(const FString& ParamName, int32 RangeIndex = -1);
+	CUSTOMIZABLEOBJECT_API int32 RemoveValueFromIntRange(const FString& ParamName, int32 RangeIndex = -1);
 
 	/** Remove the FRangeIndex element of the float range of values from the parameter ParamName. If RangeValue is -1 removes the last of the float range of values.
 		Returns the index of the last valid float, -1 if no values left. */
 	UFUNCTION(BlueprintCallable, Category = CustomizableObjectInstance)
-	UE_API int32 RemoveValueFromFloatRange(const FString& ParamName, int32 RangeIndex = -1);
+	CUSTOMIZABLEOBJECT_API int32 RemoveValueFromFloatRange(const FString& ParamName, int32 RangeIndex = -1);
 
 	/** Remove the FRangeIndex element of the projector range of values from the parameter ParamName. If RangeValue is -1 removes the last of the projector range of values.
 		Returns the index of the last valid projector, -1 if no values left. */
 	UFUNCTION(BlueprintCallable, Category = CustomizableObjectInstance)
-	UE_API int32 RemoveValueFromProjectorRange(const FString& ParamName, int32 RangeIndex = -1);
+	CUSTOMIZABLEOBJECT_API int32 RemoveValueFromProjectorRange(const FString& ParamName, int32 RangeIndex = -1);
 
 	// ------------------------------------------------------------
 	// Multilayer Projectors
@@ -630,106 +617,103 @@ public:
 
 	/** See FMultilayerProjector::NumLayers. */
 	UFUNCTION(BlueprintCallable, Category = CustomizableObjectInstance)
-	UE_API int32 MultilayerProjectorNumLayers(const FName& ProjectorParamName) const;
+	CUSTOMIZABLEOBJECT_API int32 MultilayerProjectorNumLayers(const FName& ProjectorParamName) const;
 
 	/** See FMultilayerProjector::CreateLayer. */
 	UFUNCTION(BlueprintCallable, Category = CustomizableObjectInstance)
-	UE_API void MultilayerProjectorCreateLayer(const FName& ProjectorParamName, int32 Index);
+	CUSTOMIZABLEOBJECT_API void MultilayerProjectorCreateLayer(const FName& ProjectorParamName, int32 Index);
 
 	/** See FMultilayerProjector::RemoveLayerAt. */
 	UFUNCTION(BlueprintCallable, Category = CustomizableObjectInstance)
-	UE_API void MultilayerProjectorRemoveLayerAt(const FName& ProjectorParamName, int32 Index);
+	CUSTOMIZABLEOBJECT_API void MultilayerProjectorRemoveLayerAt(const FName& ProjectorParamName, int32 Index);
 
 	/** See FMultilayerProjector::GetLayer. */
 	UFUNCTION(BlueprintCallable, Category = CustomizableObjectInstance)
-	UE_API FMultilayerProjectorLayer MultilayerProjectorGetLayer(const FName& ProjectorParamName, int32 Index) const;
+	CUSTOMIZABLEOBJECT_API FMultilayerProjectorLayer MultilayerProjectorGetLayer(const FName& ProjectorParamName, int32 Index) const;
 
 	/** See FMultilayerProjector::UpdateLayer. */
 	UFUNCTION(BlueprintCallable, Category = CustomizableObjectInstance)
-	UE_API void MultilayerProjectorUpdateLayer(const FName& ProjectorParamName, int32 Index, const FMultilayerProjectorLayer& Layer);
+	CUSTOMIZABLEOBJECT_API void MultilayerProjectorUpdateLayer(const FName& ProjectorParamName, int32 Index, const FMultilayerProjectorLayer& Layer);
 	
 	// ------------------------------------------------------------
 
 	/** Return the list of names of components generated for this instance. 
 	* This only has values when the instance has been completely generated. */
 	UFUNCTION(BlueprintCallable, Category = CustomizableObjectInstance)
-	UE_API TArray<FName> GetComponentNames() const;
+	CUSTOMIZABLEOBJECT_API TArray<FName> GetComponentNames() const;
 
 	// ------------------------------------------------------------
 
 	/** Returns the animation BP for the parameter component and slot, gathered from all the meshes that compose this instance. */
 	UFUNCTION(BlueprintCallable, Category = CustomizableObjectInstance)
-	UE_API TSubclassOf<UAnimInstance> GetAnimBP(FName ComponentName, const FName& Slot) const;
+	CUSTOMIZABLEOBJECT_API TSubclassOf<UAnimInstance> GetAnimBP(FName ComponentName, const FName& Slot) const;
 
 	/** Return the list of tags for this instance. */
 	UFUNCTION(BlueprintCallable, Category = CustomizableObjectInstance)
-	UE_API const FGameplayTagContainer& GetAnimationGameplayTags() const;
+	CUSTOMIZABLEOBJECT_API const FGameplayTagContainer& GetAnimationGameplayTags() const;
 	
 	/** Execute a delegate for each animation instance involved in this customizable object instance. */
 	UFUNCTION(BlueprintCallable, Category = CustomizableObjectInstance)
-	UE_API void ForEachComponentAnimInstance(FName ComponentName, FEachComponentAnimInstanceClassDelegate Delegate) const;
+	CUSTOMIZABLEOBJECT_API void ForEachComponentAnimInstance(FName ComponentName, FEachComponentAnimInstanceClassDelegate Delegate) const;
 
-	UE_API void ForEachComponentAnimInstance(FName ComponentName, FEachComponentAnimInstanceClassNativeDelegate Delegate) const;
+	CUSTOMIZABLEOBJECT_API void ForEachComponentAnimInstance(FName ComponentName, FEachComponentAnimInstanceClassNativeDelegate Delegate) const;
 
 	/** DEPRECATED: Use ForEachComponentAnimInstance. */
 	UFUNCTION(BlueprintCallable, Category = CustomizableObjectInstance)
-	UE_API void ForEachAnimInstance(int32 ObjectComponentIndex, FEachComponentAnimInstanceClassDelegate Delegate) const;
+	CUSTOMIZABLEOBJECT_API void ForEachAnimInstance(int32 ObjectComponentIndex, FEachComponentAnimInstanceClassDelegate Delegate) const;
 
 	/** DEPRECATED: Use ForEachComponentAnimInstance. */
-	UE_API void ForEachAnimInstance(int32 ObjectComponentIndex, FEachComponentAnimInstanceClassNativeDelegate Delegate) const;
+	CUSTOMIZABLEOBJECT_API void ForEachAnimInstance(int32 ObjectComponentIndex, FEachComponentAnimInstanceClassNativeDelegate Delegate) const;
 
 	/** Check if the given UAnimInstance class requires to be fixed up. */
-	UE_API bool AnimInstanceNeedsFixup(TSubclassOf<UAnimInstance> AnimInstance) const;
+	CUSTOMIZABLEOBJECT_API bool AnimInstanceNeedsFixup(TSubclassOf<UAnimInstance> AnimInstance) const;
 
 	/** Fix the given UAnimInstance instance. */
-	UE_API void AnimInstanceFixup(UAnimInstance* AnimInstance) const;
+	CUSTOMIZABLEOBJECT_API void AnimInstanceFixup(UAnimInstance* AnimInstance) const;
 
 	UFUNCTION(BlueprintCallable, Category = CustomizableObjectInstance)
-	UE_API FInstancedStruct GetExtensionInstanceData(const UCustomizableObjectExtension* Extension) const;
+	CUSTOMIZABLEOBJECT_API FInstancedStruct GetExtensionInstanceData(const UCustomizableObjectExtension* Extension) const;
 	
 	/** See FCustomizableObjectInstanceDescriptor::SaveDescriptor. */
-	UE_API void SaveDescriptor(FArchive &CustomizableObjectDescriptor, bool bUseCompactDescriptor);
+	CUSTOMIZABLEOBJECT_API void SaveDescriptor(FArchive &CustomizableObjectDescriptor, bool bUseCompactDescriptor);
 
 	/** See FCustomizableObjectInstanceDescriptor::LoadDescriptor. */
-	UE_API void LoadDescriptor(FArchive &CustomizableObjectDescriptor);
+	CUSTOMIZABLEOBJECT_API void LoadDescriptor(FArchive &CustomizableObjectDescriptor);
 
 	/** Enable physics asset replacement so that generated skeletal meshes have the merged physics assets of their skeletal mesh parts and reference mesh. */
 	UFUNCTION(BlueprintCallable, Category = CustomizableObjectInstance)
-	UE_API void SetReplacePhysicsAssets(bool bReplaceEnabled);
+	CUSTOMIZABLEOBJECT_API void SetReplacePhysicsAssets(bool bReplaceEnabled);
 
 	/** Enables the reuse of all possible textures when the instance is updated without any changes in geometry or state (the first update after creation doesn't reuse any)
 	  * It will only work if the textures aren't compressed, so set the instance to a Mutable state with texture compression disabled
 	  * WARNING! If texture reuse is enabled, do NOT keep external references to the textures of the instance. The instance owns the textures. */
-	UE_API void SetReuseInstanceTextures(bool bTextureReuseEnabled);
+	CUSTOMIZABLEOBJECT_API void SetReuseInstanceTextures(bool bTextureReuseEnabled);
 	
 	/** If enabled, low-priority textures will generate resident mipmaps too. */
-	UE_API void SetForceGenerateResidentMips(bool bForceGenerateResidentMips);
-
-	UFUNCTION(BlueprintCallable, Category = CustomizableObjectInstance)
-	CUSTOMIZABLEOBJECT_API TArray<UMaterialInterface*> GetSkeletalMeshComponentOverrideMaterials(const FName& ComponentName) const;
+	CUSTOMIZABLEOBJECT_API void SetForceGenerateResidentMips(bool bForceGenerateResidentMips);
 	
 	// The following methods should only be used in an LOD management class
-	UE_API void SetIsBeingUsedByComponentInPlay(bool bIsUsedByComponent);
-	UE_API bool GetIsBeingUsedByComponentInPlay() const;
-	UE_API void SetIsDiscardedBecauseOfTooManyInstances(bool bIsDiscarded);
-	UE_API bool GetIsDiscardedBecauseOfTooManyInstances() const;
-	UE_API void SetIsPlayerOrNearIt(bool NewValue);
-	UE_API float GetMinSquareDistToPlayer() const;
-	UE_API void SetMinSquareDistToPlayer(float NewValue);
+	CUSTOMIZABLEOBJECT_API void SetIsBeingUsedByComponentInPlay(bool bIsUsedByComponent);
+	CUSTOMIZABLEOBJECT_API bool GetIsBeingUsedByComponentInPlay() const;
+	CUSTOMIZABLEOBJECT_API void SetIsDiscardedBecauseOfTooManyInstances(bool bIsDiscarded);
+	CUSTOMIZABLEOBJECT_API bool GetIsDiscardedBecauseOfTooManyInstances() const;
+	CUSTOMIZABLEOBJECT_API void SetIsPlayerOrNearIt(bool NewValue);
+	CUSTOMIZABLEOBJECT_API float GetMinSquareDistToPlayer() const;
+	CUSTOMIZABLEOBJECT_API void SetMinSquareDistToPlayer(float NewValue);
 
 	/** Return the total number of components that can be generated with the CustomizableObject (unrelated to this instance parameters). 
 	* DEPRECATED: Get the CO with GetCustomizableObect and get the components there with GetComponentCount.
 	* If the actual components of this instance are required, use GetComponents to get a list of names.
 	*/
 	UFUNCTION(BlueprintCallable, Category = CustomizableObjectInstance)
-	UE_API int32 GetNumComponents() const;
+	CUSTOMIZABLEOBJECT_API int32 GetNumComponents() const;
 	
 	/** Sets an array of LODs to generate per component. Mutable will generate those plus the currently generated LODs (if any).
 	  * Requires mutable.EnableOnlyGenerateRequestedLODs and CurrentInstanceLODManagement->IsOnlyGenerateRequestedLODLevelsEnabled() to be true.
 	  * @param InMinLODs - MinLOD to generate per Component.
 	  * @param InFirstRequestedLOD - Requested LODs per component to generate. Key is the Component name. Value represents the first LOD to generate. Zero means all LODs.
 	  * @param InOutRequestedUpdates - Map from Instance to Update data that stores a request for the Instance to be updated, which will be either processed or discarded by priority (to be rerequested the next tick) */
-	UE_API void SetRequestedLODs(const TMap<FName, uint8>& InMinLODs, const TMap<FName, uint8>& InFirstRequestedLOD, FMutableInstanceUpdateMap& InOutRequestedUpdates);
+	CUSTOMIZABLEOBJECT_API void SetRequestedLODs(const TMap<FName, uint8>& InMinLODs, const TMap<FName, uint8>& InFirstRequestedLOD, FMutableInstanceUpdateMap& InOutRequestedUpdates);
 
 #if WITH_EDITOR
 	/**
@@ -737,10 +721,10 @@ public:
 	 * @param InBakingConfiguration The configuration to be using for the baking operation
 	 */
 	UFUNCTION(BlueprintCallable, Category = CustomizableObjectInstanceBaking)
-	UE_API void Bake(const FBakingConfiguration& InBakingConfiguration);
+	CUSTOMIZABLEOBJECT_API void Bake(const FBakingConfiguration& InBakingConfiguration);
 #endif
 	
-	UE_API UCustomizableInstancePrivate* GetPrivate() const;
+	CUSTOMIZABLEOBJECT_API UCustomizableInstancePrivate* GetPrivate() const;
 
 private:
 	UPROPERTY()
@@ -774,4 +758,3 @@ private:
 	bool bBuildParameterRelevancy_DEPRECATED = false;
 };
 
-#undef UE_API
diff --git a/Engine/Plugins/Mutable/Source/CustomizableObject/Public/MuCO/CustomizableObjectSystem.h b/Engine/Plugins/Mutable/Source/CustomizableObject/Public/MuCO/CustomizableObjectSystem.h
--- a/Engine/Plugins/Mutable/Source/CustomizableObject/Public/MuCO/CustomizableObjectSystem.h
+++ b/Engine/Plugins/Mutable/Source/CustomizableObject/Public/MuCO/CustomizableObjectSystem.h
@@ -217,11 +217,7 @@ public:
 
 	UE_API bool IsOnlyGenerateRequestedLODsEnabled() const;
 	UE_API void SetOnlyGenerateRequestedLODsEnabled(bool bIsEnabled);
-
-#if WITH_EDITOR
-	UE_API void SetImagePixelFormatOverride(const mu::FImageOperator::FImagePixelFormatFunc&);
-#endif
-
+	
 	// Show a warning on-screen and via a notification (if in Editor) and log an error when a CustomizableObject is
 	// being used and it's not compiled.  Callers can add additional information to the error log.
 	UE_API void AddUncompiledCOWarning(const UCustomizableObject& InObject, FString const* OptionalLogInfo = nullptr);
diff --git a/Engine/Plugins/Mutable/Source/CustomizableObjectEditor/Internal/MuCOE/CustomizableObjectInstanceBakingUtils.h b/Engine/Plugins/Mutable/Source/CustomizableObjectEditor/Internal/MuCOE/CustomizableObjectInstanceBakingUtils.h
--- a/Engine/Plugins/Mutable/Source/CustomizableObjectEditor/Internal/MuCOE/CustomizableObjectInstanceBakingUtils.h
+++ b/Engine/Plugins/Mutable/Source/CustomizableObjectEditor/Internal/MuCOE/CustomizableObjectInstanceBakingUtils.h
@@ -12,17 +12,6 @@
 CUSTOMIZABLEOBJECTEDITOR_API void ScheduleInstanceUpdateForBaking(UCustomizableObjectInstance& InInstance, FInstanceUpdateNativeDelegate& InInstanceUpdateDelegate);
 
 
-UE_DEPRECATED(5.6, "Use the function with the same name that uses an FBakingConfiguration as input instead.")
-CUSTOMIZABLEOBJECTEDITOR_API bool BakeCustomizableObjectInstance (UCustomizableObjectInstance& InInstance,
- const FString& FileName, 
- const FString& AssetPath,
- const bool bExportAllResources,
- const bool bGenerateConstantMaterialInstances, 
- bool bHasPermissionToOverride,
- bool bIsUnattendedExecution, 
- TArray<TPair<EPackageSaveResolutionType,UPackage*>>& OutSavedPackages);
-
-
 /**
  * Serializes onto disk the resources used by the targeted Customizable Object Instance. The operation can be configured in the FInstanceBakingSettings settings object.
  * @param InInstance The mutable COI instance whose resources we want to bake onto disk
diff --git a/Engine/Plugins/Mutable/Source/CustomizableObjectEditor/Private/MuCOE/CustomizableObjectInstanceBakingUtils.cpp b/Engine/Plugins/Mutable/Source/CustomizableObjectEditor/Private/MuCOE/CustomizableObjectInstanceBakingUtils.cpp
--- a/Engine/Plugins/Mutable/Source/CustomizableObjectEditor/Private/MuCOE/CustomizableObjectInstanceBakingUtils.cpp
+++ b/Engine/Plugins/Mutable/Source/CustomizableObjectEditor/Private/MuCOE/CustomizableObjectInstanceBakingUtils.cpp
@@ -343,31 +343,6 @@ void AddUniqueToSavePackage(TPair<UPackage*, EPackageSaveResolutionType> NewObje
 }
 
 
-bool BakeCustomizableObjectInstance(UCustomizableObjectInstance& InInstance, const FString& FileName, const FString& AssetPath,
-	const bool bExportAllResources, const bool bGenerateConstantMaterialInstances, bool bHasPermissionToOverride, bool bIsUnattendedExecution,
-	TArray<TPair<EPackageSaveResolutionType, UPackage*>>& OutSavedPackages)
-{
-	FBakingConfiguration BakingConfiguration;
-	BakingConfiguration.OutputPath = AssetPath;
-	BakingConfiguration.OutputFilesBaseName = FileName;
-	BakingConfiguration.bExportAllResourcesOnBake = bExportAllResources;
-	BakingConfiguration.bGenerateConstantMaterialInstancesOnBake = bGenerateConstantMaterialInstances;
-
-	TMap<UPackage*, EPackageSaveResolutionType> SavedPackages;
-	bool bWasOperationSuccessful = BakeCustomizableObjectInstance(InInstance, BakingConfiguration, bIsUnattendedExecution, SavedPackages);
-
-	// Generate the array this method expects as an output based on the "SavedPackages" map
-	OutSavedPackages.Reset(SavedPackages.Num());
-	for (TPair<UPackage*,EPackageSaveResolutionType> SavedPackagesEntry : SavedPackages)
-	{
-		TPair<EPackageSaveResolutionType, UPackage*> ArrayEntry = {SavedPackagesEntry.Value, SavedPackagesEntry.Key};
-		OutSavedPackages.Add(ArrayEntry);
-	}
-
-	return bWasOperationSuccessful;
-}
-
-
 bool BakeCustomizableObjectInstance(
 	UCustomizableObjectInstance& InInstance,
 	const FBakingConfiguration& Configuration,
diff --git a/Engine/Plugins/Mutable/Source/CustomizableObjectEditor/Private/MuCOE/Nodes/CustomizableObjectNodeObject.h b/Engine/Plugins/Mutable/Source/CustomizableObjectEditor/Private/MuCOE/Nodes/CustomizableObjectNodeObject.h
--- a/Engine/Plugins/Mutable/Source/CustomizableObjectEditor/Private/MuCOE/Nodes/CustomizableObjectNodeObject.h
+++ b/Engine/Plugins/Mutable/Source/CustomizableObjectEditor/Private/MuCOE/Nodes/CustomizableObjectNodeObject.h
@@ -113,7 +113,7 @@ public:
 
 protected:
 
-	UPROPERTY(EditAnywhere, Category=CustomizableObject)
+	UPROPERTY()
 	FString ObjectName;
 
 public:
diff --git a/Engine/Plugins/Mutable/Source/CustomizableObjectEditor/Private/MuCOE/Nodes/CustomizableObjectNodeSkeletalMeshParameter.h b/Engine/Plugins/Mutable/Source/CustomizableObjectEditor/Private/MuCOE/Nodes/CustomizableObjectNodeSkeletalMeshParameter.h
--- a/Engine/Plugins/Mutable/Source/CustomizableObjectEditor/Private/MuCOE/Nodes/CustomizableObjectNodeSkeletalMeshParameter.h
+++ b/Engine/Plugins/Mutable/Source/CustomizableObjectEditor/Private/MuCOE/Nodes/CustomizableObjectNodeSkeletalMeshParameter.h
@@ -54,7 +54,7 @@ public:
 	TSoftObjectPtr<USkeletalMesh> ReferenceValue;
 	
 	UPROPERTY()
-	FString ParameterName = "Mesh Param";
+	FString ParameterName = "Default Name";
 
 	UPROPERTY(EditAnywhere, Category = UI, meta = (DisplayName = "Parameter UI Metadata"))
 	FMutableParamUIMetadata ParamUIMetadata;
