[Mutable] Add max LOD to propagate mesh secitons to when using automatic LODs.

#rnx

diff --git a/Engine/Plugins/Mutable/Source/CustomizableObjectEditor/Private/MuCOE/GenerateMutableSource/GenerateMutableSourceSurface.cpp b/Engine/Plugins/Mutable/Source/CustomizableObjectEditor/Private/MuCOE/GenerateMutableSource/GenerateMutableSourceSurface.cpp
--- a/Engine/Plugins/Mutable/Source/CustomizableObjectEditor/Private/MuCOE/GenerateMutableSource/GenerateMutableSourceSurface.cpp
+++ b/Engine/Plugins/Mutable/Source/CustomizableObjectEditor/Private/MuCOE/GenerateMutableSource/GenerateMutableSourceSurface.cpp
@@ -97,6 +97,12 @@ mu::Ptr<mu::NodeSurface> GenerateMutableSourceSurface(const UEdGraphPin * Pin, F
 
 	if (UCustomizableObjectNodeMaterialBase* TypedNodeMat = Cast<UCustomizableObjectNodeMaterialBase>(Node))
 	{
+		if (GenerationContext.CurrentAutoLODStrategy == ECustomizableObjectAutomaticLODStrategy::AutomaticFromMesh
+			&& TypedNodeMat->GetMeshSectionMaxLOD() != INDEX_NONE && TypedNodeMat->GetMeshSectionMaxLOD() < GenerationContext.CurrentLOD)
+		{
+			return Result;
+		}
+
 		bool bGeneratingImplicitComponent = GenerationContext.ComponentMeshOverride.get() != nullptr;
 
 		const UEdGraphPin* ConnectedMaterialPin = FollowInputPin(*TypedNodeMat->GetMeshPin());
diff --git a/Engine/Plugins/Mutable/Source/CustomizableObjectEditor/Private/MuCOE/Nodes/CustomizableObjectNodeMaterialBase.cpp b/Engine/Plugins/Mutable/Source/CustomizableObjectEditor/Private/MuCOE/Nodes/CustomizableObjectNodeMaterialBase.cpp
--- a/Engine/Plugins/Mutable/Source/CustomizableObjectEditor/Private/MuCOE/Nodes/CustomizableObjectNodeMaterialBase.cpp
+++ b/Engine/Plugins/Mutable/Source/CustomizableObjectEditor/Private/MuCOE/Nodes/CustomizableObjectNodeMaterialBase.cpp
@@ -4,6 +4,11 @@
 
 #include "MuCOE/EdGraphSchema_CustomizableObject.h"
 
+int32 UCustomizableObjectNodeMaterialBase::GetMeshSectionMaxLOD() const
+{
+	return MaxLOD;
+}
+
 FLinearColor UCustomizableObjectNodeMaterialBase::GetNodeTitleColor() const
 {
 	return UEdGraphSchema_CustomizableObject::GetPinTypeColor(UEdGraphSchema_CustomizableObject::PC_Material);
diff --git a/Engine/Plugins/Mutable/Source/CustomizableObjectEditor/Private/MuCOE/Nodes/CustomizableObjectNodeMaterialBase.h b/Engine/Plugins/Mutable/Source/CustomizableObjectEditor/Private/MuCOE/Nodes/CustomizableObjectNodeMaterialBase.h
--- a/Engine/Plugins/Mutable/Source/CustomizableObjectEditor/Private/MuCOE/Nodes/CustomizableObjectNodeMaterialBase.h
+++ b/Engine/Plugins/Mutable/Source/CustomizableObjectEditor/Private/MuCOE/Nodes/CustomizableObjectNodeMaterialBase.h
@@ -62,6 +62,9 @@ public:
 		return const_cast<UCustomizableObjectNodeMaterial*>(ConstThis->GetMaterialNode());
 	}
 
+	// Max LOD to propagate this section to, when using automatic LODs.
+	int32 GetMeshSectionMaxLOD() const;
+
 	virtual TArray<UCustomizableObjectLayout*> GetLayouts() const PURE_VIRTUAL(UCustomizableObjectNodeMaterial::GetLayouts, return {}; );
 	
 	virtual UMaterialInterface* GetMaterial() const PURE_VIRTUAL(UCustomizableObjectNodeMaterial::GetMaterial, return {}; );
@@ -152,5 +155,11 @@ public:
 
 	// UEdGraphNode interface
 	virtual FLinearColor GetNodeTitleColor() const override;
+
+private:
+	
+	/* Max LOD to propagate this section to, when using automatic LODs. */
+	UPROPERTY(EditAnywhere, Category = AdvancedSettings)
+	int32 MaxLOD = INDEX_NONE;
 };
 
