[mutable] Made the UI Sections appear sorted in the COI editor details pannel.
- Previously the parameters got sorted but the sections did not, even when they got to be displayed.
- Section parameters will keep following the previous pattern of being sorted by UITargetPosition and ParamName.

#rnx

diff --git a/Engine/Plugins/Mutable/Source/CustomizableObjectEditor/Private/MuCOE/CustomizableInstanceDetails.cpp b/Engine/Plugins/Mutable/Source/CustomizableObjectEditor/Private/MuCOE/CustomizableInstanceDetails.cpp
--- a/Engine/Plugins/Mutable/Source/CustomizableObjectEditor/Private/MuCOE/CustomizableInstanceDetails.cpp
+++ b/Engine/Plugins/Mutable/Source/CustomizableObjectEditor/Private/MuCOE/CustomizableInstanceDetails.cpp
@@ -641,6 +641,73 @@ struct FParameterInfo
 
 // PARAMETERS WIDGET GENERATION -----------------------------------------------------------------------------------------------------------------
 
+void FCustomizableInstanceDetails::SortParametersAndUISections(const UCustomizableObject& CustomizableObject, TArray<FParameterInfo>& ParametersToSort) const
+{
+	// Sort the parameter tree not based in the type defined sorting behaviour but using for it the name of the section that the parameters are part of
+	ParametersToSort.Sort([&CustomizableObject, DefaultSectionName = this->DefaultSectionName](const FParameterInfo& A,const FParameterInfo& B)
+	{
+		const FMutableParamUIMetadata AParamIMetadata = CustomizableObject.GetParameterUIMetadata(A.ParamName);
+		const FMutableParamUIMetadata BParamIMetadata = CustomizableObject.GetParameterUIMetadata(B.ParamName);
+
+		const FString AParameterSectionName = !AParamIMetadata.UISectionName.IsEmpty() ? AParamIMetadata.UISectionName : DefaultSectionName;
+		const FString BParameterSectionName = !BParamIMetadata.UISectionName.IsEmpty() ? BParamIMetadata.UISectionName : DefaultSectionName;
+		
+		return AParameterSectionName < BParameterSectionName;
+	});
+
+	// Now that the parameters are sorted by section make sure each section has its parameters sorted
+	
+	// Key is section name and the amount of parameters of this section
+	TMap<FString, uint32> SectionParameterMap;
+	for (const FParameterInfo& ParameterInfo : ParametersToSort)
+	{
+		FString ParamSectionName = CustomizableObject.GetParameterUIMetadata(ParameterInfo.ParamName).UISectionName;
+		
+		// Sections with no name should use the default Section Name.
+		if (ParamSectionName.IsEmpty())
+		{
+			ParamSectionName = DefaultSectionName;
+		}
+		
+		if (SectionParameterMap.Contains(ParamSectionName))
+		{
+			SectionParameterMap[ParamSectionName] ++;
+		}
+		else
+		{
+			SectionParameterMap.Add(ParamSectionName, 1);
+		}
+	}
+
+	// Locate the section of parameters that need to be sorted from the ones found in the actual UI section
+	uint32 SectionStartIndex = 0;
+	for (const TPair<FString, uint32>& SectionAndParamIndex : SectionParameterMap)
+	{
+		const uint32 SectionParametersCount = SectionAndParamIndex.Value;
+
+		// Capture the parameters of the section so we can sort them
+		TArray<FParameterInfo> SectionParameters;
+		for (uint32 ParamIndex = SectionStartIndex; ParamIndex < SectionStartIndex + SectionParametersCount; ++ParamIndex)
+		{
+			SectionParameters.Add(ParametersToSort[ParamIndex]);
+		}
+
+		// Sort the section parameters
+		SectionParameters.Sort();
+
+		// Replace the unsorted (ParametersToSort) parameters with the sorted ones (SectionParameters)
+		int32 SectionParametersIndex = 0;
+		for (uint32 ParamIndex = SectionStartIndex; ParamIndex < SectionStartIndex + SectionParametersCount; ++ParamIndex)
+		{
+			ParametersToSort[ParamIndex] = SectionParameters[SectionParametersIndex++];
+		}
+		
+		// Move the starting index forward for the next section
+		SectionStartIndex += SectionParametersCount;
+	}
+}
+
+
 bool FCustomizableInstanceDetails::GenerateParametersView(IDetailCategoryBuilder& DetailsCategory)
 {
 	ParamChildren.Empty();
@@ -680,7 +747,16 @@ bool FCustomizableInstanceDetails::GenerateParametersView(IDetailCategoryBuilder
 			}
 		}
 
-		ParametersTree.Sort();
+		if (CustomInstance->GetPrivate()->bShowUISections)
+		{
+			// Sort the parameter tree using for it the name of the section that the parameters are part of
+			SortParametersAndUISections(*CustomizableObject, ParametersTree);
+		}
+		else
+		{
+			// Sort the parameters using for it the UIOrder and the name of them.
+			ParametersTree.Sort();
+		}
 
 		for (int32 ParamIndex = 0; ParamIndex < ParametersTree.Num(); ++ParamIndex)
 		{
@@ -737,9 +813,18 @@ bool FCustomizableInstanceDetails::GenerateParametersView(IDetailCategoryBuilder
 				ParametersTree.Add(ParameterSortInfo);
 			}
 		}
-
-		ParametersTree.Sort();
-
+		
+		if (CustomInstance->GetPrivate()->bShowUISections)
+		{
+			// Sort the parameter tree using for it the name of the section that the parameters are part of
+			SortParametersAndUISections(*CustomizableObject, ParametersTree);
+		}
+		else
+		{
+			// Sort the parameters using for it the UIOrder and the name of them.
+			ParametersTree.Sort();
+		}
+		
 		for (int32 ParamIndexInObject = 0; ParamIndexInObject < ParametersTree.Num(); ++ParamIndexInObject)
 		{
 			FillChildrenMap(ParametersTree[ParamIndexInObject].ParamIndexInObject);
@@ -910,7 +995,7 @@ bool FCustomizableInstanceDetails::IsMultidimensionalProjector(int32 ParamIndexI
 IDetailGroup* FCustomizableInstanceDetails::GenerateParameterSection(IDetailCategoryBuilder& DetailsCategory, const UCustomizableObject& CustomizableObject, const FString& ParamName)
 {
 	const FMutableParamUIMetadata UIMetadata = CustomizableObject.GetParameterUIMetadata(ParamName);
-	const FString SectionName = UIMetadata.UISectionName.IsEmpty() ? "Miscellaneous" : UIMetadata.UISectionName;
+	const FString SectionName = UIMetadata.UISectionName.IsEmpty() ? DefaultSectionName : UIMetadata.UISectionName;
 	IDetailGroup* CurrentSection = nullptr;
 
 	if (GeneratedSections.Contains(SectionName))
diff --git a/Engine/Plugins/Mutable/Source/CustomizableObjectEditor/Private/MuCOE/CustomizableInstanceDetails.h b/Engine/Plugins/Mutable/Source/CustomizableObjectEditor/Private/MuCOE/CustomizableInstanceDetails.h
--- a/Engine/Plugins/Mutable/Source/CustomizableObjectEditor/Private/MuCOE/CustomizableInstanceDetails.h
+++ b/Engine/Plugins/Mutable/Source/CustomizableObjectEditor/Private/MuCOE/CustomizableInstanceDetails.h
@@ -26,6 +26,7 @@ namespace ETextCommit { enum Type : int; }
 struct FGeometry;
 struct FPointerEvent;
 struct FUpdateContext;
+struct FParameterInfo;
 
 class FCustomizableInstanceDetails : public IDetailCustomization
 {
@@ -75,6 +76,17 @@ private:
 	// Main parameter generation functions
 	// Returns true if parameters have been hidden due to runtime type
 	bool GenerateParametersView(IDetailCategoryBuilder& MainCategory);
+
+	/**
+	 * Sort the contents of a Parameters Info array taking in consideration that the array should get sorted first by UI Section and then by
+	 * parameter UI order or Param name.
+	 * @note The result will be an array where the sections are in sequential order following an alphabetical ordering (A-Z)
+	 * where each section contains its parameters in an incremental numerical order (ParamUIOrder int) or incremental lexical order (ParamName)
+	 * @param CustomizableObject The customizable object the parameters are linked to.
+	 * @param ParametersToSort The array of parameters we want to get sorted out.
+	 */
+	void SortParametersAndUISections(const UCustomizableObject& CustomizableObject, TArray<FParameterInfo>& ParametersToSort) const;
+	
 	void RecursivelyAddParamAndChildren(const UCustomizableObject& CustomizableObject, const int32 ParamIndexInObject, const FString ParentName, IDetailCategoryBuilder& DetailsCategory);
 	void FillChildrenMap(int32 ParamIndexInObject);
 
@@ -203,5 +215,8 @@ private:
 	//** Editor gameplay tags filter and filter type. Used to filter int parameter options. */
 	FGameplayTagContainer Filter;
 	EGameplayContainerMatchType FilterType = EGameplayContainerMatchType::Any;
+
+	/** Section name used for parameters that do have an empty string for its section name. */
+	const FString DefaultSectionName = FString(TEXT("Miscellaneous"));
 };
 
