[Mutable] Store image mip tail in a single rom.

#rnx

diff --git a/Engine/Plugins/Mutable/Source/CustomizableObject/Internal/MuCO/CustomizableObjectPrivate.h b/Engine/Plugins/Mutable/Source/CustomizableObject/Internal/MuCO/CustomizableObjectPrivate.h
--- a/Engine/Plugins/Mutable/Source/CustomizableObject/Internal/MuCO/CustomizableObjectPrivate.h
+++ b/Engine/Plugins/Mutable/Source/CustomizableObject/Internal/MuCO/CustomizableObjectPrivate.h
@@ -1628,6 +1628,8 @@ enum class ECustomizableObjectVersions : int32
 
 	AddedBlockMasksToRuntimeLayouts,
 
+	StoreMipTailInASingleRom,
+
 	// -----<new versions can be added above this line>--------
 	LastCustomizableObjectVersion
 };
@@ -1893,7 +1895,7 @@ public:
 
 #endif
 	
-	static constexpr int32 DerivedDataVersion = 0x1a376aa8;
+	static constexpr int32 DerivedDataVersion = 0x28f3f788;
 };
 
 #if WITH_EDITOR
diff --git a/Engine/Plugins/Mutable/Source/MutableRuntime/Internal/MuR/ModelPrivate.h b/Engine/Plugins/Mutable/Source/MutableRuntime/Internal/MuR/ModelPrivate.h
--- a/Engine/Plugins/Mutable/Source/MutableRuntime/Internal/MuR/ModelPrivate.h
+++ b/Engine/Plugins/Mutable/Source/MutableRuntime/Internal/MuR/ModelPrivate.h
@@ -86,11 +86,12 @@ namespace mu
 	//!
 	struct FImageLODRange
 	{
-		int32 FirstIndex = 0;
-		uint16 ImageSizeX = 0;
-		uint16 ImageSizeY = 0;
-		uint16 _Padding = 0;
-		uint8 LODCount = 0;
+		int32 FirstIndex    = 0;
+		uint16 ImageSizeX   = 0;
+		uint16 ImageSizeY   = 0;
+		uint8 LODCount      = 0;
+		uint8 NumLODsInTail = 0; 
+		uint8 Flags         = 0;
 		EImageFormat ImageFormat = EImageFormat::None;
 	};
 	MUTABLE_DEFINE_POD_SERIALISABLE(FImageLODRange);
@@ -698,102 +699,152 @@ namespace mu
 			return Mip;
 		}
 
-
         /** Get a constant image, assuming at least some mips are loaded. The image constant will be composed with lodaded mips if necessary. */
 		template <typename CreateImageFunc>
-        void GetConstant( uint32 ConstantIndex, TSharedPtr<const FImage>& res, int32 MipsToSkip, const CreateImageFunc& CreateImage) const
+        void GetConstant(uint32 ConstantIndex, TSharedPtr<const FImage>& Result, int32 MipsToSkip, const CreateImageFunc& CreateImage) const
         {
-			int32 ReallySkippedLODs = FMath::Min(ConstantImages[ConstantIndex].LODCount - 1, MipsToSkip);
-			int32 FirstLODIndexIndex = ConstantImages[ConstantIndex].FirstIndex;
+			const int32 NumLODs = ConstantImages[ConstantIndex].LODCount;
 
-			// Get the first mip
-			int32 ResultLODIndexIndex = FirstLODIndexIndex + ReallySkippedLODs;
-			FConstantResourceIndex ResultLODIndex = ConstantImageLODIndices[ResultLODIndexIndex];
-			TSharedPtr<const FImage> CurrentMip = GetImageLOD(ResultLODIndex);
+			check(NumLODs > 0);
+			const int32 ReallySkippedLODs = FMath::Min(NumLODs - 1, MipsToSkip);
+			const int32 FirstLODIndexIndex = ConstantImages[ConstantIndex].FirstIndex;
+			const int32 FirstTailLOD = ConstantImages[ConstantIndex].LODCount - ConstantImages[ConstantIndex].NumLODsInTail;
 
-			// We may need to skip more LODs if they are not loaded
-			while (!CurrentMip)
+			// Compose the result image
 			{
-				++ReallySkippedLODs;
-				
-				if (ReallySkippedLODs >= ConstantImages[ConstantIndex].LODCount)
+				MUTABLE_CPUPROFILER_SCOPE(ComposeConstantImage);
+
+				TSharedPtr<FImage> ResultImage = nullptr;
+
+				int32 LOD = ReallySkippedLODs;
+
+				// Find first availabe LOD.
+				for (; LOD <= FirstTailLOD; ++LOD)
 				{
-					// We don't have a single mip loaded for the image that was requested
-					ensure(false);
-					break;
+					MUTABLE_CPUPROFILER_SCOPE(Search_FisrtValid);
+					 TSharedPtr<const FImage> FirstValidLODImage = GetImageLOD(ConstantImageLODIndices[FirstLODIndexIndex + LOD]);
+					
+					if (FirstValidLODImage)
+					{
+						break;
+					}
 				}
+			
+				if (LOD < NumLODs)
+				{						
+					int32 SizeX = ConstantImages[ConstantIndex].ImageSizeX; 
+					int32 SizeY = ConstantImages[ConstantIndex].ImageSizeY; 
 
-				++ResultLODIndexIndex;
-				ResultLODIndex = ConstantImageLODIndices[ResultLODIndexIndex];
-				CurrentMip = GetImageLOD(ResultLODIndex);
-			}
+					for (int32 I = 0; I < ReallySkippedLODs; ++I)
+					{
+						SizeX = FMath::DivideAndRoundUp(SizeX, 2);
+						SizeY = FMath::DivideAndRoundUp(SizeY, 2);
+					}
+					 
+					EImageFormat Format = ConstantImages[ConstantIndex].ImageFormat; 
 
-			int32 FinalLODs = ConstantImages[ConstantIndex].LODCount - ReallySkippedLODs;
-			check(FinalLODs > 0);
+					ResultImage = CreateImage(SizeX, SizeY, NumLODs - LOD, Format, EInitializationType::NotInitialized);
+					ResultImage->Flags = ConstantImages[ConstantIndex].Flags;
+				}
+	
+				check(ResultImage);
 
-			// Shortcut if we only want one mip
-			if (FinalLODs == 1)
-			{
-				res = CurrentMip;
-				return;
-			}
+				// Some non-block pixel formats require manual resize.
+				const bool bFormatNeedsDataResize = ResultImage->DataStorage.IsEmpty();
+				
+				int32 DestLOD = 0;
+				// Process hi-res separeted mips.
+				for (; LOD < FirstTailLOD; ++LOD, ++DestLOD)
+				{
+					MUTABLE_CPUPROFILER_SCOPE(NonTailLODs);
+					TSharedPtr<const FImage> Image = GetImageLOD(ConstantImageLODIndices[FirstLODIndexIndex + LOD]);
 
-			// Compose the result image
-			{
-				MUTABLE_CPUPROFILER_SCOPE(ComposeConstantImage);
+					if (bFormatNeedsDataResize)
+					{
+						ResultImage->DataStorage.ResizeLOD(DestLOD, Image->DataStorage.GetLOD(0).Num());
+					}
+
+					TArrayView<uint8> ResultLODView = ResultImage->DataStorage.GetLOD(DestLOD);
+					TArrayView<const uint8> ImageLODView = Image->DataStorage.GetLOD(0);
 
-				TSharedPtr<FImage> Result = CreateImage(CurrentMip->GetSizeX(), CurrentMip->GetSizeY(), FinalLODs, CurrentMip->GetFormat(), EInitializationType::NotInitialized);
-				Result->Flags = CurrentMip->Flags;
+					check(ResultLODView.Num() == ImageLODView.Num())
+					FMemory::Memcpy(ResultLODView.GetData(), ImageLODView.GetData(), ResultLODView.Num());
+				}
 
-				// Some non-block pixel formats require separate memory size calculation
-				if (Result->DataStorage.IsEmpty())
+				// Process compacted tail.
+				// NOTE: The rom tail data storage can have multiple image buffers.
+				if (LOD < NumLODs)
 				{
-					for (int32 LOD = 0; LOD < FinalLODs; ++LOD)
+					MUTABLE_CPUPROFILER_SCOPE(TailLODs);
+					TSharedPtr<const FImage> TailLODsImage = GetImageLOD(ConstantImageLODIndices[FirstLODIndexIndex + FirstTailLOD]);
+
+					if (!ensureAlwaysMsgf(TailLODsImage, TEXT("Missing tail image rom, setting image to black.")))
 					{
-						FConstantResourceIndex LODIndex = ConstantImageLODIndices[ResultLODIndexIndex + LOD];
+						ResultImage->InitToBlack();
+						Result = ResultImage;
+						return;
+					}
 
-						TSharedPtr<const FImage> Image = GetImageLOD(LODIndex);
+					check(TailLODsImage->DataStorage.GetNumLODs() >= NumLODs - LOD);		
+	
+					int32 FirstTailImageTailLOD = FMath::Min(
+							TailLODsImage->DataStorage.ComputeFirstCompactedTailLOD(),
+							TailLODsImage->DataStorage.GetNumLODs());
 
-						// This could happen in the case of missing data files
-						if (!Image)
+					int32 TailImageLOD = FMath::Max(0, ReallySkippedLODs - FirstTailLOD);
+					for (; TailImageLOD < FirstTailImageTailLOD; ++TailImageLOD, ++LOD, ++DestLOD)
+					{
+						MUTABLE_CPUPROFILER_SCOPE(TailLODCopy);
+						if (bFormatNeedsDataResize)
 						{
-							break;
+							ResultImage->DataStorage.ResizeLOD(DestLOD, TailLODsImage->DataStorage.GetLOD(TailImageLOD).Num());
 						}
 
-						int32 MipSizeBytes = Image->GetLODDataSize(0);
-						Result->DataStorage.ResizeLOD(LOD, MipSizeBytes);
-					}
-				}
-
-				for (int32 LOD = 0; LOD < FinalLODs; ++LOD)
-				{
-					check(CurrentMip->GetLODCount() == 1);
-					check(CurrentMip->GetFormat() == Result->GetFormat());
-
-					TArrayView<uint8> ResultLODView = Result->DataStorage.GetLOD(LOD);
-					TArrayView<const uint8> CurrentMipView = CurrentMip->DataStorage.GetLOD(0);
-					
-					check(CurrentMipView.Num() == ResultLODView.Num());
+						TArrayView<uint8> ResultLODView = ResultImage->DataStorage.GetLOD(DestLOD);
+						TArrayView<const uint8> ImageLODView = TailLODsImage->DataStorage.GetLOD(TailImageLOD);
 
-					FMemory::Memcpy(ResultLODView.GetData(), CurrentMipView.GetData(), ResultLODView.Num());
+						check(ResultLODView.Num() == ImageLODView.Num());
+						FMemory::Memcpy(ResultLODView.GetData(), ImageLODView.GetData(), ResultLODView.Num());
+					}	
 
-					if (LOD + 1 < FinalLODs)
+					if (LOD < NumLODs)
 					{
-						ResultLODIndex = ConstantImageLODIndices[ResultLODIndexIndex + LOD + 1];
-						CurrentMip = GetImageLOD(ResultLODIndex);
+						MUTABLE_CPUPROFILER_SCOPE(TailTailCopy);
+						check(DestLOD >= ResultImage->DataStorage.ComputeFirstCompactedTailLOD());
+						check(TailImageLOD >= TailLODsImage->DataStorage.ComputeFirstCompactedTailLOD());
 
-						// This could only happen if missing or corrupted data.
-						if (!CurrentMip)
+						TArrayView<const uint8> FirstTailLODView = TailLODsImage->DataStorage.GetLOD(TailImageLOD);
+						TArrayView<const uint8> LastTailLODView  = TailLODsImage->DataStorage.GetLOD(TailLODsImage->DataStorage.GetNumLODs() - 1);
+					
+						// The tail is stored compacted in memory.
+						int64 NumBytesInTail = LastTailLODView.GetData() + LastTailLODView.Num() - FirstTailLODView.GetData();
+						check(NumBytesInTail >= 0);
+						
+						if (bFormatNeedsDataResize)
 						{
-							break;
+							const int32 NumLODsToResize = ResultImage->DataStorage.GetNumLODs() - DestLOD;
+							check(NumLODsToResize == TailLODsImage->DataStorage.GetNumLODs() - TailImageLOD);
+
+							ResultImage->DataStorage.GetInternalArray(DestLOD).Reserve(NumBytesInTail);
+
+							for (int32 LODOffset = 0; LODOffset < NumLODsToResize; ++LODOffset)
+							{
+								int32 TailLODSizeInBytes = TailLODsImage->DataStorage.GetLOD(LODOffset + TailImageLOD).Num();
+								ResultImage->DataStorage.ResizeLOD(LODOffset + DestLOD, TailLODSizeInBytes);
+							}
 						}
+
+						TArrayView<uint8> ResultLODView = ResultImage->DataStorage.GetLOD(DestLOD);
+
+						check((int64)ResultImage->DataStorage.GetInternalArray(DestLOD).Num() >= NumBytesInTail);
+						FMemory::Memcpy(ResultLODView.GetData(), FirstTailLODView.GetData(), NumBytesInTail);
 					}
 				}
-
-				res = Result;
+	
+				Result = ResultImage;
 			}
-		}
-		
+		}		
+
 		template<class CreateMeshFunc>
         void GetConstant(
 				uint32 MeshConstantIndex, 
@@ -1070,5 +1121,4 @@ namespace mu
     };
 
 }
-
 #undef UE_API
diff --git a/Engine/Plugins/Mutable/Source/MutableRuntime/MutableRuntime.Build.cs b/Engine/Plugins/Mutable/Source/MutableRuntime/MutableRuntime.Build.cs
--- a/Engine/Plugins/Mutable/Source/MutableRuntime/MutableRuntime.Build.cs
+++ b/Engine/Plugins/Mutable/Source/MutableRuntime/MutableRuntime.Build.cs
@@ -19,6 +19,7 @@ namespace UnrealBuildTool.Rules
 				"Core",
 				"CoreUObject",
 				"Engine",
+				"ImageCore"
 				}
 			);
 
diff --git a/Engine/Plugins/Mutable/Source/MutableRuntime/Private/MuR/CodeRunner.cpp b/Engine/Plugins/Mutable/Source/MutableRuntime/Private/MuR/CodeRunner.cpp
--- a/Engine/Plugins/Mutable/Source/MutableRuntime/Private/MuR/CodeRunner.cpp
+++ b/Engine/Plugins/Mutable/Source/MutableRuntime/Private/MuR/CodeRunner.cpp
@@ -1080,14 +1080,14 @@ namespace mu
         case EOpType::IM_CONSTANT:
         {
 			OP::ResourceConstantArgs Args = Program.GetOpArgs<OP::ResourceConstantArgs>(Item.At);
-            OP::ADDRESS cat = Args.value;
 
 			int32 MipsToSkip = Item.ExecutionOptions;
             TSharedPtr<const FImage> Source;
-			Program.GetConstant(cat, Source, MipsToSkip, [this](int32 x, int32 y, int32 m, EImageFormat f, EInitializationType i)
-				{
-					return CreateImage(x, y, m, f, i);
-				});
+			Program.GetConstant(Args.value, Source, MipsToSkip, 
+			[this](int32 SizeX, int32 SizeY, int32 NumLODs, EImageFormat Format, EInitializationType InitPolicy)
+			{
+				return CreateImage(SizeX, SizeY, NumLODs, Format, InitPolicy);
+			});
 
 			// Assume the ROM has been loaded previously in a task generated at IssueOp
 			if (!Source)
diff --git a/Engine/Plugins/Mutable/Source/MutableRuntime/Private/MuR/CodeRunner.h b/Engine/Plugins/Mutable/Source/MutableRuntime/Private/MuR/CodeRunner.h
--- a/Engine/Plugins/Mutable/Source/MutableRuntime/Private/MuR/CodeRunner.h
+++ b/Engine/Plugins/Mutable/Source/MutableRuntime/Private/MuR/CodeRunner.h
@@ -286,11 +286,11 @@ namespace  mu
 		class FLoadImageRomsTask : public CodeRunner::FIssuedTask
 		{
 		public:
-			FLoadImageRomsTask(const FScheduledOp& InOp, int32 InLODIndexIndex, int32 InLODIndexCount )
+			FLoadImageRomsTask(const FScheduledOp& InOp, int32 InImageIndexBegin, int32 InImageIndexEnd)
 				: FIssuedTask(InOp)
+				, ImageIndexBegin(InImageIndexBegin)
+				, ImageIndexEnd(InImageIndexEnd)
 			{
-				LODIndexIndex = InLODIndexIndex;
-				LODIndexCount = InLODIndexCount;
 			}
 
 			// FIssuedTask interface
@@ -298,8 +298,8 @@ namespace  mu
 			virtual bool Complete(CodeRunner*) override;
 
 		private:
- 			int32 LODIndexIndex = -1;
-			int32 LODIndexCount = -1;
+ 			int32 ImageIndexBegin = -1;
+			int32 ImageIndexEnd = -1;
 
 			TArray<int32> RomIndices;
 		};
diff --git a/Engine/Plugins/Mutable/Source/MutableRuntime/Private/MuR/CodeRunnerTasked.cpp b/Engine/Plugins/Mutable/Source/MutableRuntime/Private/MuR/CodeRunnerTasked.cpp
--- a/Engine/Plugins/Mutable/Source/MutableRuntime/Private/MuR/CodeRunnerTasked.cpp
+++ b/Engine/Plugins/Mutable/Source/MutableRuntime/Private/MuR/CodeRunnerTasked.cpp
@@ -2205,20 +2205,19 @@ namespace mu
 
 		TArray<UE::Tasks::FTask> ReadCompleteEvents;
 
-		ReadCompleteEvents.Reserve(LODIndexCount); 
+		ReadCompleteEvents.Reserve(FMath::Max(0, ImageIndexEnd - ImageIndexBegin));
 
-		for (int32 LODIndex = 0; LODIndex < LODIndexCount; ++LODIndex)
+		for (int32 ImageIndex = ImageIndexBegin; ImageIndex < ImageIndexEnd; ++ImageIndex)
 		{
-			const int32 CurrentIndexIndex = LODIndexIndex + LODIndex;
-			const FConstantResourceIndex CurrentIndex = Program.ConstantImageLODIndices[CurrentIndexIndex];
+			const FConstantResourceIndex ImageConstantResourceIndex = Program.ConstantImageLODIndices[ImageIndex];
 
-			if (!CurrentIndex.Streamable)
+			if (!ImageConstantResourceIndex.Streamable)
 			{
 				// This data is always resident.
 				continue;
 			}
 
-			int32 RomIndex = CurrentIndex.Index;
+			int32 RomIndex = ImageConstantResourceIndex.Index;
 			check(RomIndex < Program.Roms.Num());
 
 			++ModelCache->PendingOpsPerRom[RomIndex];
@@ -2339,26 +2338,28 @@ namespace mu
 		// Process the constant op normally, now that the rom is loaded.	
 		bool bSuccess = Runner->RunCode_ConstantResource(Op, Runner->Model.Get());
 
-		for (int32 LODIndex = 0; bSuccess && (LODIndex < LODIndexCount); ++LODIndex)
+		if (bSuccess)
 		{
-			int32 CurrentIndexIndex = LODIndexIndex + LODIndex;
-			FConstantResourceIndex CurrentIndex = Program.ConstantImageLODIndices[CurrentIndexIndex];
-
-			if (!CurrentIndex.Streamable)
+			for (int32 ImageIndex = ImageIndexBegin; ImageIndex < ImageIndexEnd; ++ImageIndex)
 			{
-				// This data is always resident.
-				continue;
-			}
+				FConstantResourceIndex ImageConstantResourceIndex = Program.ConstantImageLODIndices[ImageIndex];
 
-			int32 RomIndex = CurrentIndex.Index;
-			check(RomIndex < Program.Roms.Num());
+				if (!ImageConstantResourceIndex.Streamable)
+				{
+					// This data is always resident.
+					continue;
+				}
 
-			Runner->System->WorkingMemoryManager.MarkRomUsed(RomIndex, Runner->Model);
-			--ModelCache->PendingOpsPerRom[RomIndex];
+				int32 RomIndex = ImageConstantResourceIndex.Index;
+				check(RomIndex < Program.Roms.Num());
 
-			if (DebugRom && (DebugRomAll || RomIndex == DebugRomIndex))
-			{
-				UE_LOG(LogMutableCore, Log, TEXT("FLoadImageRomsTask::Complete rom %d, now peding ops is %d."), RomIndex, ModelCache->PendingOpsPerRom[RomIndex]);
+				Runner->System->WorkingMemoryManager.MarkRomUsed(RomIndex, Runner->Model);
+				--ModelCache->PendingOpsPerRom[RomIndex];
+
+				if (DebugRom && (DebugRomAll || RomIndex == DebugRomIndex))
+				{
+					UE_LOG(LogMutableCore, Log, TEXT("FLoadImageRomsTask::Complete rom %d, now peding ops is %d."), RomIndex, ModelCache->PendingOpsPerRom[RomIndex]);
+				}
 			}
 		}
 
@@ -2511,37 +2512,45 @@ namespace mu
 	//---------------------------------------------------------------------------------------------
 	//---------------------------------------------------------------------------------------------
 	//---------------------------------------------------------------------------------------------
-	TSharedPtr<CodeRunner::FIssuedTask> CodeRunner::IssueOp(FScheduledOp item)
+	TSharedPtr<CodeRunner::FIssuedTask> CodeRunner::IssueOp(FScheduledOp Item)
 	{
 		TSharedPtr<FIssuedTask> Issued;
 
 		FProgram& Program = Model->GetPrivate()->Program;
 
-		EOpType type = Program.GetOpType(item.At);
+		EOpType type = Program.GetOpType(Item.At);
 
 		switch (type)
 		{
 		case EOpType::ME_CONSTANT:
 		{
-			OP::MeshConstantArgs Args = Program.GetOpArgs<OP::MeshConstantArgs>(item.At);
+			OP::MeshConstantArgs Args = Program.GetOpArgs<OP::MeshConstantArgs>(Item.At);
 		
 			const FMeshContentRange MeshContentRange = Program.ConstantMeshes[Args.Value];
-			const EMeshContentFlags ContentFilterFlags = static_cast<EMeshContentFlags>(item.ExecutionOptions);
+			const EMeshContentFlags ContentFilterFlags = static_cast<EMeshContentFlags>(Item.ExecutionOptions);
 			
-			Issued = MakeShared<FLoadMeshRomTask>(item, 
+			Issued = MakeShared<FLoadMeshRomTask>(Item, 
 					MeshContentRange.GetFirstIndex(), MeshContentRange.GetContentFlags(), ContentFilterFlags);
 			break;
 		}
 
 		case EOpType::IM_CONSTANT:
 		{
-			OP::ResourceConstantArgs args = Program.GetOpArgs<OP::ResourceConstantArgs>(item.At);
-			int32 MipsToSkip = item.ExecutionOptions;
+			OP::ResourceConstantArgs args = Program.GetOpArgs<OP::ResourceConstantArgs>(Item.At);
+			int32 MipsToSkip = Item.ExecutionOptions;
 			int32 ImageIndex = args.value;
 			int32 ReallySkip = FMath::Min(MipsToSkip, Program.ConstantImages[ImageIndex].LODCount - 1);
-			int32 LODIndexIndex = Program.ConstantImages[ImageIndex].FirstIndex + ReallySkip;
-			int32 LODIndexCount = Program.ConstantImages[ImageIndex].LODCount - ReallySkip;
-			check(LODIndexCount > 0);
+
+			const int32 NumLODs         = Program.ConstantImages[ImageIndex].LODCount;
+			const int32 NumLODsInTail   = Program.ConstantImages[ImageIndex].NumLODsInTail;
+			const int32 FirstImageIndex = Program.ConstantImages[ImageIndex].FirstIndex;
+			const int32 FirstLODInTailIndex  = FirstImageIndex + NumLODs - NumLODsInTail;
+
+			check(NumLODs >= NumLODsInTail);
+
+			int32 ImageIndexBegin = FMath::Min(FirstImageIndex + ReallySkip, FirstLODInTailIndex);
+			int32 ImageIndexEnd   = FirstLODInTailIndex + 1; 
+			check(NumLODs > 0);
 
 			// We always need to follow this path, or roms may not be protected for long enough and might be unloaded 
 			// because of memory budget contraints.
@@ -2559,7 +2568,7 @@ namespace mu
 
 			if (bAnyMissing)
 			{
-				Issued = MakeShared<FLoadImageRomsTask>(item, LODIndexIndex, LODIndexCount);
+				Issued = MakeShared<FLoadImageRomsTask>(Item, ImageIndexBegin, ImageIndexEnd);
 
 				if (DebugRom && (DebugRomAll || ImageIndex == DebugImageIndex))
 					UE_LOG(LogMutableCore, Log, TEXT("Issuing image %d skipping %d ."), ImageIndex, ReallySkip);
@@ -2575,37 +2584,37 @@ namespace mu
 
 		case EOpType::IM_PARAMETER:
 		{
-			OP::ParameterArgs args = Program.GetOpArgs<OP::ParameterArgs>(item.At);
-			TSharedPtr<FRangeIndex> Index = BuildCurrentOpRangeIndex(item, Params, Model.Get(), args.variable);
+			OP::ParameterArgs args = Program.GetOpArgs<OP::ParameterArgs>(Item.At);
+			TSharedPtr<FRangeIndex> Index = BuildCurrentOpRangeIndex(Item, Params, Model.Get(), args.variable);
 
 			UTexture* Image = Params->GetImageValue(args.variable, Index.Get());
 
 			check(ImageLOD < TNumericLimits<uint8>::Max() && ImageLOD >= 0);
-			check(ImageLOD + static_cast<int32>(item.ExecutionOptions) < TNumericLimits<uint8>::Max());
+			check(ImageLOD + static_cast<int32>(Item.ExecutionOptions) < TNumericLimits<uint8>::Max());
 
-			const uint8 MipmapsToSkip = item.ExecutionOptions + static_cast<uint8>(ImageLOD);
+			const uint8 MipmapsToSkip = Item.ExecutionOptions + static_cast<uint8>(ImageLOD);
 
 			CodeRunner::FExternalResourceId FullId;
 			FullId.ImageParameter = Image;
-			Issued = MakeShared<FImageExternalLoadTask>(item, MipmapsToSkip, FullId);
+			Issued = MakeShared<FImageExternalLoadTask>(Item, MipmapsToSkip, FullId);
 
 			break;
 		}
 
 		case EOpType::IM_REFERENCE:
 		{
-			OP::ResourceReferenceArgs Args = Model->GetPrivate()->Program.GetOpArgs<OP::ResourceReferenceArgs>(item.At);
+			OP::ResourceReferenceArgs Args = Model->GetPrivate()->Program.GetOpArgs<OP::ResourceReferenceArgs>(Item.At);
 
 			// We only convert references to images if indicated in the operation.
 			if (Args.ForceLoad)
 			{
-				check(item.Stage==0);
+				check(Item.Stage==0);
 
-				const uint8 MipmapsToSkip = item.ExecutionOptions + static_cast<uint8>(ImageLOD);
+				const uint8 MipmapsToSkip = Item.ExecutionOptions + static_cast<uint8>(ImageLOD);
 
 				FExternalResourceId FullId;
 				FullId.ReferenceResourceId = Args.ID;
-				Issued = MakeShared<FImageExternalLoadTask>(item, MipmapsToSkip, FullId);
+				Issued = MakeShared<FImageExternalLoadTask>(Item, MipmapsToSkip, FullId);
 			}
 
 			break;
@@ -2613,8 +2622,8 @@ namespace mu
 
 		case EOpType::ME_PARAMETER:
 		{
-			OP::MeshParameterArgs Args = Program.GetOpArgs<OP::MeshParameterArgs>(item.At);
-			TSharedPtr<FRangeIndex> Index = BuildCurrentOpRangeIndex(item, Params, Model.Get(), Args.variable);
+			OP::MeshParameterArgs Args = Program.GetOpArgs<OP::MeshParameterArgs>(Item.At);
+			TSharedPtr<FRangeIndex> Index = BuildCurrentOpRangeIndex(Item, Params, Model.Get(), Args.variable);
 
 			USkeletalMesh* Id = Params->GetMeshValue(Args.variable, Index.Get());
 
@@ -2623,19 +2632,19 @@ namespace mu
 
 			CodeRunner::FExternalResourceId FullId;
 			FullId.MeshParameter = Id;
-			Issued = MakeShared<FMeshExternalLoadTask>(item, FullId, LODIndex, SectionIndex, Args.MeshID);
+			Issued = MakeShared<FMeshExternalLoadTask>(Item, FullId, LODIndex, SectionIndex, Args.MeshID);
 
 			break;
 		}
 
 		case EOpType::ME_REFERENCE:
 		{
-			OP::ResourceReferenceArgs Args = Model->GetPrivate()->Program.GetOpArgs<OP::ResourceReferenceArgs>(item.At);
+			OP::ResourceReferenceArgs Args = Model->GetPrivate()->Program.GetOpArgs<OP::ResourceReferenceArgs>(Item.At);
 
 			// We only convert references to meshes if indicated in the operation.
 			if (Args.ForceLoad)
 			{
-				check(item.Stage == 0);
+				check(Item.Stage == 0);
 
 				FExternalResourceId FullId;
 				FullId.ReferenceResourceId = Args.ID;
@@ -2645,7 +2654,7 @@ namespace mu
 				int32 SectionIndex = 0;
 				int32 MeshID = 0;
 
-				Issued = MakeShared<FMeshExternalLoadTask>(item, FullId, LODIndex, SectionIndex, MeshID);
+				Issued = MakeShared<FMeshExternalLoadTask>(Item, FullId, LODIndex, SectionIndex, MeshID);
 			}
 
 			break;
@@ -2653,105 +2662,105 @@ namespace mu
 
 		case EOpType::IM_PIXELFORMAT:
 		{
-			if (item.Stage == 1)
+			if (Item.Stage == 1)
 			{
-				OP::ImagePixelFormatArgs Args = Program.GetOpArgs<OP::ImagePixelFormatArgs>(item.At);
-				Issued = MakeShared<FImagePixelFormatTask>(item, Args);
+				OP::ImagePixelFormatArgs Args = Program.GetOpArgs<OP::ImagePixelFormatArgs>(Item.At);
+				Issued = MakeShared<FImagePixelFormatTask>(Item, Args);
 			}
 			break;
 		}
 
 		case EOpType::IM_LAYERCOLOUR:
 		{
-			if (item.Stage == 1)
+			if (Item.Stage == 1)
 			{
-				OP::ImageLayerColourArgs Args = Program.GetOpArgs<OP::ImageLayerColourArgs>(item.At);
-				Issued = MakeShared<FImageLayerColourTask>(item, Args);
+				OP::ImageLayerColourArgs Args = Program.GetOpArgs<OP::ImageLayerColourArgs>(Item.At);
+				Issued = MakeShared<FImageLayerColourTask>(Item, Args);
 			}
 			break;
 		}
 
 		case EOpType::IM_LAYER:
 		{
-			if ((ExecutionStrategy == EExecutionStrategy::MinimizeMemory && item.Stage == 2)
+			if ((ExecutionStrategy == EExecutionStrategy::MinimizeMemory && Item.Stage == 2)
 				||
-				(ExecutionStrategy != EExecutionStrategy::MinimizeMemory && item.Stage == 1)
+				(ExecutionStrategy != EExecutionStrategy::MinimizeMemory && Item.Stage == 1)
 				)
 			{
-				OP::ImageLayerArgs Args = Program.GetOpArgs<OP::ImageLayerArgs>(item.At);
-				Issued = MakeShared<FImageLayerTask>(item, Args);
+				OP::ImageLayerArgs Args = Program.GetOpArgs<OP::ImageLayerArgs>(Item.At);
+				Issued = MakeShared<FImageLayerTask>(Item, Args);
 			}
 			break;
 		}
 
 		case EOpType::IM_MIPMAP:
 		{
-			if (item.Stage == 1)
+			if (Item.Stage == 1)
 			{
-				OP::ImageMipmapArgs Args = Program.GetOpArgs<OP::ImageMipmapArgs>(item.At);
-				Issued = MakeShared<FImageMipmapTask>(item, Args);
+				OP::ImageMipmapArgs Args = Program.GetOpArgs<OP::ImageMipmapArgs>(Item.At);
+				Issued = MakeShared<FImageMipmapTask>(Item, Args);
 			}
 			break;
 		}
 
 		case EOpType::IM_SWIZZLE:
 		{
-			if (item.Stage == 1)
+			if (Item.Stage == 1)
 			{
-				OP::ImageSwizzleArgs Args = Program.GetOpArgs<OP::ImageSwizzleArgs>(item.At);
-				Issued = MakeShared<FImageSwizzleTask>(item, Args);
+				OP::ImageSwizzleArgs Args = Program.GetOpArgs<OP::ImageSwizzleArgs>(Item.At);
+				Issued = MakeShared<FImageSwizzleTask>(Item, Args);
 			}
 			break;
 		}
 
 		case EOpType::IM_SATURATE:
 		{
-			if (item.Stage == 1)
+			if (Item.Stage == 1)
 			{
-				OP::ImageSaturateArgs Args = Program.GetOpArgs<OP::ImageSaturateArgs>(item.At);
-				Issued = MakeShared<FImageSaturateTask>(item, Args);
+				OP::ImageSaturateArgs Args = Program.GetOpArgs<OP::ImageSaturateArgs>(Item.At);
+				Issued = MakeShared<FImageSaturateTask>(Item, Args);
 			}
 			break;
 		}
 
 		case EOpType::IM_INVERT:
 		{
-			if (item.Stage == 1)
+			if (Item.Stage == 1)
 			{
-				OP::ImageInvertArgs Args = Program.GetOpArgs<OP::ImageInvertArgs>(item.At);
-				Issued = MakeShared<FImageInvertTask>(item, Args);
+				OP::ImageInvertArgs Args = Program.GetOpArgs<OP::ImageInvertArgs>(Item.At);
+				Issued = MakeShared<FImageInvertTask>(Item, Args);
 			}
 			break;
 		}
 
 		case EOpType::IM_RESIZE:
 		{
-			if (item.Stage == 1)
+			if (Item.Stage == 1)
 			{
-				OP::ImageResizeArgs Args = Program.GetOpArgs<OP::ImageResizeArgs>(item.At);
-				Issued = MakeShared<FImageResizeTask>(item, Args);
+				OP::ImageResizeArgs Args = Program.GetOpArgs<OP::ImageResizeArgs>(Item.At);
+				Issued = MakeShared<FImageResizeTask>(Item, Args);
 			}
 			break;
 		}
 
 		case EOpType::IM_RESIZEREL:
 		{
-			if (item.Stage == 1)
+			if (Item.Stage == 1)
 			{
-				OP::ImageResizeRelArgs Args = Program.GetOpArgs<OP::ImageResizeRelArgs>(item.At);
-				Issued = MakeShared<FImageResizeRelTask>(item, Args);
+				OP::ImageResizeRelArgs Args = Program.GetOpArgs<OP::ImageResizeRelArgs>(Item.At);
+				Issued = MakeShared<FImageResizeRelTask>(Item, Args);
 			}
 			break;
 		}
 
 		case EOpType::IM_COMPOSE:
 		{
-			if ((ExecutionStrategy == EExecutionStrategy::MinimizeMemory && item.Stage == 3) ||
-				(ExecutionStrategy != EExecutionStrategy::MinimizeMemory && item.Stage == 2))
+			if ((ExecutionStrategy == EExecutionStrategy::MinimizeMemory && Item.Stage == 3) ||
+				(ExecutionStrategy != EExecutionStrategy::MinimizeMemory && Item.Stage == 2))
 			{
-				OP::ImageComposeArgs Args = Program.GetOpArgs<OP::ImageComposeArgs>(item.At);
-				TSharedPtr<const FLayout> ComposeLayout = StaticCastSharedPtr<const FLayout>( HeapData[item.CustomState].Resource);
-				Issued = MakeShared<FImageComposeTask>(item, Args, ComposeLayout);
+				OP::ImageComposeArgs Args = Program.GetOpArgs<OP::ImageComposeArgs>(Item.At);
+				TSharedPtr<const FLayout> ComposeLayout = StaticCastSharedPtr<const FLayout>( HeapData[Item.CustomState].Resource);
+				Issued = MakeShared<FImageComposeTask>(Item, Args, ComposeLayout);
 			}
 			break;
 		}
diff --git a/Engine/Plugins/Mutable/Source/MutableRuntime/Private/MuR/Image.cpp b/Engine/Plugins/Mutable/Source/MutableRuntime/Private/MuR/Image.cpp
--- a/Engine/Plugins/Mutable/Source/MutableRuntime/Private/MuR/Image.cpp
+++ b/Engine/Plugins/Mutable/Source/MutableRuntime/Private/MuR/Image.cpp
@@ -9,6 +9,8 @@
 #include "MuR/MutableMath.h"
 #include "MuR/MutableTrace.h"
 
+#include "ImageCore.h"
+
 namespace mu
 {
 
@@ -97,47 +99,78 @@ namespace mu
 
 	TSharedPtr<FImage> FImageOperator::ExtractMip(const FImage* This, int32 Mip)
 	{
+		MUTABLE_CPUPROFILER_SCOPE(ExtractMip);
+		
 		if (Mip == 0 && This->GetLODCount() == 1)
 		{
 			return CloneImage(This);
 		}
 
 		FIntVector2 MipSize = This->CalculateMipSize(Mip);
-
-		constexpr int32 Quality = 4;
-
-		if (This->GetLODCount() > Mip)
+		
+		if (Mip < This->GetLODCount())
 		{
-			TArrayView<const uint8> SrcView = This->DataStorage.GetLOD(Mip); 
+			TSharedPtr<FImage> Result = CreateImage(MipSize.X, MipSize.Y, 1, This->GetFormat(), EInitializationType::NotInitialized);
+			Result->Flags = This->Flags;
 
-			TSharedPtr<FImage> ResultImage = CreateImage(MipSize.X, MipSize.Y, 1, This->GetFormat(), EInitializationType::NotInitialized);
-			ResultImage->Flags = This->Flags;
-			
-			// Probably an RLE texture, resize.
-			if (ResultImage->DataStorage.IsEmpty())
+			// RLE formats will not allocate memory when created. Resize manually.
+			if (Result->DataStorage.IsEmpty())
 			{
-				ResultImage->DataStorage.ResizeLOD(0, SrcView.Num());
+				Result->DataStorage.ResizeLOD(0, This->DataStorage.GetLOD(Mip).Num());
 			}
 
-			TArrayView<uint8> DestView = ResultImage->DataStorage.GetLOD(0);
+			TArrayView<const uint8> SrcView = This->DataStorage.GetLOD(Mip);
+			TArrayView<uint8> DstView = Result->DataStorage.GetLOD(0);
 			
-			check(DestView.Num() == SrcView.Num())
-			FMemory::Memcpy(DestView.GetData(), SrcView.GetData(), DestView.Num());
-			
-			return ResultImage;
-		}
+			check(SrcView.Num() == DstView.Num());
+			FMemory::Memcpy(DstView.GetData(), SrcView.GetData(), DstView.Num());
 
-		// We need to generate the mip
-		// \TODO: optimize, quality
+			return Result;	
+		}
 
 		TSharedPtr<FImage> Resized = CreateImage(MipSize.X, MipSize.Y, 1, This->GetFormat(), EInitializationType::NotInitialized);
-		ImageResizeLinear(Resized.Get(), Quality, This);
 		
-		//TODO: Do we realy need to extract the mip again? The resized Image looks already to be what we want.
-		TSharedPtr<FImage> Result = ExtractMip(Resized.Get(), 0);
-		ReleaseImage(Resized);
+		::FImageView SourceView;
+		::FImageView DestView;
 
-		return Result;
+		// We don't know the gamma space here, set to linear so no conversions are done to be consistent
+		// with old behaviour.  
+		switch(This->GetFormat())
+		{
+		case EImageFormat::L_UByte:
+		{
+			SourceView = ::FImageView((void*)This->DataStorage.GetLOD(0).GetData(),
+					This->GetSizeX(), This->GetSizeY(), 1, ERawImageFormat::Type::G8, EGammaSpace::Linear);
+
+			DestView = ::FImageView((void*)Resized->DataStorage.GetLOD(0).GetData(),
+					Resized->GetSizeX(), Resized->GetSizeY(), 1, ERawImageFormat::Type::G8, EGammaSpace::Linear);
+
+			break;
+		}
+		case EImageFormat::RGBA_UByte:
+		case EImageFormat::BGRA_UByte:
+		{
+			SourceView = ::FImageView((void*)This->DataStorage.GetLOD(0).GetData(),
+					This->GetSizeX(), This->GetSizeY(), 1, ERawImageFormat::Type::BGRA8, EGammaSpace::Linear);
+
+			DestView = ::FImageView((void*)Resized->DataStorage.GetLOD(0).GetData(),
+					Resized->GetSizeX(), Resized->GetSizeY(), 1, ERawImageFormat::Type::BGRA8, EGammaSpace::Linear);
+
+			break;
+		}
+		default:
+			
+			checkf(false, TEXT("Image format not supported."));
+		}		
+
+		if (SourceView.RawData && DestView.RawData)
+		{
+			// TODO: Consider using a better resize filter.
+			MUTABLE_CPUPROFILER_SCOPE(ImageResize);
+			FImageCore::ResizeImage(SourceView, DestView, FImageCore::EResizeImageFilter::Bilinear);
+		}
+
+		return Resized;
 	}
 	
 	uint16 FImage::GetSizeX() const
diff --git a/Engine/Plugins/Mutable/Source/MutableRuntime/Private/MuR/ImageDataStorage.cpp b/Engine/Plugins/Mutable/Source/MutableRuntime/Private/MuR/ImageDataStorage.cpp
--- a/Engine/Plugins/Mutable/Source/MutableRuntime/Private/MuR/ImageDataStorage.cpp
+++ b/Engine/Plugins/Mutable/Source/MutableRuntime/Private/MuR/ImageDataStorage.cpp
@@ -319,7 +319,7 @@ namespace mu
 		}
 	}
 
-	void FImageDataStorage::ResizeLOD(int32 LODIndex, int32 NewSizeBytes)
+	void FImageDataStorage::ResizeLOD(int32 LODIndex, int32 NewSizeBytes, EAllowShrinking AllowShrinking)
 	{
 		check(!IsVoid());
 		check(NewSizeBytes >= 0);
@@ -329,7 +329,7 @@ namespace mu
 
 		if (LODIndex < FirstCompactedTailLOD)
 		{
-			Buffer.SetNum(NewSizeBytes);
+			Buffer.SetNum(NewSizeBytes, AllowShrinking);
 			return;
 		}
 
@@ -345,7 +345,7 @@ namespace mu
 		// If the buffer is growing we need to resize before moving the content. 
 		if (SizeDifference > 0)
 		{
-			Buffer.SetNum(NewBufferSize);
+			Buffer.SetNum(NewBufferSize, AllowShrinking);
 		}
 
 		// Move content and update tail offsets.
@@ -365,7 +365,7 @@ namespace mu
 		// If the buffer is shrinking we need to resize after moving the content. 
 		if (SizeDifference < 0)
 		{
-			Buffer.SetNum(NewBufferSize);
+			Buffer.SetNum(NewBufferSize, AllowShrinking);
 		}
 
 		// Tail buffer offsets are represented with uint16 for compactness, check any overflow.
diff --git a/Engine/Plugins/Mutable/Source/MutableRuntime/Public/MuR/ImageDataStorage.h b/Engine/Plugins/Mutable/Source/MutableRuntime/Public/MuR/ImageDataStorage.h
--- a/Engine/Plugins/Mutable/Source/MutableRuntime/Public/MuR/ImageDataStorage.h
+++ b/Engine/Plugins/Mutable/Source/MutableRuntime/Public/MuR/ImageDataStorage.h
@@ -119,7 +119,7 @@ public:
 	/**
 	 * Resizes LODIndex Data to NewSizeBytes
 	 */
-	MUTABLERUNTIME_API void ResizeLOD(int32 LODIndex, int32 NewSizeBytes);
+	MUTABLERUNTIME_API void ResizeLOD(int32 LODIndex, int32 NewSizeBytes, EAllowShrinking AllowShrinking = EAllowShrinking::Default);
 
 	FORCEINLINE const FImageArray& GetInternalArray(int32 LODIndex) const
 	{	
diff --git a/Engine/Plugins/Mutable/Source/MutableTools/Private/MuT/ASTOpConstantResource.cpp b/Engine/Plugins/Mutable/Source/MutableTools/Private/MuT/ASTOpConstantResource.cpp
--- a/Engine/Plugins/Mutable/Source/MutableTools/Private/MuT/ASTOpConstantResource.cpp
+++ b/Engine/Plugins/Mutable/Source/MutableTools/Private/MuT/ASTOpConstantResource.cpp
@@ -456,125 +456,144 @@ namespace mu
 			return Program.ConstantMeshes.Add(MeshContentRange);
 		}   
 
-
 		/** Adds a constant image data to a program and returns its constant index. */
-		uint32 AddConstantImage(FProgram& Program, const TSharedPtr<const FImage>& pImage, FLinkerOptions& Options)
+		uint32 AddConstantImage(FProgram& Program, const TSharedPtr<const FImage>& Image, FLinkerOptions& Options)
 		{
 			MUTABLE_CPUPROFILER_SCOPE(AddConstantImage);
 
-			check(pImage->GetSizeX() * pImage->GetSizeY() > 0);
-			
-			// Mips to store
-			int32 MipsToStore = 1;
-
-			int32 FirstLODIndexIndex = Program.ConstantImageLODIndices.Num();
-
+			check(Image->GetSizeX() * Image->GetSizeY() > 0);
+		
 			FImageOperator& ImOp = Options.ImageOperator;
-			TSharedPtr<const FImage> pMip;
 
 			if (!Options.bSeparateImageMips)
-			{
-				pMip = pImage;
+			{	
+				// Use a map-based deduplication only if we are splitting mips.
+				int32 ImageIndex = Program.ConstantImageLODsPermanent.Add(Image);
+				Options.ImageConstantMipMap.Add(Image, ImageIndex);
+				
+				FImageLODRange LODRange
+				{
+					.FirstIndex    = Program.ConstantImageLODIndices.Add({ uint32(ImageIndex), 0 }),
+					.ImageSizeX    = Image->GetSizeX(),
+					.ImageSizeY    = Image->GetSizeY(),
+					.LODCount      = (uint8)Image->GetLODCount(),
+					.NumLODsInTail = (uint8)Image->GetLODCount(),
+					.Flags         = Image->Flags,
+					.ImageFormat   = Image->GetFormat(),
+				};
+
+				return Program.ConstantImages.Add(LODRange);
 			}
-			else
-			{
-				// We may want the full mipmaps for fragments of images, regardless of the resident mip size, for intermediate operations.
-				// \TODO: Calculate the mip ranges that makes sense to store.
-				int32 MaxMipmaps = FImage::GetMipmapCount(pImage->GetSizeX(), pImage->GetSizeY());
-				MipsToStore = MaxMipmaps;
-
-				// Some images cannot be resized or mipmaped
-				bool bCannotBeScaled = pImage->Flags & FImage::IF_CANNOT_BE_SCALED;
-				if (bCannotBeScaled)
-				{
-					// Store only the mips that we have already calculated. We assume we have calculated them correctly.
-					MipsToStore = pImage->GetLODCount();
-				}
 
-				if (pImage->GetLODCount() == 1)
+			auto AddImageToProgram = [&Program, &Options](const TSharedPtr<FImage>& Image) -> int32
+			{
+				MUTABLE_CPUPROFILER_SCOPE(Deduplicate);
+				const int32* FoundStoreIndex = Options.ImageConstantMipMap.Find(Image);
+				if (FoundStoreIndex)
 				{
-					pMip = pImage;
+					return Program.ConstantImageLODIndices.Add({ uint32(*FoundStoreIndex), 0 });
 				}
 				else
 				{
-					pMip = ImOp.ExtractMip(pImage.Get(), 0);
+					int32 StoreIndex = Program.ConstantImageLODsPermanent.Add(Image);
+					Options.ImageConstantMipMap.Add(Image, StoreIndex);
+
+					return Program.ConstantImageLODIndices.Add({ uint32(StoreIndex), 0 });
 				}
-			}
+			};
+
+			
+			// Compute number of LODs to store.
+			const int32 LODsToStore = Image->Flags & FImage::IF_CANNOT_BE_SCALED 
+					? Image->GetLODCount()
+					: FImage::GetMipmapCount(Image->GetSizeX(), Image->GetSizeY());
 
-			// Temporary uncompressed version of the image, if we need to generate the mips and the source is compressed.
-			TSharedPtr<const FImage> UncompressedMip;
-			EImageFormat UncompressedFormat = GetUncompressedFormat( pMip->GetFormat() );
+			// If the image cannot be scaled prevent from storing LODs that are not present.
+			const int32 FirstLODToStoreInTail = 
+					FMath::Min(LODsToStore, Image->DataStorage.ComputeFirstCompactedTailLOD());
 
-			for (int32 Mip = 0; Mip < MipsToStore; ++Mip)
+			FImageLODRange LODRange
 			{
-				check(pMip->GetFormat() == pImage->GetFormat());
+				.FirstIndex    = Program.ConstantImageLODIndices.Num(),
+				.ImageSizeX    = Image->GetSizeX(),
+				.ImageSizeY    = Image->GetSizeY(),
+				.LODCount      = (uint8)LODsToStore,
+				.NumLODsInTail = (uint8)FMath::Max(0, LODsToStore - FirstLODToStoreInTail),
+				.Flags         = Image->Flags,
+				.ImageFormat   = Image->GetFormat(),
+			};
 
-				// Ensure unique at mip level
-				int32 MipIndex = -1;
+			constexpr int32 MaxQuality = 4;
 
-				// Use a map-based deduplication only if we are splitting mips.
-				if (Options.bSeparateImageMips)
-				{
-					MUTABLE_CPUPROFILER_SCOPE(Deduplicate);
+			int32 LOD = 0;
 
-					const int32* IndexPtr = Options.ImageConstantMipMap.Find(pMip);
-					if (IndexPtr)
-					{
-						MipIndex = *IndexPtr;
-					}
-				}
+			// Special case for the LODs available in Image that will be stored as a single rom per LOD, just a copy. 
+			for (; LOD < FirstLODToStoreInTail && LOD < Image->GetLODCount(); ++LOD)
+			{
+				AddImageToProgram(ImOp.ExtractMip(Image.Get(), LOD));
+			}
 
-				if (MipIndex<0)
-				{
-					MipIndex = Program.ConstantImageLODsPermanent.Add(pMip);
-					Options.ImageConstantMipMap.Add(pMip, MipIndex);
-				}
+			EImageFormat ImageUnCompressedFormat = GetUncompressedFormat(Image->GetFormat());
+
+			// Extract Mip does not support RGB images if it needs to resize.
+			ImageUnCompressedFormat = ImageUnCompressedFormat == EImageFormat::RGB_UByte
+					? EImageFormat::RGBA_UByte
+					: ImageUnCompressedFormat;
+
+			TSharedPtr<const FImage> RawImage = Image;
+			if (ImageUnCompressedFormat != Image->GetFormat())
+			{
+				int32 FirstLODToDecompress = FMath::Min(Image->GetLODCount() - 1, LOD);
+				RawImage = ImOp.ImagePixelFormat(MaxQuality, RawImage.Get(), ImageUnCompressedFormat, FirstLODToDecompress);
+				RawImage = ImOp.ExtractMip(RawImage.Get(), LOD - FirstLODToDecompress);
+			}
+			else
+			{
+				RawImage = ImOp.ExtractMip(RawImage.Get(), LOD);
+			}
 
-				Program.ConstantImageLODIndices.Add({ uint32(MipIndex), 0 });
+			// Store LODs that will be stored as a single rom per LOD and need to be generated.
+			for (; LOD < FirstLODToStoreInTail; ++LOD)
+			{
+				AddImageToProgram(ImOp.ImagePixelFormat(MaxQuality, RawImage.Get(), Image->GetFormat()));
 
-				// Generate next mip if necessary
-				if (Mip + 1 < MipsToStore)
+				// This is wrong, we are losing signal, but compilation performance gets badly affected on
+				// models with lots of images if we scale down from the original image.  
+				RawImage = ImOp.ExtractMip(RawImage.Get(), 1);
+			}
+
+			// Store LODs that will be stored together in a single rom.
+			const int32 RawImageLOD = LOD;
+			if (LOD < LODsToStore)
+			{
+				TSharedPtr<FImage> ImageToStore = ImOp.ExtractMip(RawImage.Get(), LOD - RawImageLOD);
+				ImageToStore->DataStorage.SetNumLODs(LODRange.NumLODsInTail);
+				++LOD;
+				
+				for (; LOD < LODsToStore; ++LOD)
 				{
-					TSharedPtr<FImage> NewMip;
-					if (Mip+1 < pImage->GetLODCount())
-					{
-						// Extract directly from source image
-						NewMip = ImOp.ExtractMip(pImage.Get(), Mip + 1);
-					}
-					else
+					TSharedPtr<FImage> SrcLODImage = ImOp.ExtractMip(RawImage.Get(), LOD - RawImageLOD);
+
+					const int32 DestLOD = LOD - FirstLODToStoreInTail;
+					
+					// Some formats that are not block based will not allocate data when setting the number of LODs.
+					// NOTE: We are working with uncompressed formats, probably this is not needed.
+					if (ImageToStore->DataStorage.GetLOD(DestLOD).IsEmpty())
 					{
-						// Generate from the last mip.
-						if (UncompressedFormat!=pMip->GetFormat())
-						{
-							int32 Quality = 4; // TODO
-
-							if (!UncompressedMip)
-							{
-								UncompressedMip = ImOp.ImagePixelFormat(Quality, pMip.Get(), UncompressedFormat);
-							}
-
-							UncompressedMip = ImOp.ExtractMip(UncompressedMip.Get(), 1);
-							NewMip = ImOp.ImagePixelFormat(Quality, UncompressedMip.Get(), pMip->GetFormat());
-						}
-						else
-						{
-							NewMip = ImOp.ExtractMip(pMip.Get(), 1);
-						}
+						ImageToStore->DataStorage.ResizeLOD(DestLOD, SrcLODImage->DataStorage.GetLOD(0).Num());
 					}
-					check(NewMip);
+					
+					TArrayView<uint8> SrcView  = SrcLODImage->DataStorage.GetLOD(0);
+					TArrayView<uint8> DestView = ImageToStore->DataStorage.GetLOD(DestLOD);
 
-					pMip = NewMip;
+					check(DestView.Num() == SrcView.Num());
+					FMemory::Memcpy(DestView.GetData(), SrcView.GetData(), DestView.Num());	
 				}
+
+				AddImageToProgram(ImOp.ImagePixelFormat(MaxQuality, ImageToStore.Get(), Image->GetFormat()));
 			}
 
-			FImageLODRange LODRange;
-			LODRange.FirstIndex = FirstLODIndexIndex;
-			LODRange.LODCount = MipsToStore;
-			LODRange.ImageFormat = pImage->GetFormat();
-			LODRange.ImageSizeX = pImage->GetSizeX();
-			LODRange.ImageSizeY = pImage->GetSizeY();
-			uint32 ImageIndex = Program.ConstantImages.Add(LODRange);
-			return ImageIndex;
+			return Program.ConstantImages.Add(LODRange);
 		}
 	}
 
diff --git a/Engine/Plugins/Mutable/Source/MutableTools/Private/MuT/Compiler.cpp b/Engine/Plugins/Mutable/Source/MutableTools/Private/MuT/Compiler.cpp
--- a/Engine/Plugins/Mutable/Source/MutableTools/Private/MuT/Compiler.cpp
+++ b/Engine/Plugins/Mutable/Source/MutableTools/Private/MuT/Compiler.cpp
@@ -25,6 +25,7 @@
 #include "Containers/Map.h"
 #include "HAL/LowLevelMemTracker.h"
 #include "HAL/PlatformCrt.h"
+#include "HAL/PlatformTime.h"
 #include "Hash/CityHash.h"
 #include "Logging/LogCategory.h"
 #include "Logging/LogMacros.h"
@@ -429,6 +430,7 @@ namespace mu
     {
         MUTABLE_CPUPROFILER_SCOPE(Compile);
 
+		double StartTimeSeconds = FPlatformTime::Seconds();
 
         TArray< FStateCompilationData > States;
         TSharedPtr<FErrorLog> GenErrorLog;
@@ -694,6 +696,8 @@ namespace mu
 			Program.Serialise(ProgramMemoryArch);
 		}
 
+		UE_LOG(LogMutableCore, Log, TEXT("Compilation Time: %lfs"), FPlatformTime::Seconds() - StartTimeSeconds);
+
 		const double ProgramSize  = double(ProgramSizeStream.GetBufferSize());
 		const double ByteCodeSize = double(Program.ByteCode.Num() + Program.OpAddress.Num() * sizeof(uint32));
 		constexpr double ByteToMiB = 1.0/(1024.0*1024.0); 
@@ -795,7 +799,8 @@ namespace mu
                                 NumTotalLODs - (FirstOptionalLOD - OptionalLODBias), FMath::Max<int32>(0, NumTotalLODs - NumNonOptionalLODs));
                     }
 
-					for (int32 LODRangeIndex = NumOptionalMips; LODRangeIndex < LODRange.LODCount; ++LODRangeIndex)
+					int32 NumLODImages = LODRange.LODCount - LODRange.NumLODsInTail + 1; 
+					for (int32 LODRangeIndex = NumOptionalMips; LODRangeIndex < NumLODImages; ++LODRangeIndex)
 					{
 						FConstantResourceIndex LODIndex = Program.ConstantImageLODIndices[LODRange.FirstIndex + LODRangeIndex];
 						check(!LODIndex.Streamable); // Not classified yet
@@ -805,7 +810,7 @@ namespace mu
 					// Moreover, at least one mip of each image has to be non-highres
 					if (LODRange.LODCount > 0)
 					{
-						FConstantResourceIndex LastLODIndex = Program.ConstantImageLODIndices[LODRange.FirstIndex + LODRange.LODCount - 1];
+						FConstantResourceIndex LastLODIndex = Program.ConstantImageLODIndices[LODRange.FirstIndex + NumLODImages - 1];
 						check(!LastLODIndex.Streamable); // Not classified yet
 						IsLODHighRes[LastLODIndex.Index] = false;
 					}
@@ -834,7 +839,9 @@ namespace mu
 
 					uint32 SourceId = AdditionalData.SourceImagePerConstant[ImageIndex].SourceId;
 
-					for (int32 LODRangeIndex = 0; LODRangeIndex < LODRange.LODCount; ++LODRangeIndex)
+					
+					int32 NumLODImages = LODRange.LODCount - LODRange.NumLODsInTail + 1; 
+					for (int32 LODRangeIndex = 0; LODRangeIndex < NumLODImages; ++LODRangeIndex)
 					{
 						FConstantResourceIndex ResourceIndex = Program.ConstantImageLODIndices[LODRange.FirstIndex + LODRangeIndex];
 						check(!ResourceIndex.Streamable); // Not classified yet
