[Mutable] Fix crash when finalizing a DDC read request due to the request owner being destroyed within the completion callback.

#rnx

diff --git a/Engine/Plugins/Mutable/Source/CustomizableObject/Private/MuCO/MutableStreamRequest.cpp b/Engine/Plugins/Mutable/Source/CustomizableObject/Private/MuCO/MutableStreamRequest.cpp
--- a/Engine/Plugins/Mutable/Source/CustomizableObject/Private/MuCO/MutableStreamRequest.cpp
+++ b/Engine/Plugins/Mutable/Source/CustomizableObject/Private/MuCO/MutableStreamRequest.cpp
@@ -127,9 +127,9 @@ UE::Tasks::FTask FMutableStreamRequest::Stream()
 				Request.Key = ModelStreamableBulkData->DDCKey;
 				Request.Policy = PolicyBuilder.Build();
 
-				TSharedPtr<UE::DerivedData::FRequestOwner>& DDCRequest = HeapMemory->DDCReadRequest.Add_GetRef(MakeShared<FRequestOwner>(EPriority::High));
+				HeapMemory->DDCReadRequest = MakeShared<FRequestOwner>(EPriority::High);
 
-				GetCache().Get(MakeArrayView(&Request, 1), *DDCRequest,
+				GetCache().Get(MakeArrayView(&Request, 1), *HeapMemory->DDCReadRequest,
 					[HeapMemory, BlockReadInfos = MoveTemp(BlockReadInfos), CapturedResourceIds = MoveTemp(ResourceIds)](FCacheGetResponse&& Response) mutable
 					{
 						bool bSuccess = Response.Status == EStatus::Ok;
@@ -158,6 +158,8 @@ UE::Tasks::FTask FMutableStreamRequest::Stream()
 							}
 						}
 
+						HeapMemory->DDCReadRequest->KeepAlive();
+
 						for (UE::Tasks::FTaskEvent& CompletionEvent : HeapMemory->CompletionEvents)
 						{
 							CompletionEvent.Trigger();
diff --git a/Engine/Plugins/Mutable/Source/CustomizableObject/Private/MuCO/MutableStreamRequest.h b/Engine/Plugins/Mutable/Source/CustomizableObject/Private/MuCO/MutableStreamRequest.h
--- a/Engine/Plugins/Mutable/Source/CustomizableObject/Private/MuCO/MutableStreamRequest.h
+++ b/Engine/Plugins/Mutable/Source/CustomizableObject/Private/MuCO/MutableStreamRequest.h
@@ -66,7 +66,7 @@ private:
 		TArray<TSharedPtr<IBulkDataIORequest>> BulkReadRequests;
 
 #if WITH_EDITORONLY_DATA
-		TArray<TSharedPtr<UE::DerivedData::FRequestOwner>> DDCReadRequest;
+		TSharedPtr<UE::DerivedData::FRequestOwner> DDCReadRequest;
 #endif
 
 		bool bIsCancelled = false;
