[Mutable] Merge Skeletons from mesh parameters. Fix crash when converting skeletal meshes from runtime data.

#rnx

diff --git a/Engine/Plugins/Mutable/Source/CustomizableObject/Internal/MuCO/CustomizableObjectSystemPrivate.h b/Engine/Plugins/Mutable/Source/CustomizableObject/Internal/MuCO/CustomizableObjectSystemPrivate.h
--- a/Engine/Plugins/Mutable/Source/CustomizableObject/Internal/MuCO/CustomizableObjectSystemPrivate.h
+++ b/Engine/Plugins/Mutable/Source/CustomizableObject/Internal/MuCO/CustomizableObjectSystemPrivate.h
@@ -539,6 +539,8 @@ struct FInstanceUpdateData
 
 	struct FSkeletonData
 	{
+		TArray<TObjectPtr<USkeleton>> Skeletons;
+
 		TArray<uint16> SkeletonIds;
 
 		TArray<FBone> BonePose;
diff --git a/Engine/Plugins/Mutable/Source/CustomizableObject/Internal/MuCO/UnrealConversionUtils.h b/Engine/Plugins/Mutable/Source/CustomizableObject/Internal/MuCO/UnrealConversionUtils.h
--- a/Engine/Plugins/Mutable/Source/CustomizableObject/Internal/MuCO/UnrealConversionUtils.h
+++ b/Engine/Plugins/Mutable/Source/CustomizableObject/Internal/MuCO/UnrealConversionUtils.h
@@ -155,7 +155,7 @@ namespace UnrealConversionUtils
 
 	void MorphTargetVertexInfoBuffers(FSkeletalMeshLODRenderData& LODResource, const USkeletalMesh& Owner, const mu::FMesh& MutableMesh, const TMap<uint32, FMorphTargetMeshData>& MorphTargetMeshData, int32 LODIndex);
 
-	CUSTOMIZABLEOBJECT_API UE::Tasks::FTask ConvertSkeletalMeshFromRuntimeData(USkeletalMesh* SkeletalMesh, int32 LODIndex, int32 SectionIndex, UModelResources* ModelResources, mu::FMesh* Result);
+	CUSTOMIZABLEOBJECT_API UE::Tasks::FTask ConvertSkeletalMeshFromRuntimeData(USkeletalMesh* SkeletalMesh, int32 LODIndex, int32 SectionIndex, UModelResources* ModelResources, mu::FMesh* Result, TFunction<void(bool bSuccess)>& CompletionCallback);
 
 	void GetSectionClothData(const mu::FMesh& MutableMesh, int32 LODIndex, const TMap<uint32, FClothingMeshData>& ClothingMeshData, TArray<FSectionClothData>& SectionsClothData, int32& NumClothingDataNotFound);
 	
diff --git a/Engine/Plugins/Mutable/Source/CustomizableObject/Private/MuCO/CustomizableObjectInstance.cpp b/Engine/Plugins/Mutable/Source/CustomizableObject/Private/MuCO/CustomizableObjectInstance.cpp
--- a/Engine/Plugins/Mutable/Source/CustomizableObject/Private/MuCO/CustomizableObjectInstance.cpp
+++ b/Engine/Plugins/Mutable/Source/CustomizableObject/Private/MuCO/CustomizableObjectInstance.cpp
@@ -4959,6 +4959,8 @@ UE::Tasks::FTask UCustomizableInstancePrivate::LoadAdditionalAssetsAndData(const
 		}
 
 		// Add Skeletons to merge
+		ComponentData->Skeletons.SkeletonsToMerge = SkeletonData.Skeletons;
+
 		for (const uint32 SkeletonId : SkeletonData.SkeletonIds)
 		{
 			TSoftObjectPtr<USkeleton> AssetPtr = ModelResources.Skeletons.IsValidIndex(SkeletonId) ? ModelResources.Skeletons[SkeletonId] : nullptr;
diff --git a/Engine/Plugins/Mutable/Source/CustomizableObject/Private/MuCO/CustomizableObjectSystem.cpp b/Engine/Plugins/Mutable/Source/CustomizableObject/Private/MuCO/CustomizableObjectSystem.cpp
--- a/Engine/Plugins/Mutable/Source/CustomizableObject/Private/MuCO/CustomizableObjectSystem.cpp
+++ b/Engine/Plugins/Mutable/Source/CustomizableObject/Private/MuCO/CustomizableObjectSystem.cpp
@@ -2270,6 +2270,16 @@ namespace impl
 					SkeletonData.SkeletonIds.AddUnique(Mesh->GetSkeletonID(SkeletonIndex));
 				}
 
+				// Add Skeletons from Mesh Parameters
+				for (const TStrongObjectPtr<USkeletalMesh>& StrongMesh : Mesh->SkeletalMeshes)
+				{
+					USkeletalMesh* SkeletalMesh = StrongMesh.Get();
+					if (ensure(SkeletalMesh))
+					{
+						SkeletonData.Skeletons.AddUnique(SkeletalMesh->GetSkeleton());
+					}
+				}
+
 				// Append BoneMap to the array of BoneMaps
 				const TArray<mu::FBoneName>& BoneMap = Mesh->GetBoneMap();
 				LOD.FirstBoneMap = OperationData->InstanceUpdateData.BoneMaps.Num();
diff --git a/Engine/Plugins/Mutable/Source/CustomizableObject/Private/MuCO/UnrealConversionUtils.cpp b/Engine/Plugins/Mutable/Source/CustomizableObject/Private/MuCO/UnrealConversionUtils.cpp
--- a/Engine/Plugins/Mutable/Source/CustomizableObject/Private/MuCO/UnrealConversionUtils.cpp
+++ b/Engine/Plugins/Mutable/Source/CustomizableObject/Private/MuCO/UnrealConversionUtils.cpp
@@ -14,6 +14,7 @@
 #include "MuR/MutableTrace.h"
 #include "Animation/Skeleton.h"
 #include "Engine/SkeletalMesh.h"
+#include "UObject/StrongObjectPtr.h"
 #include "GPUSkinVertexFactory.h"
 #include "Rendering/SkeletalMeshRenderData.h"
 
@@ -886,12 +887,13 @@ namespace UnrealConversionUtils
 	}
 
 
-	UE::Tasks::FTask ConvertSkeletalMeshFromRuntimeData(USkeletalMesh* SkeletalMesh, int32 LODIndex, int32 SectionIndex, UModelResources* ModelResources, mu::FMesh* Result)
+	UE::Tasks::FTask ConvertSkeletalMeshFromRuntimeData(USkeletalMesh* SkeletalMesh, int32 LODIndex, int32 SectionIndex, UModelResources* ModelResources, mu::FMesh* Result, TFunction<void(bool bSuccess)>& CompletionCallback)
 	{
 		MUTABLE_CPUPROFILER_SCOPE(ConvertSkeletalMeshFromRuntimeData);
 
 		if (!SkeletalMesh || !Result)
 		{
+			CompletionCallback(false);
 			ensure(false);
 			return UE::Tasks::MakeCompletedTask<void>();
 		}
@@ -904,6 +906,8 @@ namespace UnrealConversionUtils
 			// TODO: an empty mesh is returned, but should we use the last LOD instead?
 			UE_LOG(LogMutable, Warning, TEXT("Trying to convert mesh [%s] LOD %d but it only has %d LODs"), 
 				*SkeletalMesh->GetName(), LODIndex, Data->LODRenderData.Num() );
+			
+			CompletionCallback(false);
 
 			return UE::Tasks::MakeCompletedTask<void>();
 		}
@@ -916,9 +920,13 @@ namespace UnrealConversionUtils
 			UE_LOG(LogMutable, Warning, TEXT("Trying to convert mesh [%s] section %d in LOD %d but it only has %d sections"),
 				*SkeletalMesh->GetName(), SectionIndex, LODIndex, LOD->RenderSections.Num());
 
+			CompletionCallback(false);
+			
 			return UE::Tasks::MakeCompletedTask<void>();
 		}
 
+		Result->SkeletalMeshes.AddUnique(TStrongObjectPtr<USkeletalMesh>(SkeletalMesh));
+
 		const FSkelMeshRenderSection& Section = LOD->RenderSections[SectionIndex];
 
 		// Skeleton data
@@ -932,17 +940,11 @@ namespace UnrealConversionUtils
 			USkeleton* Skeleton = SkeletalMesh->GetSkeleton();
 			if (!Skeleton)
 			{
+				CompletionCallback(false);
 				ensure(false);
 				return UE::Tasks::MakeCompletedTask<void>();
 			}
 
-			// Add the skeleton to the list of referenced skeletons and add its index to the mesh
-			if (ModelResources)
-			{
-				const int32 SkeletonID = ModelResources->Skeletons.AddUnique(Skeleton);
-				Result->AddSkeletonID(SkeletonID);
-			}
-
 			const TArray<uint16>& SourceRequiredBones = LOD->RequiredBones;
 
 			// Build RequiredBones array
@@ -1137,7 +1139,7 @@ namespace UnrealConversionUtils
 		}
 
 		UE::Tasks::FTask FinalConversionTask = UE::Tasks::Launch( TEXT("MutableRuntimeMeshConversionFinal"),
-			[Context, SectionIndex, Result]()
+			[Context, SectionIndex, Result, CompletionCallback = CompletionCallback]()
 			{
 				MUTABLE_CPUPROFILER_SCOPE(MeshConversion);
 
@@ -1319,6 +1321,8 @@ namespace UnrealConversionUtils
 				}
 
 				Result->EnsureSurfaceData();
+
+				CompletionCallback(true);
 			},
 
 			// Dependencies
diff --git a/Engine/Plugins/Mutable/Source/CustomizableObject/Private/MuCO/UnrealMutableImageProvider.cpp b/Engine/Plugins/Mutable/Source/CustomizableObject/Private/MuCO/UnrealMutableImageProvider.cpp
--- a/Engine/Plugins/Mutable/Source/CustomizableObject/Private/MuCO/UnrealMutableImageProvider.cpp
+++ b/Engine/Plugins/Mutable/Source/CustomizableObject/Private/MuCO/UnrealMutableImageProvider.cpp
@@ -505,8 +505,12 @@ TTuple<UE::Tasks::FTask, TFunction<void()>> FUnrealMutableResourceProvider::GetM
 		// This may happen when updating mips, and for those operations we don't need skeleton data.
 		//ensure(ModelResources);
 
-		UE::Tasks::FTask ConversionTask = UnrealConversionUtils::ConvertSkeletalMeshFromRuntimeData(SkeletalMesh, LODIndex, SectionIndex, ModelResources, Result.Get());
-		Completion.Trigger();
+		TFunction<void(bool)> Callback = [Completion](bool bSuccess) mutable // Mutable due Trigger not being const
+			{
+				Completion.Trigger();
+			};
+
+		UE::Tasks::FTask ConversionTask = UnrealConversionUtils::ConvertSkeletalMeshFromRuntimeData(SkeletalMesh, LODIndex, SectionIndex, ModelResources, Result.Get(), Callback);
 	});
 
 	return MakeTuple(
diff --git a/Engine/Plugins/Mutable/Source/CustomizableObjectEditor/Private/MuCOE/SMutableCodeViewer.cpp b/Engine/Plugins/Mutable/Source/CustomizableObjectEditor/Private/MuCOE/SMutableCodeViewer.cpp
--- a/Engine/Plugins/Mutable/Source/CustomizableObjectEditor/Private/MuCOE/SMutableCodeViewer.cpp
+++ b/Engine/Plugins/Mutable/Source/CustomizableObjectEditor/Private/MuCOE/SMutableCodeViewer.cpp
@@ -2922,7 +2922,8 @@ namespace
 
 			TSharedPtr<mu::FMesh> Result = MakeShared<mu::FMesh>();
 
-			UE::Tasks::FTask ConversionTask = UnrealConversionUtils::ConvertSkeletalMeshFromRuntimeData(SkeletalMesh, LODIndex, SectionIndex, nullptr, Result.Get());
+			TFunction<void(bool)> DummyCallback = [](bool) {};
+			UE::Tasks::FTask ConversionTask = UnrealConversionUtils::ConvertSkeletalMeshFromRuntimeData(SkeletalMesh, LODIndex, SectionIndex, nullptr, Result.Get(), DummyCallback);
 
 			return MakeTuple(
 				UE::Tasks::Launch( TEXT("FinalizeGetMesh"),
diff --git a/Engine/Plugins/Mutable/Source/MutableRuntime/Internal/MuR/OpMeshMerge.h b/Engine/Plugins/Mutable/Source/MutableRuntime/Internal/MuR/OpMeshMerge.h
--- a/Engine/Plugins/Mutable/Source/MutableRuntime/Internal/MuR/OpMeshMerge.h
+++ b/Engine/Plugins/Mutable/Source/MutableRuntime/Internal/MuR/OpMeshMerge.h
@@ -8,6 +8,8 @@
 #include "MuR/OpMeshFormat.h"
 #include "MuR/MutableTrace.h"
 
+#include "Engine/SkeletalMesh.h"
+#include "UObject/StrongObjectPtr.h"
 
 namespace mu
 {
@@ -1106,6 +1108,15 @@ namespace mu
 		{
 			Result->StreamedResources.AddUnique(pSecond->StreamedResources[Index]);
 		}
+
+		// Skeletal Meshes
+		Result->SkeletalMeshes = pFirst->SkeletalMeshes;
+
+		const int32 NumSkeletalMeshes = pSecond->SkeletalMeshes.Num();
+		for (int32 Index = 0; Index < NumSkeletalMeshes; ++Index)
+		{
+			Result->SkeletalMeshes.AddUnique(pSecond->SkeletalMeshes[Index]);
+		}
 	}
 	
 
diff --git a/Engine/Plugins/Mutable/Source/MutableRuntime/Private/MuR/Mesh.cpp b/Engine/Plugins/Mutable/Source/MutableRuntime/Private/MuR/Mesh.cpp
--- a/Engine/Plugins/Mutable/Source/MutableRuntime/Private/MuR/Mesh.cpp
+++ b/Engine/Plugins/Mutable/Source/MutableRuntime/Private/MuR/Mesh.cpp
@@ -11,6 +11,7 @@
 #include "MuR/MutableTrace.h"
 #include "MuR/OpMeshClipWithMesh.h"
 #include "Spatial/PointHashGrid3.h"
+#include "Engine/SkeletalMesh.h"
 
 namespace mu
 {
@@ -117,6 +118,7 @@ void FMesh::CopyFrom(const FMesh& From, EMeshCopyFlags InFlags)
 	ReferenceID = From.ReferenceID;
 
 	MeshIDPrefix = From.MeshIDPrefix;
+	SkeletalMeshes = From.SkeletalMeshes;
 
 	if (EnumHasAnyFlags(InFlags, EMeshCopyFlags::WithSurfaces))
 	{
diff --git a/Engine/Plugins/Mutable/Source/MutableRuntime/Private/MuR/OpMeshFormat.cpp b/Engine/Plugins/Mutable/Source/MutableRuntime/Private/MuR/OpMeshFormat.cpp
--- a/Engine/Plugins/Mutable/Source/MutableRuntime/Private/MuR/OpMeshFormat.cpp
+++ b/Engine/Plugins/Mutable/Source/MutableRuntime/Private/MuR/OpMeshFormat.cpp
@@ -18,6 +18,9 @@
 #include "MuR/RefCounted.h"
 #include "MuR/MutableRuntimeModule.h"
 
+#include "Engine/SkeletalMesh.h"
+#include "UObject/StrongObjectPtr.h"
+
 #include "GPUSkinPublicDefs.h"
 
 namespace mu
@@ -702,6 +705,7 @@ namespace mu
 
 		Result->Tags = Source->Tags;
 		Result->StreamedResources = Source->StreamedResources;
+		Result->SkeletalMeshes = Source->SkeletalMeshes;
 
 		Result->AdditionalBuffers = Source->AdditionalBuffers;
 
diff --git a/Engine/Plugins/Mutable/Source/MutableRuntime/Public/MuR/Mesh.h b/Engine/Plugins/Mutable/Source/MutableRuntime/Public/MuR/Mesh.h
--- a/Engine/Plugins/Mutable/Source/MutableRuntime/Public/MuR/Mesh.h
+++ b/Engine/Plugins/Mutable/Source/MutableRuntime/Public/MuR/Mesh.h
@@ -18,11 +18,14 @@
 #include "MuR/Skeleton.h"
 #include "Templates/Tuple.h"
 
+#include "UObject/StrongObjectPtr.h"
+
 #include <type_traits>
 
 #define UE_API MUTABLERUNTIME_API
 
 class FString;
+class USkeletalMesh;
 
 namespace mu
 {
@@ -422,6 +425,9 @@ namespace mu
 		 */
 		mutable EMeshFlags Flags = EMeshFlags::None;
 
+		/** Skeletal Meshes involved in the generation of the mesh. Meshes from Mesh Parameters.*/
+		TArray<TStrongObjectPtr<USkeletalMesh>> SkeletalMeshes;
+
 		/** Only valid if the right flags are set, this identifies a referenced mesh. */
 		uint32 ReferenceID = 0;
 
