[mutable] Fixed morph pin names breaking when changing the locale.
#tests Editor
#rnx

diff --git a/Engine/Plugins/Mutable/Source/CustomizableObjectEditor/Private/MuCOE/GenerateMutableSource/GenerateMutableSourceMesh.cpp b/Engine/Plugins/Mutable/Source/CustomizableObjectEditor/Private/MuCOE/GenerateMutableSource/GenerateMutableSourceMesh.cpp
--- a/Engine/Plugins/Mutable/Source/CustomizableObjectEditor/Private/MuCOE/GenerateMutableSource/GenerateMutableSourceMesh.cpp
+++ b/Engine/Plugins/Mutable/Source/CustomizableObjectEditor/Private/MuCOE/GenerateMutableSource/GenerateMutableSourceMesh.cpp
@@ -3661,12 +3661,16 @@ TArray<FMorphNodeData> GenerateMeshMorphStackDefinition(const UEdGraphPin* Pin,
 
 		for (const UEdGraphPin* MorphPin : MorphPins)
 		{
-			if (MorphPin->LinkedTo.Num() && MorphTargetNames.Contains(MorphPin->PinFriendlyName.ToString()))
+			FString MorphPinName = MorphPin->GetName();
+			// The pin name has an extra _Morph suffix at the end when compared to morph target names
+			MorphPinName.RemoveFromEnd(TEXT("_Morph"), ESearchCase::CaseSensitive);
+
+			if (MorphPin->LinkedTo.Num() && MorphTargetNames.Contains(MorphPinName))
 			{
 				// Generate Factor
 				mu::Ptr<mu::NodeScalar> FactorNode = GenerateMorphFactor(Node, *MorphPin, GenerationContext);
 
-				OutMorphData.Add({ Node, MorphPin->PinFriendlyName.ToString(), FactorNode, MeshPin });
+				OutMorphData.Add({ Node, MorphPinName, FactorNode, MeshPin});
 			}
 		}
 	}
