[mutable] Fixed crash when failing to get the MaterialAsset pin because the locale had been changed.
#tests Editor
#rnx

diff --git a/Engine/Plugins/Mutable/Source/CustomizableObjectEditor/Private/MuCOE/Nodes/CustomizableObjectNodeMaterial.cpp b/Engine/Plugins/Mutable/Source/CustomizableObjectEditor/Private/MuCOE/Nodes/CustomizableObjectNodeMaterial.cpp
--- a/Engine/Plugins/Mutable/Source/CustomizableObjectEditor/Private/MuCOE/Nodes/CustomizableObjectNodeMaterial.cpp
+++ b/Engine/Plugins/Mutable/Source/CustomizableObjectEditor/Private/MuCOE/Nodes/CustomizableObjectNodeMaterial.cpp
@@ -895,7 +895,18 @@ UEdGraphPin* UCustomizableObjectNodeMaterial::GetMaterialAssetPin() const
 	}
 	else
 	{
-		return FindPin(PinFriendlyName);
+		// Changing the locale can change the GetPinCategoryFriendlyName() result, leading to failing to find the MeatrialAsset pin.
+		// This is a provisional fix while incorrect FriendlyName usages are removed in UE-294583
+		Pin = FindPin(PinFriendlyName);
+
+		if (Pin)
+		{
+			return Pin;
+		}
+		else
+		{
+			return FindPin(FString("Material_Input_Pin"));
+		}
 	}
 }
 
