[mutable] Make the Mesh Section node's material pin non-localized, and then fix up all the existing ones at load time.
#tests Editor
#rnx

diff --git a/Engine/Plugins/Mutable/Source/CustomizableObject/Internal/MuCO/CustomizableObjectCustomVersion.h b/Engine/Plugins/Mutable/Source/CustomizableObject/Internal/MuCO/CustomizableObjectCustomVersion.h
--- a/Engine/Plugins/Mutable/Source/CustomizableObject/Internal/MuCO/CustomizableObjectCustomVersion.h
+++ b/Engine/Plugins/Mutable/Source/CustomizableObject/Internal/MuCO/CustomizableObjectCustomVersion.h
@@ -195,6 +195,8 @@ struct FCustomizableObjectCustomVersion
 		EnableMutableMacrosNewVersion,
 
 		BaseObjectNodeNamePin,
+
+		FixMaterialNodeMaterialPinIncorrectLocalization,
 		
 		// -----<new versions can be added above this line>-------------------------------------------------
 		VersionPlusOne,
diff --git a/Engine/Plugins/Mutable/Source/CustomizableObjectEditor/Private/MuCOE/Nodes/CustomizableObjectNodeMaterial.cpp b/Engine/Plugins/Mutable/Source/CustomizableObjectEditor/Private/MuCOE/Nodes/CustomizableObjectNodeMaterial.cpp
--- a/Engine/Plugins/Mutable/Source/CustomizableObjectEditor/Private/MuCOE/Nodes/CustomizableObjectNodeMaterial.cpp
+++ b/Engine/Plugins/Mutable/Source/CustomizableObjectEditor/Private/MuCOE/Nodes/CustomizableObjectNodeMaterial.cpp
@@ -330,7 +330,7 @@ void UCustomizableObjectNodeMaterial::AllocateDefaultPins(UCustomizableObjectNod
 	
 	{
 		const FString PinFriendlyName = UEdGraphSchema_CustomizableObject::GetPinCategoryFriendlyName(Schema->PC_MaterialAsset).ToString();
-		const FString PinName = PinFriendlyName + FString(TEXT("_Input_Pin"));
+		const FString PinName = FString(TEXT("Material_Input_Pin"));
 		UEdGraphPin* TableMaterialPin = CustomCreatePin(EGPD_Input, Schema->PC_MaterialAsset, FName(*PinName));
 		TableMaterialPin->PinFriendlyName = FText::FromString(PinFriendlyName);
 		TableMaterialPin->bDefaultValueIsIgnored = true;
@@ -790,6 +790,23 @@ void UCustomizableObjectNodeMaterial::BackwardsCompatibleFixup(int32 Customizabl
 			EnableTagsPinRef = TagsPin;
 		}
 	}
+
+	if (CustomizableObjectCustomVersion == FCustomizableObjectCustomVersion::FixMaterialNodeMaterialPinIncorrectLocalization)
+	{
+		for (UEdGraphPin* Pin : GetAllPins())
+		{
+			if (Pin)
+			{
+				if (Pin->PinType.PinCategory == UEdGraphSchema_CustomizableObject::PC_MaterialAsset && 
+					Pin->Direction == EGPD_Input && Pin->PinName.ToString().EndsWith(TEXT("_Input_Pin")) &&
+					Pin->PinName != FString(TEXT("Material_Input_Pin")))
+				{
+					FString PinName = FString(TEXT("Material_Input_Pin"));
+					Pin->PinName = FName(*PinName);
+				}
+			}
+		}
+	}
 }
 
 
@@ -895,7 +912,7 @@ UEdGraphPin* UCustomizableObjectNodeMaterial::GetMaterialAssetPin() const
 	}
 	else
 	{
-		// Changing the locale can change the GetPinCategoryFriendlyName() result, leading to failing to find the MeatrialAsset pin.
+		// Changing the locale can change the GetPinCategoryFriendlyName() result, leading to failing to find the MaterialAsset pin.
 		// This is a provisional fix while incorrect FriendlyName usages are removed in UE-294583
 		Pin = FindPin(PinFriendlyName);
 
