[mutable] Fixed log producing a log spam when high amounts of COIs are present in the level.

#rnx

diff --git a/Engine/Plugins/Mutable/Source/CustomizableObject/Private/MuCO/CustomizableObjectSystem.cpp b/Engine/Plugins/Mutable/Source/CustomizableObject/Private/MuCO/CustomizableObjectSystem.cpp
--- a/Engine/Plugins/Mutable/Source/CustomizableObject/Private/MuCO/CustomizableObjectSystem.cpp
+++ b/Engine/Plugins/Mutable/Source/CustomizableObject/Private/MuCO/CustomizableObjectSystem.cpp
@@ -1386,6 +1386,8 @@ void UpdateSkeletalMesh(const TSharedRef<FUpdateContextPrivate>& Context)
 		CustomizableObjectInstance->PreSetSkeletalMeshNativeDelegate.Broadcast(Params);
 	}
 	
+	uint32 ParentsWithOverrideMaterialsUpdatedCount = 0;
+	
 	UCustomizableInstancePrivate* CustomizableObjectInstancePrivateData = CustomizableObjectInstance->GetPrivate();
 	check(CustomizableObjectInstancePrivateData != nullptr);
 	for (TObjectIterator<UCustomizableObjectInstanceUsage> It; It; ++It)
@@ -1459,7 +1461,8 @@ void UpdateSkeletalMesh(const TSharedRef<FUpdateContextPrivate>& Context)
 				Parent->SetMaterial(Index, NewOverridenMaterials[Index]);
 			}
 
-			UE_LOG(LogMutable, Log, TEXT("Skeletal Mesh Component override materials have been updated."));
+			ParentsWithOverrideMaterialsUpdatedCount++;
+			
 			bAttachedParentUpdated = true;
 		}
 
@@ -1494,6 +1497,11 @@ void UpdateSkeletalMesh(const TSharedRef<FUpdateContextPrivate>& Context)
 			Context->AttachedParentUpdated.Add(Usage);
 		}
 	}
+
+	if (ParentsWithOverrideMaterialsUpdatedCount > 0)
+	{
+		UE_LOG(LogMutable, Log, TEXT("A total of %u Skeletal Mesh Components got their override materials updated."), ParentsWithOverrideMaterialsUpdatedCount);
+	}
 }
 
 
