Replace some usages of FORCEINLINE with inline in Mutable modules.

For context about this change, see this doc: https://docs.google.com/document/d/1LjpT7NpzTecfXdFNyv5KM6sCu67827BRljNi1iaWVSE/edit?tab=t.0#heading=h.fvtz8gjukz7x

Includes modules: CustomizableObject, MutableRuntime

Reviewers were automatically determined. Please comment if anyone else should be added.

#rnx

diff --git a/Engine/Plugins/Mutable/Source/CustomizableObject/Public/MuCO/CustomizableObjectParameterTypeDefinitions.h b/Engine/Plugins/Mutable/Source/CustomizableObject/Public/MuCO/CustomizableObjectParameterTypeDefinitions.h
--- a/Engine/Plugins/Mutable/Source/CustomizableObject/Public/MuCO/CustomizableObjectParameterTypeDefinitions.h
+++ b/Engine/Plugins/Mutable/Source/CustomizableObject/Public/MuCO/CustomizableObjectParameterTypeDefinitions.h
@@ -329,7 +329,7 @@ inline uint32 GetTypeHash(const FCustomizableObjectProjector& Key)
 }
 
 
-FORCEINLINE FArchive& operator<<(FArchive& Ar, FCustomizableObjectProjector& Data)
+inline FArchive& operator<<(FArchive& Ar, FCustomizableObjectProjector& Data)
 {
 	Ar << Data.Position;
 	Ar << Data.Direction;
diff --git a/Engine/Plugins/Mutable/Source/MutableRuntime/Public/MuR/ImageDataStorage.h b/Engine/Plugins/Mutable/Source/MutableRuntime/Public/MuR/ImageDataStorage.h
--- a/Engine/Plugins/Mutable/Source/MutableRuntime/Public/MuR/ImageDataStorage.h
+++ b/Engine/Plugins/Mutable/Source/MutableRuntime/Public/MuR/ImageDataStorage.h
@@ -29,7 +29,7 @@ namespace mu
 
 namespace mu::FImageDataStorageInternal
 {
-	FORCEINLINE int32 ComputeNumLODsForSize(FImageSize Size)
+	inline int32 ComputeNumLODsForSize(FImageSize Size)
 	{
 		return static_cast<int32>(FMath::CeilLogTwo(FMath::Max(Size.X, Size.Y))) + 1;
 	}
@@ -65,12 +65,12 @@ public:
 
 	MUTABLERUNTIME_API bool operator==(const FImageDataStorage& Other) const;
 
-	FORCEINLINE FImageDesc MakeImageDesc() const
+	inline FImageDesc MakeImageDesc() const
 	{
 		return FImageDesc(ImageSize, ImageFormat, NumLODs);
 	}
 
-	FORCEINLINE void InitVoid(const FImageDesc& Desc) 
+	inline void InitVoid(const FImageDesc& Desc) 
 	{
 		Buffers.Empty();	
 		CompactedTailOffsets[0] = TNumericLimits<uint16>::Max();
@@ -80,7 +80,7 @@ public:
 		NumLODs = Desc.m_lods;
 	}
 
-	FORCEINLINE bool IsVoid() const
+	inline bool IsVoid() const
 	{
 		return Buffers.IsEmpty() && CompactedTailOffsets[0] == TNumericLimits<uint16>::Max();
 	}
@@ -90,7 +90,7 @@ public:
 
 	MUTABLERUNTIME_API void CopyFrom(const FImageDataStorage& Other);
 
-	FORCEINLINE int32 GetNumLODs() const
+	inline int32 GetNumLODs() const
 	{
 		return NumLODs;
 	}
@@ -121,21 +121,21 @@ public:
 	 */
 	MUTABLERUNTIME_API void ResizeLOD(int32 LODIndex, int32 NewSizeBytes, EAllowShrinking AllowShrinking = EAllowShrinking::Default);
 
-	FORCEINLINE const FImageArray& GetInternalArray(int32 LODIndex) const
+	inline const FImageArray& GetInternalArray(int32 LODIndex) const
 	{	
 		check(!IsVoid());
 		check(LODIndex < NumLODs);
 		return Buffers[FMath::Min(LODIndex, ComputeFirstCompactedTailLOD())];
 	}
 
-	FORCEINLINE FImageArray& GetInternalArray(int32 LODIndex)
+	inline FImageArray& GetInternalArray(int32 LODIndex)
 	{	
 		check(!IsVoid());
 		check(LODIndex < NumLODs);
 		return Buffers[FMath::Min(LODIndex, ComputeFirstCompactedTailLOD())];
 	}
 
-	FORCEINLINE int32 ComputeFirstCompactedTailLOD() const
+	inline int32 ComputeFirstCompactedTailLOD() const
 	{
 		using namespace FImageDataStorageInternal;
 		return FMath::Max(0, ComputeNumLODsForSize(ImageSize) - NumLODsInCompactedTail); 
diff --git a/Engine/Plugins/Mutable/Source/MutableRuntime/Public/MuR/ImageTypes.h b/Engine/Plugins/Mutable/Source/MutableRuntime/Public/MuR/ImageTypes.h
--- a/Engine/Plugins/Mutable/Source/MutableRuntime/Public/MuR/ImageTypes.h
+++ b/Engine/Plugins/Mutable/Source/MutableRuntime/Public/MuR/ImageTypes.h
@@ -93,7 +93,7 @@ namespace mu
 		uint8 m_lods = 0;
 
 		/** */
-		FORCEINLINE bool operator==(const FImageDesc& Other) const
+		inline bool operator==(const FImageDesc& Other) const
 		{
 			return
 				m_size == Other.m_size && 
diff --git a/Engine/Plugins/Mutable/Source/MutableRuntime/Public/MuR/MemoryTrackingAllocationPolicy.h b/Engine/Plugins/Mutable/Source/MutableRuntime/Public/MuR/MemoryTrackingAllocationPolicy.h
--- a/Engine/Plugins/Mutable/Source/MutableRuntime/Public/MuR/MemoryTrackingAllocationPolicy.h
+++ b/Engine/Plugins/Mutable/Source/MutableRuntime/Public/MuR/MemoryTrackingAllocationPolicy.h
@@ -50,7 +50,7 @@ namespace mu
 			}
 
 			/** Destructor. */
-			FORCEINLINE ~ForAnyElementType()
+			inline ~ForAnyElementType()
 			{
 				CounterType::Get().fetch_sub(AllocSize, std::memory_order_relaxed);
 
@@ -64,7 +64,7 @@ namespace mu
 			ForAnyElementType(const ForAnyElementType&) = delete;
 			ForAnyElementType& operator=(const ForAnyElementType&) = delete;
 
-			FORCEINLINE void MoveToEmpty(ForAnyElementType& Other)
+			inline void MoveToEmpty(ForAnyElementType& Other)
 			{
 				BaseAlloc::ForAnyElementType::MoveToEmpty(Other);
 
@@ -78,7 +78,7 @@ namespace mu
 				Other.AllocSize = 0;
 			}
 
-			FORCEINLINE void ResizeAllocation(
+			inline void ResizeAllocation(
 				SizeType CurrentNum,
 				SizeType NewMax,
 				SIZE_T NumBytesPerElement
@@ -102,7 +102,7 @@ namespace mu
 				class T = BaseAlloc
 				UE_REQUIRES(TAllocatorTraits<T>::SupportsElementAlignment)
 			>
-			FORCEINLINE void ResizeAllocation(
+			inline void ResizeAllocation(
 				SizeType CurrentNum,
 				SizeType NewMax,
 				SIZE_T NumBytesPerElement,
@@ -148,7 +148,7 @@ namespace mu
 			{
 			}
 
-			FORCEINLINE ElementType* GetAllocation() const
+			inline ElementType* GetAllocation() const
 			{
 				return (ElementType*)ForAnyElementType::GetAllocation();
 			}
diff --git a/Engine/Plugins/Mutable/Source/MutableRuntime/Public/MuR/Ptr.h b/Engine/Plugins/Mutable/Source/MutableRuntime/Public/MuR/Ptr.h
--- a/Engine/Plugins/Mutable/Source/MutableRuntime/Public/MuR/Ptr.h
+++ b/Engine/Plugins/Mutable/Source/MutableRuntime/Public/MuR/Ptr.h
@@ -123,7 +123,7 @@ namespace mu
 		}
 
 
-		FORCEINLINE explicit operator bool() const
+		inline explicit operator bool() const
 		{
             return m_ptr != nullptr;
 		}
diff --git a/Engine/Plugins/Mutable/Source/MutableRuntime/Public/MuR/RefCounted.h b/Engine/Plugins/Mutable/Source/MutableRuntime/Public/MuR/RefCounted.h
--- a/Engine/Plugins/Mutable/Source/MutableRuntime/Public/MuR/RefCounted.h
+++ b/Engine/Plugins/Mutable/Source/MutableRuntime/Public/MuR/RefCounted.h
@@ -23,12 +23,12 @@ namespace mu
 	{
 	public:
 
-		FORCEINLINE void IncRef() const
+		inline void IncRef() const
 		{
 			std::atomic_fetch_add_explicit( &m_refCount, 1, std::memory_order_relaxed );
 		}
 
-		FORCEINLINE void DecRef() const
+		inline void DecRef() const
 		{
 			LLM_SCOPE_BYNAME(TEXT("MutableRuntime"));
 
@@ -39,7 +39,7 @@ namespace mu
 			}
 		}
 
-		FORCEINLINE bool IsUnique() const
+		inline bool IsUnique() const
 		{
 			return m_refCount.load(std::memory_order_relaxed)==1;
 		}
@@ -56,7 +56,7 @@ namespace mu
 			m_refCount = 0;
 		}
 
-		FORCEINLINE virtual ~RefCounted() = default;
+		inline virtual ~RefCounted() = default;
 
 	private:
 
@@ -65,13 +65,13 @@ namespace mu
 	};
 
 
-	FORCEINLINE void mutable_ptr_add_ref(const RefCounted* p)
+	inline void mutable_ptr_add_ref(const RefCounted* p)
 	{
 		if (p) p->IncRef();
 	}
 
 
-	FORCEINLINE void mutable_ptr_release(const RefCounted* p)
+	inline void mutable_ptr_release(const RefCounted* p)
 	{
 		if (p)
 		{
