[Mutable] Fixed crash when compiling with the "Only Selected" option. This crash occurs when there is more than one Table node with the same DT and different compilation options.

#rnx

diff --git a/Engine/Plugins/Mutable/Source/CustomizableObjectEditor/Private/MuCOE/GenerateMutableSource/GenerateMutableSourceTable.cpp b/Engine/Plugins/Mutable/Source/CustomizableObjectEditor/Private/MuCOE/GenerateMutableSource/GenerateMutableSourceTable.cpp
--- a/Engine/Plugins/Mutable/Source/CustomizableObjectEditor/Private/MuCOE/GenerateMutableSource/GenerateMutableSourceTable.cpp
+++ b/Engine/Plugins/Mutable/Source/CustomizableObjectEditor/Private/MuCOE/GenerateMutableSource/GenerateMutableSourceTable.cpp
@@ -689,7 +689,7 @@ void RestrictRowNamesToSelectedOption(TArray<FName>& InOutRowNames, const UCusto
 
 				if (SelectedOptionString)
 				{
-					if (!(*SelectedOptionString == FString("None") && TableNode.bAddNoneOption))
+					if (*SelectedOptionString != FString("None")) // The None option will always be added in the following compilation step.
 					{
 						SelectedOptionNames.Add(FName(*SelectedOptionString));
 					}
