[Mutable] Fix performance regression when using mesh reshape operations.

#rnx

diff --git a/Engine/Plugins/Mutable/Source/MutableRuntime/Private/MuR/OpMeshBind.h b/Engine/Plugins/Mutable/Source/MutableRuntime/Private/MuR/OpMeshBind.h
--- a/Engine/Plugins/Mutable/Source/MutableRuntime/Private/MuR/OpMeshBind.h
+++ b/Engine/Plugins/Mutable/Source/MutableRuntime/Private/MuR/OpMeshBind.h
@@ -14,6 +14,8 @@
 #include "MuR/Operations.h"
 #include "Spatial/MeshAABBTree3.h"
 
+#include "PackedNormal.h"
+
 #include "MuR/OpMeshSmoothing.h"
 
 // TODO: Make the handling of rotations an option. It is more expensive on CPU and memory, and for some
@@ -36,7 +38,7 @@ namespace mu
 	{
 		TArray<UE::Math::TVector<PositionCompType>> Positions;
 		TArray<UE::Math::TVector<NormalCompType>> Normals;
-		TArray<UE::Geometry::FIndex3i> Triangles;
+		TArrayView<const UE::Geometry::FIndex3i> Triangles;
 	};
 
 	using FShapeMeshDescriptorBind = FShapeMeshDescriptor<double, float>;
@@ -728,7 +730,9 @@ namespace mu
 
 			BindData.SetNum(MeshVertexCount);
 
-			const UntypedMeshBufferIteratorConst ItPositionBase(BaseMesh->GetVertexBuffers(), EMeshBufferSemantic::Position);
+			MeshBufferIteratorConst<EMeshBufferFormat::Float32, float, 3> PositionIter(BaseMesh->GetVertexBuffers(), EMeshBufferSemantic::Position);
+			TArrayView<const FVector3f> Positions = TArrayView<const FVector3f>(reinterpret_cast<const FVector3f*>(PositionIter.ptr()), BaseMesh->GetVertexCount());
+
 			const UntypedMeshBufferIteratorConst ItNormalBase(BaseMesh->GetVertexBuffers(), EMeshBufferSemantic::Normal);
 
 			// Disable vertex color reads if the color is not used for mask weights.
@@ -743,14 +747,33 @@ namespace mu
 			if (NonRigidClusterIdx < 0)
 			{
 				const int32 VertexCount = BaseMesh->GetVertexCount();
-				for (int32 VertexIndex = 0; VertexIndex < VertexCount; ++VertexIndex)
-				{
-					const FVector3f VertexPosition = (ItPositionBase + VertexIndex).GetAsVec3f();
-					const FVector3f VertexNormal = ItNormalBase.ptr() ? (ItNormalBase + VertexIndex).GetAsVec3f() : FVector3f::ZAxisVector;
-					const float MaskWeight = ItColorBase.ptr() ? GetVertexMaskWeight(ItColorBase + VertexIndex, ColorUsageMasks) : 1.0f;
 
-					BindReshapePoint(ShapeMeshTree, VertexPosition, VertexNormal, MaskWeight, BindData[VertexIndex], BindValidityTolerance);
+				constexpr int32 NumBatchElems = 1 << 5; // Small batches perform better.
+				const int32 NumBatches = FMath::DivideAndRoundUp<int32>(VertexCount, NumBatchElems);
+
+				TArrayView<const FPackedNormal> PackedNormals;
+				if (ItNormalBase.ptr())
+				{
+					PackedNormals = TArrayView<const FPackedNormal>(reinterpret_cast<const FPackedNormal*>(ItNormalBase.ptr()), VertexCount);
 				}
+
+				auto ProcessVertexBatch = [
+					Positions, PackedNormals, ItColorBase, &ShapeMeshTree, &BindData, BindValidityTolerance, &ColorUsageMasks, VertexCount, NumBatchElems
+				] (int32 BatchId)
+					{
+						const int32 BatchBeginVertexIndex = BatchId * NumBatchElems;
+						const int32 BatchEndVertexIndex = FMath::Min(BatchBeginVertexIndex + NumBatchElems, VertexCount);
+
+						for (int32 VertexIndex = BatchBeginVertexIndex; VertexIndex < BatchEndVertexIndex; ++VertexIndex)
+						{
+							const FVector3f VertexNormal = PackedNormals.Num() ? PackedNormals[VertexIndex].ToFVector3f() : FVector3f::ZAxisVector;
+							const float MaskWeight = ItColorBase.ptr() ? GetVertexMaskWeight(ItColorBase + VertexIndex, ColorUsageMasks) : 1.0f;
+
+							BindReshapePoint(ShapeMeshTree, Positions[VertexIndex], VertexNormal, MaskWeight, BindData[VertexIndex], BindValidityTolerance);
+						}
+					};
+
+				ParallelFor(NumBatches, ProcessVertexBatch);
 			}
 			else
 			{
@@ -760,11 +783,10 @@ namespace mu
 				for (int32 I = 0; I < NonRigidVertexCount; ++I)
 				{
 					const int32 VertexIndex = NonRigidCluster[I];
-					const FVector3f VertexPosition = (ItPositionBase + VertexIndex).GetAsVec3f();
 					const FVector3f VertexNormal = ItNormalBase.ptr() ? (ItNormalBase + VertexIndex).GetAsVec3f() : FVector3f::ZAxisVector;
 					const float MaskWeight = ItColorBase.ptr() ? GetVertexMaskWeight(ItColorBase + VertexIndex, ColorUsageMasks) : 1.0f;
 					
-					BindReshapePoint(ShapeMeshTree, VertexPosition, VertexNormal, MaskWeight, BindData[VertexIndex], BindValidityTolerance);
+					BindReshapePoint(ShapeMeshTree, Positions[VertexIndex], VertexNormal, MaskWeight, BindData[VertexIndex], BindValidityTolerance);
 				}
 
 				// Remove data form the non rigid cluster so it ios not processed in the rigid parts binding step.
@@ -1065,21 +1087,28 @@ namespace mu
 			}
 		}
 		// Generate the temp face query data for the shape
-		// TODO: Index data copy may be saved in most cases.
+		TArray<UE::Geometry::FIndex3i> FormattedTriangles;
+
 		{
-			MUTABLE_CPUPROFILER_SCOPE(GenerateTrianglesQueryData);	
-			ShapeMeshDescriptor.Triangles.SetNum(ShapeTriangleCount);
-			// \TODO: Simple but inefficient
 			const UntypedMeshBufferIteratorConst ItIndices(ShapeMesh->GetIndexBuffers(), EMeshBufferSemantic::VertexIndex);
-			for (int32 TriangleIndex = 0; TriangleIndex < ShapeTriangleCount; ++TriangleIndex)
+			if (ItIndices.GetFormat() == EMeshBufferFormat::UInt32 || ItIndices.GetFormat() == EMeshBufferFormat::Int32)
 			{
-				UE::Geometry::FIndex3i Triangle;
-				Triangle.A = int32((ItIndices + TriangleIndex*3 + 0).GetAsUINT32());
-				Triangle.B = int32((ItIndices + TriangleIndex*3 + 1).GetAsUINT32());
-				Triangle.C = int32((ItIndices + TriangleIndex*3 + 2).GetAsUINT32());
+				ShapeMeshDescriptor.Triangles = TArrayView<const UE::Geometry::FIndex3i>(reinterpret_cast<const UE::Geometry::FIndex3i*>(ItIndices.ptr()), ShapeTriangleCount);
+			}
+			else
+			{
+				FormattedTriangles.SetNumUninitialized(ShapeTriangleCount);
+				for (int32 TriangleIndex = 0; TriangleIndex < ShapeTriangleCount; ++TriangleIndex)
+				{
+					UE::Geometry::FIndex3i Triangle;
+					Triangle.A = int32((ItIndices + TriangleIndex * 3 + 0).GetAsUINT32());
+					Triangle.B = int32((ItIndices + TriangleIndex * 3 + 1).GetAsUINT32());
+					Triangle.C = int32((ItIndices + TriangleIndex * 3 + 2).GetAsUINT32());
 
-				ShapeMeshDescriptor.Triangles[TriangleIndex] = Triangle;
-			}	
+					FormattedTriangles[TriangleIndex] = Triangle;
+				}
+				ShapeMeshDescriptor.Triangles = TArrayView<const UE::Geometry::FIndex3i>(FormattedTriangles);
+			}
 		}
 
 		FShapeMeshAdapter ShapeMeshAdapter(ShapeMeshDescriptor);
@@ -1489,21 +1518,28 @@ namespace mu
 			}
 		}
 		// Generate the temp face query data for the shape
-		// TODO: Index data copy may be saved in most cases.
+		TArray<UE::Geometry::FIndex3i> FormattedTriangles;
+
 		{
-			MUTABLE_CPUPROFILER_SCOPE(GenerateTrianglesQueryData);	
-			ShapeMeshDescriptor.Triangles.SetNum(ShapeTriangleCount);
-			// \TODO: Simple but inefficient
 			const UntypedMeshBufferIteratorConst ItIndices(ShapeMesh->GetIndexBuffers(), EMeshBufferSemantic::VertexIndex);
-			for (int32 TriangleIndex = 0; TriangleIndex < ShapeTriangleCount; ++TriangleIndex)
+			if (ItIndices.GetFormat() == EMeshBufferFormat::UInt32 || ItIndices.GetFormat() == EMeshBufferFormat::Int32)
+			{
+				ShapeMeshDescriptor.Triangles = TArrayView<const UE::Geometry::FIndex3i>(reinterpret_cast<const UE::Geometry::FIndex3i*>(ItIndices.ptr()), ShapeTriangleCount);
+			}
+			else
 			{
-				UE::Geometry::FIndex3i Triangle;
-				Triangle.A = int32((ItIndices + TriangleIndex*3 + 0).GetAsUINT32());
-				Triangle.B = int32((ItIndices + TriangleIndex*3 + 1).GetAsUINT32());
-				Triangle.C = int32((ItIndices + TriangleIndex*3 + 2).GetAsUINT32());
+				FormattedTriangles.SetNumUninitialized(ShapeTriangleCount);
+				for (int32 TriangleIndex = 0; TriangleIndex < ShapeTriangleCount; ++TriangleIndex)
+				{
+					UE::Geometry::FIndex3i Triangle;
+					Triangle.A = int32((ItIndices + TriangleIndex * 3 + 0).GetAsUINT32());
+					Triangle.B = int32((ItIndices + TriangleIndex * 3 + 1).GetAsUINT32());
+					Triangle.C = int32((ItIndices + TriangleIndex * 3 + 2).GetAsUINT32());
 
-				ShapeMeshDescriptor.Triangles[TriangleIndex] = Triangle;
-			}	
+					FormattedTriangles[TriangleIndex] = Triangle;
+				}
+				ShapeMeshDescriptor.Triangles = TArrayView<const UE::Geometry::FIndex3i>(FormattedTriangles);
+			}
 		}
 
 		FShapeMeshAdapter ShapeMeshAdapter(ShapeMeshDescriptor);
diff --git a/Engine/Plugins/Mutable/Source/MutableRuntime/Private/MuR/OpMeshReshape.h b/Engine/Plugins/Mutable/Source/MutableRuntime/Private/MuR/OpMeshReshape.h
--- a/Engine/Plugins/Mutable/Source/MutableRuntime/Private/MuR/OpMeshReshape.h
+++ b/Engine/Plugins/Mutable/Source/MutableRuntime/Private/MuR/OpMeshReshape.h
@@ -18,6 +18,8 @@
 #include "MuR/OpMeshSmoothing.h"
 #include "MuR/OpMeshComputeNormals.h"
 
+#include "PackedNormal.h"
+
 #include "Algo/AnyOf.h" 
 
 
@@ -39,9 +41,9 @@ namespace mu
 
 	struct FShapeMeshDescriptorApply
 	{
-		TArray<FVector3f> Positions;
+		TArrayView<const FVector3f> Positions;
 		TArray<FVector3f> Normals;
-		TArray<UE::Geometry::FIndex3i> Triangles;
+		TArrayView<const UE::Geometry::FIndex3i> Triangles;
 	};
 
 	// Methods to actually deform a point
@@ -374,7 +376,9 @@ namespace mu
 		check(Mesh);
 		check(Mesh->GetVertexCount() == BindingData.Num());
 		
-		const UntypedMeshBufferIterator PositionIter(Mesh->GetVertexBuffers(), EMeshBufferSemantic::Position);
+		MeshBufferIterator<EMeshBufferFormat::Float32, float, 3> PositionIter(Mesh->GetVertexBuffers(), EMeshBufferSemantic::Position);
+		TArrayView<FVector3f> Positions = TArrayView<FVector3f>(reinterpret_cast<FVector3f*>(PositionIter.ptr()), Mesh->GetVertexCount());
+
 		const UntypedMeshBufferIterator NormalIter(Mesh->GetVertexBuffers(), EMeshBufferSemantic::Normal);
 		const UntypedMeshBufferIterator TangentIter(Mesh->GetVertexBuffers(), EMeshBufferSemantic::Tangent);
 		const UntypedMeshBufferIterator BiNormalIter(Mesh->GetVertexBuffers(), EMeshBufferSemantic::Binormal);
@@ -416,21 +420,17 @@ namespace mu
 				
 				GetDeform(Shape, Binding, NewPosition, NewNormalPosition);
 
-				const FVector3f OldPosition = (PositionIter + VertexIndex).GetAsVec3f();
+				const FVector3f& OldPosition = Positions[VertexIndex];
   
-                FVector3f PositionDisplacement;
-                if ((OldPosition - Binding.AttachmentPoint).IsNearlyZero())
-                {
-				    PositionDisplacement = (NewPosition - OldPosition) * Binding.Weight;
+				if ((OldPosition - Binding.AttachmentPoint).IsNearlyZero())
+				{
+					Positions[VertexIndex] += (NewPosition - OldPosition) * Binding.Weight;
+				}
+				else
+				{
+					Positions[VertexIndex] += (NewPosition - Binding.AttachmentPoint) * Binding.Weight;
 				}
-                else
-                {
-				    PositionDisplacement = (NewPosition - Binding.AttachmentPoint) * Binding.Weight;
-                }
 
-				// TODO: Separeat rigid from not rigid processing.
-				(PositionIter + VertexIndex).SetFromVec3f(OldPosition + PositionDisplacement);
-				
 				if (bComputeNormal)
 				{
 					const FVector3f OldNormal = (NormalIter + VertexIndex).GetAsVec3f();
@@ -815,44 +815,56 @@ namespace mu
 
 		// Generate the temp vertex query data for the shape
 		// TODO: Vertex data copy may be saved in most cases.
-		TArray<FVector3f> ShapePositions;
 		FShapeMeshDescriptorApply ShapeDescriptor;
 		{
-			MUTABLE_CPUPROFILER_SCOPE(GenerateVertexQueryData);
+			const UntypedMeshBufferIteratorConst ItPosition(ShapeMesh->GetVertexBuffers(), EMeshBufferSemantic::Position);
+			ShapeDescriptor.Positions = TArrayView<const FVector3f>(reinterpret_cast<const FVector3f*>(ItPosition.ptr()), ShapeVertexCount);
 
-			ShapeDescriptor.Positions.SetNum(ShapeVertexCount);
-			ShapeDescriptor.Normals.SetNum(ShapeVertexCount);
 
-			// \TODO: Simple but inefficient
-			const UntypedMeshBufferIteratorConst ItPosition(ShapeMesh->GetVertexBuffers(), EMeshBufferSemantic::Position);
 			const UntypedMeshBufferIteratorConst ItNormal(ShapeMesh->GetVertexBuffers(), EMeshBufferSemantic::Normal);
-			for (int32 ShapeVertexIndex = 0; ShapeVertexIndex < ShapeVertexCount; ++ShapeVertexIndex)
+			ShapeDescriptor.Normals.SetNumUninitialized(ShapeVertexCount);
+			
+			if (ItNormal.GetFormat() == EMeshBufferFormat::PackedDirS8_W_TangentSign)
 			{
-				FVector3f Position = (ItPosition + ShapeVertexIndex).GetAsVec3f();
-				ShapeDescriptor.Positions[ShapeVertexIndex] = Position;
-
-				FVector3f Normal = (ItNormal + ShapeVertexIndex).GetAsVec3f();
-				ShapeDescriptor.Normals[ShapeVertexIndex] = Normal;
+				for (int32 ShapeVertexIndex = 0; ShapeVertexIndex < ShapeVertexCount; ++ShapeVertexIndex)
+				{
+					const FPackedNormal* PackedNormal = reinterpret_cast<const FPackedNormal*>((ItNormal + ShapeVertexIndex).ptr());
+					ShapeDescriptor.Normals[ShapeVertexIndex] = PackedNormal->ToFVector3f();
+				}
+			}
+			else
+			{
+				for (int32 ShapeVertexIndex = 0; ShapeVertexIndex < ShapeVertexCount; ++ShapeVertexIndex)
+				{
+					ShapeDescriptor.Normals[ShapeVertexIndex] = (ItNormal + ShapeVertexIndex).GetAsVec3f();
+				}
 			}
 		}
 
 		// Generate the temp face query data for the shape
 		// TODO: Index data copy may be saved in most cases.
-		{
-			MUTABLE_CPUPROFILER_SCOPE(GenerateTriangleQueryData);
 
-			ShapeDescriptor.Triangles.SetNum(ShapeTriangleCount);
+		TArray<UE::Geometry::FIndex3i> FormattedTriangles;
 
-			// \TODO: Simple but inefficient
+		{
 			const UntypedMeshBufferIteratorConst ItIndices(ShapeMesh->GetIndexBuffers(), EMeshBufferSemantic::VertexIndex);
-			for (int32 TriangleIndex = 0; TriangleIndex < ShapeTriangleCount; ++TriangleIndex)
+			if (ItIndices.GetFormat() == EMeshBufferFormat::UInt32)
+			{
+				ShapeDescriptor.Triangles = TArrayView<const UE::Geometry::FIndex3i>(reinterpret_cast<const UE::Geometry::FIndex3i*>(ItIndices.ptr()), ShapeTriangleCount);
+			}
+			else
 			{
-				UE::Geometry::FIndex3i Triangle;
-				Triangle.A = int32((ItIndices + TriangleIndex*3 + 0).GetAsUINT32());
-				Triangle.B = int32((ItIndices + TriangleIndex*3 + 1).GetAsUINT32());
-				Triangle.C = int32((ItIndices + TriangleIndex*3 + 2).GetAsUINT32());
+				FormattedTriangles.SetNumUninitialized(ShapeTriangleCount);
+				for (int32 TriangleIndex = 0; TriangleIndex < ShapeTriangleCount; ++TriangleIndex)
+				{
+					UE::Geometry::FIndex3i Triangle;
+					Triangle.A = int32((ItIndices + TriangleIndex * 3 + 0).GetAsUINT32());
+					Triangle.B = int32((ItIndices + TriangleIndex * 3 + 1).GetAsUINT32());
+					Triangle.C = int32((ItIndices + TriangleIndex * 3 + 2).GetAsUINT32());
 
-				ShapeDescriptor.Triangles[TriangleIndex] = Triangle;
+					FormattedTriangles[TriangleIndex] = Triangle;
+				}
+				ShapeDescriptor.Triangles = TArrayView<const UE::Geometry::FIndex3i>(FormattedTriangles);
 			}
 		}
 
