[Mutable] Allow load and recompile of the CustomizableObject model when in-game mode.
Fix crash in SCustomizableObjectEditorAdvancedPreviewSettings::LoadProfileEnvironment() due to missing EnvironmentCubeMap.

#rnx

diff --git a/Engine/Plugins/Experimental/HairStrandsMutable/HairStrandsMutable.uplugin b/Engine/Plugins/Experimental/HairStrandsMutable/HairStrandsMutable.uplugin
--- a/Engine/Plugins/Experimental/HairStrandsMutable/HairStrandsMutable.uplugin
+++ b/Engine/Plugins/Experimental/HairStrandsMutable/HairStrandsMutable.uplugin
@@ -23,7 +23,7 @@
 		},
 		{
 			"Name": "HairStrandsMutableEditor",
-			"Type": "Editor",
+			"Type": "UncookedOnly",
 			"LoadingPhase": "Default"
 		}
 
diff --git a/Engine/Plugins/Experimental/RigLogicMutable/RigLogicMutable.uplugin b/Engine/Plugins/Experimental/RigLogicMutable/RigLogicMutable.uplugin
--- a/Engine/Plugins/Experimental/RigLogicMutable/RigLogicMutable.uplugin
+++ b/Engine/Plugins/Experimental/RigLogicMutable/RigLogicMutable.uplugin
@@ -23,7 +23,7 @@
 		},
 		{
 			"Name": "RigLogicMutableEditor",
-			"Type": "Editor",
+			"Type": "UncookedOnly",
 			"LoadingPhase": "Default"
 		}
 
diff --git a/Engine/Plugins/Mutable/Mutable.uplugin b/Engine/Plugins/Mutable/Mutable.uplugin
--- a/Engine/Plugins/Mutable/Mutable.uplugin
+++ b/Engine/Plugins/Mutable/Mutable.uplugin
@@ -20,7 +20,7 @@
 		},
 		{
 			"Name": "MutableTools",
-			"Type": "Editor",
+			"Type": "UncookedOnly",
 			"LoadingPhase": "PreDefault"
 		},
 		{
@@ -30,7 +30,7 @@
 		},
 		{
 			"Name": "CustomizableObjectEditor",
-			"Type": "Editor",
+			"Type": "UncookedOnly",
 			"LoadingPhase": "Default"
 		}
 	],
diff --git a/Engine/Plugins/Mutable/Source/CustomizableObject/Internal/MuCO/ICustomizableObjectEditorModule.h b/Engine/Plugins/Mutable/Source/CustomizableObject/Internal/MuCO/ICustomizableObjectEditorModule.h
--- a/Engine/Plugins/Mutable/Source/CustomizableObject/Internal/MuCO/ICustomizableObjectEditorModule.h
+++ b/Engine/Plugins/Mutable/Source/CustomizableObject/Internal/MuCO/ICustomizableObjectEditorModule.h
@@ -41,26 +41,18 @@ class ICustomizableObjectEditorModule : public IModuleInterface
 public:
 	static ICustomizableObjectEditorModule* Get()
 	{
-		// Prevent access to this module if the game is being played (in Standalone mode for example)
-		if (IsRunningGame())
-		{
-			return nullptr;
-		}
-
-		// If this is a dedicated server no access to the editor module should be allowed.
-		if (IsRunningDedicatedServer())
-		{
-			return nullptr;
-		}
-		
+#if WITH_EDITOR
 		return FModuleManager::LoadModulePtr<ICustomizableObjectEditorModule>(MODULE_NAME_COE);
+#else
+		return nullptr; // IsRunningGame
+#endif
 	}
 	
 	static ICustomizableObjectEditorModule& GetChecked()
 	{
-		// This module is editor-only. DO NOT try to access it during gameplay
-		checkf(!IsRunningGame(), TEXT("Unable to access \"ICustomizableObjectEditorModule\" module. Running standalone game."))
-		checkf(!IsRunningDedicatedServer(), TEXT("Unable to access \"ICustomizableObjectEditorModule\" module. Running dedicated server."))
+#if !WITH_EDITOR
+		check(false); // This module is editor-only. DO NOT try to access it during gameplay
+#endif
 		return FModuleManager::LoadModuleChecked<ICustomizableObjectEditorModule>(MODULE_NAME_COE);
 	}
 
diff --git a/Engine/Plugins/Mutable/Source/CustomizableObject/Private/MuCO/CustomizableObject.cpp b/Engine/Plugins/Mutable/Source/CustomizableObject/Private/MuCO/CustomizableObject.cpp
--- a/Engine/Plugins/Mutable/Source/CustomizableObject/Private/MuCO/CustomizableObject.cpp
+++ b/Engine/Plugins/Mutable/Source/CustomizableObject/Private/MuCO/CustomizableObject.cpp
@@ -663,7 +663,6 @@ void UCustomizableObjectPrivate::LoadCompiledDataFromDisk()
 			AuxMemoryReader << DataTypeHeader;
 
 			if(DataTypeHeader.InternalVersion != GetECustomizableObjectVersionEnumHash()
-				|| (!IsRunningGame() && DataTypeHeader.VersionId != GetVersionId()) 
 				|| (CompiledVersionId.IsValid() && CompiledVersionId != DataTypeHeader.VersionId))
 			{
 				bHasCompiledData = false;
@@ -761,7 +760,6 @@ bool UCustomizableObject::ConditionalAutoCompile()
 		return false;
 	}
 
-	// Get the EditorModule : Will be null if IsRunningGame() is true
 	if (ICustomizableObjectEditorModule* EditorModule = ICustomizableObjectEditorModule::Get())
 	{
 		if (System->IsAutoCompilationSync() && GetPrivate()->Status.Get() == FCustomizableObjectStatusTypes::EState::Loading)
diff --git a/Engine/Plugins/Mutable/Source/CustomizableObject/Private/MuCO/CustomizableObjectInstance.cpp b/Engine/Plugins/Mutable/Source/CustomizableObject/Private/MuCO/CustomizableObjectInstance.cpp
--- a/Engine/Plugins/Mutable/Source/CustomizableObject/Private/MuCO/CustomizableObjectInstance.cpp
+++ b/Engine/Plugins/Mutable/Source/CustomizableObject/Private/MuCO/CustomizableObjectInstance.cpp
@@ -6469,7 +6469,7 @@ void UCustomizableInstancePrivate::BuildMaterials(const TSharedRef<FUpdateContex
 			SkeletalMesh->SetMaterials(Materials);
 
 #if WITH_EDITOR
-			if (RecreateRenderStateOnInstanceComponent[InstanceComponentIndex])
+			if (GEditor && RecreateRenderStateOnInstanceComponent[InstanceComponentIndex])
 			{
 				// Close all open editors for this mesh to invalidate viewports.
 				GEditor->GetEditorSubsystem<UAssetEditorSubsystem>()->CloseAllEditorsForAsset(SkeletalMesh);
diff --git a/Engine/Plugins/Mutable/Source/CustomizableObjectEditor/Internal/ICustomizableObjectEditorModulePrivate.h b/Engine/Plugins/Mutable/Source/CustomizableObjectEditor/Internal/ICustomizableObjectEditorModulePrivate.h
--- a/Engine/Plugins/Mutable/Source/CustomizableObjectEditor/Internal/ICustomizableObjectEditorModulePrivate.h
+++ b/Engine/Plugins/Mutable/Source/CustomizableObjectEditor/Internal/ICustomizableObjectEditorModulePrivate.h
@@ -1,4 +1,4 @@
-﻿// Copyright Epic Games, Inc. All Rights Reserved.
+// Copyright Epic Games, Inc. All Rights Reserved.
 
 #pragma once
 
@@ -14,18 +14,11 @@ class ICustomizableObjectEditorModulePrivate : public ICustomizableObjectEditorM
 public:
 	static ICustomizableObjectEditorModulePrivate* Get()
 	{
-		// Prevent access to this module if the game is being played (in Standalone mode for example)
-		if (IsRunningGame())
-		{
-			return nullptr;
-		}
-		
 		return FModuleManager::LoadModulePtr<ICustomizableObjectEditorModulePrivate>(MODULE_NAME_COE);
 	}
 	
 	static ICustomizableObjectEditorModulePrivate& GetChecked()
 	{
-		check(!IsRunningGame()) // This module is editor-only. DO NOT try to access it during gameplay
 		return FModuleManager::LoadModuleChecked<ICustomizableObjectEditorModulePrivate>(MODULE_NAME_COE);
 	}
 	
diff --git a/Engine/Plugins/Mutable/Source/CustomizableObjectEditor/Private/MuCOE/CustomizableObjectCompiler.cpp b/Engine/Plugins/Mutable/Source/CustomizableObjectEditor/Private/MuCOE/CustomizableObjectCompiler.cpp
--- a/Engine/Plugins/Mutable/Source/CustomizableObjectEditor/Private/MuCOE/CustomizableObjectCompiler.cpp
+++ b/Engine/Plugins/Mutable/Source/CustomizableObjectEditor/Private/MuCOE/CustomizableObjectCompiler.cpp
@@ -423,17 +423,6 @@ void FCustomizableObjectCompiler::EnqueueCompileRequest(const TSharedRef<FCompil
 		return;
 	}
 	
-	if (IsRunningGame())
-	{
-		FCompileCallbackParams Params;
-		Params.bRequestFailed = true;
-		Params.bCompiled = CompileRequest->GetCustomizableObject()->IsCompiled();
-
-		CompileRequest->Callback.ExecuteIfBound(Params);
-		CompileRequest->CallbackNative.ExecuteIfBound(Params);
-		return;
-	}
-	
 	{
 		const FAssetRegistryModule& AssetRegistryModule = FModuleManager::LoadModuleChecked<FAssetRegistryModule>("AssetRegistry");
 		IAssetRegistry& AssetRegistry = AssetRegistryModule.Get();
diff --git a/Engine/Plugins/Mutable/Source/CustomizableObjectEditor/Private/MuCOE/CustomizableObjectCompiler.h b/Engine/Plugins/Mutable/Source/CustomizableObjectEditor/Private/MuCOE/CustomizableObjectCompiler.h
--- a/Engine/Plugins/Mutable/Source/CustomizableObjectEditor/Private/MuCOE/CustomizableObjectCompiler.h
+++ b/Engine/Plugins/Mutable/Source/CustomizableObjectEditor/Private/MuCOE/CustomizableObjectCompiler.h
@@ -36,7 +36,7 @@ namespace mu
 mu::Ptr<mu::NodeObject> GenerateMutableRoot(const UCustomizableObject* Object, FMutableGraphGenerationContext& GenerationContext);
 
 
-class FCustomizableObjectCompiler : public TSharedFromThis<FCustomizableObjectCompiler>, public FTickableEditorObject, public FTickableCookObject, public FGCObject
+class FCustomizableObjectCompiler : public TSharedFromThis<FCustomizableObjectCompiler>, public FTickableEditorObject, public FTickableCookObject, public FTickableGameObject, public FGCObject
 {
 public:
 	/** Check for pending compilation process. Returns true if an object has been updated. */
diff --git a/Engine/Plugins/Mutable/Source/CustomizableObjectEditor/Private/MuCOE/CustomizableObjectEditorModule.cpp b/Engine/Plugins/Mutable/Source/CustomizableObjectEditor/Private/MuCOE/CustomizableObjectEditorModule.cpp
--- a/Engine/Plugins/Mutable/Source/CustomizableObjectEditor/Private/MuCOE/CustomizableObjectEditorModule.cpp
+++ b/Engine/Plugins/Mutable/Source/CustomizableObjectEditor/Private/MuCOE/CustomizableObjectEditorModule.cpp
@@ -661,7 +661,7 @@ void IsCompilationOutOfDate_Async(const TSharedRef<FCompilationOutOfDateContext>
 		const bool bOutOfDate = Context->bVersionDiff || !Context->OutOfDatePackages.IsEmpty() || !Context->AddedPackages.IsEmpty() || !Context->RemovedPackages.IsEmpty();
 		Context->Callback(bOutOfDate, Context->bVersionDiff, Context->OutOfDatePackages, Context->AddedPackages, Context->RemovedPackages);
 	}
-	else
+	else if (GEditor)
 	{
 		GEditor->GetTimerManager()->SetTimerForNextTick([=]()
 		{
@@ -1265,7 +1265,7 @@ bool FCustomizableObjectEditorModule::IsCompiling(const UCustomizableObject& Obj
 
 void FCustomizableObjectEditorModule::OnPreBeginPIE(const bool bIsSimulatingInEditor)
 {
-	if (IsRunningGame() || !UCustomizableObjectSystem::IsActive())
+	if (!UCustomizableObjectSystem::IsActive())
 	{
 		return;
 	}
diff --git a/Engine/Plugins/Mutable/Source/CustomizableObjectEditor/Private/MuCOE/SCustomizableObjectEditorAdvancedPreviewSettings.cpp b/Engine/Plugins/Mutable/Source/CustomizableObjectEditor/Private/MuCOE/SCustomizableObjectEditorAdvancedPreviewSettings.cpp
--- a/Engine/Plugins/Mutable/Source/CustomizableObjectEditor/Private/MuCOE/SCustomizableObjectEditorAdvancedPreviewSettings.cpp
+++ b/Engine/Plugins/Mutable/Source/CustomizableObjectEditor/Private/MuCOE/SCustomizableObjectEditorAdvancedPreviewSettings.cpp
@@ -488,7 +488,10 @@ void SCustomizableObjectEditorAdvancedPreviewSettings::LoadProfileEnvironment()
 				LoadedObject = Redirector->DestinationObject;
 			}
 
-			DefaultSettings->Profiles[i].EnvironmentCubeMap = CastChecked<UTextureCube>(LoadedObject);
+			if (LoadedObject)
+			{
+				DefaultSettings->Profiles[i].EnvironmentCubeMap = CastChecked<UTextureCube>(LoadedObject);
+			}
 		}
 	}
 }
