[mutable] COI Baking : Removed unnecessary duplication of materials when not asking for the bake of all resources of the COI.
- Previously we were duplicating the material/instance even if no mutable textures were being used by it.

#rnx

diff --git a/Engine/Plugins/Mutable/Source/CustomizableObjectEditor/Private/MuCOE/CustomizableObjectInstanceBakingUtils.cpp b/Engine/Plugins/Mutable/Source/CustomizableObjectEditor/Private/MuCOE/CustomizableObjectInstanceBakingUtils.cpp
--- a/Engine/Plugins/Mutable/Source/CustomizableObjectEditor/Private/MuCOE/CustomizableObjectInstanceBakingUtils.cpp
+++ b/Engine/Plugins/Mutable/Source/CustomizableObjectEditor/Private/MuCOE/CustomizableObjectInstanceBakingUtils.cpp
@@ -653,6 +653,12 @@ bool BakeCustomizableObjectInstance(
 			for (int32 MaterialIndex = 0; MaterialIndex < Mesh->GetMaterials().Num(); ++MaterialIndex)
 			{
 				UMaterialInterface* Interface = Mesh->GetMaterials()[MaterialIndex].MaterialInterface;
+				
+				// If not a mutable material/instance continue as we do not want to duplicate resources that have been untouched by mutable
+				if (Interface->HasAnyFlags(RF_Standalone))
+				{
+					continue;
+				}
 
 				// Check that the resource has not yet been processed, and if so, skip it as we only want to duplicate each interface once
 				if (HandledSourceObjects.Contains(Interface))
