[Mutable] Fix crash when finalizing a DDC read request due to the request owner being destroyed within the completion callback.
Fix crash after failed DDC read request.
Fix missing DDC FValues.
Fix performance regression when streaming data from DDC due to the size of the files to download and the decompression method used.

#rnx

diff --git a/Engine/Plugins/Mutable/Source/CustomizableObject/Internal/MuCO/CustomizableObjectPrivate.h b/Engine/Plugins/Mutable/Source/CustomizableObject/Internal/MuCO/CustomizableObjectPrivate.h
--- a/Engine/Plugins/Mutable/Source/CustomizableObject/Internal/MuCO/CustomizableObjectPrivate.h
+++ b/Engine/Plugins/Mutable/Source/CustomizableObject/Internal/MuCO/CustomizableObjectPrivate.h
@@ -924,7 +924,8 @@ namespace MutablePrivate
 		/** Id generated from a hash of the file content + offset to avoid collisions. */
 		uint32 Id = 0;
 
-		uint32 Padding = 0;
+		/** File size */
+		uint64 Size = 0;
 
 		/** List of blocks that are contained in the file, in order. */
 		TArray<FBlock> Blocks;
@@ -1052,6 +1053,7 @@ namespace MutablePrivate
 		TSharedPtr<const mu::FModel, ESPMode::ThreadSafe> Model,
 		FModelStreamableBulkData& ModelStreamableBulkData,
 		uint32 NumFilesPerBucket,
+		bool bAllowSplit,
 		TArray<FFile>& OutBulkDataFiles);
 
 	/** Generate the list of BulkData files with a soft restriction to the size of the files.
@@ -1102,6 +1104,7 @@ struct FModelStreamableBulkData
 	bool bIsStoredInDDC = false;
 	UE::DerivedData::FCacheKey DDCKey = UE::DerivedData::FCacheKey::Empty;
 	UE::DerivedData::ECachePolicy DDCDefaultPolicy = UE::DerivedData::ECachePolicy::Default;
+	TArray<FIoHash> DDCValues;
 #endif
 
 	// File path to stream resources from when not using FByteBulkData or DDC.
@@ -1115,7 +1118,7 @@ struct FModelStreamableBulkData
 		Ar << Data.ModelStreamables;
 		Ar << Data.ClothingStreamables;
 		Ar << Data.RealTimeMorphStreamables;
-
+		Ar << Data.DDCValues;
 		// Don't serialize FByteBulkData manually, the data will be skipped.
 		
 		Ar << Data.FullFilePath;
@@ -1630,6 +1633,8 @@ enum class ECustomizableObjectVersions : int32
 
 	StoreMipTailInASingleRom,
 
+	FixDDCCrashesAndLoadTimes,
+
 	// -----<new versions can be added above this line>--------
 	LastCustomizableObjectVersion
 };
@@ -1895,7 +1900,7 @@ public:
 
 #endif
 	
-	static constexpr int32 DerivedDataVersion = 0x28f3f788;
+	static constexpr int32 DerivedDataVersion = 0x2c537a2f;
 };
 
 #if WITH_EDITOR
diff --git a/Engine/Plugins/Mutable/Source/CustomizableObject/Private/MuCO/CustomizableObject.cpp b/Engine/Plugins/Mutable/Source/CustomizableObject/Private/MuCO/CustomizableObject.cpp
--- a/Engine/Plugins/Mutable/Source/CustomizableObject/Private/MuCO/CustomizableObject.cpp
+++ b/Engine/Plugins/Mutable/Source/CustomizableObject/Private/MuCO/CustomizableObject.cpp
@@ -2425,6 +2425,7 @@ namespace MutablePrivate
 		TSharedPtr<const mu::FModel, ESPMode::ThreadSafe> Model,
 		FModelStreamableBulkData& ModelStreamableBulkData,
 		uint32 NumFilesPerBucket,
+		bool bAllowSplit,
 		TArray<FFile>& OutBulkDataFiles)
 	{
 		MUTABLE_CPUPROFILER_SCOPE(GenerateBulkDataFilesListWithFileLimit);
@@ -2524,12 +2525,13 @@ namespace MutablePrivate
 		}
 		
 		// Function to create the list of bulk data files. Blocks will be grouped by source Id.
-		const auto CreateFileList = [Categories](const FFileCategoryID& CategoryID, const FClassifyNode& Node, TArray<FFile>& OutBulkDataFiles)
+		const auto CreateFileList = [Categories, bAllowSplit](const FFileCategoryID& CategoryID, const FClassifyNode& Node, TArray<FFile>& OutBulkDataFiles)
 			{
 				const FFileCategory* Category = Categories.FindByPredicate(
 					[CategoryID](const FFileCategory& C) { return C.Id == CategoryID; });
 				check(Category);
 
+				const int32 FirstFile = OutBulkDataFiles.Num();
 				int32 NumBulkDataFiles = OutBulkDataFiles.Num();
 				OutBulkDataFiles.Reserve(NumBulkDataFiles + Category->NumFiles);
 
@@ -2537,6 +2539,8 @@ namespace MutablePrivate
 				TArray<int64> BulkDataFileIndex;
 				BulkDataFileIndex.Init(INDEX_NONE, Category->NumFiles);
 
+				uint64 MaxSize = 0;
+
 				const int32 NumBlocks = Node.Blocks.Num();
 				for (const FBlock& Block : Node.Blocks)
 				{
@@ -2559,6 +2563,56 @@ namespace MutablePrivate
 
 					// Add block to the file 
 					OutBulkDataFiles[FileIndex].Blocks.Add(Block);
+					OutBulkDataFiles[FileIndex].Size += Block.Size;
+
+					MaxSize = FMath::Max(OutBulkDataFiles[FileIndex].Size, MaxSize);
+				}
+
+				if (bAllowSplit)
+				{
+					// Split big files
+					constexpr uint64 MinSplitSize = 10 * 1024 * 1024;
+					const int32 PartitionsPerFile = FMath::Min(FMath::DivideAndRoundDown(Category->NumFiles, (uint32)(NumBulkDataFiles - FirstFile)), 3u);
+
+					if (MaxSize > MinSplitSize && PartitionsPerFile > 1)
+					{
+						for (int32 FileIndex = FirstFile; FileIndex < NumBulkDataFiles; ++FileIndex)
+						{
+							FFile& BaseFile = OutBulkDataFiles[FileIndex];
+							if (BaseFile.Size < MinSplitSize)
+							{
+								continue;
+							}
+
+							const uint64 MaxPartitionSize = FMath::DivideAndRoundUp(BaseFile.Size, (uint64)PartitionsPerFile);
+							BaseFile.Size = 0;
+
+							TArray<FBlock> BaseBlocks;
+							Exchange(BaseBlocks, BaseFile.Blocks);
+
+							TArray<int32> NewFileIndices;
+							NewFileIndices.Add(FileIndex);
+							for (int32 Index = 1; Index < PartitionsPerFile; ++Index)
+							{
+								NewFileIndices.Add(OutBulkDataFiles.Num());
+								FFile NewFile = BaseFile;
+								OutBulkDataFiles.Add(MoveTemp(NewFile));
+							}
+
+							int32 NewFileIndex = 0;
+							for (const FBlock& BaseBlock : BaseBlocks)
+							{
+								FFile& NewFile = OutBulkDataFiles[NewFileIndices[NewFileIndex]];
+								NewFile.Blocks.Add(BaseBlock);
+								NewFile.Size += BaseBlock.Size;
+
+								if (NewFile.Size > MaxPartitionSize)
+								{
+									NewFileIndex++;
+								}
+							}
+						}
+					}
 				}
 			};
 
diff --git a/Engine/Plugins/Mutable/Source/CustomizableObject/Private/MuCO/MutableStreamRequest.cpp b/Engine/Plugins/Mutable/Source/CustomizableObject/Private/MuCO/MutableStreamRequest.cpp
--- a/Engine/Plugins/Mutable/Source/CustomizableObject/Private/MuCO/MutableStreamRequest.cpp
+++ b/Engine/Plugins/Mutable/Source/CustomizableObject/Private/MuCO/MutableStreamRequest.cpp
@@ -70,6 +70,11 @@ UE::Tasks::FTask FMutableStreamRequest::Stream()
 
 	bIsStreaming = true;
 
+	if (BlockReadInfos.IsEmpty())
+	{
+		return UE::Tasks::MakeCompletedTask<void>();
+	}
+
 	for (const FBlockReadInfo& Block : BlockReadInfos)
 	{
 		HeapMemory->CompletionEvents.Emplace(TEXT("AsyncReadDataReadyEvent"));
@@ -102,69 +107,70 @@ UE::Tasks::FTask FMutableStreamRequest::Stream()
 			{
 				using namespace UE::DerivedData;
 
-				// Skip loading values by default
-				FCacheRecordPolicyBuilder PolicyBuilder(ECachePolicy::Default | ECachePolicy::SkipData);
+				UE::FSharedString SharedName = UE::FSharedString(TEXT("MutableStreamRequest"));
 
-				TArray<FValueId> ResourceIds;
-				ResourceIds.Reserve(BlockReadInfos.Num());
+				TArray<FCacheGetValueRequest> Requests;
+				Requests.Reserve(BlockReadInfos.Num());
 
-				// Override policy for the resources to load.
+				TArray<uint32> UniqueFileIds;
 				for (const FBlockReadInfo& Block : BlockReadInfos)
 				{
-					FValueId ResourceId = GetDerivedDataValueIdForResource(Block.DataType, Block.FileId, Block.ResourceType, Block.ResourceFlags);
-
-					if (ResourceIds.Find(ResourceId) == INDEX_NONE)
+					if (UniqueFileIds.Find(Block.FileId) == INDEX_NONE)
 					{
-						// Add value policy once.
-						PolicyBuilder.AddValuePolicy(ResourceId, ECachePolicy::Default);
+						FCacheGetValueRequest Request;
+						Request.Name = SharedName;
+						Request.Key = ModelStreamableBulkData->DDCKey;
+						Request.Key.Hash = ModelStreamableBulkData->DDCValues[Block.FileId];
+						Request.UserData = (uint32)Block.FileId;
+						Request.Policy = ModelStreamableBulkData->DDCDefaultPolicy;
+						Requests.Add(MoveTemp(Request));
+
+						UniqueFileIds.Add(Block.FileId);
 					}
-
-					ResourceIds.Add(ResourceId);
 				}
 
-				FCacheGetRequest Request;
-				Request.Name = ModelStreamableBulkData->FullFilePath;
-				Request.Key = ModelStreamableBulkData->DDCKey;
-				Request.Policy = PolicyBuilder.Build();
-
-				HeapMemory->DDCReadRequest = MakeShared<FRequestOwner>(EPriority::High);
-
-				GetCache().Get(MakeArrayView(&Request, 1), *HeapMemory->DDCReadRequest,
-					[HeapMemory, BlockReadInfos = MoveTemp(BlockReadInfos), CapturedResourceIds = MoveTemp(ResourceIds)](FCacheGetResponse&& Response) mutable
+				HeapMemory->DDCRequestOwner = MakeShared<FRequestOwner>(EPriority::High);
+				HeapMemory->DDCRequestOwner->KeepAlive();
+				
+				GetCache().GetValue(Requests, *HeapMemory->DDCRequestOwner,
+					[HeapMemory, BlockReadInfos = MoveTemp(BlockReadInfos)](FCacheGetValueResponse&& Response) mutable
 					{
+						MUTABLE_CPUPROFILER_SCOPE(MutableStreamRequest_DDC_Callback);
 						bool bSuccess = Response.Status == EStatus::Ok;
 
-						if (ensure(bSuccess))
+						if (!bSuccess)
+						{
+							for (UE::Tasks::FTaskEvent& CompletionEvent : HeapMemory->CompletionEvents)
+							{
+								CompletionEvent.Trigger();
+							}
+
+							return;
+						}
+						
+						for (int32 Index = 0; Index < BlockReadInfos.Num(); ++Index)
 						{
-							const int32 NumResources = CapturedResourceIds.Num();
-							for (int32 Index = 0; Index < NumResources; ++Index)
+							const FBlockReadInfo& Block = BlockReadInfos[Index];
+							if (Block.FileId == Response.UserData)
 							{
-								const FCompressedBuffer& CompressedBuffer = Response.Record.GetValue(CapturedResourceIds[Index]).GetData();
-								if (ensure(!CompressedBuffer.IsNull()))
+								const FCompressedBuffer& CompressedBuffer = Response.Value.GetData();
+								if (!CompressedBuffer.IsNull())
 								{
-									const FBlockReadInfo& Block = BlockReadInfos[Index];
 									const uint64 Size = Block.AllocatedMemoryView.Num();
-									if (Size < CompressedBuffer.GetRawSize())
+									if (Size == CompressedBuffer.GetRawSize())
 									{
-										check(CompressedBuffer.GetRawSize() >= Block.Offset + Size);
-										FSharedBuffer DecompressedBuffer = CompressedBuffer.Decompress();
-										FMemory::Memcpy(Block.AllocatedMemoryView.GetData(), reinterpret_cast<const uint8*>(DecompressedBuffer.GetData()) + Block.Offset, Size);
+										bSuccess = CompressedBuffer.TryDecompressTo(MakeMemoryView(Block.AllocatedMemoryView.GetData(), Size));
 									}
-									else
+									else if (Block.Offset + Size <= CompressedBuffer.GetRawSize())
 									{
-										check(CompressedBuffer.TryDecompressTo(MakeMemoryView(Block.AllocatedMemoryView.GetData(), Size)));
+										FCompressedBufferReader BufferReader(CompressedBuffer);
+										bSuccess = BufferReader.TryDecompressTo(MakeMemoryView(Block.AllocatedMemoryView.GetData(), Size), Block.Offset);
 									}
 								}
-							}
-						}
 
-						HeapMemory->DDCReadRequest->KeepAlive();
-
-						for (UE::Tasks::FTaskEvent& CompletionEvent : HeapMemory->CompletionEvents)
-						{
-							CompletionEvent.Trigger();
+								HeapMemory->CompletionEvents[Index].Trigger();
+							}
 						}
-
 					});
 
 				return;
diff --git a/Engine/Plugins/Mutable/Source/CustomizableObject/Private/MuCO/MutableStreamRequest.h b/Engine/Plugins/Mutable/Source/CustomizableObject/Private/MuCO/MutableStreamRequest.h
--- a/Engine/Plugins/Mutable/Source/CustomizableObject/Private/MuCO/MutableStreamRequest.h
+++ b/Engine/Plugins/Mutable/Source/CustomizableObject/Private/MuCO/MutableStreamRequest.h
@@ -66,7 +66,7 @@ private:
 		TArray<TSharedPtr<IBulkDataIORequest>> BulkReadRequests;
 
 #if WITH_EDITORONLY_DATA
-		TSharedPtr<UE::DerivedData::FRequestOwner> DDCReadRequest;
+		TSharedPtr<UE::DerivedData::FRequestOwner> DDCRequestOwner;
 #endif
 
 		bool bIsCancelled = false;
diff --git a/Engine/Plugins/Mutable/Source/CustomizableObject/Private/MuCO/UnrealMutableModelDiskStreamer.cpp b/Engine/Plugins/Mutable/Source/CustomizableObject/Private/MuCO/UnrealMutableModelDiskStreamer.cpp
--- a/Engine/Plugins/Mutable/Source/CustomizableObject/Private/MuCO/UnrealMutableModelDiskStreamer.cpp
+++ b/Engine/Plugins/Mutable/Source/CustomizableObject/Private/MuCO/UnrealMutableModelDiskStreamer.cpp
@@ -355,52 +355,48 @@ mu::FModelReader::FOperationID FUnrealMutableModelBulkReader::BeginReadBlock(con
 #if WITH_EDITOR
 		if (PinnedModelStreamableBulkData->bIsStoredInDDC)
 		{
-			MUTABLE_CPUPROFILER_SCOPE(FUnrealMutableModelBulkStreamer::OpenReadFileDDC);
-
 			using namespace UE::DerivedData;
 
+			UE::FSharedString SharedName = UE::FSharedString(TEXT("UnrealMutableModelBulkReader"));
 
-			// Skip loading values by default
-			FCacheRecordPolicyBuilder PolicyBuilder(ECachePolicy::Default | ECachePolicy::SkipData);
-
-			// Override policy for the resource to load.
-			FValueId ResourceId = GetDerivedDataValueIdForResource(MutablePrivate::EStreamableDataType::Model, Block->FileId, (uint16)ResourceType, Block->Flags);
-			PolicyBuilder.AddValuePolicy(ResourceId, ECachePolicy::Default);
-
-			FCacheGetRequest Request;
-			Request.Name = ObjectData->BulkFilePrefix;
+			FCacheGetValueRequest Request;
+			Request.Name = SharedName;
 			Request.Key = PinnedModelStreamableBulkData->DDCKey;
-			Request.Policy = PolicyBuilder.Build();
-
+			Request.Key.Hash = PinnedModelStreamableBulkData->DDCValues[Block->FileId];
+			Request.Policy = PinnedModelStreamableBulkData->DDCDefaultPolicy;
+			
 			FReadRequest ReadRequest;
 			ReadRequest.DDCReadRequest = MakeShared<FDDCReadRequest>();
+			GetCache().GetValue(MakeArrayView(&Request,1), ReadRequest.DDCReadRequest->RequestOwner,
+				[Request = ReadRequest.DDCReadRequest, CompletionCallbackCapture = *CompletionCallback, Block, Buffer, Size](FCacheGetValueResponse&& Response)
+			{
+				MUTABLE_CPUPROFILER_SCOPE(FUnrealMutableModelBulkStreamer::OpenReadFile_DDC_Callback);
+				bool bSuccess = Response.Status == EStatus::Ok;
 
-			GetCache().Get(MakeArrayView(&Request, 1), ReadRequest.DDCReadRequest->RequestOwner,
-				[CompletionCallbackCapture = *CompletionCallback, Request = ReadRequest.DDCReadRequest, Block, ResourceId, Buffer, Size](FCacheGetResponse&& Response)
+				if (bSuccess)
 				{
-					const bool bSuccess = Response.Status == EStatus::Ok;
-					if (bSuccess)
+					const FCompressedBuffer& CompressedBuffer = Response.Value.GetData();
+					if (ensure(!CompressedBuffer.IsNull()))
 					{
-						const FCompressedBuffer& CompressedBuffer = Response.Record.GetValue(ResourceId).GetData();
-						if (ensure(!CompressedBuffer.IsNull()))
+						if (Size == CompressedBuffer.GetRawSize())
 						{
-							if (Size < CompressedBuffer.GetRawSize())
-							{
-								check(CompressedBuffer.GetRawSize() >= Block->Offset + Size);
-								FSharedBuffer DecompressedBuffer = CompressedBuffer.Decompress();
-								FMemory::Memcpy(Buffer, reinterpret_cast<const uint8*>(DecompressedBuffer.GetData()) + Block->Offset, Size);
-							}
-							else
-							{
-								check(CompressedBuffer.TryDecompressTo(MakeMemoryView(Buffer, Size)));
-							}
-
-							Request->bSuccess = true;
+							bSuccess = CompressedBuffer.TryDecompressTo(MakeMemoryView(Buffer, Size));
+						}
+						else if (Block->Offset + Size <= CompressedBuffer.GetRawSize())
+						{
+							FCompressedBufferReader BufferReader(CompressedBuffer);
+							bSuccess = BufferReader.TryDecompressTo(MakeMemoryView(Buffer, Size), Block->Offset);
+						}
+						else 
+						{
+							bSuccess = false;
 						}
 					}
+				}
 
-					CompletionCallbackCapture(bSuccess);
-				});
+				Request->bSuccess = bSuccess;
+				CompletionCallbackCapture(bSuccess);
+			});
 
 			ObjectData->CurrentReadRequests.Add(Result, ReadRequest);
 		}
@@ -624,9 +620,8 @@ void FUnrealMutableModelBulkWriterCook::CloseWriteFile()
 {
 }
 
-
 FDDCReadRequest::FDDCReadRequest() :
-RequestOwner(UE::DerivedData::EPriority::High)
+	RequestOwner(UE::DerivedData::EPriority::High)
 {
 }
 
diff --git a/Engine/Plugins/Mutable/Source/CustomizableObjectEditor/Private/MuCOE/CustomizableObjectCompileRunnable.cpp b/Engine/Plugins/Mutable/Source/CustomizableObjectEditor/Private/MuCOE/CustomizableObjectCompileRunnable.cpp
--- a/Engine/Plugins/Mutable/Source/CustomizableObjectEditor/Private/MuCOE/CustomizableObjectCompileRunnable.cpp
+++ b/Engine/Plugins/Mutable/Source/CustomizableObjectEditor/Private/MuCOE/CustomizableObjectCompileRunnable.cpp
@@ -424,7 +424,7 @@ void FCustomizableObjectSaveDDRunnable::CachePlatformData()
 		// Generate list of files and update streamable blocks ids and offsets
 		if (Options.bUseBulkData)
 		{
-			MutablePrivate::GenerateBulkDataFilesListWithFileLimit(PlatformData->Model, *PlatformData->ModelStreamableBulkData.Get(), MAX_uint8, PlatformData->BulkDataFiles);
+			MutablePrivate::GenerateBulkDataFilesListWithFileLimit(PlatformData->Model, *PlatformData->ModelStreamableBulkData.Get(), MAX_uint8, false, PlatformData->BulkDataFiles);
 		}
 		else
 		{
@@ -454,44 +454,80 @@ void FCustomizableObjectSaveDDRunnable::StoreCachedPlatformDataInDDC(bool& bStor
 	check(PlatformData->Model.Get() != nullptr);
 	check(DDCKey.Hash.IsZero() == false);
 
-	bStoredSuccessfully = false;
+	UE::FSharedString SharedName = UE::FSharedString(CustomizableObjectName);
+
+	bStoredSuccessfully = true;
 
 	// DDC record
 	FCacheRecordBuilder RecordBuilder(DDCKey);
 
-	// Store streamable resources info as FValues
+	// ModelStreamable  will be modified for the DDC record. Modify a copy
 	FModelStreamableBulkData ModelStreamablesDDC = *PlatformData->ModelStreamableBulkData.Get();
+	
+	// Store streamable resources as FValues.
 	{
-		MUTABLE_CPUPROFILER_SCOPE(SerializeModelStreamables);
-
-		// ModelStreamable  will be modified for the DDC record. Modify a copy
+		MUTABLE_CPUPROFILER_SCOPE(SerializeBulkDataForDDC);
 
 		// Generate list of files and update streamable blocks ids and offsets
 		constexpr int32 MaxDDCFiles = 1 << 13;
-		MutablePrivate::GenerateBulkDataFilesListWithFileLimit(PlatformData->Model, ModelStreamablesDDC, MaxDDCFiles, BulkDataFilesDDC);
+		MutablePrivate::GenerateBulkDataFilesListWithFileLimit(PlatformData->Model, ModelStreamablesDDC, MaxDDCFiles, !Options.bIsCooking, BulkDataFilesDDC);
 
-		TArray64<uint8> ModelStreamablesBytesDDC;
-		FMemoryWriter64 MemoryWriterDDC(ModelStreamablesBytesDDC);
-		MemoryWriterDDC << ModelStreamablesDDC;
+		if (!Options.bIsCooking)
+		{
+			TArray<FCachePutValueRequest> PutValueRequests;
+			PutValueRequests.Reserve(BulkDataFilesDDC.Num());
 
-		const FValue ModelStreamablesValue = FValue::Compress(FSharedBuffer::MakeView(ModelStreamablesBytesDDC.GetData(), ModelStreamablesBytesDDC.Num()));
-		RecordBuilder.AddValue(MutablePrivate::GetDerivedDataModelStreamableBulkDataId(), ModelStreamablesValue);
+			const auto WriteBulkDataDDC = [&PutValueRequests, &ModelStreamablesDDC, SharedName, DefaultDDCPolicy = DefaultDDCPolicy, DDCKey = DDCKey]
+			(MutablePrivate::FFile& File, TArray64<uint8>& FileBulkData, uint32 FileIndex)
+				{
+					FCachePutValueRequest& ValueRequest = PutValueRequests.AddDefaulted_GetRef();
+					ValueRequest.Name = SharedName;
+					ValueRequest.Policy = DefaultDDCPolicy;
+					ValueRequest.Value = FValue::Compress(FSharedBuffer::MakeView(FileBulkData.GetData(), FileBulkData.Num()));
+					ValueRequest.Key = DDCKey;
+					ValueRequest.Key.Hash = ValueRequest.Value.GetRawHash();
+					ModelStreamablesDDC.DDCValues.Add(ValueRequest.Value.GetRawHash());
+				};
+
+			constexpr bool bDropData = false;
+			MutablePrivate::SerializeBulkDataFiles(*PlatformData.Get(), BulkDataFilesDDC, WriteBulkDataDDC, bDropData);
+
+			FRequestOwner RequestOwner(UE::DerivedData::EPriority::Blocking);
+			GetCache().PutValue(PutValueRequests, RequestOwner,
+				[&bStoredSuccessfully](FCachePutValueResponse&& Response)
+				{
+					if (Response.Status != EStatus::Ok)
+					{
+						bStoredSuccessfully = false;
+					}
+				});
+			RequestOwner.Wait();
+		}
+		else
+		{
+			const auto WriteBulkDataDDC = [&RecordBuilder](MutablePrivate::FFile& File, TArray64<uint8>& FileBulkData, uint32 FileIndex)
+				{
+					const FValueId ValueId = GetDerivedDataValueIdForResource(File.DataType, FileIndex, File.ResourceType, File.Flags);
+					const FValue Value = FValue::Compress(FSharedBuffer::MakeView(FileBulkData.GetData(), FileBulkData.Num()));
+
+					RecordBuilder.AddValue(ValueId, Value);
+				};
+
+			constexpr bool bDropData = false;
+			MutablePrivate::SerializeBulkDataFiles(*PlatformData.Get(), BulkDataFilesDDC, WriteBulkDataDDC, bDropData);
+		}
 	}
 
-	// Store streamable resources as FValues
+	// Store streamable resources info as FValues
 	{
-		MUTABLE_CPUPROFILER_SCOPE(SerializeBulkDataForDDC);
-	
-		const auto WriteBulkDataDDC = [&RecordBuilder](MutablePrivate::FFile& File, TArray64<uint8>& FileBulkData, uint32 FileIndex)
-			{
-				const FValueId ValueId = GetDerivedDataValueIdForResource(File.DataType, FileIndex, File.ResourceType, File.Flags);
-				const FValue Value = FValue::Compress(FSharedBuffer::MakeView(FileBulkData.GetData(), FileBulkData.Num()));
+		MUTABLE_CPUPROFILER_SCOPE(SerializeModelStreamables);
 
-				RecordBuilder.AddValue(ValueId, Value);
-			};
+		TArray64<uint8> ModelStreamablesBytesDDC;
+		FMemoryWriter64 MemoryWriterDDC(ModelStreamablesBytesDDC);
+		MemoryWriterDDC << ModelStreamablesDDC;
 
-		constexpr bool bDropData = false;
-		MutablePrivate::SerializeBulkDataFiles(*PlatformData.Get(), BulkDataFilesDDC, WriteBulkDataDDC, bDropData);
+		const FValue ModelStreamablesValue = FValue::Compress(FSharedBuffer::MakeView(ModelStreamablesBytesDDC.GetData(), ModelStreamablesBytesDDC.Num()));
+		RecordBuilder.AddValue(MutablePrivate::GetDerivedDataModelStreamableBulkDataId(), ModelStreamablesValue);
 	}
 
 
@@ -527,19 +563,21 @@ void FCustomizableObjectSaveDDRunnable::StoreCachedPlatformDataInDDC(bool& bStor
 	// Push record to the DDC
 	{
 		MUTABLE_CPUPROFILER_SCOPE(PushRecordToDDC);
-
-		FRequestOwner RequestOwner(UE::DerivedData::EPriority::Blocking);
-		const FCachePutRequest PutRequest = { UE::FSharedString(CustomizableObjectName), RecordBuilder.Build(), DefaultDDCPolicy };
-		GetCache().Put(MakeArrayView(&PutRequest, 1), RequestOwner,
-			[&bStoredSuccessfully](FCachePutResponse&& Response)
-			{
-				if (Response.Status == EStatus::Ok)
+		if (bStoredSuccessfully)
+		{
+			FRequestOwner RequestOwner(UE::DerivedData::EPriority::Blocking);
+			const FCachePutRequest PutRequest = { UE::FSharedString(CustomizableObjectName), RecordBuilder.Build(), DefaultDDCPolicy };
+			GetCache().Put(MakeArrayView(&PutRequest, 1), RequestOwner,
+				[&bStoredSuccessfully](FCachePutResponse&& Response)
 				{
-					bStoredSuccessfully = true;
-				}
-			});
+					if (Response.Status != EStatus::Ok)
+					{
+						bStoredSuccessfully = false;
+					}
+				});
 
-		RequestOwner.Wait();
+			RequestOwner.Wait();
+		}
 
 		if (bStoredSuccessfully)
 		{
diff --git a/Engine/Plugins/Mutable/Source/CustomizableObjectEditor/Private/MuCOE/CustomizableObjectCompiler.cpp b/Engine/Plugins/Mutable/Source/CustomizableObjectEditor/Private/MuCOE/CustomizableObjectCompiler.cpp
--- a/Engine/Plugins/Mutable/Source/CustomizableObjectEditor/Private/MuCOE/CustomizableObjectCompiler.cpp
+++ b/Engine/Plugins/Mutable/Source/CustomizableObjectEditor/Private/MuCOE/CustomizableObjectCompiler.cpp
@@ -1542,7 +1542,7 @@ void FCustomizableObjectCompiler::FinishLoadingModelDataFromDDC()
 					{
 						const uint32 NumBulkDataFilesPerBucket = MAX_uint8;
 						MutablePrivate::GenerateBulkDataFilesListWithFileLimit(PlatformData->Model, *PlatformData->ModelStreamableBulkData.Get(),
-							NumBulkDataFilesPerBucket, PlatformData->BulkDataFiles);
+							NumBulkDataFilesPerBucket, false, PlatformData->BulkDataFiles);
 					}
 					else
 					{
