[mutable] Fixed crash when drag&dropping an instance to a level and then undoing.
#tests Editor
#rnx

diff --git a/Engine/Plugins/Mutable/Source/CustomizableObject/Private/MuCO/CustomizableObjectSystem.cpp b/Engine/Plugins/Mutable/Source/CustomizableObject/Private/MuCO/CustomizableObjectSystem.cpp
--- a/Engine/Plugins/Mutable/Source/CustomizableObject/Private/MuCO/CustomizableObjectSystem.cpp
+++ b/Engine/Plugins/Mutable/Source/CustomizableObject/Private/MuCO/CustomizableObjectSystem.cpp
@@ -4013,10 +4013,17 @@ int32 UCustomizableObjectSystem::TickInternal(const bool bBlocking)
 
 	Private->UpdateStats();
 
-	for (TObjectIterator<UCustomizableObjectInstanceUsage> Usage; Usage; ++Usage)
+	for (TObjectIterator<UCustomizableObjectInstanceUsage> CustomizableObjectInstanceUsageItr; CustomizableObjectInstanceUsageItr; ++CustomizableObjectInstanceUsageItr)
 	{
-		UCustomizableObjectInstanceUsagePrivate* UsagePrivate = Usage->GetPrivate();
-		
+		TObjectPtr<UCustomizableObjectInstanceUsage> CustomizableObjectInstanceUsage = *CustomizableObjectInstanceUsageItr;
+
+		if (!IsValid(CustomizableObjectInstanceUsage))
+		{
+			continue;
+		}
+
+		UCustomizableObjectInstanceUsagePrivate* UsagePrivate = CustomizableObjectInstanceUsage->GetPrivate();
+
 		if (!UsagePrivate->bPendingSetSkeletalMesh)
 		{
 			continue;
@@ -4024,7 +4031,7 @@ int32 UCustomizableObjectSystem::TickInternal(const bool bBlocking)
 
 		UsagePrivate->bPendingSetSkeletalMesh = false;
 
-		UCustomizableObjectInstance* Instance = Usage->GetCustomizableObjectInstance();
+		UCustomizableObjectInstance* Instance = CustomizableObjectInstanceUsage->GetCustomizableObjectInstance();
 		if (!Instance)
 		{
 			continue;
@@ -4038,16 +4045,16 @@ int32 UCustomizableObjectSystem::TickInternal(const bool bBlocking)
 			continue;
 		}
 
-		USkeletalMeshComponent* Parent = Cast<USkeletalMeshComponent>(Usage->GetAttachParent());
+		USkeletalMeshComponent* Parent = Cast<USkeletalMeshComponent>(CustomizableObjectInstanceUsage->GetAttachParent());
 		if (!Parent)
 		{
 			continue;
 		}
 		
-		if (!Usage->GetSkipSetReferenceSkeletalMesh() &&
+		if (!CustomizableObjectInstanceUsage->GetSkipSetReferenceSkeletalMesh() &&
 			 Object->bEnableUseRefSkeletalMeshAsPlaceholder)
 		{
-			const FName& ComponentName = Usage->GetComponentName();
+			const FName& ComponentName = CustomizableObjectInstanceUsage->GetComponentName();
 			
 			if (USkeletalMesh* ReferenceSkeletalMesh = Object->GetComponentMeshReferenceSkeletalMesh(ComponentName))
 			{
@@ -4321,7 +4328,7 @@ void UCustomizableObjectSystem::DiscardInstances()
 						for (TObjectIterator<UCustomizableObjectInstanceUsage> It; It; ++It)
 						{
 							UCustomizableObjectInstanceUsage* InstanceUsage = *It;
-							if (InstanceUsage->GetCustomizableObjectInstance() != COI)
+							if (!IsValid(InstanceUsage) || InstanceUsage->GetCustomizableObjectInstance() != COI)
 							{
 								continue;
 							}
