[mutable] Fixed logical error related to the checking for the validity of an SKM soft object Ptr to determine if an SKM was set or not.
- Previously we were checking if the SKM was not loaded into memory. Now we check if the SKM SoftObjectPtr is actually null (pointing to nothing).

#rnx

diff --git a/Engine/Plugins/Mutable/Source/CustomizableObjectEditor/Private/MuCOE/GenerateMutableSource/GenerateMutableSourceComponent.cpp b/Engine/Plugins/Mutable/Source/CustomizableObjectEditor/Private/MuCOE/GenerateMutableSource/GenerateMutableSourceComponent.cpp
--- a/Engine/Plugins/Mutable/Source/CustomizableObjectEditor/Private/MuCOE/GenerateMutableSource/GenerateMutableSourceComponent.cpp
+++ b/Engine/Plugins/Mutable/Source/CustomizableObjectEditor/Private/MuCOE/GenerateMutableSource/GenerateMutableSourceComponent.cpp
@@ -236,7 +236,7 @@ mu::Ptr<mu::NodeComponent> GenerateMutableSourceComponent(const UEdGraphPin* Pin
 			return nullptr;
 		}
 
-		if (!TypedComponentPassthroughMesh->SkeletalMesh.IsValid())
+		if (TypedComponentPassthroughMesh->SkeletalMesh.IsNull())
 		{
 			FString Msg = FString::Printf(TEXT("No mesh set for component node."));
 			GenerationContext.Log(FText::FromString(Msg), TypedComponentPassthroughMesh, EMessageSeverity::Warning);
