[Mutable] Material Parameter API

- Fix Skeletal Mesh Parameter Name issues.
- Convert UCustomizableObjectNodeMesh to UInterface

diff --git a/Engine/Plugins/Mutable/Source/CustomizableObject/Internal/MuCO/CustomizableObjectPrivate.h b/Engine/Plugins/Mutable/Source/CustomizableObject/Internal/MuCO/CustomizableObjectPrivate.h
--- a/Engine/Plugins/Mutable/Source/CustomizableObject/Internal/MuCO/CustomizableObjectPrivate.h
+++ b/Engine/Plugins/Mutable/Source/CustomizableObject/Internal/MuCO/CustomizableObjectPrivate.h
@@ -1338,6 +1338,9 @@ public:
 
 	UPROPERTY()
 	TMap<FName, TObjectPtr<USkeletalMesh>> SkeletalMeshParameterDefaultValues;
+
+	UPROPERTY()
+	TMap<FName, TObjectPtr<UMaterialInterface>> MaterialParameterDefaultValues;
 	
 	UPROPERTY()
 	FString ReleaseVersion;
diff --git a/Engine/Plugins/Mutable/Source/CustomizableObject/Private/MuCO/CustomizableObject.cpp b/Engine/Plugins/Mutable/Source/CustomizableObject/Private/MuCO/CustomizableObject.cpp
--- a/Engine/Plugins/Mutable/Source/CustomizableObject/Private/MuCO/CustomizableObject.cpp
+++ b/Engine/Plugins/Mutable/Source/CustomizableObject/Private/MuCO/CustomizableObject.cpp
@@ -1799,7 +1799,14 @@ UTexture* UCustomizableObject::GetTextureParameterDefaultValue(const FString& Pa
 	return Result ? *Result : nullptr;
 }
  
- 
+
+UMaterialInterface* UCustomizableObject::GetMaterialParameterDefaultValue(const FString& ParameterName) const
+{
+	const TObjectPtr<UMaterialInterface>* Result = GetPrivate()->GetModelResources()->MaterialParameterDefaultValues.Find(FName(ParameterName));
+	return Result ? *Result : nullptr;
+}
+
+
 USkeletalMesh* UCustomizableObject::GetSkeletalMeshParameterDefaultValue(const FString& ParameterName) const
 {
 	const TObjectPtr<USkeletalMesh>* Result = GetPrivate()->GetModelResources()->SkeletalMeshParameterDefaultValues.Find(FName(ParameterName));
diff --git a/Engine/Plugins/Mutable/Source/CustomizableObject/Private/MuCO/CustomizableObjectInstance.cpp b/Engine/Plugins/Mutable/Source/CustomizableObject/Private/MuCO/CustomizableObjectInstance.cpp
--- a/Engine/Plugins/Mutable/Source/CustomizableObject/Private/MuCO/CustomizableObjectInstance.cpp
+++ b/Engine/Plugins/Mutable/Source/CustomizableObject/Private/MuCO/CustomizableObjectInstance.cpp
@@ -2607,6 +2607,18 @@ void UCustomizableObjectInstance::SetSkeletalMeshParameterSelectedOption(const F
 }
 
 
+UMaterialInterface* UCustomizableObjectInstance::GetMaterialParameterSelectedOption(const FString& MaterialParamName, int32 RangeIndex)
+{
+	return Descriptor.GetMaterialParameterSelectedOption(MaterialParamName, RangeIndex);
+}
+
+
+void UCustomizableObjectInstance::SetMaterialParameterSelectedOption(const FString& MaterialParamName, UMaterialInterface* MaterialValue, int32 RangeIndex)
+{
+	Descriptor.SetMaterialParameterSelectedOption(MaterialParamName, MaterialValue, RangeIndex);
+}
+
+
 FLinearColor UCustomizableObjectInstance::GetColorParameterSelectedOption(const FString& ColorParamName) const
 {
 	return Descriptor.GetColorParameterSelectedOption(ColorParamName);
diff --git a/Engine/Plugins/Mutable/Source/CustomizableObject/Private/MuCO/CustomizableObjectInstanceDescriptor.cpp b/Engine/Plugins/Mutable/Source/CustomizableObject/Private/MuCO/CustomizableObjectInstanceDescriptor.cpp
--- a/Engine/Plugins/Mutable/Source/CustomizableObject/Private/MuCO/CustomizableObjectInstanceDescriptor.cpp
+++ b/Engine/Plugins/Mutable/Source/CustomizableObject/Private/MuCO/CustomizableObjectInstanceDescriptor.cpp
@@ -595,7 +595,8 @@ TSharedPtr<mu::FParameters> FCustomizableObjectInstanceDescriptor::GetParameters
 	TSharedPtr<mu::FParameters> MutableParameters = mu::FModel::NewParameters(CustomizableObject->GetPrivate()->GetModel(),
 		nullptr,
 		&ModelResources.TextureParameterDefaultValues,
-		&ModelResources.SkeletalMeshParameterDefaultValues);
+		&ModelResources.SkeletalMeshParameterDefaultValues,
+		&ModelResources.MaterialParameterDefaultValues);
  
 	const int32 ParamCount = MutableParameters->GetCount();
 	for (int32 ParamIndex = 0; ParamIndex < ParamCount; ++ParamIndex)
@@ -896,7 +897,8 @@ void FCustomizableObjectInstanceDescriptor::ReloadParameters()
 		CustomizableObject->GetPrivate()->GetModel(),
 		nullptr,
 		&ModelResources->TextureParameterDefaultValues,
-		&ModelResources->SkeletalMeshParameterDefaultValues);
+		&ModelResources->SkeletalMeshParameterDefaultValues,
+		&ModelResources->MaterialParameterDefaultValues);
  
 	int32 ParamCount = MutableParameters->GetCount();
 	for (int32 ParamIndex = 0; ParamIndex < ParamCount; ++ParamIndex)
@@ -1745,8 +1747,85 @@ void FCustomizableObjectInstanceDescriptor::SetSkeletalMeshParameterSelectedOpti
 		SkeletalMeshParameters[ParameterIndexInInstance].ParameterRangeValues[RangeIndex] = SkeletalMeshValue;
 	}
 }
+
+
+UMaterialInterface* FCustomizableObjectInstanceDescriptor::GetMaterialParameterSelectedOption(const FString& MaterialParamName,	int32 RangeIndex) const
+{
+	if (!CustomizableObject)
+	{
+		CustomizableObjectNullErrorMessage();
+		return nullptr;
+	}
+ 
+	const int32 ParameterIndexInObject = CustomizableObject->GetPrivate()->FindParameter(MaterialParamName);
+	const int32 ParameterIndexInInstance = FindTypedParameterIndex(MaterialParamName, EMutableParameterType::Material);
+ 
+	if (ParameterIndexInObject >= 0 && ParameterIndexInInstance >= 0)
+	{
+		// Due to optimizations a parameter may lose its multidimensionality (if it becomes constant). 
+		// In that case it means there is only one possible option so it is ok to set it as if wasn't multidimensional.
+		bool bIsMultidimensional = CustomizableObject->GetPrivate()->IsParameterMultidimensional(ParameterIndexInObject);
+		if (!bIsMultidimensional || RangeIndex == INDEX_NONE)
+		{
+			return MaterialParameters[ParameterIndexInInstance].ParameterValue;
+		}
+		else
+		{
+			if (MaterialParameters[ParameterIndexInInstance].ParameterRangeValues.IsValidIndex(RangeIndex))
+			{
+				return MaterialParameters[ParameterIndexInInstance].ParameterRangeValues[RangeIndex];
+			}
+		}
+	}
+	
+	LogParameterNotFoundWarning(MaterialParamName, ParameterIndexInObject, ParameterIndexInInstance, CustomizableObject, __FUNCTION__);
+ 
+	return nullptr; 
+}
+
+
+void FCustomizableObjectInstanceDescriptor::SetMaterialParameterSelectedOption(const FString& MaterialParamName, UMaterialInterface* MaterialValue,	int32 RangeIndex)
+{
+	if (!CustomizableObject)
+	{
+		CustomizableObjectNullErrorMessage();
+		return;
+	}
  
+	RETURN_ON_UNCOMPILED_CO(CustomizableObject, TEXT("Error: Cannot set Material parameter "));
+	
+	const int32 ParameterIndexInObject = CustomizableObject->GetPrivate()->FindParameter(MaterialParamName);
+	const int32 ParameterIndexInInstance = FindTypedParameterIndex(MaterialParamName, EMutableParameterType::Material);
  
+	if (ParameterIndexInObject < 0 || ParameterIndexInInstance < 0)
+	{
+		// Early out since we could not find the parameter to set.
+		LogParameterNotFoundWarning(MaterialParamName, ParameterIndexInObject, ParameterIndexInInstance, CustomizableObject, __FUNCTION__);
+		return;
+	}
+ 
+	// Due to optimizations a parameter may lose its multidimensionality (if it becomes constant). 
+	// In that case it means there is only one possible option so it is ok to set it as if wasn't multidimensional.
+	bool bIsMultidimensional = CustomizableObject->GetPrivate()->IsParameterMultidimensional(ParameterIndexInObject);
+	if (!bIsMultidimensional || RangeIndex == INDEX_NONE)
+	{
+		MaterialParameters[ParameterIndexInInstance].ParameterValue = MaterialValue;
+	}
+	else
+	{
+		if (!MaterialParameters[ParameterIndexInInstance].ParameterRangeValues.IsValidIndex(RangeIndex))
+		{
+			const int32 InsertionIndex = MaterialParameters[ParameterIndexInInstance].ParameterRangeValues.Num();
+			const int32 NumInsertedElements = RangeIndex + 1 - MaterialParameters[ParameterIndexInInstance].ParameterRangeValues.Num();
+			MaterialParameters[ParameterIndexInInstance].ParameterRangeValues.InsertDefaulted(InsertionIndex, NumInsertedElements);
+		}
+ 
+		check(MaterialParameters[ParameterIndexInInstance].ParameterRangeValues.IsValidIndex(RangeIndex));
+		MaterialParameters[ParameterIndexInInstance].ParameterRangeValues[RangeIndex] = MaterialValue;
+	}
+}
+
+
 FLinearColor FCustomizableObjectInstanceDescriptor::GetColorParameterSelectedOption(const FString& ColorParamName) const
 {
 	if (!CustomizableObject)
diff --git a/Engine/Plugins/Mutable/Source/CustomizableObject/Public/MuCO/CustomizableObject.h b/Engine/Plugins/Mutable/Source/CustomizableObject/Public/MuCO/CustomizableObject.h
--- a/Engine/Plugins/Mutable/Source/CustomizableObject/Public/MuCO/CustomizableObject.h
+++ b/Engine/Plugins/Mutable/Source/CustomizableObject/Public/MuCO/CustomizableObject.h
@@ -412,13 +412,17 @@ public:
 	CUSTOMIZABLEOBJECT_API FCustomizableObjectProjector GetProjectorParameterDefaultValue(const FString& ParameterName) const;
 	
 	/** Get the default value of a parameter of type Texture.
-	  * @param InParameterName The name of the Projector parameter to get the default value of. */
+	  * @param InParameterName The name of the Texture parameter to get the default value of. */
 	CUSTOMIZABLEOBJECT_API UTexture* GetTextureParameterDefaultValue(const FString& InParameterName) const;
 
 	/** Get the default value of a parameter of type Skeletal Mesh.
-	  * @param ParameterName The name of the Projector parameter to get the default value of. */
+	  * @param ParameterName The name of the Mesh parameter to get the default value of. */
 	CUSTOMIZABLEOBJECT_API USkeletalMesh* GetSkeletalMeshParameterDefaultValue(const FString& ParameterName) const;
 	
+	/** Get the default value of a parameter of type Material.
+	* @param ParameterName The name of the Material parameter to get the default value of. */
+	CUSTOMIZABLEOBJECT_API UMaterialInterface* GetMaterialParameterDefaultValue(const FString& ParameterName) const;
+
 	/** Return true if the parameter at the index provided is multidimensional.
 	  * @param InParameterName The name of the parameter to check.
 	  * @return True if the parameter is multidimensional and false if it is not. */
diff --git a/Engine/Plugins/Mutable/Source/CustomizableObject/Public/MuCO/CustomizableObjectInstance.h b/Engine/Plugins/Mutable/Source/CustomizableObject/Public/MuCO/CustomizableObjectInstance.h
--- a/Engine/Plugins/Mutable/Source/CustomizableObject/Public/MuCO/CustomizableObjectInstance.h
+++ b/Engine/Plugins/Mutable/Source/CustomizableObject/Public/MuCO/CustomizableObjectInstance.h
@@ -451,11 +451,21 @@ public:
 	UFUNCTION(BlueprintCallable, Category = CustomizableObjectInstance)
 	CUSTOMIZABLEOBJECT_API void SetTextureParameterSelectedOption(const FString& TextureParamName, UTexture* TextureValue, int32 RangeIndex = -1);
 
+	/** Gets the value of a Skeletal Mesh parameter with name "SkeletalMeshParamName". */
 	UFUNCTION(BlueprintCallable, Category = CustomizableObjectInstance)
 	CUSTOMIZABLEOBJECT_API USkeletalMesh* GetSkeletalMeshParameterSelectedOption(const FString& SkeletalMeshParamName, int32 RangeIndex = -1);
 
+	/** Sets the Skeletal Mesh value "SkeletalMeshValue" of a Skeletal Mesh parameter with index "SkeletalMeshParamName". */
 	UFUNCTION(BlueprintCallable, Category = CustomizableObjectInstance)
 	CUSTOMIZABLEOBJECT_API void SetSkeletalMeshParameterSelectedOption(const FString& SkeletalMeshParamName, USkeletalMesh* SkeletalMeshValue, int32 RangeIndex = -1);
+
+	/** Gets the value of a Material parameter with name "MaterialParamName". */
+	UFUNCTION(BlueprintCallable, Category = CustomizableObjectInstance)
+	CUSTOMIZABLEOBJECT_API UMaterialInterface* GetMaterialParameterSelectedOption(const FString& MaterialParamName, int32 RangeIndex = -1);
+
+	/** Sets the Material value "MaterialParamName" of a Material parameter with index "MaterialParamName". */
+	UFUNCTION(BlueprintCallable, Category = CustomizableObjectInstance)
+	CUSTOMIZABLEOBJECT_API void SetMaterialParameterSelectedOption(const FString& MaterialParamName, UMaterialInterface* MaterialValue, int32 RangeIndex = -1);
 	
 	/** Gets the value of a color parameter with name "ColorParamName". */
 	UFUNCTION(BlueprintCallable, Category = CustomizableObjectInstance)
diff --git a/Engine/Plugins/Mutable/Source/CustomizableObject/Public/MuCO/CustomizableObjectInstanceDescriptor.h b/Engine/Plugins/Mutable/Source/CustomizableObject/Public/MuCO/CustomizableObjectInstanceDescriptor.h
--- a/Engine/Plugins/Mutable/Source/CustomizableObject/Public/MuCO/CustomizableObjectInstanceDescriptor.h
+++ b/Engine/Plugins/Mutable/Source/CustomizableObject/Public/MuCO/CustomizableObjectInstanceDescriptor.h
@@ -17,6 +17,7 @@ class FArchive;
 class UCustomizableInstancePrivate;
 class UCustomizableObject;
 class UCustomizableObjectInstance;
+class UMaterialInterface;
 class FDescriptorHash;
 class FMutableUpdateCandidate;
 struct FCustomizableObjectComponentIndex;
@@ -107,11 +108,17 @@ struct FCustomizableObjectInstanceDescriptor
 	/** Sets the texture value "TextureValue" of a texture parameter with index "TextureParamIndex". */
 	UE_API void SetTextureParameterSelectedOption(const FString& TextureParamName, UTexture* TextureValue, int32 RangeIndex);
 	
-	/** Gets the value of a skeletalMesh parameter with name "SkeletalMeshParamName". */
+	/** Gets the value of a Skeletal Mesh parameter with name "SkeletalMeshParamName". */
 	UE_API USkeletalMesh* GetSkeletalMeshParameterSelectedOption(const FString& SkeletalMeshParamName, int32 RangeIndex = -1) const;
 
 	/** Sets the skeletalMesh value "SkeletalMeshValue" of a skeletalMesh parameter with index "SkeletalMeshParamIndex". */
 	UE_API void SetSkeletalMeshParameterSelectedOption(const FString& SkeletalMeshParamName, USkeletalMesh* SkeletalMeshValue, int32 RangeIndex = -1);
+
+	/** Gets the value of a Material parameter with name "MaterialMeshParamName". */
+	UE_API UMaterialInterface* GetMaterialParameterSelectedOption(const FString& MaterialMeshParamName, int32 RangeIndex = -1) const;
+
+	/** Sets the material value "MaterialParamName" of a Material parameter with index "MaterialMeshParamName". */
+	UE_API void SetMaterialParameterSelectedOption(const FString& MaterialParamName, UMaterialInterface* MaterialValue, int32 RangeIndex = -1);
 	
 	/** Gets the value of a color parameter with name "ColorParamName". */
 	UE_API FLinearColor GetColorParameterSelectedOption(const FString& ColorParamName) const;
@@ -304,13 +311,16 @@ struct FCustomizableObjectInstanceDescriptor
 	
 	UPROPERTY()
 	TArray<FCustomizableObjectFloatParameterValue> FloatParameters;
-
+	
 	UPROPERTY()
 	TArray<FCustomizableObjectTextureParameterValue> TextureParameters;
 
 	UPROPERTY()
 	TArray<FCustomizableObjectSkeletalMeshParameterValue> SkeletalMeshParameters;
 
+	UPROPERTY()
+	TArray<FCustomizableObjectMaterialParameterValue> MaterialParameters;
+	
 	UPROPERTY()
 	TArray<FCustomizableObjectVectorParameterValue> VectorParameters;
 	
diff --git a/Engine/Plugins/Mutable/Source/CustomizableObject/Public/MuCO/CustomizableObjectParameterTypeDefinitions.h b/Engine/Plugins/Mutable/Source/CustomizableObject/Public/MuCO/CustomizableObjectParameterTypeDefinitions.h
--- a/Engine/Plugins/Mutable/Source/CustomizableObject/Public/MuCO/CustomizableObjectParameterTypeDefinitions.h
+++ b/Engine/Plugins/Mutable/Source/CustomizableObject/Public/MuCO/CustomizableObjectParameterTypeDefinitions.h
@@ -8,6 +8,7 @@
 
 class UTexture;
 class USkeletalMesh;
+class UMaterialInterface;
 
 
 UENUM(BlueprintType)
@@ -22,6 +23,7 @@ enum class EMutableParameterType : uint8
 	Texture			UMETA(DisplayName = "Texture"),
 	Transform		UMETA(DisplayName = "Transform"),
 	SkeletalMesh	UMETA(DisplayName = "SkeletalMesh"),
+	Material	UMETA(DisplayName = "Material"),
 };
 
 
@@ -196,16 +198,16 @@ struct FCustomizableObjectSkeletalMeshParameterValue
 {
 	GENERATED_USTRUCT_BODY()
 	
-	UPROPERTY(Category = CustomizableObjectTextureParameterValue, VisibleAnywhere)
+	UPROPERTY(Category = CustomizableObjectSkeletalMeshParameterValue, VisibleAnywhere)
 	FString ParameterName;
 
-	UPROPERTY(Category = CustomizableObjectTextureParameterValue, VisibleAnywhere)
+	UPROPERTY(Category = CustomizableObjectSkeletalMeshParameterValue, VisibleAnywhere)
 	TObjectPtr<USkeletalMesh> ParameterValue;
 
-	UPROPERTY(Category = CustomizableObjectTextureParameterValue, VisibleAnywhere)
+	UPROPERTY(Category = CustomizableObjectSkeletalMeshParameterValue, VisibleAnywhere)
 	FGuid Id;
 
-	UPROPERTY(Category = CustomizableObjectTextureParameterValue, VisibleAnywhere)
+	UPROPERTY(Category = CustomizableObjectSkeletalMeshParameterValue, VisibleAnywhere)
 	TArray<TObjectPtr<USkeletalMesh>> ParameterRangeValues;
 };
 
@@ -224,6 +226,39 @@ inline uint32 GetTypeHash(const FCustomizableObjectSkeletalMeshParameterValue& K
 }
 
 
+USTRUCT()
+struct FCustomizableObjectMaterialParameterValue
+{
+	GENERATED_USTRUCT_BODY()
+	
+	UPROPERTY(Category = CustomizableObjectMaterialParameterValue, VisibleAnywhere)
+	FString ParameterName;
+
+	UPROPERTY(Category = CustomizableObjectMaterialParameterValue, VisibleAnywhere)
+	TObjectPtr<UMaterialInterface> ParameterValue;
+
+	UPROPERTY(Category = CustomizableObjectMaterialParameterValue, VisibleAnywhere)
+	FGuid Id;
+
+	UPROPERTY(Category = CustomizableObjectMaterialParameterValue, VisibleAnywhere)
+	TArray<TObjectPtr<UMaterialInterface>> ParameterRangeValues;
+};
+
+
+inline uint32 GetTypeHash(const FCustomizableObjectMaterialParameterValue& Key)
+{
+	uint32 Hash = GetTypeHash(Key.ParameterName);
+	Hash = HashCombine(Hash, GetTypeHash(Key.ParameterValue));
+
+	for (const TObjectPtr<UMaterialInterface>& Value : Key.ParameterRangeValues)
+	{
+		Hash = HashCombine(Hash, GetTypeHash(Value.GetPathName()));
+	}
+
+	return Hash;
+}
+
+
 USTRUCT(BlueprintType)
 struct FCustomizableObjectVectorParameterValue
 {
diff --git a/Engine/Plugins/Mutable/Source/CustomizableObjectEditor/Private/MuCOE/CustomizableObjectLayout.cpp b/Engine/Plugins/Mutable/Source/CustomizableObjectEditor/Private/MuCOE/CustomizableObjectLayout.cpp
--- a/Engine/Plugins/Mutable/Source/CustomizableObjectEditor/Private/MuCOE/CustomizableObjectLayout.cpp
+++ b/Engine/Plugins/Mutable/Source/CustomizableObjectEditor/Private/MuCOE/CustomizableObjectLayout.cpp
@@ -11,7 +11,7 @@
 #include "MuCOE/MutableUtils.h"
 #include "MuCOE/CustomizableObjectEditor.h"
 #include "MuCOE/Nodes/CustomizableObjectNodeSkeletalMesh.h"
-#include "MuCOE/Nodes/CustomizableObjectNodeMesh.h"
+#include "MuCOE/Nodes/CustomizableObjectNodeMeshInterface.h"
 #include "MuCOE/Nodes/CustomizableObjectNodeTable.h"
 
 #define LOCTEXT_NAMESPACE "CustomizableObjectEditor"
@@ -223,7 +223,7 @@ TSoftObjectPtr<UObject> UCustomizableObjectLayout::GetMesh() const
 {
 	if (UObject* Node = GetOuter())
 	{
-		if (const UCustomizableObjectNodeMesh* TypedNodeSkeletalMesh = Cast<UCustomizableObjectNodeMesh>(Node))
+		if (const ICustomizableObjectNodeMeshInterface* TypedNodeSkeletalMesh = Cast<ICustomizableObjectNodeMeshInterface>(Node))
 		{
 			return TypedNodeSkeletalMesh->GetMesh();
 		}
diff --git a/Engine/Plugins/Mutable/Source/CustomizableObjectEditor/Private/MuCOE/EdGraphSchema_CustomizableObject.cpp b/Engine/Plugins/Mutable/Source/CustomizableObjectEditor/Private/MuCOE/EdGraphSchema_CustomizableObject.cpp
--- a/Engine/Plugins/Mutable/Source/CustomizableObjectEditor/Private/MuCOE/EdGraphSchema_CustomizableObject.cpp
+++ b/Engine/Plugins/Mutable/Source/CustomizableObjectEditor/Private/MuCOE/EdGraphSchema_CustomizableObject.cpp
@@ -35,8 +35,9 @@
 #include "MuCOE/Nodes/CustomizableObjectNodeFloatArithmeticOp.h"
 #include "MuCOE/Nodes/CustomizableObjectNodeGroupProjectorParameter.h"
 #include "MuCOE/Nodes/CustomizableObjectNodeMacroInstance.h"
-#include "MuCOE/Nodes/CustomizableObjectNodeMaterialVariation.h"
+#include "MuCOE/Nodes/CustomizableObjectNodeMaterialParameter.h"
 #include "MuCOE/Nodes/CustomizableObjectNodeMaterialSwitch.h"
+#include "MuCOE/Nodes/CustomizableObjectNodeMaterialVariation.h"
 #include "MuCOE/Nodes/CustomizableObjectNodeModifierExtendMeshSection.h"
 #include "MuCOE/Nodes/CustomizableObjectNodeModifierEditMeshSection.h"
 #include "MuCOE/Nodes/CustomizableObjectNodeModifierRemoveMesh.h"
@@ -264,7 +265,7 @@ UEdGraphNode* FCustomizableObjectSchemaAction_Paste::PerformAction(class UEdGrap
 // DO NOT change the values because it will break the external pin nodes!
 const FName UEdGraphSchema_CustomizableObject::PC_Object("object");
 const FName UEdGraphSchema_CustomizableObject::PC_Component("component");
-const FName UEdGraphSchema_CustomizableObject::PC_Material("material");
+const FName UEdGraphSchema_CustomizableObject::PC_MeshSection("material");
 const FName UEdGraphSchema_CustomizableObject::PC_Modifier("modifier");
 const FName UEdGraphSchema_CustomizableObject::PC_Mesh("mesh");
 const FName UEdGraphSchema_CustomizableObject::PC_PassThroughMesh("passThroughMesh");
@@ -278,15 +279,15 @@ const FName UEdGraphSchema_CustomizableObject::PC_Float("float");
 const FName UEdGraphSchema_CustomizableObject::PC_Bool("bool");
 const FName UEdGraphSchema_CustomizableObject::PC_Enum("enum");
 const FName UEdGraphSchema_CustomizableObject::PC_Stack("stack");
-const FName UEdGraphSchema_CustomizableObject::PC_MaterialAsset("materialAsset");
+const FName UEdGraphSchema_CustomizableObject::PC_Material("materialAsset");
 const FName UEdGraphSchema_CustomizableObject::PC_Wildcard("wildcard");
 const FName UEdGraphSchema_CustomizableObject::PC_PoseAsset("poseAsset");
 const FName UEdGraphSchema_CustomizableObject::PC_Transform("transform");
 const FName UEdGraphSchema_CustomizableObject::PC_String("string");
 
 // Add more pin types to this array if needed
-const TArray<FName> UEdGraphSchema_CustomizableObject::SupportedMacroPinTypes = { PC_Object, PC_Component, PC_Material, PC_Modifier, PC_Mesh, PC_Image,
-PC_PassThroughImage, PC_Projector, PC_GroupProjector, PC_Color, PC_Float, PC_Enum, PC_Stack, PC_MaterialAsset, PC_PoseAsset, PC_Transform, PC_String };
+const TArray<FName> UEdGraphSchema_CustomizableObject::SupportedMacroPinTypes = { PC_Object, PC_Component, PC_MeshSection, PC_Modifier, PC_Mesh, PC_Image,
+PC_PassThroughImage, PC_Projector, PC_GroupProjector, PC_Color, PC_Float, PC_Enum, PC_Stack, PC_Material, PC_PoseAsset, PC_Transform, PC_String };
 
 // Node categories
 const FText UEdGraphSchema_CustomizableObject::NC_Experimental(FText::FromString("Experimental"));
@@ -609,6 +610,7 @@ void UEdGraphSchema_CustomizableObject::GetGraphContextActions(FGraphContextMenu
 			ContextMenuBuilder.CreateTemplateNode<UCustomizableObjectNodeModifierTransformInMesh>(),
 			ContextMenuBuilder.CreateTemplateNode<UCustomizableObjectNodeTextureParameter>(),
 			ContextMenuBuilder.CreateTemplateNode<UCustomizableObjectNodeSkeletalMeshParameter>(),
+			ContextMenuBuilder.CreateTemplateNode<UCustomizableObjectNodeMaterialParameter>(),
 		};
 
 		AddNewNodeCategoryActionsFiltered(ExperimentalTemplateNodes, ContextMenuBuilder, NC_Experimental.ToString(), 2, Filter);
@@ -616,7 +618,7 @@ void UEdGraphSchema_CustomizableObject::GetGraphContextActions(FGraphContextMenu
 
 	{
 		// External Pin Nodes
-		TArray<FName> PinTypes({ PC_Material, PC_Modifier, PC_Mesh, PC_Image, PC_Projector, PC_GroupProjector, PC_Color, PC_Float, PC_Bool, PC_Enum, PC_Transform, PC_Stack, PC_PassThroughImage, PC_MaterialAsset, PC_PoseAsset, PC_Component });
+		TArray<FName> PinTypes({ PC_MeshSection, PC_Modifier, PC_Mesh, PC_Image, PC_Projector, PC_GroupProjector, PC_Color, PC_Float, PC_Bool, PC_Enum, PC_Transform, PC_Stack, PC_PassThroughImage, PC_Material, PC_PoseAsset, PC_Component });
 
 		// Add pin types from extensions
 		for (const FRegisteredCustomizableObjectPinType& PinType : ICustomizableObjectModule::Get().GetExtendedPinTypes())
@@ -796,7 +798,7 @@ FLinearColor UEdGraphSchema_CustomizableObject::GetPinTypeColor(const FName& Typ
 	{
 		return FLinearColor(0.75294f, 0.08098f, 0.08098f, 1.000000f); // Soft Red
 	}
-	else if (TypeString == PC_Material)
+	else if (TypeString == PC_MeshSection)
 	{
 		return FLinearColor(0.000000f, 0.100000f, 0.600000f, 1.000000f); // Blue
 	}
@@ -816,7 +818,7 @@ FLinearColor UEdGraphSchema_CustomizableObject::GetPinTypeColor(const FName& Typ
 	{
 		return FLinearColor(1.000000f, 0.000000f, 0.800000f, 1.000000f); // Pink
 	}
-	else if (TypeString == PC_MaterialAsset)
+	else if (TypeString == PC_Material)
 	{
 		return FLinearColor(0.000000f, 1.000000f, 0.1000f, 1.000000f); // Cian
 	}
@@ -1374,7 +1376,7 @@ FName UEdGraphSchema_CustomizableObject::GetPinCategoryName(const FName& PinCate
 	{
 		return TEXT("Component");
 	}
-	else if (PinCategory == UEdGraphSchema_CustomizableObject::PC_Material)
+	else if (PinCategory == UEdGraphSchema_CustomizableObject::PC_MeshSection)
 	{
 		return TEXT("Material");
 	}
@@ -1426,7 +1428,7 @@ FName UEdGraphSchema_CustomizableObject::GetPinCategoryName(const FName& PinCate
 	{
 		return TEXT("Stack");
 	}
-	else if (PinCategory == UEdGraphSchema_CustomizableObject::PC_MaterialAsset)
+	else if (PinCategory == UEdGraphSchema_CustomizableObject::PC_Material)
 	{
 		return TEXT("Material");
 	}
@@ -1473,7 +1475,7 @@ FText UEdGraphSchema_CustomizableObject::GetPinCategoryFriendlyName(const FName&
 	{
 		return LOCTEXT("Component_Pin_Category", "Component");
 	}
-	else if (PinCategory == UEdGraphSchema_CustomizableObject::PC_Material)
+	else if (PinCategory == UEdGraphSchema_CustomizableObject::PC_MeshSection)
 	{
 		return LOCTEXT("MeshSection_Pin_Category", "Mesh Section");
 	}
@@ -1525,7 +1527,7 @@ FText UEdGraphSchema_CustomizableObject::GetPinCategoryFriendlyName(const FName&
 	{
 		return LOCTEXT("Stack_Pin_Category", "Stack");
 	}
-	else if (PinCategory == UEdGraphSchema_CustomizableObject::PC_MaterialAsset)
+	else if (PinCategory == UEdGraphSchema_CustomizableObject::PC_Material)
 	{
 		return LOCTEXT("Material_Asset_Pin_Category", "Material");
 	}
diff --git a/Engine/Plugins/Mutable/Source/CustomizableObjectEditor/Private/MuCOE/GenerateMutableSource/GenerateMutableSource.cpp b/Engine/Plugins/Mutable/Source/CustomizableObjectEditor/Private/MuCOE/GenerateMutableSource/GenerateMutableSource.cpp
--- a/Engine/Plugins/Mutable/Source/CustomizableObjectEditor/Private/MuCOE/GenerateMutableSource/GenerateMutableSource.cpp
+++ b/Engine/Plugins/Mutable/Source/CustomizableObjectEditor/Private/MuCOE/GenerateMutableSource/GenerateMutableSource.cpp
@@ -20,7 +20,7 @@
 #include "MuCOE/ICustomizableObjectExtensionNode.h"
 #include "MuCOE/Nodes/CustomizableObjectNodeMacroInstance.h"
 #include "MuCOE/Nodes/CustomizableObjectNodeMaterial.h"
-#include "MuCOE/Nodes/CustomizableObjectNodeMesh.h"
+#include "MuCOE/Nodes/CustomizableObjectNodeMeshInterface.h"
 #include "MuCOE/Nodes/CustomizableObjectNodeModifierClipWithMesh.h"
 #include "MuCOE/Nodes/CustomizableObjectNodeObjectGroup.h"
 #include "MuCOE/Nodes/CustomizableObjectNodeStaticString.h"
@@ -459,7 +459,7 @@ UTexture2D* FindReferenceImage(const UEdGraphPin* Pin, FMutableGraphGenerationCo
 		Result = ParamNodeTex->ReferenceValue;
 	}
 
-	else if (const UCustomizableObjectNodeMesh* TypedNodeMesh = Cast<UCustomizableObjectNodeMesh>(Node))
+	else if (const ICustomizableObjectNodeMeshInterface* TypedNodeMesh = Cast<ICustomizableObjectNodeMeshInterface>(Node))
 	{
 		Result = TypedNodeMesh->FindTextureForPin(Pin);
 	}
@@ -640,7 +640,7 @@ UTexture2D* FindReferenceImage(const UEdGraphPin* Pin, FMutableGraphGenerationCo
 
 	else if (const UCustomizableObjectNodeTable* TypedNodeTable = Cast<UCustomizableObjectNodeTable>(Node))
 	{
-		if (Pin->PinType.PinCategory == Schema->PC_MaterialAsset)
+		if (Pin->PinType.PinCategory == Schema->PC_Material)
 		{
 			Result = TypedNodeTable->FindReferenceTextureParameter(Pin, GenerationContext.CurrentMaterialTableParameter);
 		}
diff --git a/Engine/Plugins/Mutable/Source/CustomizableObjectEditor/Private/MuCOE/GenerateMutableSource/GenerateMutableSourceFloat.cpp b/Engine/Plugins/Mutable/Source/CustomizableObjectEditor/Private/MuCOE/GenerateMutableSource/GenerateMutableSourceFloat.cpp
--- a/Engine/Plugins/Mutable/Source/CustomizableObjectEditor/Private/MuCOE/GenerateMutableSource/GenerateMutableSourceFloat.cpp
+++ b/Engine/Plugins/Mutable/Source/CustomizableObjectEditor/Private/MuCOE/GenerateMutableSource/GenerateMutableSourceFloat.cpp
@@ -314,7 +314,7 @@ mu::Ptr<mu::NodeScalar> GenerateMutableSourceFloat(const UEdGraphPin* Pin, FMuta
 
 		Result = ConstantValue;
 
-		if (Pin->PinType.PinCategory == Schema->PC_MaterialAsset)
+		if (Pin->PinType.PinCategory == Schema->PC_Material)
 		{
 			// Material pins have to skip the cache of nodes or they will return always the same column node
 			bCacheNode = false;
@@ -346,7 +346,7 @@ mu::Ptr<mu::NodeScalar> GenerateMutableSourceFloat(const UEdGraphPin* Pin, FMuta
 				{
 					mu::Ptr<mu::NodeScalarTable> ScalarTableNode = new mu::NodeScalarTable();
 
-					if (Pin->PinType.PinCategory == Schema->PC_MaterialAsset)
+					if (Pin->PinType.PinCategory == Schema->PC_Material)
 					{
 						// Materials use the parameter id as column names
 						ColumnName = GenerationContext.CurrentMaterialTableParameterId;
diff --git a/Engine/Plugins/Mutable/Source/CustomizableObjectEditor/Private/MuCOE/GenerateMutableSource/GenerateMutableSourceImage.cpp b/Engine/Plugins/Mutable/Source/CustomizableObjectEditor/Private/MuCOE/GenerateMutableSource/GenerateMutableSourceImage.cpp
--- a/Engine/Plugins/Mutable/Source/CustomizableObjectEditor/Private/MuCOE/GenerateMutableSource/GenerateMutableSourceImage.cpp
+++ b/Engine/Plugins/Mutable/Source/CustomizableObjectEditor/Private/MuCOE/GenerateMutableSource/GenerateMutableSourceImage.cpp
@@ -14,7 +14,7 @@
 #include "MuCOE/GenerateMutableSource/GenerateMutableSourceTable.h"
 #include "MuCOE/GraphTraversal.h"
 #include "MuCOE/Nodes/CustomizableObjectNodeMacroInstance.h"
-#include "MuCOE/Nodes/CustomizableObjectNodeMesh.h"
+#include "MuCOE/Nodes/CustomizableObjectNodeMeshInterface.h"
 #include "MuCOE/Nodes/CustomizableObjectNodeTable.h"
 #include "MuCOE/Nodes/CustomizableObjectNodeMaterial.h"
 #include "MuCOE/Nodes/CustomizableObjectNodeTexture.h"
@@ -219,7 +219,7 @@ mu::Ptr<mu::NodeImage> GenerateMutableSourceImage(const UEdGraphPin* Pin, FMutab
 		Result = ResizeNode;
 	}
 
-	else if (const UCustomizableObjectNodeMesh* TypedNodeMesh = Cast<UCustomizableObjectNodeMesh>(Node))
+	else if (const ICustomizableObjectNodeMeshInterface* TypedNodeMesh = Cast<ICustomizableObjectNodeMeshInterface>(Node))
 	{
 		mu::Ptr<mu::NodeImageConstant> ImageNode = new mu::NodeImageConstant();
 		Result = ImageNode;
@@ -1021,7 +1021,7 @@ mu::Ptr<mu::NodeImage> GenerateMutableSourceImage(const UEdGraphPin* Pin, FMutab
 
 			bool bSuccess = true;
 
-			if (Pin->PinType.PinCategory == Schema->PC_MaterialAsset)
+			if (Pin->PinType.PinCategory == Schema->PC_Material)
 			{
 				// Material pins have to skip the cache of nodes or they will return always the same column node
 				bCacheNode = false;
@@ -1044,7 +1044,7 @@ mu::Ptr<mu::NodeImage> GenerateMutableSourceImage(const UEdGraphPin* Pin, FMutab
 
 				UTexture* DefaultTexture = TypedNodeTable->GetColumnDefaultAssetByType<UTexture>(Pin);
 
-				if (bSuccess && Pin->PinType.PinCategory != Schema->PC_MaterialAsset && !DefaultTexture)
+				if (bSuccess && Pin->PinType.PinCategory != Schema->PC_Material && !DefaultTexture)
 				{
 					FString Msg = FString::Printf(TEXT("Couldn't find a default value in the data table's struct for the column [%s]. The default value is null or not a supported Texture"), *ColumnName);
 					GenerationContext.Log(FText::FromString(Msg), Node);
@@ -1062,7 +1062,7 @@ mu::Ptr<mu::NodeImage> GenerateMutableSourceImage(const UEdGraphPin* Pin, FMutab
 					{
 						mu::Ptr<mu::NodeImageTable> ImageTableNode = new mu::NodeImageTable();
 
-						if (Pin->PinType.PinCategory == Schema->PC_MaterialAsset)
+						if (Pin->PinType.PinCategory == Schema->PC_Material)
 						{
 							// Material parameters use the Data Table Column Name + Parameter id as mutable column Name to aboid duplicated names (i.e. two MI columns with the same parents but different values).
 							ColumnName = Property->GetDisplayNameText().ToString() + GenerationContext.CurrentMaterialTableParameterId;
diff --git a/Engine/Plugins/Mutable/Source/CustomizableObjectEditor/Private/MuCOE/GenerateMutableSource/GenerateMutableSourceMesh.cpp b/Engine/Plugins/Mutable/Source/CustomizableObjectEditor/Private/MuCOE/GenerateMutableSource/GenerateMutableSourceMesh.cpp
--- a/Engine/Plugins/Mutable/Source/CustomizableObjectEditor/Private/MuCOE/GenerateMutableSource/GenerateMutableSourceMesh.cpp
+++ b/Engine/Plugins/Mutable/Source/CustomizableObjectEditor/Private/MuCOE/GenerateMutableSource/GenerateMutableSourceMesh.cpp
@@ -2920,7 +2920,7 @@ TSharedPtr<mu::FMesh> BuildMorphedMutableMesh(const UEdGraphPin* BaseSourcePin,
 	USkeletalMesh* SkeletalMesh = nullptr;
 	UCustomizableObjectNode* Node = Cast<UCustomizableObjectNode>(BaseSourcePin->GetOwningNode());
 
-	if (const UCustomizableObjectNodeMesh* TypedNodeSkeletalMesh = Cast<UCustomizableObjectNodeMesh>(Node))
+	if (const ICustomizableObjectNodeMeshInterface* TypedNodeSkeletalMesh = Cast<ICustomizableObjectNodeMeshInterface>(Node))
 	{
 		int32 LayoutIndex;
 		TypedNodeSkeletalMesh->GetPinSection(*BaseSourcePin, LODIndexConnected, SectionIndexConnected, LayoutIndex);
@@ -4603,12 +4603,12 @@ mu::Ptr<mu::NodeMesh> GenerateMutableSourceMesh(const UEdGraphPin* Pin,
 	{
 		mu::Ptr<mu::NodeMeshParameter> ParameterNode = new mu::NodeMeshParameter();
 
-		ParameterNode->Name = TypedNodeParam->ParameterName;
+		ParameterNode->Name = TypedNodeParam->GetParameterName();
 		ParameterNode->UID = GenerationContext.GetNodeIdUnique(Node).ToString();
 
 		if (TypedNodeParam->DefaultValue)
 		{
-			GenerationContext.SkeletalMeshParameterDefaultValues.Add(FName(TypedNodeParam->ParameterName), MutablePrivate::LoadObject(TypedNodeParam->DefaultValue));
+			GenerationContext.SkeletalMeshParameterDefaultValues.Add(FName(TypedNodeParam->GetParameterName()), MutablePrivate::LoadObject(TypedNodeParam->DefaultValue));
 		}
 
 		const int32 LODIndex = 0;
@@ -4616,7 +4616,7 @@ mu::Ptr<mu::NodeMesh> GenerateMutableSourceMesh(const UEdGraphPin* Pin,
 		TypedNodeParam->GetPinSection(*Pin, SectionIndex);
 		ParameterNode->SectionIndex = SectionIndex;
 
-		GenerationContext.ParameterUIDataMap.Add(TypedNodeParam->ParameterName, FMutableParameterData(
+		GenerationContext.ParameterUIDataMap.Add(TypedNodeParam->GetParameterName(), FMutableParameterData(
 			TypedNodeParam->ParamUIMetadata,
 			EMutableParameterType::Texture));
 
diff --git a/Engine/Plugins/Mutable/Source/CustomizableObjectEditor/Private/MuCOE/GenerateMutableSource/GenerateMutableSourceTable.cpp b/Engine/Plugins/Mutable/Source/CustomizableObjectEditor/Private/MuCOE/GenerateMutableSource/GenerateMutableSourceTable.cpp
--- a/Engine/Plugins/Mutable/Source/CustomizableObjectEditor/Private/MuCOE/GenerateMutableSource/GenerateMutableSourceTable.cpp
+++ b/Engine/Plugins/Mutable/Source/CustomizableObjectEditor/Private/MuCOE/GenerateMutableSource/GenerateMutableSourceTable.cpp
@@ -896,7 +896,7 @@ bool GenerateTableColumn(const UCustomizableObjectNodeTable* TableNode, const UE
 	check(Schema);
 
 	// At some point we may want more default values types
-	TSet<FName> ValidNonePinTypes({ Schema->PC_Color, Schema->PC_MaterialAsset });
+	TSet<FName> ValidNonePinTypes({ Schema->PC_Color, Schema->PC_Material });
 	
 	// Only set the "None" value of Material and Color columns
 	if (ValidNonePinTypes.Contains(Pin->PinType.PinCategory))
@@ -1295,7 +1295,7 @@ bool GenerateNoneRow(const UCustomizableObjectNodeTable& TableNode, const UEdGra
 	//Ensure that we are generating a material column and not a Texture column
 	FString EncodedSwitchParameterName = "__MutableMaterialId";
 
-	if (Pin->PinType.PinCategory == Schema->PC_MaterialAsset && ColumnName.Contains(EncodedSwitchParameterName))
+	if (Pin->PinType.PinCategory == Schema->PC_Material && ColumnName.Contains(EncodedSwitchParameterName))
 	{
 		// Add a scalar column if it has not been added yet
 		if (ColumnId == INDEX_NONE)
diff --git a/Engine/Plugins/Mutable/Source/CustomizableObjectEditor/Private/MuCOE/GraphTraversal.cpp b/Engine/Plugins/Mutable/Source/CustomizableObjectEditor/Private/MuCOE/GraphTraversal.cpp
--- a/Engine/Plugins/Mutable/Source/CustomizableObjectEditor/Private/MuCOE/GraphTraversal.cpp
+++ b/Engine/Plugins/Mutable/Source/CustomizableObjectEditor/Private/MuCOE/GraphTraversal.cpp
@@ -641,7 +641,7 @@ const UEdGraphPin* FindMeshBaseSource(const UEdGraphPin& Pin, const bool bOnlyLo
 		||
 		Pin.PinType.PinCategory == UEdGraphSchema_CustomizableObject::PC_PassThroughMesh
 		||
-		Pin.PinType.PinCategory == UEdGraphSchema_CustomizableObject::PC_Material
+		Pin.PinType.PinCategory == UEdGraphSchema_CustomizableObject::PC_MeshSection
 		||
 		Pin.PinType.PinCategory == UEdGraphSchema_CustomizableObject::PC_Modifier
 	);
diff --git a/Engine/Plugins/Mutable/Source/CustomizableObjectEditor/Private/MuCOE/Nodes/CustomizableObjectNodeComponentMesh.cpp b/Engine/Plugins/Mutable/Source/CustomizableObjectEditor/Private/MuCOE/Nodes/CustomizableObjectNodeComponentMesh.cpp
--- a/Engine/Plugins/Mutable/Source/CustomizableObjectEditor/Private/MuCOE/Nodes/CustomizableObjectNodeComponentMesh.cpp
+++ b/Engine/Plugins/Mutable/Source/CustomizableObjectEditor/Private/MuCOE/Nodes/CustomizableObjectNodeComponentMesh.cpp
@@ -63,7 +63,7 @@ void UCustomizableObjectNodeComponentMesh::AllocateDefaultPins(UCustomizableObje
 	// Mesh Component Pins
 	const FString PinFriendlyName = TEXT("Overlay Material");
 	const FString PinName = PinFriendlyName + FString(TEXT("_Input_Pin"));
-	UEdGraphPin* OverlayMaterialPin = CustomCreatePin(EGPD_Input, UEdGraphSchema_CustomizableObject::PC_MaterialAsset, FName(*PinName));
+	UEdGraphPin* OverlayMaterialPin = CustomCreatePin(EGPD_Input, UEdGraphSchema_CustomizableObject::PC_Material, FName(*PinName));
 	OverlayMaterialPin->PinFriendlyName = FText::FromString(PinFriendlyName);
 	OverlayMaterialPin->bDefaultValueIsIgnored = true;
 	OverlayMaterialPin->PinToolTip = "Pin for an Overlay Material from a Table Node";
@@ -74,7 +74,7 @@ void UCustomizableObjectNodeComponentMesh::AllocateDefaultPins(UCustomizableObje
 	{
 		FString LODName = FString::Printf(TEXT("LOD %d"), LODIndex);
 
-		UEdGraphPin* Pin = CustomCreatePin(EGPD_Input, UEdGraphSchema_CustomizableObject::PC_Material, FName(*LODName), true);
+		UEdGraphPin* Pin = CustomCreatePin(EGPD_Input, UEdGraphSchema_CustomizableObject::PC_MeshSection, FName(*LODName), true);
 		LODPins.Add(Pin);
 	}
 }
diff --git a/Engine/Plugins/Mutable/Source/CustomizableObjectEditor/Private/MuCOE/Nodes/CustomizableObjectNodeComponentMeshAddTo.cpp b/Engine/Plugins/Mutable/Source/CustomizableObjectEditor/Private/MuCOE/Nodes/CustomizableObjectNodeComponentMeshAddTo.cpp
--- a/Engine/Plugins/Mutable/Source/CustomizableObjectEditor/Private/MuCOE/Nodes/CustomizableObjectNodeComponentMeshAddTo.cpp
+++ b/Engine/Plugins/Mutable/Source/CustomizableObjectEditor/Private/MuCOE/Nodes/CustomizableObjectNodeComponentMeshAddTo.cpp
@@ -21,7 +21,7 @@ void UCustomizableObjectNodeComponentMeshAddTo::AllocateDefaultPins(UCustomizabl
 	{
 		FString LODName = FString::Printf(TEXT("LOD %d"), LODIndex);
 
-		UEdGraphPin* Pin = CustomCreatePin(EGPD_Input, UEdGraphSchema_CustomizableObject::PC_Material, FName(*LODName), true);
+		UEdGraphPin* Pin = CustomCreatePin(EGPD_Input, UEdGraphSchema_CustomizableObject::PC_MeshSection, FName(*LODName), true);
 		LODPins.Add(Pin);
 	}
 
diff --git a/Engine/Plugins/Mutable/Source/CustomizableObjectEditor/Private/MuCOE/Nodes/CustomizableObjectNodeComponentPassthroughMesh.cpp b/Engine/Plugins/Mutable/Source/CustomizableObjectEditor/Private/MuCOE/Nodes/CustomizableObjectNodeComponentPassthroughMesh.cpp
--- a/Engine/Plugins/Mutable/Source/CustomizableObjectEditor/Private/MuCOE/Nodes/CustomizableObjectNodeComponentPassthroughMesh.cpp
+++ b/Engine/Plugins/Mutable/Source/CustomizableObjectEditor/Private/MuCOE/Nodes/CustomizableObjectNodeComponentPassthroughMesh.cpp
@@ -126,7 +126,7 @@ void UCustomizableObjectNodeComponentPassthroughMesh::AllocateDefaultPins(UCusto
 					FString SectionFriendlyName = MaterialInterface ? MaterialInterface->GetName() : FString::Printf(TEXT("Section %i"), SectionIndex);
 					FString MaterialPinName = FString::Printf(TEXT("LOD %i - Section %i - Material"), LODIndex, SectionIndex);
 
-					UEdGraphPin* Pin = CustomCreatePin(EGPD_Input, UEdGraphSchema_CustomizableObject::PC_Material, FName(*MaterialPinName), PinData);
+					UEdGraphPin* Pin = CustomCreatePin(EGPD_Input, UEdGraphSchema_CustomizableObject::PC_MeshSection, FName(*MaterialPinName), PinData);
 					Pin->PinFriendlyName = FText::FromString(FString::Printf(TEXT("LOD %i - Section %i - %s"), LODIndex, SectionIndex, *SectionFriendlyName));
 					Pin->PinToolTip = MaterialPinName;
 				}
diff --git a/Engine/Plugins/Mutable/Source/CustomizableObjectEditor/Private/MuCOE/Nodes/CustomizableObjectNodeCopyMaterial.cpp b/Engine/Plugins/Mutable/Source/CustomizableObjectEditor/Private/MuCOE/Nodes/CustomizableObjectNodeCopyMaterial.cpp
--- a/Engine/Plugins/Mutable/Source/CustomizableObjectEditor/Private/MuCOE/Nodes/CustomizableObjectNodeCopyMaterial.cpp
+++ b/Engine/Plugins/Mutable/Source/CustomizableObjectEditor/Private/MuCOE/Nodes/CustomizableObjectNodeCopyMaterial.cpp
@@ -41,12 +41,12 @@ void UCustomizableObjectNodeCopyMaterial::AllocateDefaultPins(UCustomizableObjec
 	Pin->PinFriendlyName = LOCTEXT("MeshPin", "Mesh");
 	Pin->bDefaultValueIsIgnored = true;
 
-	Pin = CustomCreatePin(EGPD_Input, Schema->PC_Material, FName(CustomizableObjectNodeCopyMaterial::MeshSectionPinName));
+	Pin = CustomCreatePin(EGPD_Input, Schema->PC_MeshSection, FName(CustomizableObjectNodeCopyMaterial::MeshSectionPinName));
 	Pin->PinFriendlyName = LOCTEXT("BaseMeshSectionPin", "Base Mesh Section");;
 	Pin->bDefaultValueIsIgnored = true;
 
 	// Output pins
-	Pin = CustomCreatePin(EGPD_Output, Schema->PC_Material, FName(CustomizableObjectNodeCopyMaterial::OutputPinName));
+	Pin = CustomCreatePin(EGPD_Output, Schema->PC_MeshSection, FName(CustomizableObjectNodeCopyMaterial::OutputPinName));
 	Pin->PinFriendlyName = LOCTEXT("MeshSectionPin", "Mesh Section");
 }
 
@@ -237,13 +237,13 @@ bool UCustomizableObjectNodeCopyMaterial::IsPinRelevant(const UEdGraphPin* Pin)
 	if (Pin->Direction == EGPD_Output)
 	{
 		return (Node->IsA(UCustomizableObjectNodeMaterial::StaticClass()) && !Node->IsA(UCustomizableObjectNodeCopyMaterial::StaticClass())) ||
-			(Node->IsA(UCustomizableObjectNodeExternalPin::StaticClass()) && Pin->PinType.PinCategory == UEdGraphSchema_CustomizableObject::PC_Material) ||
+			(Node->IsA(UCustomizableObjectNodeExternalPin::StaticClass()) && Pin->PinType.PinCategory == UEdGraphSchema_CustomizableObject::PC_MeshSection) ||
 			Pin->PinType.PinCategory == UEdGraphSchema_CustomizableObject::PC_Mesh;
 	}
 	else
 	{
 		return Node->IsA(UCustomizableObjectNodeObject::StaticClass()) ||
-			(Node->IsA(UCustomizableObjectNodeExposePin::StaticClass()) && Pin->PinType.PinCategory == UEdGraphSchema_CustomizableObject::PC_Material);
+			(Node->IsA(UCustomizableObjectNodeExposePin::StaticClass()) && Pin->PinType.PinCategory == UEdGraphSchema_CustomizableObject::PC_MeshSection);
 			
 	}
 }
diff --git a/Engine/Plugins/Mutable/Source/CustomizableObjectEditor/Private/MuCOE/Nodes/CustomizableObjectNodeExposePin.cpp b/Engine/Plugins/Mutable/Source/CustomizableObjectEditor/Private/MuCOE/Nodes/CustomizableObjectNodeExposePin.cpp
--- a/Engine/Plugins/Mutable/Source/CustomizableObjectEditor/Private/MuCOE/Nodes/CustomizableObjectNodeExposePin.cpp
+++ b/Engine/Plugins/Mutable/Source/CustomizableObjectEditor/Private/MuCOE/Nodes/CustomizableObjectNodeExposePin.cpp
@@ -94,7 +94,7 @@ void UCustomizableObjectNodeExposePin::BackwardsCompatibleFixup(int32 Customizab
 				Pin->PinName = TEXT("Component");
 				Pin->PinFriendlyName = LOCTEXT("Component_Pin_Category", "Component");
 			}
-			else if (PinType == UEdGraphSchema_CustomizableObject::PC_Material)
+			else if (PinType == UEdGraphSchema_CustomizableObject::PC_MeshSection)
 			{
 				Pin->PinName = TEXT("Material");
 				Pin->PinFriendlyName = LOCTEXT("MeshSection_Pin_Category", "Mesh Section");
@@ -159,7 +159,7 @@ void UCustomizableObjectNodeExposePin::BackwardsCompatibleFixup(int32 Customizab
 				Pin->PinName = TEXT("Stack");
 				Pin->PinFriendlyName = LOCTEXT("Stack_Pin_Category", "Stack");
 			}
-			else if (PinType == UEdGraphSchema_CustomizableObject::PC_MaterialAsset)
+			else if (PinType == UEdGraphSchema_CustomizableObject::PC_Material)
 			{
 				Pin->PinName = TEXT("Material");
 				Pin->PinFriendlyName = LOCTEXT("Material_Asset_Pin_Category", "Material");
diff --git a/Engine/Plugins/Mutable/Source/CustomizableObjectEditor/Private/MuCOE/Nodes/CustomizableObjectNodeExternalPin.cpp b/Engine/Plugins/Mutable/Source/CustomizableObjectEditor/Private/MuCOE/Nodes/CustomizableObjectNodeExternalPin.cpp
--- a/Engine/Plugins/Mutable/Source/CustomizableObjectEditor/Private/MuCOE/Nodes/CustomizableObjectNodeExternalPin.cpp
+++ b/Engine/Plugins/Mutable/Source/CustomizableObjectEditor/Private/MuCOE/Nodes/CustomizableObjectNodeExternalPin.cpp
@@ -52,7 +52,7 @@ void UCustomizableObjectNodeExternalPin::BackwardsCompatibleFixup(int32 Customiz
 				Pin->PinName = TEXT("Component");
 				Pin->PinFriendlyName = LOCTEXT("Component_Pin_Category", "Component");
 			}
-			else if (PinType == UEdGraphSchema_CustomizableObject::PC_Material)
+			else if (PinType == UEdGraphSchema_CustomizableObject::PC_MeshSection)
 			{
 				Pin->PinName = TEXT("Material");
 				Pin->PinFriendlyName = LOCTEXT("MeshSection_Pin_Category", "Mesh Section");
@@ -117,7 +117,7 @@ void UCustomizableObjectNodeExternalPin::BackwardsCompatibleFixup(int32 Customiz
 				Pin->PinName = TEXT("Stack");
 				Pin->PinFriendlyName = LOCTEXT("Stack_Pin_Category", "Stack");
 			}
-			else if (PinType == UEdGraphSchema_CustomizableObject::PC_MaterialAsset)
+			else if (PinType == UEdGraphSchema_CustomizableObject::PC_Material)
 			{
 				Pin->PinName = TEXT("Material");
 				Pin->PinFriendlyName = LOCTEXT("Material_Asset_Pin_Category", "Material");
diff --git a/Engine/Plugins/Mutable/Source/CustomizableObjectEditor/Private/MuCOE/Nodes/CustomizableObjectNodeMaterial.cpp b/Engine/Plugins/Mutable/Source/CustomizableObjectEditor/Private/MuCOE/Nodes/CustomizableObjectNodeMaterial.cpp
--- a/Engine/Plugins/Mutable/Source/CustomizableObjectEditor/Private/MuCOE/Nodes/CustomizableObjectNodeMaterial.cpp
+++ b/Engine/Plugins/Mutable/Source/CustomizableObjectEditor/Private/MuCOE/Nodes/CustomizableObjectNodeMaterial.cpp
@@ -329,9 +329,9 @@ void UCustomizableObjectNodeMaterial::AllocateDefaultPins(UCustomizableObjectNod
 	}
 	
 	{
-		const FString PinFriendlyName = UEdGraphSchema_CustomizableObject::GetPinCategoryFriendlyName(Schema->PC_MaterialAsset).ToString();
+		const FString PinFriendlyName = UEdGraphSchema_CustomizableObject::GetPinCategoryFriendlyName(Schema->PC_Material).ToString();
 		const FString PinName = FString(TEXT("Material_Input_Pin"));
-		UEdGraphPin* TableMaterialPin = CustomCreatePin(EGPD_Input, Schema->PC_MaterialAsset, FName(*PinName));
+		UEdGraphPin* TableMaterialPin = CustomCreatePin(EGPD_Input, Schema->PC_Material, FName(*PinName));
 		TableMaterialPin->PinFriendlyName = FText::FromString(PinFriendlyName);
 		TableMaterialPin->bDefaultValueIsIgnored = true;
 		TableMaterialPin->PinToolTip = "Pin for a Material from a Table Node";
@@ -354,7 +354,7 @@ void UCustomizableObjectNodeMaterial::AllocateDefaultPins(UCustomizableObjectNod
 	{
 		const FString PinFriendlyName = TEXT("Mesh Section");
 		const FString PinName = PinFriendlyName + FString(TEXT("_Output_Pin"));
-		UEdGraphPin* OutputPin = CustomCreatePin(EGPD_Output, Schema->PC_Material, FName(*PinName));
+		UEdGraphPin* OutputPin = CustomCreatePin(EGPD_Output, Schema->PC_MeshSection, FName(*PinName));
 		OutputPin->PinFriendlyName = FText::FromString(PinFriendlyName);
 	}
 }
@@ -797,7 +797,7 @@ void UCustomizableObjectNodeMaterial::BackwardsCompatibleFixup(int32 Customizabl
 		{
 			if (Pin)
 			{
-				if (Pin->PinType.PinCategory == UEdGraphSchema_CustomizableObject::PC_MaterialAsset && 
+				if (Pin->PinType.PinCategory == UEdGraphSchema_CustomizableObject::PC_Material && 
 					Pin->Direction == EGPD_Input && Pin->PinName.ToString().EndsWith(TEXT("_Input_Pin")) &&
 					Pin->PinName != FString(TEXT("Material_Input_Pin")))
 				{
@@ -901,7 +901,7 @@ UEdGraphPin* UCustomizableObjectNodeMaterial::GetMeshPin() const
 
 UEdGraphPin* UCustomizableObjectNodeMaterial::GetMaterialAssetPin() const
 {
-	const FString PinFriendlyName = UEdGraphSchema_CustomizableObject::GetPinCategoryFriendlyName(UEdGraphSchema_CustomizableObject::PC_MaterialAsset).ToString();
+	const FString PinFriendlyName = UEdGraphSchema_CustomizableObject::GetPinCategoryFriendlyName(UEdGraphSchema_CustomizableObject::PC_Material).ToString();
 	const FString PinName = PinFriendlyName + FString(TEXT("_Input_Pin"));
 
 	UEdGraphPin* Pin = FindPin(PinName);
@@ -1361,7 +1361,7 @@ bool UCustomizableObjectNodeMaterial::IsPinRelevant(const UEdGraphPin* Pin) cons
 	}
 	else if (Pin->Direction == EEdGraphPinDirection::EGPD_Input)
 	{	
-		return Pin->PinType.PinCategory == Schema->PC_Material;
+		return Pin->PinType.PinCategory == Schema->PC_MeshSection;
 	}
 	else
 	{
diff --git a/Engine/Plugins/Mutable/Source/CustomizableObjectEditor/Private/MuCOE/Nodes/CustomizableObjectNodeMaterialBase.cpp b/Engine/Plugins/Mutable/Source/CustomizableObjectEditor/Private/MuCOE/Nodes/CustomizableObjectNodeMaterialBase.cpp
--- a/Engine/Plugins/Mutable/Source/CustomizableObjectEditor/Private/MuCOE/Nodes/CustomizableObjectNodeMaterialBase.cpp
+++ b/Engine/Plugins/Mutable/Source/CustomizableObjectEditor/Private/MuCOE/Nodes/CustomizableObjectNodeMaterialBase.cpp
@@ -11,5 +11,5 @@ int32 UCustomizableObjectNodeMaterialBase::GetMeshSectionMaxLOD() const
 
 FLinearColor UCustomizableObjectNodeMaterialBase::GetNodeTitleColor() const
 {
-	return UEdGraphSchema_CustomizableObject::GetPinTypeColor(UEdGraphSchema_CustomizableObject::PC_Material);
+	return UEdGraphSchema_CustomizableObject::GetPinTypeColor(UEdGraphSchema_CustomizableObject::PC_MeshSection);
 }
diff --git a/Engine/Plugins/Mutable/Source/CustomizableObjectEditor/Private/MuCOE/Nodes/CustomizableObjectNodeMaterialInterface.h b/Engine/Plugins/Mutable/Source/CustomizableObjectEditor/Private/MuCOE/Nodes/CustomizableObjectNodeMaterialInterface.h
new file mode 100644
--- /dev/null
+++ b/Engine/Plugins/Mutable/Source/CustomizableObjectEditor/Private/MuCOE/Nodes/CustomizableObjectNodeMaterialInterface.h
@@ -0,0 +1,32 @@
+// Copyright Epic Games, Inc. All Rights Reserved.
+
+#pragma once
+
+#include "MuCOE/Nodes/CustomizableObjectNode.h"
+
+#include "CustomizableObjectNodeMaterialInterface.generated.h"
+
+class UEdGraphPin;
+class UMaterialInterface;
+
+UINTERFACE(MinimalAPI)
+class UCustomizableObjectNodeMaterialInterface : public UInterface
+{
+	GENERATED_BODY()
+};
+
+
+class ICustomizableObjectNodeMaterialInterface
+{
+	GENERATED_BODY()
+
+public:
+	// Own interface
+
+	/** Returns the Unreal mesh (e.g., USkeletalMesh, UStaticMesh...). */
+	virtual TSoftObjectPtr<UMaterialInterface> GetMaterial() const PURE_VIRTUAL(UCustomizableObjectNodeMaterialInterface::GetMaterial, return {};);
+
+	/** Returns the output Mesh pin associated to the given LODIndex and SectionIndex. Override. */
+	virtual UEdGraphPin* GetMaterialPin() const PURE_VIRTUAL(UCustomizableObjectNodeMaterialInterface::GetMaterialPin, return {};);
+
+};
diff --git a/Engine/Plugins/Mutable/Source/CustomizableObjectEditor/Private/MuCOE/Nodes/CustomizableObjectNodeMaterialParameter.cpp b/Engine/Plugins/Mutable/Source/CustomizableObjectEditor/Private/MuCOE/Nodes/CustomizableObjectNodeMaterialParameter.cpp
new file mode 100644
--- /dev/null
+++ b/Engine/Plugins/Mutable/Source/CustomizableObjectEditor/Private/MuCOE/Nodes/CustomizableObjectNodeMaterialParameter.cpp
@@ -0,0 +1,31 @@
+﻿// Copyright Epic Games, Inc. All Rights Reserved.
+
+#include "CustomizableObjectNodeMaterialParameter.h"
+
+#include "MuCOE/EdGraphSchema_CustomizableObject.h"
+
+
+bool UCustomizableObjectNodeMaterialParameter::IsExperimental() const
+{
+	return true;
+}
+
+
+FName UCustomizableObjectNodeMaterialParameter::GetCategory() const
+{
+	return UEdGraphSchema_CustomizableObject::PC_Material;
+}
+
+
+TSoftObjectPtr<UMaterialInterface> UCustomizableObjectNodeMaterialParameter::GetMaterial() const
+{
+	unimplemented(); // TODO Does not make sense
+	return nullptr;
+}
+
+
+UEdGraphPin* UCustomizableObjectNodeMaterialParameter::GetMaterialPin() const
+{
+	unimplemented(); // TODO Does not make sense
+	return nullptr;
+}
diff --git a/Engine/Plugins/Mutable/Source/CustomizableObjectEditor/Private/MuCOE/Nodes/CustomizableObjectNodeMaterialParameter.h b/Engine/Plugins/Mutable/Source/CustomizableObjectEditor/Private/MuCOE/Nodes/CustomizableObjectNodeMaterialParameter.h
new file mode 100644
--- /dev/null
+++ b/Engine/Plugins/Mutable/Source/CustomizableObjectEditor/Private/MuCOE/Nodes/CustomizableObjectNodeMaterialParameter.h
@@ -0,0 +1,41 @@
+// Copyright Epic Games, Inc. All Rights Reserved.
+
+#pragma once
+
+#include "MuCOE/Nodes/CustomizableObjectNodeParameter.h"
+#include "MuCOE/Nodes/CustomizableObjectNodeMaterialInterface.h"
+
+#include "CustomizableObjectNodeMaterialParameter.generated.h"
+
+#define UE_API CUSTOMIZABLEOBJECTEDITOR_API
+
+namespace ENodeTitleType { enum Type : int; }
+
+class UMaterialInterface;
+
+
+UCLASS(MinimalAPI)
+class UCustomizableObjectNodeMaterialParameter : public UCustomizableObjectNodeParameter, public ICustomizableObjectNodeMaterialInterface
+{
+public:
+	GENERATED_BODY()
+
+	/** Default value of the parameter. */
+	UPROPERTY(EditAnywhere, Category=CustomizableObject)
+	TSoftObjectPtr<UMaterialInterface> DefaultValue;
+
+	UPROPERTY(EditAnywhere, Category=CustomizableObject)
+	TSoftObjectPtr<UMaterialInterface> ReferenceValue;
+	
+	// UCustomizableObjectNode interface
+	UE_API virtual bool IsExperimental() const override;
+
+	// CustomizableObjectNodeParameter interface
+	UE_API virtual FName GetCategory() const override;
+
+	// ICustomizableObjectNodeMaterialInterface interface
+	UE_API virtual TSoftObjectPtr<UMaterialInterface> GetMaterial() const override;
+	UE_API virtual UEdGraphPin* GetMaterialPin() const override;
+};
+
+#undef UE_API
diff --git a/Engine/Plugins/Mutable/Source/CustomizableObjectEditor/Private/MuCOE/Nodes/CustomizableObjectNodeMaterialSwitch.cpp b/Engine/Plugins/Mutable/Source/CustomizableObjectEditor/Private/MuCOE/Nodes/CustomizableObjectNodeMaterialSwitch.cpp
--- a/Engine/Plugins/Mutable/Source/CustomizableObjectEditor/Private/MuCOE/Nodes/CustomizableObjectNodeMaterialSwitch.cpp
+++ b/Engine/Plugins/Mutable/Source/CustomizableObjectEditor/Private/MuCOE/Nodes/CustomizableObjectNodeMaterialSwitch.cpp
@@ -10,7 +10,7 @@
 
 FName UCustomizableObjectNodeMaterialSwitch::GetCategory() const
 {
-	return UEdGraphSchema_CustomizableObject::PC_Material;
+	return UEdGraphSchema_CustomizableObject::PC_MeshSection;
 }
 
 void UCustomizableObjectNodeMaterialSwitch::BackwardsCompatibleFixup(int32 CustomizableObjectCustomVersion)
@@ -33,7 +33,7 @@ void UCustomizableObjectNodeMaterialSwitch::BackwardsCompatibleFixup(int32 Custo
 			UEdGraphPin* NewPin = FindPin(TEXT("Material"), EEdGraphPinDirection::EGPD_Output);
 			if (!NewPin)
 			{
-				NewPin = CustomCreatePin(EGPD_Output, Schema->PC_Material, FName("Material"));
+				NewPin = CustomCreatePin(EGPD_Output, Schema->PC_MeshSection, FName("Material"));
 			}
 
 			FGuid PinId = NewPin->PinId;
@@ -57,7 +57,7 @@ void UCustomizableObjectNodeMaterialSwitch::BackwardsCompatibleFixup(int32 Custo
 				UEdGraphPin* NewPin = FindPin(NewName, EEdGraphPinDirection::EGPD_Input);
 				if (!NewPin)
 				{
-					NewPin = CustomCreatePin(EEdGraphPinDirection::EGPD_Input, Schema->PC_Material, FName(*NewName));
+					NewPin = CustomCreatePin(EEdGraphPinDirection::EGPD_Input, Schema->PC_MeshSection, FName(*NewName));
 				}
 
 				FGuid PinId = NewPin->PinId;
diff --git a/Engine/Plugins/Mutable/Source/CustomizableObjectEditor/Private/MuCOE/Nodes/CustomizableObjectNodeMaterialVariation.cpp b/Engine/Plugins/Mutable/Source/CustomizableObjectEditor/Private/MuCOE/Nodes/CustomizableObjectNodeMaterialVariation.cpp
--- a/Engine/Plugins/Mutable/Source/CustomizableObjectEditor/Private/MuCOE/Nodes/CustomizableObjectNodeMaterialVariation.cpp
+++ b/Engine/Plugins/Mutable/Source/CustomizableObjectEditor/Private/MuCOE/Nodes/CustomizableObjectNodeMaterialVariation.cpp
@@ -26,7 +26,7 @@ void UCustomizableObjectNodeMaterialVariation::BackwardsCompatibleFixup(int32 Cu
 
 FName UCustomizableObjectNodeMaterialVariation::GetCategory() const
 {
-	return UEdGraphSchema_CustomizableObject::PC_Material;
+	return UEdGraphSchema_CustomizableObject::PC_MeshSection;
 }
 
 
diff --git a/Engine/Plugins/Mutable/Source/CustomizableObjectEditor/Private/MuCOE/Nodes/CustomizableObjectNodeMesh.cpp b/Engine/Plugins/Mutable/Source/CustomizableObjectEditor/Private/MuCOE/Nodes/CustomizableObjectNodeMesh.cpp
deleted file mode 100644
diff --git a/Engine/Plugins/Mutable/Source/CustomizableObjectEditor/Private/MuCOE/Nodes/CustomizableObjectNodeMesh.h b/Engine/Plugins/Mutable/Source/CustomizableObjectEditor/Private/MuCOE/Nodes/CustomizableObjectNodeMesh.h
deleted file mode 100644
diff --git a/Engine/Plugins/Mutable/Source/CustomizableObjectEditor/Private/MuCOE/Nodes/CustomizableObjectNodeMeshInterface.h b/Engine/Plugins/Mutable/Source/CustomizableObjectEditor/Private/MuCOE/Nodes/CustomizableObjectNodeMeshInterface.h
new file mode 100644
--- /dev/null
+++ b/Engine/Plugins/Mutable/Source/CustomizableObjectEditor/Private/MuCOE/Nodes/CustomizableObjectNodeMeshInterface.h
@@ -0,0 +1,50 @@
+// Copyright Epic Games, Inc. All Rights Reserved.
+
+#pragma once
+
+#include "MuCOE/Nodes/CustomizableObjectNode.h"
+
+#include "CustomizableObjectNodeMeshInterface.generated.h"
+
+class UCustomizableObjectLayout;
+class UEdGraphPin;
+class UObject;
+class UTexture2D;
+
+
+UINTERFACE(MinimalAPI)
+class UCustomizableObjectNodeMeshInterface : public UInterface
+{
+	GENERATED_BODY()
+};
+
+
+class ICustomizableObjectNodeMeshInterface
+{
+public:
+	GENERATED_BODY()
+
+	// Own interface
+	/** Returns the multiple Layouts of a given Mesh pin. A pin can have multiple Layouts since it can have multiple UVs. Override. */
+	virtual TArray<UCustomizableObjectLayout*> GetLayouts(const UEdGraphPin& OutPin) const PURE_VIRTUAL(UCustomizableObjectNodeMesh::GetLayouts, return {};);
+
+	virtual UTexture2D* FindTextureForPin(const UEdGraphPin* Pin) const PURE_VIRTUAL(UCustomizableObjectNodeMesh::FindTextureForPin, return {};);
+
+	/** Returns the Unreal mesh (e.g., USkeletalMesh, UStaticMesh...). */
+	virtual TSoftObjectPtr<UObject> GetMesh() const PURE_VIRTUAL(UCustomizableObjectNodeMesh::GetMesh, return {};);
+
+	/** Returns the output Mesh pin associated to the given LODIndex and SectionIndex. Override. */
+	virtual UEdGraphPin* GetMeshPin(int32 LOD, int32 SectionIndex) const PURE_VIRTUAL(UCustomizableObjectNodeMesh::GetMeshPin, return {};);
+
+	/** Given a pin, return the Section and Layout index.
+	 *
+	 * Will always return a valid result. The result can be valid but out of sync with respect the Unreal mesh asset.
+	 * In other words, the pin still represents a LOD 3 but the asset no longer has this third LOD.
+	 * 
+	 * @param Pin In. Node's owning pin.
+	 * @param LODIndex Out.
+	 * @param SectionIndex Out.
+	 * @param LayoutIndex Out. */
+	virtual void GetPinSection(const UEdGraphPin& Pin, int32& OutLODIndex, int32& OutSectionIndex, int32& OutLayoutIndex) const PURE_VIRTUAL(UCustomizableObjectNodeMesh::GetPinSection, );
+};
+
diff --git a/Engine/Plugins/Mutable/Source/CustomizableObjectEditor/Private/MuCOE/Nodes/CustomizableObjectNodeMeshMorphDetails.cpp b/Engine/Plugins/Mutable/Source/CustomizableObjectEditor/Private/MuCOE/Nodes/CustomizableObjectNodeMeshMorphDetails.cpp
--- a/Engine/Plugins/Mutable/Source/CustomizableObjectEditor/Private/MuCOE/Nodes/CustomizableObjectNodeMeshMorphDetails.cpp
+++ b/Engine/Plugins/Mutable/Source/CustomizableObjectEditor/Private/MuCOE/Nodes/CustomizableObjectNodeMeshMorphDetails.cpp
@@ -54,7 +54,7 @@ void FCustomizableObjectNodeMeshMorphDetails::CustomizeDetails( IDetailLayoutBui
 		if ( BaseSourcePin )
 		{
 			USkeletalMesh* SkeletalMesh = nullptr;
-			if ( const UCustomizableObjectNodeMesh* TypedNodeSkeletalMesh = Cast<UCustomizableObjectNodeMesh>(BaseSourcePin->GetOwningNode()) )
+			if (const ICustomizableObjectNodeMeshInterface* TypedNodeSkeletalMesh = Cast<ICustomizableObjectNodeMeshInterface>(BaseSourcePin->GetOwningNode()) )
 			{
 				SkeletalMesh = Cast<USkeletalMesh>(MutablePrivate::LoadObject(TypedNodeSkeletalMesh->GetMesh()));
 			}
diff --git a/Engine/Plugins/Mutable/Source/CustomizableObjectEditor/Private/MuCOE/Nodes/CustomizableObjectNodeMeshMorphStackApplication.cpp b/Engine/Plugins/Mutable/Source/CustomizableObjectEditor/Private/MuCOE/Nodes/CustomizableObjectNodeMeshMorphStackApplication.cpp
--- a/Engine/Plugins/Mutable/Source/CustomizableObjectEditor/Private/MuCOE/Nodes/CustomizableObjectNodeMeshMorphStackApplication.cpp
+++ b/Engine/Plugins/Mutable/Source/CustomizableObjectEditor/Private/MuCOE/Nodes/CustomizableObjectNodeMeshMorphStackApplication.cpp
@@ -107,7 +107,7 @@ TArray<FString> UCustomizableObjectNodeMeshMorphStackApplication::GetMorphList()
 	{
 		USkeletalMesh* SkeletalMesh = nullptr;
 
-		if (const UCustomizableObjectNodeMesh* MeshNode = Cast< UCustomizableObjectNodeMesh >(MeshNodePin->GetOwningNode()))
+		if (const ICustomizableObjectNodeMeshInterface* MeshNode = Cast<ICustomizableObjectNodeMeshInterface>(MeshNodePin->GetOwningNode()))
 		{
 			SkeletalMesh = Cast<USkeletalMesh>(MutablePrivate::LoadObject(MeshNode->GetMesh()));
 		}
diff --git a/Engine/Plugins/Mutable/Source/CustomizableObjectEditor/Private/MuCOE/Nodes/CustomizableObjectNodeMeshMorphStackDefinition.cpp b/Engine/Plugins/Mutable/Source/CustomizableObjectEditor/Private/MuCOE/Nodes/CustomizableObjectNodeMeshMorphStackDefinition.cpp
--- a/Engine/Plugins/Mutable/Source/CustomizableObjectEditor/Private/MuCOE/Nodes/CustomizableObjectNodeMeshMorphStackDefinition.cpp
+++ b/Engine/Plugins/Mutable/Source/CustomizableObjectEditor/Private/MuCOE/Nodes/CustomizableObjectNodeMeshMorphStackDefinition.cpp
@@ -98,7 +98,7 @@ void UCustomizableObjectNodeMeshMorphStackDefinition::UpdateMorphList()
 	{
 		USkeletalMesh* SkeletalMesh = nullptr;
 
-		if (const UCustomizableObjectNodeMesh* MeshNode = Cast<UCustomizableObjectNodeMesh>(MeshNodePin->GetOwningNode()))
+		if (const ICustomizableObjectNodeMeshInterface* MeshNode = Cast<ICustomizableObjectNodeMeshInterface>(MeshNodePin->GetOwningNode()))
 		{
 			SkeletalMesh = Cast<USkeletalMesh>(MutablePrivate::LoadObject(MeshNode->GetMesh()));
 		}
diff --git a/Engine/Plugins/Mutable/Source/CustomizableObjectEditor/Private/MuCOE/Nodes/CustomizableObjectNodeObject.cpp b/Engine/Plugins/Mutable/Source/CustomizableObjectEditor/Private/MuCOE/Nodes/CustomizableObjectNodeObject.cpp
--- a/Engine/Plugins/Mutable/Source/CustomizableObjectEditor/Private/MuCOE/Nodes/CustomizableObjectNodeObject.cpp
+++ b/Engine/Plugins/Mutable/Source/CustomizableObjectEditor/Private/MuCOE/Nodes/CustomizableObjectNodeObject.cpp
@@ -150,7 +150,7 @@ void UCustomizableObjectNodeObject::BackwardsCompatibleFixup(int32 CustomizableO
 			{
 				FString LODName = FString::Printf(TEXT("LOD %d"), NodeComponentLODIndex);
 
-				UEdGraphPin* Pin = Node->CustomCreatePin(EGPD_Input, Schema->PC_Material, FName(*LODName), true);
+				UEdGraphPin* Pin = Node->CustomCreatePin(EGPD_Input, Schema->PC_MeshSection, FName(*LODName), true);
 				NodeComponentMesh->GetLODPins().Add(Pin);
 			}
 						
diff --git a/Engine/Plugins/Mutable/Source/CustomizableObjectEditor/Private/MuCOE/Nodes/CustomizableObjectNodeObjectDetails.cpp b/Engine/Plugins/Mutable/Source/CustomizableObjectEditor/Private/MuCOE/Nodes/CustomizableObjectNodeObjectDetails.cpp
--- a/Engine/Plugins/Mutable/Source/CustomizableObjectEditor/Private/MuCOE/Nodes/CustomizableObjectNodeObjectDetails.cpp
+++ b/Engine/Plugins/Mutable/Source/CustomizableObjectEditor/Private/MuCOE/Nodes/CustomizableObjectNodeObjectDetails.cpp
@@ -469,7 +469,7 @@ void FCustomizableObjectNodeObjectDetails::FillParameterNamesArray()
 		{
 			if (MeshParameterNode)
 			{
-				BaseObjectNode->ParameterNames.Add(MeshParameterNode->ParameterName);
+				BaseObjectNode->ParameterNames.Add(MeshParameterNode->GetParameterName());
 			}
 		}
 
diff --git a/Engine/Plugins/Mutable/Source/CustomizableObjectEditor/Private/MuCOE/Nodes/CustomizableObjectNodeSkeletalMesh.h b/Engine/Plugins/Mutable/Source/CustomizableObjectEditor/Private/MuCOE/Nodes/CustomizableObjectNodeSkeletalMesh.h
--- a/Engine/Plugins/Mutable/Source/CustomizableObjectEditor/Private/MuCOE/Nodes/CustomizableObjectNodeSkeletalMesh.h
+++ b/Engine/Plugins/Mutable/Source/CustomizableObjectEditor/Private/MuCOE/Nodes/CustomizableObjectNodeSkeletalMesh.h
@@ -3,7 +3,7 @@
 #pragma once
 
 #include "GameplayTagContainer.h"
-#include "MuCOE/Nodes/CustomizableObjectNodeMesh.h"
+#include "MuCOE/Nodes/CustomizableObjectNodeMeshInterface.h"
 #include "MuCOE/CustomizableObjectEditor_Deprecated.h"
 #include "MuCOE/Nodes/SCustomizableObjectNode.h"
 #include "MuCOE/RemapPins/CustomizableObjectNodeRemapPinsByNameDefaultPin.h"
@@ -179,7 +179,7 @@ private:
 
 
 UCLASS(MinimalAPI)
-class UCustomizableObjectNodeSkeletalMesh : public UCustomizableObjectNodeMesh
+class UCustomizableObjectNodeSkeletalMesh : public UCustomizableObjectNode, public ICustomizableObjectNodeMeshInterface
 {
 public:
 	GENERATED_BODY()
diff --git a/Engine/Plugins/Mutable/Source/CustomizableObjectEditor/Private/MuCOE/Nodes/CustomizableObjectNodeSkeletalMeshParameter.cpp b/Engine/Plugins/Mutable/Source/CustomizableObjectEditor/Private/MuCOE/Nodes/CustomizableObjectNodeSkeletalMeshParameter.cpp
--- a/Engine/Plugins/Mutable/Source/CustomizableObjectEditor/Private/MuCOE/Nodes/CustomizableObjectNodeSkeletalMeshParameter.cpp
+++ b/Engine/Plugins/Mutable/Source/CustomizableObjectEditor/Private/MuCOE/Nodes/CustomizableObjectNodeSkeletalMeshParameter.cpp
@@ -15,15 +15,15 @@ class UCustomizableObjectNodeRemapPins;
 
 void UCustomizableObjectNodeSkeletalMeshParameter::AllocateDefaultPins(UCustomizableObjectNodeRemapPins* RemapPins)
 {
+	NamePin = CustomCreatePin(EGPD_Input, UEdGraphSchema_CustomizableObject::PC_String, FName("Name"));
+
 	USkeletalMesh* SkeletalMesh = ReferenceValue.LoadSynchronous();
 
 	if (!SkeletalMesh)
 	{
 		return;
 	}
-
-	NamePin = CustomCreatePin(EGPD_Input, UEdGraphSchema_CustomizableObject::PC_String, FName("Name"));
-
+	
 	const UEdGraphSchema_CustomizableObject* Schema = GetDefault<UEdGraphSchema_CustomizableObject>();
 
 	if (const FSkeletalMeshModel* ImportedModel = SkeletalMesh->GetImportedModel())
@@ -60,7 +60,6 @@ void UCustomizableObjectNodeSkeletalMeshParameter::AllocateDefaultPins(UCustomiz
 			}
 		}
 	}
-
 }
 
 
@@ -96,6 +95,12 @@ void UCustomizableObjectNodeSkeletalMeshParameter::BackwardsCompatibleFixup(int3
 }
 
 
+FName UCustomizableObjectNodeSkeletalMeshParameter::GetCategory() const
+{
+	return UEdGraphSchema_CustomizableObject::PC_Mesh;
+}
+
+
 FText UCustomizableObjectNodeSkeletalMeshParameter::GetNodeTitle(ENodeTitleType::Type TitleType) const
 {
 	return LOCTEXT("Skeletal_Mesh_Parameter", "Skeletal Mesh Parameter");
@@ -115,15 +120,6 @@ FText UCustomizableObjectNodeSkeletalMeshParameter::GetTooltipText() const
 }
 
 
-void UCustomizableObjectNodeSkeletalMeshParameter::OnRenameNode(const FString& NewName)
-{
-	if (!NewName.IsEmpty())
-	{
-		ParameterName = NewName;
-	}
-}
-
-
 void UCustomizableObjectNodeSkeletalMeshParameter::PinConnectionListChanged(UEdGraphPin* Pin)
 {
 	if (Pin == NamePin.Get())
diff --git a/Engine/Plugins/Mutable/Source/CustomizableObjectEditor/Private/MuCOE/Nodes/CustomizableObjectNodeSkeletalMeshParameter.h b/Engine/Plugins/Mutable/Source/CustomizableObjectEditor/Private/MuCOE/Nodes/CustomizableObjectNodeSkeletalMeshParameter.h
--- a/Engine/Plugins/Mutable/Source/CustomizableObjectEditor/Private/MuCOE/Nodes/CustomizableObjectNodeSkeletalMeshParameter.h
+++ b/Engine/Plugins/Mutable/Source/CustomizableObjectEditor/Private/MuCOE/Nodes/CustomizableObjectNodeSkeletalMeshParameter.h
@@ -2,7 +2,8 @@
 
 #pragma once
 
-#include "MuCOE/Nodes/CustomizableObjectNodeMesh.h"
+#include "CustomizableObjectNodeParameter.h"
+#include "MuCOE/Nodes/CustomizableObjectNodeMeshInterface.h"
 
 #include "CustomizableObjectNodeSkeletalMeshParameter.generated.h"
 
@@ -43,7 +44,7 @@ private:
 
 
 UCLASS(MinimalAPI)
-class UCustomizableObjectNodeSkeletalMeshParameter : public UCustomizableObjectNodeMesh
+class UCustomizableObjectNodeSkeletalMeshParameter : public UCustomizableObjectNodeParameter, public ICustomizableObjectNodeMeshInterface
 {
 public:
 	GENERATED_BODY()
@@ -55,17 +56,6 @@ public:
 	UPROPERTY(EditAnywhere, Category=CustomizableObject)
 	TSoftObjectPtr<USkeletalMesh> ReferenceValue;
 	
-	UPROPERTY()
-	FString ParameterName = "Default Name";
-
-	UPROPERTY(EditAnywhere, Category = UI, meta = (DisplayName = "Parameter UI Metadata"))
-	FMutableParamUIMetadata ParamUIMetadata;
-
-	UPROPERTY()
-	FEdGraphPinReference NamePin;
-
-public:
-
 	// UObject interface.
 	UE_API virtual void PostEditChangeProperty(FPropertyChangedEvent& PropertyChangedEvent) override;
 
@@ -73,7 +63,6 @@ public:
 	UE_API virtual FText GetNodeTitle(ENodeTitleType::Type TitleType) const override;
 	UE_API virtual FLinearColor GetNodeTitleColor() const override;
 	UE_API virtual FText GetTooltipText() const override;
-	UE_API virtual void OnRenameNode(const FString& NewName) override;
 	UE_API virtual void PinConnectionListChanged(UEdGraphPin* Pin) override;
 
 	// UCustomizableObjectNode interface
@@ -82,7 +71,10 @@ public:
 	UE_API virtual bool IsExperimental() const override;
 	UE_API virtual void BackwardsCompatibleFixup(int32 CustomizableObjectCustomVersion) override;
 
-	// UCustomizableObjectNodeMesh interface
+	// UCustomizableObjectNodeParameter iterface
+	virtual FName GetCategory() const override;
+	
+	// ICustomizableObjectNodeMeshInterface interface
 	UE_API virtual UTexture2D* FindTextureForPin(const UEdGraphPin* Pin) const override;
 	UE_API virtual TArray<UCustomizableObjectLayout*> GetLayouts(const UEdGraphPin& MeshPin) const override;
 	UE_API virtual TSoftObjectPtr<UObject> GetMesh() const override;
@@ -95,7 +87,6 @@ public:
 	UE_API FSkeletalMaterial* GetSkeletalMaterialFor(const int32 SectionIndex) const;
 	UE_API int32 GetSkeletalMaterialIndexFor(const UEdGraphPin& Pin) const;
 	UE_API int32 GetSkeletalMaterialIndexFor(const int32 SectionIndex) const;
-
 };
 
 #undef UE_API
diff --git a/Engine/Plugins/Mutable/Source/CustomizableObjectEditor/Private/MuCOE/Nodes/CustomizableObjectNodeStaticMesh.h b/Engine/Plugins/Mutable/Source/CustomizableObjectEditor/Private/MuCOE/Nodes/CustomizableObjectNodeStaticMesh.h
--- a/Engine/Plugins/Mutable/Source/CustomizableObjectEditor/Private/MuCOE/Nodes/CustomizableObjectNodeStaticMesh.h
+++ b/Engine/Plugins/Mutable/Source/CustomizableObjectEditor/Private/MuCOE/Nodes/CustomizableObjectNodeStaticMesh.h
@@ -2,7 +2,7 @@
 
 #pragma once
 
-#include "MuCOE/Nodes/CustomizableObjectNodeMesh.h"
+#include "MuCOE/Nodes/CustomizableObjectNodeMeshInterface.h"
 #include "MuCOE/Nodes/SCustomizableObjectNode.h"
 
 #include "CustomizableObjectNodeStaticMesh.generated.h"
@@ -116,7 +116,7 @@ struct FCustomizableObjectNodeStaticMeshLOD
 
 
 UCLASS(MinimalAPI)
-class UCustomizableObjectNodeStaticMesh : public UCustomizableObjectNodeMesh
+class UCustomizableObjectNodeStaticMesh : public UCustomizableObjectNode, public ICustomizableObjectNodeMeshInterface
 { 
 public:
 	GENERATED_BODY()
diff --git a/Engine/Plugins/Mutable/Source/CustomizableObjectEditor/Private/MuCOE/Nodes/CustomizableObjectNodeTable.cpp b/Engine/Plugins/Mutable/Source/CustomizableObjectEditor/Private/MuCOE/Nodes/CustomizableObjectNodeTable.cpp
--- a/Engine/Plugins/Mutable/Source/CustomizableObjectEditor/Private/MuCOE/Nodes/CustomizableObjectNodeTable.cpp
+++ b/Engine/Plugins/Mutable/Source/CustomizableObjectEditor/Private/MuCOE/Nodes/CustomizableObjectNodeTable.cpp
@@ -654,7 +654,7 @@ void UCustomizableObjectNodeTable::AllocateDefaultPins(UCustomizableObjectNodeRe
 						PinData->ColumnDisplayName = PinName;
 						PinData->StructColumnId = PropertyID;
 
-						OutPin = CustomCreatePin(EGPD_Output, Schema->PC_MaterialAsset, FName(*PinName), PinData);
+						OutPin = CustomCreatePin(EGPD_Output, Schema->PC_Material, FName(*PinName), PinData);
 					}
 
 					else if (Object->IsA(UPoseAsset::StaticClass()))
@@ -860,7 +860,7 @@ bool UCustomizableObjectNodeTable::IsNodeOutDatedAndNeedsRefresh()
 				}
 				else if (Object->IsA(UMaterialInterface::StaticClass()))
 				{
-					if (CheckPinUpdated(PinName, Schema->PC_MaterialAsset))
+					if (CheckPinUpdated(PinName, Schema->PC_Material))
 					{
 						bNeedsUpdate = true;
 					}
@@ -991,7 +991,7 @@ bool UCustomizableObjectNodeTable::IsPinRelevant(const UEdGraphPin* Pin) const
 	const UEdGraphSchema_CustomizableObject* Schema = GetDefault<UEdGraphSchema_CustomizableObject>();
 	
 	return Pin->Direction == EGPD_Input &&
-		(Pin->PinType.PinCategory == Schema->PC_MaterialAsset ||
+		(Pin->PinType.PinCategory == Schema->PC_Material ||
 		Pin->PinType.PinCategory == Schema->PC_Image ||
 		Pin->PinType.PinCategory == Schema->PC_PassThroughImage ||
 		Pin->PinType.PinCategory == Schema->PC_Color ||
diff --git a/Engine/Plugins/Mutable/Source/CustomizableObjectEditor/Private/MuCOE/Nodes/SCustomizableObjectNodeStaticStringPin.cpp b/Engine/Plugins/Mutable/Source/CustomizableObjectEditor/Private/MuCOE/Nodes/SCustomizableObjectNodeStaticStringPin.cpp
--- a/Engine/Plugins/Mutable/Source/CustomizableObjectEditor/Private/MuCOE/Nodes/SCustomizableObjectNodeStaticStringPin.cpp
+++ b/Engine/Plugins/Mutable/Source/CustomizableObjectEditor/Private/MuCOE/Nodes/SCustomizableObjectNodeStaticStringPin.cpp
@@ -55,7 +55,7 @@ FText SCustomizableObjectNodeStaticStringPin::GetNodeStringValue() const
 	}
 	else if (const UCustomizableObjectNodeSkeletalMeshParameter* MeshParamNode = Cast<UCustomizableObjectNodeSkeletalMeshParameter>(Node))
 	{
-		return FText::FromString(MeshParamNode->ParameterName);
+		return FText::FromString(MeshParamNode->GetParameterName());
 	}
 	else if (const UCustomizableObjectNodeObject* ObjectNode = Cast<UCustomizableObjectNodeObject>(Node))
 	{
@@ -104,10 +104,6 @@ void SCustomizableObjectNodeStaticStringPin::OnTextCommited(const FText& InValue
 	{
 		AddComponentNode->SetParentComponentName(FName(*InValue.ToString()));
 	}
-	else if (UCustomizableObjectNodeSkeletalMeshParameter* MeshParamNode = Cast<UCustomizableObjectNodeSkeletalMeshParameter>(Node))
-	{
-		MeshParamNode->ParameterName = InValue.ToString();
-	}
 	else if (UCustomizableObjectNodeObject* ObjectNode = Cast<UCustomizableObjectNodeObject>(Node))
 	{
 		ObjectNode->SetObjectName(InValue.ToString());
diff --git a/Engine/Plugins/Mutable/Source/CustomizableObjectEditor/Public/MuCOE/EdGraphSchema_CustomizableObject.h b/Engine/Plugins/Mutable/Source/CustomizableObjectEditor/Public/MuCOE/EdGraphSchema_CustomizableObject.h
--- a/Engine/Plugins/Mutable/Source/CustomizableObjectEditor/Public/MuCOE/EdGraphSchema_CustomizableObject.h
+++ b/Engine/Plugins/Mutable/Source/CustomizableObjectEditor/Public/MuCOE/EdGraphSchema_CustomizableObject.h
@@ -28,7 +28,7 @@ public:
 	// Allowable PinType.PinCategory values
 	static UE_API const FName PC_Object;
 	static UE_API const FName PC_Component;
-	static UE_API const FName PC_Material;
+	static UE_API const FName PC_MeshSection;
 	static UE_API const FName PC_Modifier;
 	static UE_API const FName PC_Mesh;
 	static UE_API const FName PC_PassThroughMesh;
@@ -42,7 +42,7 @@ public:
 	static UE_API const FName PC_Bool;
 	static UE_API const FName PC_Enum;
 	static UE_API const FName PC_Stack;
-	static UE_API const FName PC_MaterialAsset;
+	static UE_API const FName PC_Material;
 	static UE_API const FName PC_Wildcard;
 	static UE_API const FName PC_PoseAsset;
 	static UE_API const FName PC_Transform;
diff --git a/Engine/Plugins/Mutable/Source/MutableRuntime/Private/MuR/Model.cpp b/Engine/Plugins/Mutable/Source/MutableRuntime/Private/MuR/Model.cpp
--- a/Engine/Plugins/Mutable/Source/MutableRuntime/Private/MuR/Model.cpp
+++ b/Engine/Plugins/Mutable/Source/MutableRuntime/Private/MuR/Model.cpp
@@ -421,7 +421,8 @@ namespace mu
 	TSharedPtr<FParameters> FModel::NewParameters(TSharedPtr<const FModel> Model,
 		const FParameters* OldParameters,
 		const TMap<FName, TObjectPtr<UTexture>>* ImageDefaults,
-		const TMap<FName, TObjectPtr<USkeletalMesh>>* MeshDefaults)
+		const TMap<FName, TObjectPtr<USkeletalMesh>>* MeshDefaults,
+		const TMap<FName, TObjectPtr<UMaterialInterface>>* MaterialDefaults)
 	{
 		LLM_SCOPE_BYNAME(TEXT("MutableRuntime"));
 
diff --git a/Engine/Plugins/Mutable/Source/MutableRuntime/Public/MuR/Model.h b/Engine/Plugins/Mutable/Source/MutableRuntime/Public/MuR/Model.h
--- a/Engine/Plugins/Mutable/Source/MutableRuntime/Public/MuR/Model.h
+++ b/Engine/Plugins/Mutable/Source/MutableRuntime/Public/MuR/Model.h
@@ -8,6 +8,9 @@
 #include "Templates/Function.h"
 #include "Templates/SharedPointer.h"
 
+class UMaterialInterface;
+
+
 #define UE_API MUTABLERUNTIME_API
 
 
@@ -19,7 +22,6 @@ namespace mu
 	class FSystem;
     class FModelParametersGenerator;
 
-
     /** A Model represents a customisable object with any number of parameters.
     * When values are given to the parameters, specific Instances can be built, which hold the
     * built application-usable data.
@@ -70,7 +72,8 @@ namespace mu
         static UE_API TSharedPtr<FParameters> NewParameters(TSharedPtr<const FModel> Model,
         	const FParameters* OldParameters = nullptr,
         	const TMap<FName, TObjectPtr<UTexture>>* ImageDefaults = nullptr,
-        	const TMap<FName, TObjectPtr<USkeletalMesh>>* MeshDefaults = nullptr);
+        	const TMap<FName, TObjectPtr<USkeletalMesh>>* MeshDefaults = nullptr,
+        	const TMap<FName, TObjectPtr<UMaterialInterface>>* MaterialDefaults = nullptr);
 
 		/** Return true if the parameter is multi-dimensional */
 		UE_API bool IsParameterMultidimensional(int32 ParameterIndex) const;
