[Mutable] Fix rotation to rigid clusters reshape. Compact data used for vertex binding and cluster bindig using a union.

#rnx

diff --git a/Engine/Plugins/Mutable/Source/CustomizableObject/Internal/MuCO/CustomizableObjectPrivate.h b/Engine/Plugins/Mutable/Source/CustomizableObject/Internal/MuCO/CustomizableObjectPrivate.h
--- a/Engine/Plugins/Mutable/Source/CustomizableObject/Internal/MuCO/CustomizableObjectPrivate.h
+++ b/Engine/Plugins/Mutable/Source/CustomizableObject/Internal/MuCO/CustomizableObjectPrivate.h
@@ -1638,6 +1638,8 @@ enum class ECustomizableObjectVersions : int32
 
 	FixDDCCrashesAndLoadTimes,
 
+	CompactBindingDataStructure,
+
 	// -----<new versions can be added above this line>--------
 	LastCustomizableObjectVersion
 };
diff --git a/Engine/Plugins/Mutable/Source/MutableRuntime/Private/MuR/OpMeshBind.h b/Engine/Plugins/Mutable/Source/MutableRuntime/Private/MuR/OpMeshBind.h
--- a/Engine/Plugins/Mutable/Source/MutableRuntime/Private/MuR/OpMeshBind.h
+++ b/Engine/Plugins/Mutable/Source/MutableRuntime/Private/MuR/OpMeshBind.h
@@ -108,40 +108,84 @@ namespace mu
 
 	using FShapeMeshTree = UE::Geometry::TMeshAABBTree3<const FShapeMeshAdapter>;
 
-	// Structure used for vertex bing data in vertexbuffers for reshape operations.
 	struct FReshapeVertexBindingData
 	{
+		// Index of the shape triangle and the type of binding, Cluster or Vertex. 
+		uint32 TriangleIndex_BindingType_NotValid; // the 30 low bits are the triangle index, 30 ith bit is the binding type, 31 ith bit is set if the biding is valid. 
+		
 		// Barycentric coordinates on the shape triangle
-		float S, T;
+		FVector2f PositionBaryCoords;
 
 		// Distance along the normals of the shape triangle.
 		FVector3f D;
 
-		// Index of the shape triangle.
-		int32 Triangle;
+		union
+		{
+			struct
+			{
+				FVector3f NormalBaryCoords; // NormalBaryCoords.Z is not used.
+				FVector3f NormalD;
+			} Vertex;
 
-		// Used to calculate the rotation to apply to the reshaped vertex tangent space.
-		float NS, NT;
-		FVector3f NormalD;
-		
-		// Bind point, if the point belongs to a rigid cluster, the attachment point, otherwise
-		// the original point.
-		FVector3f AttachmentPoint;
+			struct
+			{
+				FVector3f AttachmentPoint;
+				FVector3f AttachmentNormal;
+			} Cluster;
+		};	
 
 		// Weight of the effect for this binding. Ranged between 0 and 1 where 0 denotes no effect at all and 1
-		// full effect. This weight should be proportional to the confidence we have that the binding data is valid. 
+		// full effect.
 		float Weight;
+
+		FORCEINLINE uint32 GetTriangleIndex() const
+		{
+			return TriangleIndex_BindingType_NotValid & 0x3FFFFFFF;
+		}
+
+		FORCEINLINE bool IsClusterBinding() const
+		{
+			return (TriangleIndex_BindingType_NotValid & 0x40000000) != 0;
+		}
+
+		FORCEINLINE bool IsValidBinding() const
+		{
+			return (TriangleIndex_BindingType_NotValid & 0x80000000) == 0;
+		}
+
+		FORCEINLINE void SetTriangleIndex(uint32 Index)
+		{	
+			check(Index < 0x3FFFFFFF)
+			TriangleIndex_BindingType_NotValid = 
+					(TriangleIndex_BindingType_NotValid & ~0x3FFFFFFF) | (Index & 0x3FFFFFFF);
+		}
+
+		FORCEINLINE void SetClusterBinding(bool Value)
+		{
+			uint32 ValueBits = uint32(Value) << 30;
+			TriangleIndex_BindingType_NotValid = 
+					(TriangleIndex_BindingType_NotValid & ~0x40000000) | ValueBits;
+		}
+
+		FORCEINLINE void SetBindNotValid(bool Value)
+		{
+			uint32 ValueBits = uint32(Value) << 31;
+			TriangleIndex_BindingType_NotValid = 
+					(TriangleIndex_BindingType_NotValid & ~0x80000000) | ValueBits;
+		}
+
+
 	};
-	static_assert(sizeof(FReshapeVertexBindingData) == 4*15);
+	static_assert(sizeof(FReshapeVertexBindingData) == 4*13);
 
 	struct FReshapeVertexBindingDataBufferDescriptor
 	{
 		constexpr static int32 ElementSize = sizeof(FReshapeVertexBindingData);
 		constexpr static int32 Channels = 4;
-		constexpr static EMeshBufferSemantic Semantics[Channels] = { EMeshBufferSemantic::BarycentricCoords, EMeshBufferSemantic::Distance, EMeshBufferSemantic::TriangleIndex, EMeshBufferSemantic::Other };
-		constexpr static EMeshBufferFormat Formats[Channels] = { EMeshBufferFormat::Float32, EMeshBufferFormat::Float32, EMeshBufferFormat::Int32, EMeshBufferFormat::Float32 };
-		constexpr static int32 Components[Channels] = { 2, 3, 1, 2+3+3+1 };
-		constexpr static int32 Offsets[Channels] = { 0, 4*2, (2+3)*4, (2+3+1)*4};
+		constexpr static EMeshBufferSemantic Semantics[Channels] = { EMeshBufferSemantic::TriangleIndex, EMeshBufferSemantic::BarycentricCoords, EMeshBufferSemantic::Distance, EMeshBufferSemantic::Other };
+		constexpr static EMeshBufferFormat Formats[Channels] = { EMeshBufferFormat::UInt32, EMeshBufferFormat::Float32, EMeshBufferFormat::Float32, EMeshBufferFormat::Float32 };
+		constexpr static int32 Components[Channels] = { 1, 2, 3, 3+3+1 };
+		constexpr static int32 Offsets[Channels] = { 0, 1*4, (1+2)*4, (1+2+3)*4 };
 
 		const int32 SemanticIndices[Channels] = { 0, 0, 0, 0 };
 
@@ -285,15 +329,14 @@ namespace mu
 	{
 
 		const FShapeMeshDescriptorBind& ShapeMesh = ShapeMeshTree.GetMesh()->Mesh;
-		
-		OutBindData.S = 0.0f;
-		OutBindData.T = 0.0f;
-		OutBindData.Triangle = -1;
-		OutBindData.AttachmentPoint = Point;
+	
+		FMemory::Memzero(OutBindData);
+
 		float Weight = MaskWeight;
 
 		if (FMath::IsNearlyZero(MaskWeight))
 		{
+			OutBindData.SetBindNotValid(true);
 			return;
 		}
 
@@ -303,6 +346,7 @@ namespace mu
 			
 		if (FoundIndex < 0)
 		{
+			OutBindData.SetBindNotValid(true);
 			return;
 		}
 
@@ -329,7 +373,7 @@ namespace mu
 				FMath::RayPlaneIntersectionParam(TriangleC, ShapeMesh.Normals[Triangle.C], VertexPlane)
 			};
 
-		OutBindData.NormalD = 
+		OutBindData.Vertex.NormalD = 
 			FVector3f
 			{
 				FMath::RayPlaneIntersectionParam(TriangleA, ShapeMesh.Normals[Triangle.A], NormalPlane),
@@ -343,21 +387,22 @@ namespace mu
 					TriangleB + ShapeMesh.Normals[Triangle.B] * OutBindData.D.Y, 
 					TriangleC + ShapeMesh.Normals[Triangle.C] * OutBindData.D.Z); 
 
-		OutBindData.S = PositionBarycentric.X;
-		OutBindData.T = PositionBarycentric.Y;
+		OutBindData.PositionBaryCoords.X = PositionBarycentric.X;
+		OutBindData.PositionBaryCoords.Y = PositionBarycentric.Y;
 	
 		FVector2f NormalBarycentric = ComputeBarycentric((Point + Normal), 
-				TriangleA + ShapeMesh.Normals[Triangle.A] * OutBindData.NormalD.X, 
-				TriangleB + ShapeMesh.Normals[Triangle.B] * OutBindData.NormalD.Y, 
-				TriangleC + ShapeMesh.Normals[Triangle.C] * OutBindData.NormalD.Z); 
+				TriangleA + ShapeMesh.Normals[Triangle.A] * OutBindData.Vertex.NormalD.X, 
+				TriangleB + ShapeMesh.Normals[Triangle.B] * OutBindData.Vertex.NormalD.Y, 
+				TriangleC + ShapeMesh.Normals[Triangle.C] * OutBindData.Vertex.NormalD.Z); 
 
-		OutBindData.NS = NormalBarycentric.X;
-		OutBindData.NT = NormalBarycentric.Y;
-
-		OutBindData.Triangle = FoundIndex;
+		OutBindData.Vertex.NormalBaryCoords.X = NormalBarycentric.X;
+		OutBindData.Vertex.NormalBaryCoords.Y = NormalBarycentric.Y;
 
+		check(FoundIndex >= 0);
+		OutBindData.SetTriangleIndex((uint32)FoundIndex);
 		OutBindData.Weight = FMath::Clamp(MaskWeight, 0.0f, 1.0f);
-		OutBindData.Triangle = FMath::IsNearlyZero(OutBindData.Weight) ? -1 : OutBindData.Triangle;
+
+		OutBindData.SetBindNotValid(FMath::IsNearlyZero(OutBindData.Weight));
 	}
 	
     //---------------------------------------------------------------------------------------------
@@ -440,6 +485,24 @@ namespace mu
 		// This will be filled in afterwards.
 		constexpr float MaskWeight = 1.0f;
 		BindReshapePoint(ShapeMeshTree, BoundPoint, FVector3f::ZAxisVector, MaskWeight, OutBindingData, BindTolerance);
+
+		if (OutBindingData.IsValidBinding())
+		{
+			const FShapeMeshDescriptorBind& ShapeMesh = ShapeMeshTree.GetMesh()->Mesh;
+			UE::Geometry::FIndex3i Triangle = ShapeMesh.Triangles[OutBindingData.GetTriangleIndex()];
+
+			FVector3f TriangleA = (FVector3f)ShapeMesh.Positions[Triangle.A];
+			FVector3f TriangleB = (FVector3f)ShapeMesh.Positions[Triangle.B];
+			FVector3f TriangleC = (FVector3f)ShapeMesh.Positions[Triangle.C];
+
+			//TODO: store triangle area for scale.
+			//OutBindingData.Cluster.AttachementArea = FVector3f::CrossProduct((TriangleB - TriangleA), (TriangleC - TriangleA)).Size() / 2.0f;
+			OutBindingData.Cluster.AttachmentPoint = BoundPoint;
+			OutBindingData.Cluster.AttachmentNormal = FPlane4f(TriangleA, TriangleB, TriangleC).GetNormal(); 
+
+			// Mark binding as cluster.
+			OutBindingData.SetClusterBinding(true);
+		}
 	}
 
 	inline TTuple<TArray<FReshapePointBindingData>, TArray<int32>, TArray<int32>> BindPhysicsBodies( 
@@ -697,8 +760,19 @@ namespace mu
 		{
 			constexpr float MaskWeight = 1.0f;
 			BindReshapePoint(ShapeMeshTree, Points[PointIndex], FVector3f::ZAxisVector, MaskWeight, VertexBindData, 0.1f);
+
+			int32 TriangleIndex = VertexBindData.IsValidBinding() 
+					? (int32)VertexBindData.GetTriangleIndex() 
+					: -1;
+
 			PhysicsBodyBindData[PointIndex] = FReshapePointBindingData
-				{ VertexBindData.S, VertexBindData.T, VertexBindData.D, VertexBindData.Triangle, VertexBindData.Weight };
+			{ 
+				.S = VertexBindData.PositionBaryCoords.X, 
+				.T = VertexBindData.PositionBaryCoords.Y, 
+				.D = VertexBindData.D, 
+				.Triangle = TriangleIndex,
+				.Weight = VertexBindData.Weight,
+			};
 		}
 
 		return ReturnValue;
@@ -799,6 +873,7 @@ namespace mu
 					const int32 ClusterVertexCount = RigidCluster.Num();
 					if (ClusterVertexCount)
 					{
+
 						FindBindingForCluster(BaseMesh, ShapeMeshTree, RigidCluster, ClusterBinding, BindValidityTolerance);
 						// Copy cluster binding to every vertex of the cluster modifing weight if the vertex color is used.
 						for (int32 I = 0; I < ClusterVertexCount; ++I)
@@ -807,11 +882,11 @@ namespace mu
 
 							const float MaskWeight = ItColorBase.ptr() ? GetVertexMaskWeight(ItColorBase + VertexIndex, ColorUsageMasks) : 1.0f;
 
-							ClusterBinding.Weight = FMath::Min(ClusterBinding.Weight, MaskWeight);
-							// Invalidate binding for very small weights.
-							ClusterBinding.Triangle = FMath::IsNearlyZero(ClusterBinding.Weight) ? -1 : ClusterBinding.Triangle;
-
 							BindData[VertexIndex] = ClusterBinding;
+							BindData[VertexIndex].Weight = FMath::Min(ClusterBinding.Weight, MaskWeight);
+
+							// Invalidate binding for very small weights.
+							BindData[VertexIndex].SetBindNotValid(FMath::IsNearlyZero(ClusterBinding.Weight));	
 						}
 					}
 				}
@@ -974,11 +1049,17 @@ namespace mu
 			constexpr float MaskWeight = 1.0f;
 			BindReshapePoint(ShapeMeshTree, Mesh->BonePoses[BoneIndex].BoneTransform.GetLocation(), FVector3f::ZAxisVector, MaskWeight, BindData, BindValidityTolerance);
 
-			// Only add binding  if there is a chance of the bone moving.
-			if (BindData.Weight > UE_SMALL_NUMBER && BindData.Triangle >= 0)
+			// Only add binding if there is a chance of the bone moving.
+			if (BindData.Weight > UE_SMALL_NUMBER && BindData.IsValidBinding())
 			{
-				SkeletonBindDataArray.Emplace(
-					FReshapePointBindingData{ BindData.S, BindData.T, BindData.D, BindData.Triangle, BindData.Weight });
+				SkeletonBindDataArray.Emplace(FReshapePointBindingData
+				{ 
+					.S = BindData.PositionBaryCoords.X, 
+					.T = BindData.PositionBaryCoords.Y, 
+					.D = BindData.D, 
+					.Triangle = (int32)BindData.GetTriangleIndex(), 
+					.Weight = BindData.Weight, 
+				});
 				BoneIndices.Add(BoneIndex);
 			}	
 		}
diff --git a/Engine/Plugins/Mutable/Source/MutableRuntime/Private/MuR/OpMeshReshape.h b/Engine/Plugins/Mutable/Source/MutableRuntime/Private/MuR/OpMeshReshape.h
--- a/Engine/Plugins/Mutable/Source/MutableRuntime/Private/MuR/OpMeshReshape.h
+++ b/Engine/Plugins/Mutable/Source/MutableRuntime/Private/MuR/OpMeshReshape.h
@@ -15,6 +15,7 @@
 #include "Spatial/MeshAABBTree3.h"
 #include "Distance/DistPoint3Triangle3.h"
 #include "LineTypes.h"
+
 #include "MuR/OpMeshSmoothing.h"
 #include "MuR/OpMeshComputeNormals.h"
 
@@ -49,19 +50,25 @@ namespace mu
 	// Methods to actually deform a point
 	inline void GetDeform( const FShapeMeshDescriptorApply& Shape, const FReshapeVertexBindingData& Binding, FVector3f& NewPosition, FVector3f& NewNormalPosition)
 	{
-		const UE::Geometry::FIndex3i& Triangle = Shape.Triangles[Binding.Triangle];
+		const UE::Geometry::FIndex3i& Triangle = Shape.Triangles[Binding.GetTriangleIndex()];
 	
 		FVector3f VA = Shape.Positions[Triangle.A] + Shape.Normals[Triangle.A]*Binding.D.X;
 		FVector3f VB = Shape.Positions[Triangle.B] + Shape.Normals[Triangle.B]*Binding.D.Y;
 		FVector3f VC = Shape.Positions[Triangle.C] + Shape.Normals[Triangle.C]*Binding.D.Z;
 
-		NewPosition = VA*Binding.S + VB*Binding.T + VC*(1.0f - Binding.S - Binding.T);
+		NewPosition = 
+				VA*Binding.PositionBaryCoords.X + 
+				VB*Binding.PositionBaryCoords.Y + 
+				VC*(1.0f - Binding.PositionBaryCoords.X - Binding.PositionBaryCoords.Y);
 	
-		FVector3f NVA = Shape.Positions[Triangle.A] + Shape.Normals[Triangle.A]*Binding.NormalD.X;
-		FVector3f NVB = Shape.Positions[Triangle.B] + Shape.Normals[Triangle.B]*Binding.NormalD.Y;
-		FVector3f NVC = Shape.Positions[Triangle.C] + Shape.Normals[Triangle.C]*Binding.NormalD.Z;
-
-		NewNormalPosition = (NVA*Binding.NS + NVB*Binding.NT + NVC*(1.0f - Binding.NS - Binding.NT));
+		FVector3f NVA = Shape.Positions[Triangle.A] + Shape.Normals[Triangle.A]*Binding.Vertex.NormalD.X;
+		FVector3f NVB = Shape.Positions[Triangle.B] + Shape.Normals[Triangle.B]*Binding.Vertex.NormalD.Y;
+		FVector3f NVC = Shape.Positions[Triangle.C] + Shape.Normals[Triangle.C]*Binding.Vertex.NormalD.Z;
+
+		NewNormalPosition = 
+				NVA*Binding.Vertex.NormalBaryCoords.X + 
+				NVB*Binding.Vertex.NormalBaryCoords.Y + 
+				NVC*(1.0f - Binding.Vertex.NormalBaryCoords.X - Binding.Vertex.NormalBaryCoords.Y);
 	}
 
 	inline void GetDeform( const FShapeMeshDescriptorApply& Shape, const FReshapePointBindingData& Binding, FVector3f& NewPosition)
@@ -75,7 +82,6 @@ namespace mu
 		NewPosition = VA*Binding.S + VB*Binding.T + VC*(1.0f - Binding.S - Binding.T);
 	}
 
-
     //---------------------------------------------------------------------------------------------
     //! Physics Bodies Reshape 
     //---------------------------------------------------------------------------------------------
@@ -394,7 +400,10 @@ namespace mu
 #if DO_CHECK
 		// checking if the Base shape has more triangles than the target shape
 		const bool bTriangleOutOfScopeFound = Algo::AnyOf(BindingData, 
-			[NumShapeTriangles = Shape.Triangles.Num()](const FReshapeVertexBindingData& B) { return B.Triangle >= NumShapeTriangles; });
+		[NumShapeTriangles = (uint32)Shape.Triangles.Num()](const FReshapeVertexBindingData& B) 
+		{ 
+			return (B.GetTriangleIndex() >= NumShapeTriangles) && (B.IsValidBinding()); 
+		});
 
 		if (bTriangleOutOfScopeFound)
 		{
@@ -412,7 +421,7 @@ namespace mu
 		{
 			const FReshapeVertexBindingData& Binding = BindingData[VertexIndex];
 
-			const bool bModified = (Binding.Triangle >= 0) & (Binding.Triangle < ShapeTriangleCount);
+			const bool bModified = Binding.IsValidBinding() & (Binding.GetTriangleIndex() < (uint32)ShapeTriangleCount);
 			if (bModified)
 			{
 				FVector3f NewPosition;
@@ -420,25 +429,50 @@ namespace mu
 				
 				GetDeform(Shape, Binding, NewPosition, NewNormalPosition);
 
-				const FVector3f& OldPosition = Positions[VertexIndex];
-  
-				if ((OldPosition - Binding.AttachmentPoint).IsNearlyZero())
-				{
-					Positions[VertexIndex] += (NewPosition - OldPosition) * Binding.Weight;
-				}
-				else
-				{
-					Positions[VertexIndex] += (NewPosition - Binding.AttachmentPoint) * Binding.Weight;
+				FVector3f OldPosition = (PositionIter + VertexIndex).GetAsVec3f();
+ 
+                FVector3f Position;
+				FQuat4f ClusterQuat; // Only used for cluster binding.
+                if (!Binding.IsClusterBinding())
+                {
+					Position = OldPosition + (NewPosition - OldPosition) * Binding.Weight;
 				}
+                else
+                {
+					const UE::Geometry::FIndex3i& Triangle = Shape.Triangles[Binding.GetTriangleIndex()];
+	
+					const FVector3f& VA = Shape.Positions[Triangle.A];
+					const FVector3f& VB = Shape.Positions[Triangle.B];
+					const FVector3f& VC = Shape.Positions[Triangle.C];
+
+					FVector3f TriangleNormal = FPlane4f(VA, VB, VC).GetNormal();
+
+					ClusterQuat = FQuat4f::FindBetweenNormals(Binding.Cluster.AttachmentNormal, TriangleNormal);
+					ClusterQuat = FQuat4f::FastLerp(FQuat4f::Identity, ClusterQuat, Binding.Weight).GetNormalized();
+
+					Position = (ClusterQuat.RotateVector(OldPosition - Binding.Cluster.AttachmentPoint)) + Binding.Cluster.AttachmentPoint;
+					Position = Position + (NewPosition - Binding.Cluster.AttachmentPoint) * Binding.Weight;
+                }
+				
+				(PositionIter + VertexIndex).SetFromVec3f(Position);
 
 				if (bComputeNormal)
 				{
+					FVector3f Normal;
+
 					const FVector3f OldNormal = (NormalIter + VertexIndex).GetAsVec3f();
-					const FVector3f OldNormalPosition = OldPosition + OldNormal;
-					
-					FVector3f NormalDisplacement = (NewNormalPosition - OldNormalPosition) * Binding.Weight;
 
-					FVector3f Normal = ((OldNormalPosition + NormalDisplacement) - NewPosition).GetSafeNormal();
+					if (!Binding.IsClusterBinding())
+					{
+						const FVector3f OldNormalPosition = OldPosition + OldNormal;
+						FVector3f NormalDisplacement = (NewNormalPosition - OldNormalPosition) * Binding.Weight;
+						Normal = ((OldNormalPosition + NormalDisplacement) - NewPosition).GetSafeNormal();
+					}
+					else
+					{
+						Normal = ClusterQuat.RotateVector(OldNormal);	
+					}
+
 					FVector3f Tangent = TangentIter.ptr() ? (TangentIter + VertexIndex).GetAsVec3f() : FVector3f::XAxisVector;
 					FVector3f BiNormal = BiNormalIter.ptr() ? (BiNormalIter + VertexIndex).GetAsVec3f() : FVector3f::YAxisVector;
 				
@@ -873,7 +907,7 @@ namespace mu
 			{
 				MUTABLE_CPUPROFILER_SCOPE(ReshapeVertices);
 				// \TODO: More checks
-				check(BarycentricDataChannel == 0);
+				check(BarycentricDataChannel == 1);
 				check(VB.GetElementSize(BarycentricDataBuffer) == (int)sizeof(FReshapeVertexBindingData));
 
 				TArrayView<const FReshapeVertexBindingData> VerticesBindingData(
