[Mutable] Add LODs to Skeletal Mesh parameters.

diff --git a/Engine/Plugins/Mutable/Source/CustomizableObjectEditor/Private/MuCOE/GenerateMutableSource/GenerateMutableSourceMesh.cpp b/Engine/Plugins/Mutable/Source/CustomizableObjectEditor/Private/MuCOE/GenerateMutableSource/GenerateMutableSourceMesh.cpp
--- a/Engine/Plugins/Mutable/Source/CustomizableObjectEditor/Private/MuCOE/GenerateMutableSource/GenerateMutableSourceMesh.cpp
+++ b/Engine/Plugins/Mutable/Source/CustomizableObjectEditor/Private/MuCOE/GenerateMutableSource/GenerateMutableSourceMesh.cpp
@@ -4611,11 +4611,17 @@ mu::Ptr<mu::NodeMesh> GenerateMutableSourceMesh(const UEdGraphPin* Pin,
 			GenerationContext.SkeletalMeshParameterDefaultValues.Add(FName(TypedNodeParam->GetParameterName()), MutablePrivate::LoadObject(TypedNodeParam->DefaultValue));
 		}
 
-		const int32 LODIndex = 0;
+		int32 LODIndex = 0;
 		int32 SectionIndex = 0;
-		TypedNodeParam->GetPinSection(*Pin, SectionIndex);
-		ParameterNode->SectionIndex = SectionIndex;
-
+		
+		{
+			int32 LayoutIndex = 0;
+			TypedNodeParam->GetPinSection(*Pin, LODIndex, SectionIndex, LayoutIndex);
+			ParameterNode->LODIndex = LODIndex;
+			ParameterNode->SectionIndex = SectionIndex;
+			ParameterNode->bAutomaticLODs = GenerationContext.CurrentAutoLODStrategy == ECustomizableObjectAutomaticLODStrategy::AutomaticFromMesh;
+		}
+		
 		GenerationContext.ParameterUIDataMap.Add(TypedNodeParam->GetParameterName(), FMutableParameterData(
 			TypedNodeParam->ParamUIMetadata,
 			EMutableParameterType::Texture));
diff --git a/Engine/Plugins/Mutable/Source/CustomizableObjectEditor/Private/MuCOE/Nodes/CustomizableObjectNodeSkeletalMeshParameter.cpp b/Engine/Plugins/Mutable/Source/CustomizableObjectEditor/Private/MuCOE/Nodes/CustomizableObjectNodeSkeletalMeshParameter.cpp
--- a/Engine/Plugins/Mutable/Source/CustomizableObjectEditor/Private/MuCOE/Nodes/CustomizableObjectNodeSkeletalMeshParameter.cpp
+++ b/Engine/Plugins/Mutable/Source/CustomizableObjectEditor/Private/MuCOE/Nodes/CustomizableObjectNodeSkeletalMeshParameter.cpp
@@ -28,7 +28,7 @@ void UCustomizableObjectNodeSkeletalMeshParameter::AllocateDefaultPins(UCustomiz
 
 	if (const FSkeletalMeshModel* ImportedModel = SkeletalMesh->GetImportedModel())
 	{
-		int32 LODIndex = 0;
+		for (int32 LODIndex = 0; LODIndex < ImportedModel->LODModels.Num(); LODIndex++)
 		{
 			const int32 NumSections = ImportedModel->LODModels[LODIndex].Sections.Num();
 			for (int32 SectionIndex = 0; SectionIndex < NumSections; ++SectionIndex)
@@ -41,23 +41,20 @@ void UCustomizableObjectNodeSkeletalMeshParameter::AllocateDefaultPins(UCustomiz
 
 				UMaterialInterface* MaterialInterface = GetMaterialInterfaceFor( SectionIndex);
 
-				FString Section = FString::Printf(TEXT("Section %i"), SectionIndex);
-				if (MaterialInterface)
-				{
-					Section.Append(FString::Printf(TEXT(" : %s"), *MaterialInterface->GetName()));
-				}
+				const FName MaterialName = MaterialInterface ? MaterialInterface->GetFName() : NAME_None;
 
 				// Mesh
 				{
 					UCustomizableObjectNodeSkeletalMeshParameterPinDataSection* PinData = NewObject<UCustomizableObjectNodeSkeletalMeshParameterPinDataSection>(this);
-					PinData->Init( SectionIndex, ImportedModel->LODModels[LODIndex].NumTexCoords );
-
-					const FString MeshPinName = FString::Printf(TEXT("Section %i - Mesh"), SectionIndex);
+					PinData->Init(LODIndex, SectionIndex, ImportedModel->LODModels[LODIndex].NumTexCoords);
 
+					const FString MeshPinName = FString::Printf(TEXT("LOD %i - Section %i: %s"), LODIndex, SectionIndex, *MaterialName.ToString());
+					const FText MeshPinFriendlyName = FText::Format(LOCTEXT("SkeletalMeshParameterMeshPin", "LOD {0} - Section {1}: {2}"), LODIndex, SectionIndex, FText::FromString(*MaterialName.ToString()));
+					
 					UEdGraphPin* Pin = CustomCreatePin(EGPD_Output, Schema->PC_Mesh, FName(*MeshPinName), PinData);
-					Pin->PinFriendlyName = FText::FromString(FString::Printf(TEXT("%s"), *Section));
+					Pin->PinFriendlyName = MeshPinFriendlyName;
 				}
-			}
+			}			
 		}
 	}
 }
@@ -173,7 +170,7 @@ void UCustomizableObjectNodeSkeletalMeshParameter::GetPinSection(const UEdGraphP
 
 	if (const UCustomizableObjectNodeSkeletalMeshParameterPinDataSection* PinData = Cast<UCustomizableObjectNodeSkeletalMeshParameterPinDataSection>(GetPinData(Pin)))
 	{
-		OutLODIndex = 0;
+		OutLODIndex = PinData->GetLODIndex();
 		OutSectionIndex = PinData->GetSectionIndex();
 	}
 }
@@ -266,8 +263,9 @@ int32 UCustomizableObjectNodeSkeletalMeshParameter::GetSkeletalMaterialIndexFor(
 }
 
 
-void UCustomizableObjectNodeSkeletalMeshParameterPinDataSection::Init(int32 InSectionIndex, int32 NumTexCoords)
+void UCustomizableObjectNodeSkeletalMeshParameterPinDataSection::Init(int32 InLODIndex, int32 InSectionIndex, int32 NumTexCoords)
 {
+	LODIndex = InLODIndex;
 	SectionIndex = InSectionIndex;
 
 	if (NumTexCoords > 0)
@@ -276,8 +274,6 @@ void UCustomizableObjectNodeSkeletalMeshParameterPinDataSection::Init(int32 InSe
 
 		Layouts.SetNum(NumTexCoords);
 
-		const int LODIndex = 0;
-
 		for (int32 Index = 0; Index < NumTexCoords; ++Index)
 		{
 			Layouts[Index] = NewObject<UCustomizableObjectLayout>(Outer);
@@ -287,6 +283,12 @@ void UCustomizableObjectNodeSkeletalMeshParameterPinDataSection::Init(int32 InSe
 }
 
 
+int32 UCustomizableObjectNodeSkeletalMeshParameterPinDataSection::GetLODIndex() const
+{
+	return LODIndex;
+}
+
+
 int32 UCustomizableObjectNodeSkeletalMeshParameterPinDataSection::GetSectionIndex() const
 {
 	return SectionIndex;
diff --git a/Engine/Plugins/Mutable/Source/CustomizableObjectEditor/Private/MuCOE/Nodes/CustomizableObjectNodeSkeletalMeshParameter.h b/Engine/Plugins/Mutable/Source/CustomizableObjectEditor/Private/MuCOE/Nodes/CustomizableObjectNodeSkeletalMeshParameter.h
--- a/Engine/Plugins/Mutable/Source/CustomizableObjectEditor/Private/MuCOE/Nodes/CustomizableObjectNodeSkeletalMeshParameter.h
+++ b/Engine/Plugins/Mutable/Source/CustomizableObjectEditor/Private/MuCOE/Nodes/CustomizableObjectNodeSkeletalMeshParameter.h
@@ -23,8 +23,10 @@ class UCustomizableObjectNodeSkeletalMeshParameterPinDataSection : public UCusto
 	GENERATED_BODY()
 
 public:
-	void Init(int32 InSectionIndex, int32 NumTexCoords);
+	void Init(int32 LODIndex, int32 InSectionIndex, int32 NumTexCoords);
 
+	int32 GetLODIndex() const;
+	
 	int32 GetSectionIndex() const;
 
 	/** Layouts related to this Mesh pin */
@@ -36,10 +38,11 @@ protected:
 	virtual bool Equals(const UCustomizableObjectNodePinData& Other) const override;
 
 private:
-
+	UPROPERTY()
+	int32 LODIndex = -1;
+	
 	UPROPERTY()
 	int32 SectionIndex = -1;
-
 };
 
 
diff --git a/Engine/Plugins/Mutable/Source/MutableTools/Private/MuT/CodeGenerator_Mesh.cpp b/Engine/Plugins/Mutable/Source/MutableTools/Private/MuT/CodeGenerator_Mesh.cpp
--- a/Engine/Plugins/Mutable/Source/MutableTools/Private/MuT/CodeGenerator_Mesh.cpp
+++ b/Engine/Plugins/Mutable/Source/MutableTools/Private/MuT/CodeGenerator_Mesh.cpp
@@ -1342,7 +1342,7 @@ namespace mu
 				Op = new ASTOpParameter();
 				Op->Type = EOpType::ME_PARAMETER;
 				Op->Parameter = param;
-				Op->LODIndex = StaticOptions.LODIndex;
+				Op->LODIndex = Node->LODIndex + (Node->bAutomaticLODs ? StaticOptions.LODIndex : 0);
 				Op->SectionIndex = Node->SectionIndex;
 
 				// Assign an ID prefix to the mesh. 
diff --git a/Engine/Plugins/Mutable/Source/MutableTools/Public/MuT/NodeMeshParameter.h b/Engine/Plugins/Mutable/Source/MutableTools/Public/MuT/NodeMeshParameter.h
--- a/Engine/Plugins/Mutable/Source/MutableTools/Public/MuT/NodeMeshParameter.h
+++ b/Engine/Plugins/Mutable/Source/MutableTools/Public/MuT/NodeMeshParameter.h
@@ -18,12 +18,14 @@ namespace mu
 	class NodeMeshParameter : public NodeMesh
 	{
 	public:
-		/** */
 		TArray<Ptr<NodeLayout>> Layouts;
 
-		/** */
+		uint8 LODIndex = 0;
+
 		uint8 SectionIndex = 0;
 
+		bool bAutomaticLODs = false;
+		
 		/** Name of the parameter */
 		FString Name;
 
