[Mutable] Fixed compilation of Modifier nodes linked to a Macro Instance node.

#rnx

diff --git a/Engine/Plugins/Mutable/Source/CustomizableObjectEditor/Private/MuCOE/GenerateMutableSource/GenerateMutableSource.cpp b/Engine/Plugins/Mutable/Source/CustomizableObjectEditor/Private/MuCOE/GenerateMutableSource/GenerateMutableSource.cpp
--- a/Engine/Plugins/Mutable/Source/CustomizableObjectEditor/Private/MuCOE/GenerateMutableSource/GenerateMutableSource.cpp
+++ b/Engine/Plugins/Mutable/Source/CustomizableObjectEditor/Private/MuCOE/GenerateMutableSource/GenerateMutableSource.cpp
@@ -953,44 +953,17 @@ mu::Ptr<mu::NodeObject> GenerateMutableSource(const UEdGraphPin * Pin, FMutableG
 			TArray<UEdGraphPin*> ConnectedModifierPins = FollowInputPinArray(*ModifierPin);
 			for (UEdGraphPin* const ConnectedModifier : ConnectedModifierPins)
 			{
-				bool bMacroModifier = false;
-				UEdGraphPin* ChildNodePin = ConnectedModifier;
-
-				if (const UCustomizableObjectNodeMacroInstance* TypedNodeMacro = Cast<UCustomizableObjectNodeMacroInstance>(ChildNodePin->GetOwningNode()))
-				{
-					if (const UEdGraphPin* OutputPin = TypedNodeMacro->GetMacroIOPin(ECOMacroIOType::COMVT_Output, ChildNodePin->PinName))
-					{
-						// Get the modifier pin inside the macro
-						if (UEdGraphPin* FollowPin = FollowInputPin(*OutputPin))
-						{
-							bMacroModifier = true;
-							GenerationContext.MacroNodesStack.Push(TypedNodeMacro);
-							ChildNodePin = FollowPin;
-						}
-					}
-				}
-				
-				else if (!Cast<UCustomizableObjectNodeModifierBase>(ChildNodePin->GetOwningNode()))
-				{
-					continue;
-				}
-				
 				for (const FMutableComponentInfo& Pair : GenerationContext.CompilationContext.ComponentInfos)
 				{
 					GenerationContext.CurrentMeshComponent = Pair.ComponentName;
 					GenerationContext.CurrentAutoLODStrategy = Pair.Node->AutoLODStrategy;
 					
-					mu::Ptr<mu::NodeModifier> ModifierNode = GenerateMutableSourceModifier(ChildNodePin, GenerationContext);
+					mu::Ptr<mu::NodeModifier> ModifierNode = GenerateMutableSourceModifier(ConnectedModifier, GenerationContext);
 					ObjectNode->Modifiers.Add(ModifierNode);
 
 					GenerationContext.CurrentMeshComponent = {};
 					GenerationContext.CurrentAutoLODStrategy = {};
 				}
-
-				if (bMacroModifier)
-				{
-					GenerationContext.MacroNodesStack.Pop();
-				}
 			}
 		}
 		
