[Mutable] Material Parameter Instance Details.

diff --git a/Engine/Plugins/Mutable/Source/CustomizableObjectEditor/Private/MuCOE/CustomizableInstanceDetails.cpp b/Engine/Plugins/Mutable/Source/CustomizableObjectEditor/Private/MuCOE/CustomizableInstanceDetails.cpp
--- a/Engine/Plugins/Mutable/Source/CustomizableObjectEditor/Private/MuCOE/CustomizableInstanceDetails.cpp
+++ b/Engine/Plugins/Mutable/Source/CustomizableObjectEditor/Private/MuCOE/CustomizableInstanceDetails.cpp
@@ -1074,6 +1074,16 @@ TSharedRef<SWidget> FCustomizableInstanceDetails::GenerateParameterWidget(const
 				.OnObjectChanged(this, &FCustomizableInstanceDetails::OnSkeletalMeshParameterSelected, ParamName)
 				.ObjectPath(ObjectPath);
 		}
+		case EMutableParameterType::Material:
+		{
+			UMaterialInterface* Material = CustomInstance->GetMaterialParameterSelectedOption(ParamName);
+			const FString ObjectPath = Material ? Material->GetPathName() : TEXT("");
+
+			return SNew(SObjectPropertyEntryBox)
+				.AllowedClass(UMaterialInterface::StaticClass())
+				.OnObjectChanged(this, &FCustomizableInstanceDetails::OnSkeletalMeshParameterSelected, ParamName)
+				.ObjectPath(ObjectPath);
+		}
 		case EMutableParameterType::Transform:
 		{
 			return GenerateTransformWidget(ParamName);
@@ -1420,6 +1430,16 @@ void FCustomizableInstanceDetails::OnSkeletalMeshParameterSelected(const FAssetD
 	UpdateInstance();
 }
 
+
+// MATERIAL PARAMETERS -----------------------------------------------------------------------------------------------------------------
+
+void FCustomizableInstanceDetails::OnMaterialParameterSelected(const FAssetData& AssetData, const FString ParamName)
+{
+	CustomInstance->SetMaterialParameterSelectedOption(ParamName, Cast<UMaterialInterface>(AssetData.GetAsset()));
+	UpdateInstance();
+}
+
+
 // COLOR PARAMETERS -----------------------------------------------------------------------------------------------------------------
 
 TSharedRef<SWidget> FCustomizableInstanceDetails::GenerateColorWidget(const FString& ParamName)
diff --git a/Engine/Plugins/Mutable/Source/CustomizableObjectEditor/Private/MuCOE/CustomizableInstanceDetails.h b/Engine/Plugins/Mutable/Source/CustomizableObjectEditor/Private/MuCOE/CustomizableInstanceDetails.h
--- a/Engine/Plugins/Mutable/Source/CustomizableObjectEditor/Private/MuCOE/CustomizableInstanceDetails.h
+++ b/Engine/Plugins/Mutable/Source/CustomizableObjectEditor/Private/MuCOE/CustomizableInstanceDetails.h
@@ -122,6 +122,9 @@ private:
 	// Skeletal Mesh FParameters Functions
 	void OnSkeletalMeshParameterSelected(const FAssetData& AssetData, const FString ParamName);
 
+	// Material FParameters Functions
+	void OnMaterialParameterSelected(const FAssetData& AssetData, const FString ParamName);
+	
 	// Color FParameters Functions
 	TSharedRef<SWidget> GenerateColorWidget(const FString& ParamName);
 	FLinearColor GetColorParameterValue(const FString ParamName) const;
