[mutable] Fixed "Passthrough Skeletal Mesh Component" node compilation and runtime use (COI update).
- Fixed check for the state of the SkeletalMesh property of the UE node when generating the inputs.
- Fixed wrong property name being used when checking for changes in the SkeletalMesh property of the UE node.
- During the update, an empty no-passthrough mesh was being returned as an indirect consequence of the constant passthrough mesh generated during the compilation being lost due to an issue with the generation of the pins of the actual UE node.

#rnx

diff --git a/Engine/Plugins/Mutable/Source/CustomizableObjectEditor/Private/MuCOE/GenerateMutableSource/GenerateMutableSourceComponent.cpp b/Engine/Plugins/Mutable/Source/CustomizableObjectEditor/Private/MuCOE/GenerateMutableSource/GenerateMutableSourceComponent.cpp
--- a/Engine/Plugins/Mutable/Source/CustomizableObjectEditor/Private/MuCOE/GenerateMutableSource/GenerateMutableSourceComponent.cpp
+++ b/Engine/Plugins/Mutable/Source/CustomizableObjectEditor/Private/MuCOE/GenerateMutableSource/GenerateMutableSourceComponent.cpp
@@ -17,7 +17,6 @@
 #include "MuCOE/Nodes/CustomizableObjectNodeComponentVariation.h"
 #include "MuCOE/Nodes/CustomizableObjectNodeModifierBase.h"
 #include "MuCOE/Nodes/CustomizableObjectNodeMacroInstance.h"
-#include "MuCOE/Nodes/CustomizableObjectNodeStaticString.h"
 #include "MuCOE/Nodes/CustomizableObjectNodeTunnel.h"
 #include "MuT/NodeComponent.h"
 #include "MuT/NodeComponentEdit.h"
@@ -302,9 +301,21 @@ mu::Ptr<mu::NodeComponent> GenerateMutableSourceComponent(const UEdGraphPin* Pin
 						// Add an empty surface node anyway.
 						mu::Ptr<mu::NodeSurfaceNew> SurfaceNode = new mu::NodeSurfaceNew;
 						SurfaceNode->Mesh = MeshNode;
+						SurfaceNode->SurfaceGuid = GenerationContext.GetNodeIdUnique(TypedComponentPassthroughMesh);
 						LODNode->Surfaces.Add(SurfaceNode);
 					}
 				}
+				else
+				{
+					// The SKM has changed but the node has not been refreshed
+					// This may happen if the node was saved without pins but with a SKM set (old bug)
+					
+					// Add an empty surface node anyway.
+					mu::Ptr<mu::NodeSurfaceNew> SurfaceNode = new mu::NodeSurfaceNew;
+					SurfaceNode->Mesh = MeshNode;
+					SurfaceNode->SurfaceGuid = GenerationContext.GetNodeIdUnique(TypedComponentPassthroughMesh);
+					LODNode->Surfaces.Add(SurfaceNode);
+				}
 			}
 		}
 
diff --git a/Engine/Plugins/Mutable/Source/CustomizableObjectEditor/Private/MuCOE/Nodes/CustomizableObjectNodeComponentPassthroughMesh.cpp b/Engine/Plugins/Mutable/Source/CustomizableObjectEditor/Private/MuCOE/Nodes/CustomizableObjectNodeComponentPassthroughMesh.cpp
--- a/Engine/Plugins/Mutable/Source/CustomizableObjectEditor/Private/MuCOE/Nodes/CustomizableObjectNodeComponentPassthroughMesh.cpp
+++ b/Engine/Plugins/Mutable/Source/CustomizableObjectEditor/Private/MuCOE/Nodes/CustomizableObjectNodeComponentPassthroughMesh.cpp
@@ -6,14 +6,8 @@
 #include "Modules/ModuleManager.h"
 #include "MuCO/CustomizableObjectCustomVersion.h"
 #include "MuCO/LoadUtils.h"
-#include "MuCOE/CustomizableObjectEditor.h"
-#include "MuCOE/CustomizableObjectEditorStyle.h"
-#include "MuCOE/CustomizableObjectLayout.h"
 #include "MuCOE/EdGraphSchema_CustomizableObject.h"
-#include "MuCOE/GraphTraversal.h"
-#include "MuCOE/MutableUtils.h"
 #include "MuCOE/Nodes/CustomizableObjectNodeLayoutBlocks.h"
-#include "MuCOE/RemapPins/CustomizableObjectNodeRemapPinsByNameDefaultPin.h"
 #include "MuCOE/UnrealEditorPortabilityHelpers.h"
 
 class UCustomizableObjectNodeRemapPinsByName;
@@ -67,7 +61,7 @@ void UCustomizableObjectNodeComponentPassthroughMesh::PostEditChangeProperty(FPr
 	Super::PostEditChangeProperty(PropertyChangedEvent);
 
 	FProperty* PropertyThatChanged = PropertyChangedEvent.Property;
-	if (PropertyThatChanged && PropertyThatChanged->GetName() == TEXT("Mesh"))
+	if (PropertyThatChanged && PropertyThatChanged->GetName() == GET_MEMBER_NAME_CHECKED(UCustomizableObjectNodeComponentPassthroughMesh, SkeletalMesh))
 	{
 		ReconstructNode();
 	}
@@ -89,7 +83,7 @@ void UCustomizableObjectNodeComponentPassthroughMesh::AllocateDefaultPins(UCusto
 {
 	Super::AllocateDefaultPins(RemapPins);
 
-	if (!SkeletalMesh.IsValid())
+	if (SkeletalMesh.IsNull())
 	{
 		return;
 	}
