[Mutable] Fix Material Parameter widget.

#nrx

diff --git a/Engine/Plugins/Mutable/Source/CustomizableObjectEditor/Private/MuCOE/CustomizableInstanceDetails.cpp b/Engine/Plugins/Mutable/Source/CustomizableObjectEditor/Private/MuCOE/CustomizableInstanceDetails.cpp
--- a/Engine/Plugins/Mutable/Source/CustomizableObjectEditor/Private/MuCOE/CustomizableInstanceDetails.cpp
+++ b/Engine/Plugins/Mutable/Source/CustomizableObjectEditor/Private/MuCOE/CustomizableInstanceDetails.cpp
@@ -1081,7 +1081,7 @@ TSharedRef<SWidget> FCustomizableInstanceDetails::GenerateParameterWidget(const
 
 			return SNew(SObjectPropertyEntryBox)
 				.AllowedClass(UMaterialInterface::StaticClass())
-				.OnObjectChanged(this, &FCustomizableInstanceDetails::OnSkeletalMeshParameterSelected, ParamName)
+				.OnObjectChanged(this, &FCustomizableInstanceDetails::OnMaterialParameterSelected, ParamName)
 				.ObjectPath(ObjectPath);
 		}
 		case EMutableParameterType::Transform:
