Added support for morphs for Mutable mesh parameters. Added MorphTargetVertexCodec helper functions to read and write batch headers and encode/decode quantized morph data.

- Now Mutable meshes can load morphs encoded in the original source morph codec and apply them in ME_MOPRH operations.
- Added compressed morph vertex remove for ME_CLIP operations family.
- Added compressed morph merge.
- Added support for modifiers with mesh parameters.

#rnx

diff --git a/Engine/Plugins/Mutable/Source/CustomizableObject/Internal/MuCO/CustomizableObjectPrivate.h b/Engine/Plugins/Mutable/Source/CustomizableObject/Internal/MuCO/CustomizableObjectPrivate.h
--- a/Engine/Plugins/Mutable/Source/CustomizableObject/Internal/MuCO/CustomizableObjectPrivate.h
+++ b/Engine/Plugins/Mutable/Source/CustomizableObject/Internal/MuCO/CustomizableObjectPrivate.h
@@ -1640,6 +1640,8 @@ enum class ECustomizableObjectVersions : int32
 
 	CompactBindingDataStructure,
 
+	MorphMeshParameters,
+
 	// -----<new versions can be added above this line>--------
 	LastCustomizableObjectVersion
 };
diff --git a/Engine/Plugins/Mutable/Source/CustomizableObject/Private/MuCO/UnrealConversionUtils.cpp b/Engine/Plugins/Mutable/Source/CustomizableObject/Private/MuCO/UnrealConversionUtils.cpp
--- a/Engine/Plugins/Mutable/Source/CustomizableObject/Private/MuCO/UnrealConversionUtils.cpp
+++ b/Engine/Plugins/Mutable/Source/CustomizableObject/Private/MuCO/UnrealConversionUtils.cpp
@@ -13,10 +13,12 @@
 #include "MuR/MeshBufferSet.h"
 #include "MuR/MutableTrace.h"
 #include "Animation/Skeleton.h"
+#include "Containers/StaticArray.h"
 #include "Engine/SkeletalMesh.h"
 #include "UObject/StrongObjectPtr.h"
 #include "GPUSkinVertexFactory.h"
 #include "Rendering/SkeletalMeshRenderData.h"
+#include "Rendering/MorphTargetVertexCodec.h"
 
 class USkeleton;
 
@@ -886,7 +888,7 @@ namespace UnrealConversionUtils
 		LODResource.MorphTargetVertexInfoBuffers.InitMorphResourcesStreaming(LODResource.RenderSections, MorphsToInit, LODResource.StaticVertexBuffers.StaticMeshVertexBuffer.GetNumVertices(), ErrorTolerance);
 	}
 
-
+	
 	UE::Tasks::FTask ConvertSkeletalMeshFromRuntimeData(USkeletalMesh* SkeletalMesh, int32 LODIndex, int32 SectionIndex, UModelResources* ModelResources, mu::FMesh* Result, TFunction<void(bool bSuccess)>& CompletionCallback)
 	{
 		MUTABLE_CPUPROFILER_SCOPE(ConvertSkeletalMeshFromRuntimeData);
@@ -1139,7 +1141,7 @@ namespace UnrealConversionUtils
 		}
 
 		UE::Tasks::FTask FinalConversionTask = UE::Tasks::Launch( TEXT("MutableRuntimeMeshConversionFinal"),
-			[Context, SectionIndex, Result, CompletionCallback = CompletionCallback]()
+			[Context, SkeletalMesh, LODIndex, SectionIndex, Result, CompletionCallback]()
 			{
 				MUTABLE_CPUPROFILER_SCOPE(MeshConversion);
 
@@ -1266,6 +1268,7 @@ namespace UnrealConversionUtils
 					++CurrentVertexBuffer;
 				}
 
+
 				// Indices
 				{
 					int32 FirstIndexIndex = Section.BaseIndex;
@@ -1320,6 +1323,387 @@ namespace UnrealConversionUtils
 					}
 				}
 
+				// Morphs
+
+				using namespace UE::MorphTargetVertexCodec;
+				const FSkeletalMeshLODInfo* SkeletalMeshLODInfo = SkeletalMesh->GetLODInfo(LODIndex);
+
+				const TArray<TObjectPtr<UMorphTarget>>& MorphTargets = SkeletalMesh->GetMorphTargets();
+				for (UMorphTarget* MorphTarget : MorphTargets)
+				{
+					TArray<FMorphTargetLODModel>& MorphLODModels = MorphTarget->GetMorphLODModels();
+
+					const FMorphTargetLODModel* MorphTargetLOD = MorphLODModels.IsValidIndex(LODIndex) 
+							? &MorphLODModels[LODIndex] 
+							: nullptr;
+					Result->Morph.Names.Add(MorphTarget->GetFName());
+				}
+
+				const FMorphTargetVertexInfoBuffers& MorphTargetBuffers = Context->DataLOD->MorphTargetVertexInfoBuffers;
+				TArray<uint32> RuntimeCompressedMorphData;
+				if (MorphTargetBuffers.GetMorphDataSizeInBytes() == 0)
+				{
+					MUTABLE_CPUPROFILER_SCOPE(CompressMorphs)
+
+					// In Editor morph data is not compressed and is located in FMorphTargetLODModel.	
+					TArray<const FMorphTargetLODModel*> MorphTargetLODs;
+					MorphTargetLODs.Reserve(MorphTargets.Num());
+	
+					for (UMorphTarget* MorphTarget : MorphTargets)
+					{
+						TArray<FMorphTargetLODModel>& MorphLODModels = MorphTarget->GetMorphLODModels();
+	
+						const FMorphTargetLODModel* MorphTargetLOD = MorphLODModels.IsValidIndex(LODIndex) 
+								? &MorphLODModels[LODIndex] 
+								: nullptr;
+						MorphTargetLODs.Add(MorphTargetLOD);
+					}
+
+					check(MorphTargetLODs.Num() == Result->Morph.Names.Num());
+					
+					Result->Morph.BatchStartOffsetPerMorph.Empty(MorphTargetLODs.Num());
+					Result->Morph.BatchesPerMorph.Empty(MorphTargetLODs.Num());
+					Result->Morph.MaximumValuePerMorph.Empty(MorphTargetLODs.Num());
+					Result->Morph.MinimumValuePerMorph.Empty(MorphTargetLODs.Num());
+
+					// NOTE: Here we are compressing the morphs for all sections. This could be avoided pre-filtering
+					// the morph vertices in the section. This is done this way so the post-compress filtering is applied
+					// in editor to mimic what will happen with cooked data.
+
+					const int32 NumOriginalVertices = 
+							Context->OriginalLOD->RenderSections.Last().BaseVertexIndex +
+							Context->OriginalLOD->RenderSections.Last().NumVertices;
+
+					// From the Mutable point of view, all morphs need tangents.
+					TBitArray<> VertexNeedsTangents;
+					VertexNeedsTangents.Init(true, NumOriginalVertices);
+			
+					const float PositionPrecision = ComputePositionPrecision(SkeletalMeshLODInfo->MorphTargetPositionErrorTolerance);
+					const float TangentZPrecision = ComputeTangentPrecision();
+
+					TArray<FDeltaBatchHeader> BatchHeaders;
+					TArray<uint32> BitstreamData;
+				
+					const int32 NumMorphTargets = MorphTargetLODs.Num();
+					for (const FMorphTargetLODModel* MorphModel : MorphTargetLODs)
+					{
+						uint32 BatchStartOffset = BatchHeaders.Num();
+
+						float MaximumValues[4] = { -FLT_MAX, -FLT_MAX, -FLT_MAX, -FLT_MAX };
+						float MinimumValues[4] = { +FLT_MAX, +FLT_MAX, +FLT_MAX, +FLT_MAX };
+
+						int32 NumSrcDeltas = MorphModel ? MorphModel->Vertices.Num() : 0;
+						
+						TConstArrayView<FMorphTargetDelta> MorphDeltas(MorphModel->Vertices);
+						for (int32 DeltaIndex = 0; DeltaIndex < NumSrcDeltas; ++DeltaIndex)
+						{
+							const FMorphTargetDelta& MorphDelta = MorphDeltas[DeltaIndex];
+							const FVector3f TangentZDelta = MorphDelta.TangentZDelta;
+
+							MaximumValues[0] = FMath::Max(MaximumValues[0], MorphDelta.PositionDelta.X);
+							MaximumValues[1] = FMath::Max(MaximumValues[1], MorphDelta.PositionDelta.Y);
+							MaximumValues[2] = FMath::Max(MaximumValues[2], MorphDelta.PositionDelta.Z);
+							MaximumValues[3] = FMath::Max(MaximumValues[3], FMath::Max(TangentZDelta.X, FMath::Max(TangentZDelta.Y, TangentZDelta.Z)));
+
+							MinimumValues[0] = FMath::Min(MinimumValues[0], MorphDelta.PositionDelta.X);
+							MinimumValues[1] = FMath::Min(MinimumValues[1], MorphDelta.PositionDelta.Y);
+							MinimumValues[2] = FMath::Min(MinimumValues[2], MorphDelta.PositionDelta.Z);
+							MinimumValues[3] = FMath::Min(MinimumValues[3], FMath::Min(TangentZDelta.X, FMath::Min(TangentZDelta.Y, TangentZDelta.Z)));
+						}
+
+						// Encode the actual morph vertex info into the quantized bitstream.
+						Encode(MorphModel->Vertices, &VertexNeedsTangents, PositionPrecision, TangentZPrecision, BatchHeaders, BitstreamData);
+					
+						const uint32 MorphNumBatches = BatchHeaders.Num() - BatchStartOffset;
+						Result->Morph.BatchStartOffsetPerMorph.Add(BatchStartOffset);
+						Result->Morph.BatchesPerMorph.Add(MorphNumBatches);
+						Result->Morph.MaximumValuePerMorph.Add(FVector4f(MaximumValues[0], MaximumValues[1], MaximumValues[2], MaximumValues[3]));
+						Result->Morph.MinimumValuePerMorph.Add(FVector4f(MinimumValues[0], MinimumValues[1], MinimumValues[2], MinimumValues[3]));
+					}
+				
+					Result->Morph.PositionPrecision = PositionPrecision;
+					Result->Morph.TangentZPrecision = TangentZPrecision;
+					Result->Morph.NumTotalBatches = BatchHeaders.Num();
+					
+					// Fix batch headers and write them packed.
+					for (FDeltaBatchHeader& BatchHeader : BatchHeaders)
+					{
+						BatchHeader.DataOffset += BatchHeaders.Num() * NumBatchHeaderDwords*sizeof(uint32);
+
+						TStaticArray<uint32, NumBatchHeaderDwords> HeaderData;
+						WriteHeader(BatchHeader, HeaderData);
+						RuntimeCompressedMorphData.Append(HeaderData);
+					}
+
+					RuntimeCompressedMorphData.Append(BitstreamData);
+				}
+
+				check(!RuntimeCompressedMorphData.IsEmpty() || MorphTargetBuffers.GetMorphDataSizeInBytes() == 0);
+				const TConstArrayView<uint32> SourceMorphDataView = !RuntimeCompressedMorphData.IsEmpty() 
+						? TConstArrayView<uint32>(RuntimeCompressedMorphData)
+						: TConstArrayView<uint32>(MorphTargetBuffers.GetData(), MorphTargetBuffers.GetMorphDataSizeInBytes()/sizeof(uint32));
+
+				{	
+					MUTABLE_CPUPROFILER_SCOPE(ExtractSectionMorphs)
+					mu::FMeshMorph& Morph = Result->Morph;
+				
+					if (RuntimeCompressedMorphData.IsEmpty())
+					{
+						const int32 NumMorphs = MorphTargetBuffers.GetNumMorphs();
+						Morph.MaximumValuePerMorph.SetNumUninitialized(NumMorphs);
+						Morph.MinimumValuePerMorph.SetNumUninitialized(NumMorphs);
+						Morph.BatchStartOffsetPerMorph.SetNumUninitialized(NumMorphs);
+						Morph.BatchesPerMorph.SetNumUninitialized(NumMorphs);
+
+						for (int32 MorphIndex = 0; MorphIndex < NumMorphs; ++MorphIndex)
+						{
+							Morph.MaximumValuePerMorph[MorphIndex] = MorphTargetBuffers.GetMaximumMorphScale(MorphIndex);
+							Morph.MinimumValuePerMorph[MorphIndex] = MorphTargetBuffers.GetMinimumMorphScale(MorphIndex);
+							Morph.BatchStartOffsetPerMorph[MorphIndex] = MorphTargetBuffers.GetBatchStartOffset(MorphIndex);
+							Morph.BatchesPerMorph[MorphIndex] = MorphTargetBuffers.GetNumBatches(MorphIndex);
+							Morph.NumTotalBatches = MorphTargetBuffers.GetNumBatches();
+							Morph.PositionPrecision = MorphTargetBuffers.GetPositionPrecision();
+							Morph.TangentZPrecision = MorphTargetBuffers.GetTangentZPrecision();
+						}
+					}
+
+					// Strip all vertices not in this section.
+					const int32 SectionVertexRangeBegin = Section.BaseVertexIndex;
+					const int32 SectionVertexRangeEnd   = Section.BaseVertexIndex + Section.NumVertices;
+
+					uint32 CumulativeDataOffsetInDwords = 0;
+
+					int32 NumMorphs = Morph.BatchStartOffsetPerMorph.Num();
+					for (int32 MorphIndex = 0; MorphIndex < NumMorphs; ++MorphIndex)
+					{	
+						const int32 SourceMorphNumBatches = Morph.BatchesPerMorph[MorphIndex];
+						const int32 SourceBatchOffset = Morph.BatchStartOffsetPerMorph[MorphIndex]; 
+
+						int32 BatchIndex = 0;
+						for ( ; BatchIndex < SourceMorphNumBatches; ++BatchIndex)
+						{
+							FDeltaBatchHeader BatchHeader;
+							TConstArrayView<uint32> BatchHeaderDataView(
+									SourceMorphDataView.GetData() + (SourceBatchOffset + BatchIndex) * NumBatchHeaderDwords, 
+									NumBatchHeaderDwords); 
+
+							ReadHeader(BatchHeader, BatchHeaderDataView);
+
+							if (BatchHeader.IndexMin >= (uint32)SectionVertexRangeBegin)
+							{
+								break;
+							}
+						}
+
+						// The previous batch could have deltas in the section range.
+						int32 FirstSectionBatchIndex = FMath::Max(0, BatchIndex - 1);
+
+						for ( ; BatchIndex < SourceMorphNumBatches; ++BatchIndex)
+						{
+							FDeltaBatchHeader BatchHeader;
+							TConstArrayView<uint32> BatchHeaderDataView(
+									SourceMorphDataView.GetData() + (SourceBatchOffset + BatchIndex) * NumBatchHeaderDwords, 
+									NumBatchHeaderDwords); 
+
+							ReadHeader(BatchHeader, BatchHeaderDataView);
+
+							if (BatchHeader.IndexMin >= (uint32)SectionVertexRangeEnd)
+							{
+								break;
+							}
+						}
+						
+						int32 LastSectionBatchIndex = FMath::Max(0, BatchIndex - 1);
+
+						uint32 NumSectionMorphBatches = LastSectionBatchIndex + 1 - FirstSectionBatchIndex;
+						
+						TConstArrayView<uint32> SectionMorphHeadersView(
+									SourceMorphDataView.GetData() + (SourceBatchOffset + FirstSectionBatchIndex) * NumBatchHeaderDwords, 
+									NumSectionMorphBatches * NumBatchHeaderDwords);
+
+						Result->MorphDataBuffer.Append(SectionMorphHeadersView);
+
+						Result->Morph.BatchStartOffsetPerMorph[MorphIndex] = CumulativeDataOffsetInDwords / NumBatchHeaderDwords;
+						Result->Morph.BatchesPerMorph[MorphIndex] = NumSectionMorphBatches;
+						CumulativeDataOffsetInDwords += SectionMorphHeadersView.Num();
+					}
+
+					auto TrimBatchDeltasNotInRange = [](
+							int32 VertexRangeBegin, int32 VertexRangeEnd,
+							FDeltaBatchHeader& InOutBatchHeader,
+							TConstArrayView<uint32> DeltasData,
+							TArrayView<uint32> OutBatchDeltasData)
+					{
+						TConstArrayView<uint32> BatchDeltasData(
+								DeltasData.GetData() + (InOutBatchHeader.DataOffset / sizeof(uint32)), 
+								CalculateBatchDwords(InOutBatchHeader));
+					
+						TArray<FQuantizedDelta, TInlineAllocator<UE::MorphTargetVertexCodec::BatchSize>> QuantizedDeltas;
+						QuantizedDeltas.SetNumUninitialized(UE::MorphTargetVertexCodec::BatchSize);
+
+						ReadQuantizedDeltas(QuantizedDeltas, InOutBatchHeader, BatchDeltasData);
+
+						int32 NewIndexMin = InOutBatchHeader.IndexMin;
+						uint32 DeltaIndex = 0;
+						for ( ; DeltaIndex < InOutBatchHeader.NumElements; ++DeltaIndex)
+						{
+							int32 DeltaVertexIdx = QuantizedDeltas[DeltaIndex].Index;
+
+							if (DeltaVertexIdx >= VertexRangeBegin)
+							{
+								NewIndexMin = DeltaVertexIdx;
+								break;
+							}
+						}
+
+						int32 DeltaIndexBegin = DeltaIndex;
+
+						for ( ; DeltaIndex < InOutBatchHeader.NumElements; ++DeltaIndex)
+						{
+							int32 DeltaVertexIdx = QuantizedDeltas[DeltaIndex].Index;
+							
+							if (DeltaVertexIdx >= VertexRangeEnd)
+							{
+								break;
+							}
+						}
+
+						InOutBatchHeader.IndexMin    = NewIndexMin;
+						InOutBatchHeader.NumElements = DeltaIndex - DeltaIndexBegin;
+
+						check(InOutBatchHeader.NumElements * sizeof(FQuantizedDelta) <= QuantizedDeltas.NumBytes());
+						FMemory::Memmove(
+								QuantizedDeltas.GetData(), 
+								QuantizedDeltas.GetData() + DeltaIndexBegin,
+								InOutBatchHeader.NumElements * sizeof(FQuantizedDelta));
+						
+						WriteQuantizedDeltas(QuantizedDeltas, InOutBatchHeader, OutBatchDeltasData);
+					};
+
+					for (int32 MorphIndex = 0; MorphIndex < NumMorphs; ++MorphIndex)
+					{
+						const int32 MorphNumBatches = Result->Morph.BatchesPerMorph[MorphIndex];
+						const int32 BatchOffset = Result->Morph.BatchStartOffsetPerMorph[MorphIndex]; 
+				
+						if (MorphNumBatches == 0)
+						{
+							continue;
+						}
+
+						TArrayView<uint32> MorphBatchHeadersView(
+								Result->MorphDataBuffer.GetData() + BatchOffset*NumBatchHeaderDwords,
+								MorphNumBatches * NumBatchHeaderDwords);
+
+						// Batches still point to the source data.
+						FDeltaBatchHeader FirstBatchHeader;
+						TArrayView<uint32> FirstBatchHeaderDataView(
+								MorphBatchHeadersView.GetData(), NumBatchHeaderDwords);
+						ReadHeader(FirstBatchHeader, FirstBatchHeaderDataView);
+
+						FDeltaBatchHeader LastBatchHeader;
+						TArrayView<uint32> LastBatchHeaderDataView(
+								MorphBatchHeadersView.GetData() + (MorphNumBatches - 1)*NumBatchHeaderDwords, NumBatchHeaderDwords);
+						ReadHeader(LastBatchHeader, LastBatchHeaderDataView);
+
+						FDeltaBatchHeader FirstFullBatchHeader;
+						FDeltaBatchHeader LastFullBatchHeader;
+						
+						TConstArrayView<uint32> SourceFullBatchesDataView;
+						if (MorphNumBatches > 2)
+						{
+							TArrayView<uint32> FirstFullBatchHeaderDataView(
+									MorphBatchHeadersView.GetData() + 1*NumBatchHeaderDwords, NumBatchHeaderDwords);
+							ReadHeader(FirstFullBatchHeader, FirstFullBatchHeaderDataView);
+
+							TArrayView<uint32> LastFullBatchHeaderDataView(
+									MorphBatchHeadersView.GetData() + (MorphNumBatches - 2)*NumBatchHeaderDwords, NumBatchHeaderDwords);
+							ReadHeader(LastFullBatchHeader, LastFullBatchHeaderDataView);
+
+							SourceFullBatchesDataView = TConstArrayView<uint32>(
+								SourceMorphDataView.GetData() + (FirstFullBatchHeader.DataOffset / sizeof(uint32)),
+								((LastFullBatchHeader.DataOffset - FirstFullBatchHeader.DataOffset) / sizeof(uint32)) + CalculateBatchDwords(LastFullBatchHeader));
+						}
+
+						TArray<uint32, TInlineAllocator<BatchSize*sizeof(FQuantizedDelta)/sizeof(uint32)>> PartialBatchDeltaStorage;
+						uint32 PartialBatchDeltaStorageDwords = FMath::Max(CalculateBatchDwords(FirstBatchHeader), CalculateBatchDwords(LastBatchHeader));
+						PartialBatchDeltaStorage.SetNumUninitialized(PartialBatchDeltaStorageDwords);
+
+						// We are potentially reallocating, views will be invalidated after the reserve or any append.
+						// TODO: We can know an accurate upper-bound of the final size beforehand only looking at source 
+						// headers, pre-compute the size and only make one allocation.
+						Result->MorphDataBuffer.Reserve(
+								MorphBatchHeadersView.Num() + SourceFullBatchesDataView.Num() + 
+								CalculateBatchDwords(FirstBatchHeader) + CalculateBatchDwords(LastBatchHeader));
+
+						TrimBatchDeltasNotInRange(
+								SectionVertexRangeBegin, SectionVertexRangeEnd, FirstBatchHeader, SourceMorphDataView, PartialBatchDeltaStorage);
+						
+						FirstBatchHeaderDataView = TArrayView<uint32>(
+								Result->MorphDataBuffer.GetData() + BatchOffset*NumBatchHeaderDwords, NumBatchHeaderDwords);
+						WriteHeader(FirstBatchHeader, FirstBatchHeaderDataView);			
+	
+						// Batch size may have changed, CalculateBatchDwords value cannot be cached from previous query. 
+						Result->MorphDataBuffer.Append(
+								TArrayView<uint32>(PartialBatchDeltaStorage.GetData(), CalculateBatchDwords(FirstBatchHeader)));
+					
+						// SourceFullBatchesDataView is expected to be empty in some cases.
+						Result->MorphDataBuffer.Append(SourceFullBatchesDataView);
+
+						if (MorphNumBatches > 1)
+						{
+							PartialBatchDeltaStorage.SetNumUninitialized(CalculateBatchDwords(LastBatchHeader), EAllowShrinking::No);
+
+							TrimBatchDeltasNotInRange(
+									SectionVertexRangeBegin, SectionVertexRangeEnd, LastBatchHeader, SourceMorphDataView, PartialBatchDeltaStorage);
+							
+							LastBatchHeaderDataView = TArrayView<uint32>(
+									Result->MorphDataBuffer.GetData() + (BatchOffset + MorphNumBatches - 1)*NumBatchHeaderDwords, NumBatchHeaderDwords);
+							WriteHeader(LastBatchHeader, LastBatchHeaderDataView);					
+	
+							// Batch size may have changed, CalculateBatchDwords value cannot be cached from previous query. 
+							Result->MorphDataBuffer.Append(
+									TArrayView<uint32>(PartialBatchDeltaStorage.GetData(), CalculateBatchDwords(LastBatchHeader)));
+						}
+
+						int32 NumMorphBatches = Result->Morph.BatchesPerMorph[MorphIndex]; 
+
+						// Batch header fixup.
+						TArrayView<uint32> BatchHeadersData(
+								Result->MorphDataBuffer.GetData() + Result->Morph.BatchStartOffsetPerMorph[MorphIndex] * NumBatchHeaderDwords, 
+								NumMorphBatches * NumBatchHeaderDwords);
+						
+						for (int32 BatchIndex = 0; BatchIndex < NumMorphBatches; ++BatchIndex)
+						{
+							TArrayView<uint32> BatchHeaderData(
+									BatchHeadersData.GetData() + BatchIndex*NumBatchHeaderDwords, NumBatchHeaderDwords);
+
+							FDeltaBatchHeader BatchHeader;
+							ReadHeader(BatchHeader, BatchHeaderData);
+
+							if (BatchHeader.NumElements != 0)
+							{
+								check(BatchHeader.IndexMin >= (uint32)SectionVertexRangeBegin);
+								BatchHeader.IndexMin = BatchHeader.IndexMin - SectionVertexRangeBegin;
+								BatchHeader.DataOffset = CumulativeDataOffsetInDwords * sizeof(uint32);
+								CumulativeDataOffsetInDwords += CalculateBatchDwords(BatchHeader);
+							}
+							else
+							{
+								BatchHeader.IndexMin   = 0;
+								BatchHeader.DataOffset = CumulativeDataOffsetInDwords;
+							}
+
+							WriteHeader(BatchHeader, BatchHeaderData);
+						}
+					}
+				}
+	
+				// Make sure data can be loaded 4 Dwords at a time. Probably not needed for Mutable.
+				if (Result->MorphDataBuffer.Num())
+				{
+					Result->MorphDataBuffer.Append({0, 0, 0});
+				}
+
 				Result->EnsureSurfaceData();
 
 				CompletionCallback(true);
diff --git a/Engine/Plugins/Mutable/Source/CustomizableObjectEditor/Private/MuCOE/GenerateMutableSource/GenerateMutableSourceMesh.cpp b/Engine/Plugins/Mutable/Source/CustomizableObjectEditor/Private/MuCOE/GenerateMutableSource/GenerateMutableSourceMesh.cpp
--- a/Engine/Plugins/Mutable/Source/CustomizableObjectEditor/Private/MuCOE/GenerateMutableSource/GenerateMutableSourceMesh.cpp
+++ b/Engine/Plugins/Mutable/Source/CustomizableObjectEditor/Private/MuCOE/GenerateMutableSource/GenerateMutableSourceMesh.cpp
@@ -2920,13 +2920,18 @@ TSharedPtr<mu::FMesh> BuildMorphedMutableMesh(const UEdGraphPin* BaseSourcePin,
 	USkeletalMesh* SkeletalMesh = nullptr;
 	UCustomizableObjectNode* Node = Cast<UCustomizableObjectNode>(BaseSourcePin->GetOwningNode());
 
-	if (const ICustomizableObjectNodeMeshInterface* TypedNodeSkeletalMesh = Cast<ICustomizableObjectNodeMeshInterface>(Node))
+	if (Cast<UCustomizableObjectNodeSkeletalMeshParameter>(Node))
+	{
+		// Do not generate the morphed mesh. Deltas will be included in the Mesh Parameter.
+	}
+
+	else if (const ICustomizableObjectNodeMeshInterface* TypedNodeSkeletalMesh = Cast<ICustomizableObjectNodeMeshInterface>(Node))
 	{
 		int32 LayoutIndex;
 		TypedNodeSkeletalMesh->GetPinSection(*BaseSourcePin, LODIndexConnected, SectionIndexConnected, LayoutIndex);
 		SkeletalMesh = Cast<USkeletalMesh>(MutablePrivate::LoadObject(TypedNodeSkeletalMesh->GetMesh()));
 	}
-
+	
 	else if (const UCustomizableObjectNodeTable* TypedNodeTable = Cast<UCustomizableObjectNodeTable>(Node))
 	{
 		UDataTable* DataTable = GetDataTable(TypedNodeTable, GenerationContext);
@@ -3042,6 +3047,10 @@ TArray<TTuple<USkeletalMesh*, TSoftClassPtr<UAnimInstance>>> GetSkeletalMeshesIn
 			SkeletalMeshesInfo.Emplace(SkeletalMeshNode->SkeletalMesh.Get(), SkeletalMeshNode->AnimInstance);
 		}
 	}
+	else if (const UCustomizableObjectNodeSkeletalMeshParameter* SkeletalMeshParameterNode = Cast<UCustomizableObjectNodeSkeletalMeshParameter>(SkeletalMeshOrTableNode))
+	{
+		// Nothing
+	}
 	else if (const UCustomizableObjectNodeTable* TableNode = Cast<UCustomizableObjectNodeTable>(SkeletalMeshOrTableNode))
 	{
 		UDataTable* DataTable = GetDataTable(TableNode, GenerationContext);
@@ -3487,6 +3496,7 @@ mu::NodeMeshPtr GenerateMorphMesh(const UEdGraphPin* Pin,
 	
 	// Factor
 	//GenerateMorphFactor(MorphNode, , GenerationContext, Result);
+	Result->Name = FName(TypedNodeMorphs[MorphIndex].MorphTargetName);
 	Result->Factor = TypedNodeMorphs[MorphIndex].FactorNode;
 
 	// Base
@@ -3569,11 +3579,10 @@ mu::NodeMeshPtr GenerateMorphMesh(const UEdGraphPin* Pin,
 	else
 	{
 		MorphedSourceMesh = BuildMorphedMutableMesh(Pin, TypedNodeMorphs[MorphIndex].MorphTargetName, GenerationContext, bOnlyConnectedLOD);
+		bSuccess = true;
 
 		if (MorphedSourceMesh)
 		{
-			bSuccess = true;
-
 			mu::Ptr<mu::NodeMeshConstant> MorphedSourceMeshNode = new mu::NodeMeshConstant;
 			MorphedSourceMeshNode->Value = MorphedSourceMesh;
 			MorphedSourceMeshNode->SetMessageContext(MorphNode);
diff --git a/Engine/Plugins/Mutable/Source/MutableRuntime/Internal/MuR/OpMeshMerge.h b/Engine/Plugins/Mutable/Source/MutableRuntime/Internal/MuR/OpMeshMerge.h
--- a/Engine/Plugins/Mutable/Source/MutableRuntime/Internal/MuR/OpMeshMerge.h
+++ b/Engine/Plugins/Mutable/Source/MutableRuntime/Internal/MuR/OpMeshMerge.h
@@ -9,6 +9,8 @@
 #include "MuR/MutableTrace.h"
 
 #include "Engine/SkeletalMesh.h"
+#include "Rendering/MorphTargetVertexCodec.h"
+
 #include "UObject/StrongObjectPtr.h"
 
 namespace mu
@@ -20,6 +22,240 @@ namespace mu
 		TSharedPtr<FMesh> SecondReformat;
 	};
 
+	inline void FixMorphIndices(FMesh& Result, int32 VertexIndexOffset)
+	{
+        MUTABLE_CPUPROFILER_SCOPE(CompressedMorphs_Merge_FixMorphIndices)
+
+		using namespace UE::MorphTargetVertexCodec;
+
+		if (VertexIndexOffset == 0)
+		{
+			return;
+		}
+
+		int32 NumMorphs = Result.Morph.Names.Num();
+		for (int32 MorphIndex = 0; MorphIndex < NumMorphs; ++MorphIndex)
+		{
+			TArrayView<uint32> MorphHeadersDataView = TArrayView<uint32>(
+					Result.MorphDataBuffer.GetData() + Result.Morph.BatchStartOffsetPerMorph[MorphIndex],
+					Result.Morph.BatchesPerMorph[MorphIndex]);
+
+			int32 NumBatches = Result.Morph.BatchesPerMorph[MorphIndex];
+			for (int32 BatchIndex = 0; BatchIndex < NumBatches; ++BatchIndex)
+			{
+				TArrayView<uint32> HeaderDataView = TArrayView<uint32>(
+						MorphHeadersDataView.GetData() + BatchIndex*NumBatchHeaderDwords,
+						NumBatchHeaderDwords);
+
+				FDeltaBatchHeader BatchHeader;
+				ReadHeader(BatchHeader, HeaderDataView);
+
+				BatchHeader.IndexMin += VertexIndexOffset;
+
+				WriteHeader(BatchHeader, HeaderDataView);
+			}
+		}
+	}
+
+	inline void MergeMorphs(FMesh& Result, const FMesh& MeshA, const FMesh& MeshB)
+	{
+        MUTABLE_CPUPROFILER_SCOPE(CompressedMorphs_Merge_Merge)
+
+		using namespace UE::MorphTargetVertexCodec;
+		
+		// Compute memory requirements.	
+		Result.Morph.Names.Reserve(MeshA.Morph.Names.Num() + MeshB.Morph.Names.Num());
+		Result.Morph.Names.Append(MeshA.Morph.Names);
+
+		TArray<int32, TInlineAllocator<128>> ResultToBMorphNameMap; 
+		ResultToBMorphNameMap.Init(-1, MeshB.Morph.Names.Num() + MeshA.Morph.Names.Num());
+
+		// MeshA morphs map is the identity, compute only the MorphB morphs map.
+		for (int32 Index = 0; Index < MeshB.Morph.Names.Num(); ++Index)
+		{
+			int32 MappedIndex = Result.Morph.Names.AddUnique(MeshB.Morph.Names[Index]);
+			ResultToBMorphNameMap[MappedIndex] = Index;
+		}
+	
+		const uint32 MergedDataBufferSizeInDwords = MeshA.MorphDataBuffer.Num() + MeshB.MorphDataBuffer.Num();  
+		
+		Result.MorphDataBuffer.SetNum(MergedDataBufferSizeInDwords);
+		Result.Morph.MinimumValuePerMorph.SetNumUninitialized(Result.Morph.Names.Num());
+		Result.Morph.MaximumValuePerMorph.SetNumUninitialized(Result.Morph.Names.Num());
+		Result.Morph.BatchesPerMorph.SetNumUninitialized(Result.Morph.Names.Num());
+		Result.Morph.BatchStartOffsetPerMorph.SetNumUninitialized(Result.Morph.Names.Num());
+		
+		Result.Morph.PositionPrecision = MeshA.Morph.PositionPrecision;
+		Result.Morph.TangentZPrecision = MeshA.Morph.TangentZPrecision;
+
+		const int32 MeshANumMorphs = MeshA.Morph.Names.Num();
+		const int32 ResultNumMorphs = Result.Morph.Names.Num();
+
+		// Merge headers.
+		int32 BatchAccumulatedOffsetInDwords = 0;
+		int32 MorphIndex = 0;
+		for (; MorphIndex < MeshANumMorphs; ++MorphIndex)
+		{
+			Result.Morph.MaximumValuePerMorph[MorphIndex] = MeshA.Morph.MaximumValuePerMorph[MorphIndex];
+			Result.Morph.MinimumValuePerMorph[MorphIndex] = MeshA.Morph.MinimumValuePerMorph[MorphIndex];
+
+			Result.Morph.BatchStartOffsetPerMorph[MorphIndex] = BatchAccumulatedOffsetInDwords / NumBatchHeaderDwords;
+			Result.Morph.BatchesPerMorph[MorphIndex] = MeshA.Morph.BatchesPerMorph[MorphIndex];
+
+			FMemory::Memcpy(
+					Result.MorphDataBuffer.GetData() + BatchAccumulatedOffsetInDwords, 
+					MeshA.MorphDataBuffer.GetData() + MeshA.Morph.BatchStartOffsetPerMorph[MorphIndex]*NumBatchHeaderDwords,
+					MeshA.Morph.BatchesPerMorph[MorphIndex] * NumBatchHeaderDwords*sizeof(uint32));
+
+			BatchAccumulatedOffsetInDwords += MeshA.Morph.BatchesPerMorph[MorphIndex] * NumBatchHeaderDwords;
+			
+			int32 MeshBMorphIndex = ResultToBMorphNameMap[MorphIndex];
+			if (MeshBMorphIndex >= 0)
+			{
+				Result.Morph.BatchesPerMorph[MorphIndex] += MeshB.Morph.BatchesPerMorph[MeshBMorphIndex];
+
+				const FVector4f& SourceMaxValue = MeshB.Morph.MaximumValuePerMorph[MeshBMorphIndex];
+				const FVector4f& SourceMinValue = MeshB.Morph.MaximumValuePerMorph[MeshBMorphIndex];
+
+				Result.Morph.MaximumValuePerMorph[MorphIndex] = 
+						SourceMaxValue.ComponentMax(Result.Morph.MaximumValuePerMorph[MorphIndex]);
+				Result.Morph.MinimumValuePerMorph[MorphIndex] = 
+						SourceMinValue.ComponentMin(Result.Morph.MinimumValuePerMorph[MorphIndex]);
+
+				FMemory::Memcpy(
+						Result.MorphDataBuffer.GetData() + BatchAccumulatedOffsetInDwords, 
+						MeshB.MorphDataBuffer.GetData() + MeshB.Morph.BatchStartOffsetPerMorph[MeshBMorphIndex]*NumBatchHeaderDwords,
+						MeshB.Morph.BatchesPerMorph[MeshBMorphIndex] * NumBatchHeaderDwords*sizeof(uint32));
+				
+				BatchAccumulatedOffsetInDwords += MeshB.Morph.BatchesPerMorph[MorphIndex] * NumBatchHeaderDwords;
+			}
+		}
+
+		for (; MorphIndex < ResultNumMorphs; ++MorphIndex)
+		{
+			int32 MeshBMorphIndex = ResultToBMorphNameMap[MorphIndex];
+			
+			Result.Morph.BatchesPerMorph[MorphIndex] = MeshB.Morph.BatchesPerMorph[MeshBMorphIndex]; 
+			Result.Morph.BatchStartOffsetPerMorph[MorphIndex] = BatchAccumulatedOffsetInDwords / NumBatchHeaderDwords;
+
+			Result.Morph.MaximumValuePerMorph[MorphIndex] = MeshB.Morph.MaximumValuePerMorph[MeshBMorphIndex];
+			Result.Morph.MinimumValuePerMorph[MorphIndex] = MeshB.Morph.MinimumValuePerMorph[MeshBMorphIndex];
+
+			FMemory::Memcpy(
+					Result.MorphDataBuffer.GetData() + BatchAccumulatedOffsetInDwords, 
+					MeshB.MorphDataBuffer.GetData() + MeshB.Morph.BatchStartOffsetPerMorph[MeshBMorphIndex]*NumBatchHeaderDwords,
+					MeshB.Morph.BatchesPerMorph[MeshBMorphIndex] * NumBatchHeaderDwords*sizeof(uint32));
+
+			BatchAccumulatedOffsetInDwords += Result.Morph.BatchesPerMorph[MorphIndex] * NumBatchHeaderDwords;
+		}
+
+		// Merge morph data and fix up headers.
+
+		const uint32 MeshANumVertices = MeshA.GetVertexCount();
+
+		// Add Morphs in MeshA and common morphs with MeshB.
+		MorphIndex = 0;
+		for (; MorphIndex < MeshANumMorphs; ++MorphIndex)
+		{		
+			TArrayView<uint32> MorphHeadersDataView = TArrayView<uint32>(
+					Result.MorphDataBuffer.GetData() + Result.Morph.BatchStartOffsetPerMorph[MorphIndex]*NumBatchHeaderDwords,
+					Result.Morph.BatchesPerMorph[MorphIndex] * NumBatchHeaderDwords);
+
+			int32 BatchIndex = 0;
+			{
+				int32 NumBatches = MeshA.Morph.BatchesPerMorph[MorphIndex];
+				for ( ; BatchIndex < NumBatches; ++BatchIndex)
+				{
+					TArrayView<uint32> HeaderDataView = TArrayView<uint32>(
+							MorphHeadersDataView.GetData() + BatchIndex*NumBatchHeaderDwords,
+							NumBatchHeaderDwords);
+
+					FDeltaBatchHeader BatchHeader;
+					ReadHeader(BatchHeader, HeaderDataView);
+
+					const int32 SourceDataOffset = BatchHeader.DataOffset;
+					BatchHeader.DataOffset = BatchAccumulatedOffsetInDwords * sizeof(uint32);
+
+					WriteHeader(BatchHeader, HeaderDataView);
+
+					const int32 NumBatchDwords = CalculateBatchDwords(BatchHeader);
+					FMemory::Memcpy(
+							Result.MorphDataBuffer.GetData() + BatchHeader.DataOffset / sizeof(uint32),
+							MeshA.MorphDataBuffer.GetData() + SourceDataOffset / sizeof(uint32),
+							NumBatchDwords * sizeof(uint32));
+
+					BatchAccumulatedOffsetInDwords += NumBatchDwords;
+				}
+			}
+
+			int32 MeshBMorphIndex = ResultToBMorphNameMap[MorphIndex];
+			if (MeshBMorphIndex >= 0)
+			{
+				int32 NumBatches = MeshB.Morph.BatchesPerMorph[MeshBMorphIndex];
+
+				for (; BatchIndex < NumBatches; ++BatchIndex)
+				{
+					TArrayView<uint32> HeaderDataView = TArrayView<uint32>(
+							MorphHeadersDataView.GetData() + BatchIndex*NumBatchHeaderDwords,
+							NumBatchHeaderDwords);
+
+					FDeltaBatchHeader BatchHeader;
+					ReadHeader(BatchHeader, HeaderDataView);
+
+					const int32 SourceDataOffset = BatchHeader.DataOffset;
+					BatchHeader.DataOffset = BatchAccumulatedOffsetInDwords * sizeof(uint32);
+					BatchHeader.IndexMin += MeshANumVertices;
+
+					WriteHeader(BatchHeader, HeaderDataView);
+
+					const int32 NumBatchDwords = CalculateBatchDwords(BatchHeader);
+					FMemory::Memcpy(
+							Result.MorphDataBuffer.GetData() + BatchHeader.DataOffset / sizeof(uint32),
+							MeshB.MorphDataBuffer.GetData() + SourceDataOffset / sizeof(uint32),
+							NumBatchDwords * sizeof(uint32));
+					
+					BatchAccumulatedOffsetInDwords += NumBatchDwords;
+				}
+			}
+		}
+
+		// Add Morphs in MeshB that are not present in MeshA.
+		for (; MorphIndex < ResultNumMorphs; ++MorphIndex)
+		{
+			TArrayView<uint32> MorphHeadersDataView = TArrayView<uint32>(
+					Result.MorphDataBuffer.GetData() + Result.Morph.BatchStartOffsetPerMorph[MorphIndex],
+					Result.Morph.BatchesPerMorph[MorphIndex]);
+
+			int32 BatchIndex = 0;
+			{
+				int32 NumBatches = MeshB.Morph.BatchesPerMorph[MorphIndex];
+				for (; BatchIndex < NumBatches; ++BatchIndex)
+				{
+					TArrayView<uint32> HeaderDataView = TArrayView<uint32>(
+							MorphHeadersDataView.GetData() + BatchIndex*NumBatchHeaderDwords,
+							NumBatchHeaderDwords);
+
+					FDeltaBatchHeader BatchHeader;
+					ReadHeader(BatchHeader, HeaderDataView);
+
+					const int32 SourceDataOffset = BatchHeader.DataOffset;
+					BatchHeader.DataOffset = BatchAccumulatedOffsetInDwords * sizeof(uint32);
+					BatchHeader.IndexMin += MeshANumVertices;
+
+					WriteHeader(BatchHeader, HeaderDataView);
+
+					const int32 NumBatchDwords = CalculateBatchDwords(BatchHeader);
+					FMemory::Memcpy(
+							Result.MorphDataBuffer.GetData() + BatchHeader.DataOffset / sizeof(uint32),
+							MeshA.MorphDataBuffer.GetData() + SourceDataOffset / sizeof(uint32),
+							NumBatchDwords * sizeof(uint32));
+
+					BatchAccumulatedOffsetInDwords += NumBatchDwords;
+				}
+			}
+		}
+	}
+
 	//---------------------------------------------------------------------------------------------
 	//! Merge two meshes into one new mesh
 	//---------------------------------------------------------------------------------------------
@@ -1117,6 +1353,25 @@ namespace mu
 		{
 			Result->SkeletalMeshes.AddUnique(pSecond->SkeletalMeshes[Index]);
 		}
+
+		// Mesh morphs.
+		
+		if (pFirst->HasMorphs() && !pSecond->HasMorphs())
+		{
+			Result->Morph = pFirst->Morph;
+			Result->MorphDataBuffer = pFirst->MorphDataBuffer;
+		}
+		else if (!pFirst->HasMorphs() && pSecond->HasMorphs())
+		{
+			Result->Morph = pSecond->Morph;
+			Result->MorphDataBuffer = pSecond->MorphDataBuffer;
+
+			FixMorphIndices(*Result, pFirst->GetVertexCount());
+		}
+		else if (pFirst->HasMorphs() && pSecond->HasMorphs())
+		{
+			MergeMorphs(*Result, *pFirst, *pSecond);	
+		}
 	}
 	
 
diff --git a/Engine/Plugins/Mutable/Source/MutableRuntime/Private/MuR/CodeRunner.cpp b/Engine/Plugins/Mutable/Source/MutableRuntime/Private/MuR/CodeRunner.cpp
--- a/Engine/Plugins/Mutable/Source/MutableRuntime/Private/MuR/CodeRunner.cpp
+++ b/Engine/Plugins/Mutable/Source/MutableRuntime/Private/MuR/CodeRunner.cpp
@@ -1381,6 +1381,10 @@ namespace mu
         {
 			const uint8* Data = Program.GetOpArgsPointer(Item.At);
 
+			OP::ADDRESS NameAt = 0;
+			FMemory::Memcpy(&NameAt, Data, sizeof(OP::ADDRESS)); 
+			Data += sizeof(OP::ADDRESS);
+				
 			OP::ADDRESS FactorAt = 0;
 			FMemory::Memcpy(&FactorAt, Data, sizeof(OP::ADDRESS)); 
 			Data += sizeof(OP::ADDRESS);
@@ -1446,7 +1450,6 @@ namespace mu
 					// The Morph, partial or full
 					else
 					{
-						// We will need the base again
 						AddOp(FScheduledOp(Item.At, Item, 2, DataAddress),
 							FScheduledOp(BaseAt, Item),
 							FScheduledOp(TargetAt, Item));
@@ -1468,34 +1471,37 @@ namespace mu
 
 				if (BaseMesh)
 				{
-					// No morph
 					if (FMath::IsNearlyZero(Factor))
 					{
 						StoreMesh(Item, BaseMesh);
 					}
-					// The Morph, partial or full
 					else 
 					{
 						TSharedPtr<const FMesh> MorphMesh = LoadMesh(FCacheAddress(TargetAt, Item));
 						
-						if (MorphMesh)
+						TSharedPtr<FMesh> Result = CloneOrTakeOver(BaseMesh);
+						if (Result->HasMorphs())
 						{
-							TSharedPtr<FMesh> Result = CloneOrTakeOver(BaseMesh);
-							MeshMorph(Result.Get(), MorphMesh.Get(), Factor);
-
-							Release(MorphMesh);
-							StoreMesh(Item, Result);
+							// Morph data from the Base Mesh
+							FName Name(Program.ConstantStrings[NameAt]);
+							MeshMorph(Result.Get(), Name, Factor);
+						
 						}
-						else
+						else if (MorphMesh)
 						{
-							StoreMesh(Item, BaseMesh);
+							MeshMorph(Result.Get(), MorphMesh.Get(), Factor);
 						}
+
+						StoreMesh(Item, Result);
+						Release(MorphMesh);
 					}
 				}
 				else
 				{
 					StoreMesh(Item, nullptr);
 				}
+
+
                 break;
             }
 
diff --git a/Engine/Plugins/Mutable/Source/MutableRuntime/Private/MuR/Mesh.cpp b/Engine/Plugins/Mutable/Source/MutableRuntime/Private/MuR/Mesh.cpp
--- a/Engine/Plugins/Mutable/Source/MutableRuntime/Private/MuR/Mesh.cpp
+++ b/Engine/Plugins/Mutable/Source/MutableRuntime/Private/MuR/Mesh.cpp
@@ -195,6 +195,11 @@ void FMesh::CopyFrom(const FMesh& From, EMeshCopyFlags InFlags)
 		AdditionalPhysicsBodies = From.AdditionalPhysicsBodies;
 	}
 
+	if (EnumHasAnyFlags(InFlags, EMeshCopyFlags::WithMorphData))
+	{
+		Morph = From.Morph;
+		MorphDataBuffer = From.MorphDataBuffer;	
+	}
 }
 
 
@@ -204,6 +209,12 @@ uint32 FMesh::GetId() const
 }
 
 
+bool FMesh::HasMorphs() const
+{
+	return Morph.NumTotalBatches != 0;
+}
+
+
 int FMesh::GetVertexCount() const
 {
     return GetVertexBuffers().GetElementCount();
diff --git a/Engine/Plugins/Mutable/Source/MutableRuntime/Private/MuR/OpMeshMorph.cpp b/Engine/Plugins/Mutable/Source/MutableRuntime/Private/MuR/OpMeshMorph.cpp
new file mode 100644
--- /dev/null
+++ b/Engine/Plugins/Mutable/Source/MutableRuntime/Private/MuR/OpMeshMorph.cpp
@@ -0,0 +1,527 @@
+// Copyright Epic Games, Inc. All Rights Reserved.
+
+#include "OpMeshMorph.h"
+
+#include "MuR/MeshPrivate.h"
+#include "MuR/ConvertData.h"
+#include "MuR/MutableTrace.h"
+#include "MuR/SparseIndexMap.h"
+#include "PackedNormal.h"
+#include "MuR/Mesh.h"
+#include "UObject/NameTypes.h"
+#include "Animation/MorphTarget.h"
+#include "Rendering/MorphTargetVertexCodec.h"
+
+void mu::MeshMorph(FMesh* Mesh, const FName& MorphName, const float Factor)
+{
+	MUTABLE_CPUPROFILER_SCOPE(MeshMorph_Parameter);
+
+	using namespace UE::MorphTargetVertexCodec;
+
+	if (!Mesh)
+	{
+		return;
+	}
+
+	const int32 MorphIndex = Mesh->Morph.Names.Find(MorphName);
+	if (MorphIndex == INDEX_NONE)
+	{
+		return;
+	}
+	
+	// Get pointers to vertex position data
+	MeshBufferIterator<EMeshBufferFormat::Float32, float, 3> MeshPositionIter(Mesh->VertexBuffers, EMeshBufferSemantic::Position, 0);	
+	const bool bHasPositions = MeshPositionIter.ptr() != nullptr;
+	
+	// {BiNormal, Tangent, Normal}
+	TStaticArray<UntypedMeshBufferIterator, 3> MeshTangentFrameChannelsIters;
+	
+	const FMeshBufferSet& VertexBufferSet = Mesh->GetVertexBuffers();
+
+	const uint32 NumVertices = VertexBufferSet.GetElementCount();
+
+	int32 NormalBufferIndex   = -1;
+	int32 NormalBufferChannel = -1;
+	
+	VertexBufferSet.FindChannel(EMeshBufferSemantic::Normal, 0, &NormalBufferIndex, &NormalBufferChannel);
+	int32 NormalNumChannels = NormalBufferIndex >= 0 ? VertexBufferSet.Buffers[NormalBufferIndex].Channels.Num() : 0;
+	
+	const bool bHasNormals = NormalBufferIndex >= 0;
+
+	for (int32 ChannelIndex = 0; ChannelIndex < NormalNumChannels; ++ChannelIndex)
+	{
+		const FMeshBufferChannel& Channel = VertexBufferSet.Buffers[NormalBufferIndex].Channels[ChannelIndex];
+
+		EMeshBufferSemantic Sem = Channel.Semantic;
+		int32 SemIndex = Channel.SemanticIndex;
+
+		if (Sem == EMeshBufferSemantic::Normal && bHasNormals)
+		{
+			MeshTangentFrameChannelsIters[2] = UntypedMeshBufferIterator(Mesh->GetVertexBuffers(), Sem, SemIndex);
+		}
+		else if (Sem == EMeshBufferSemantic::Tangent && bHasNormals)
+		{
+			MeshTangentFrameChannelsIters[1] = UntypedMeshBufferIterator(Mesh->GetVertexBuffers(), Sem, SemIndex);
+		}
+		else if (Sem == EMeshBufferSemantic::Binormal && bHasNormals)
+		{
+			MeshTangentFrameChannelsIters[0] = UntypedMeshBufferIterator(Mesh->GetVertexBuffers(), Sem, SemIndex);
+		}
+	}
+
+	const UntypedMeshBufferIterator& MeshNormalIter = MeshTangentFrameChannelsIters[2];
+	const UntypedMeshBufferIterator& MeshTangentIter = MeshTangentFrameChannelsIters[1];
+	const UntypedMeshBufferIterator& MeshBiNormalIter = MeshTangentFrameChannelsIters[0];
+
+	const EMeshBufferFormat NormalFormat = MeshNormalIter.GetFormat();
+	const int32 NormalComps = MeshNormalIter.GetComponents();
+
+	const EMeshBufferFormat TangentFormat = MeshTangentIter.GetFormat();
+	const int32 TangentComps = MeshTangentIter.GetComponents();
+
+	const EMeshBufferFormat BiNormalFormat = MeshBiNormalIter.GetFormat();
+	const int32 BiNormalComps = MeshBiNormalIter.GetComponents();
+	
+	const bool bHasOptimizedNormals = NormalFormat == EMeshBufferFormat::PackedDirS8_W_TangentSign
+		&& (!MeshTangentIter.ptr() || TangentFormat == EMeshBufferFormat::PackedDirS8) && !MeshBiNormalIter.ptr();
+	
+	const uint32 BatchStartOffset = Mesh->Morph.BatchStartOffsetPerMorph[MorphIndex];
+	const uint32 MorphNumBatches  = Mesh->Morph.BatchesPerMorph[MorphIndex];
+	for (uint32 BatchHeaderIndex = 0; BatchHeaderIndex < MorphNumBatches; ++BatchHeaderIndex)
+	{
+		TArrayView<uint32> BatchHeaderData(
+				&Mesh->MorphDataBuffer[(BatchStartOffset + BatchHeaderIndex) * NumBatchHeaderDwords], 
+				NumBatchHeaderDwords);
+
+		FDeltaBatchHeader BatchHeader;
+		ReadHeader(BatchHeader, BatchHeaderData);
+
+		if (BatchHeader.NumElements == 0)
+		{
+			continue;
+		}
+
+		TArray<FQuantizedDelta, TInlineAllocator<UE::MorphTargetVertexCodec::BatchSize>> QuantizedDeltas;
+		QuantizedDeltas.SetNumUninitialized(BatchHeader.NumElements);
+
+		TArrayView<uint32> Data(
+				&Mesh->MorphDataBuffer[BatchHeader.DataOffset / sizeof(uint32)], 
+				CalculateBatchDwords(BatchHeader));
+
+		ReadQuantizedDeltas(QuantizedDeltas, BatchHeader, Data);
+
+		for (FQuantizedDelta& QuantizedDelta : QuantizedDeltas)
+		{
+			FMorphTargetDelta Delta;
+			DequantizeDelta(Delta, BatchHeader.bTangents, QuantizedDelta, Mesh->Morph.PositionPrecision, Mesh->Morph.TangentZPrecision);
+
+			// Positions
+			if (bHasPositions)
+			{
+				check(Delta.SourceIdx < NumVertices)
+				FVector3f& Position = *reinterpret_cast<FVector3f*>(*(MeshPositionIter + Delta.SourceIdx));
+				Position += Delta.PositionDelta * Factor;
+			}
+
+			// Normals
+			if (bHasOptimizedNormals)
+			{
+				// Normal
+				check(Delta.SourceIdx < NumVertices)
+				FPackedNormal* PackedNormal = reinterpret_cast<FPackedNormal*>((MeshNormalIter + Delta.SourceIdx).ptr());
+				int8 W = PackedNormal->Vector.W;
+				const FVector3f BaseNormal = PackedNormal->ToFVector3f();
+				
+				const FVector3f Normal = (BaseNormal + Delta.TangentZDelta * Factor).GetSafeNormal();
+
+				*PackedNormal = Normal;
+				PackedNormal->Vector.W = W;
+
+				// Tangent
+				if (MeshTangentIter.ptr())
+				{
+					check(Delta.SourceIdx < NumVertices)
+					FPackedNormal* PackedTangent = reinterpret_cast<FPackedNormal*>((MeshTangentIter + Delta.SourceIdx).ptr());
+					const FVector3f BaseTangent = PackedTangent->ToFVector3f();
+
+					// Orthogonalize Tangent based on new Normal. This assumes Normal and BaseTangent are normalized and different.
+					const FVector3f Tangent = (BaseTangent - FVector3f::DotProduct(Normal, BaseTangent) * Normal).GetSafeNormal();
+
+					*PackedTangent = Tangent;
+				}
+			}
+			else if (MeshNormalIter.ptr())
+			{
+				// When normal is packed, binormal channel is not expected. It is not a big deal if it's there but we would be doing extra unused work in that case. 
+				ensure(!(NormalFormat == EMeshBufferFormat::PackedDir8_W_TangentSign || NormalFormat == EMeshBufferFormat::PackedDirS8_W_TangentSign) || !MeshBiNormalIter.ptr());
+				
+				MUTABLE_CPUPROFILER_SCOPE(ApplyNormalMorph_SlowPath);
+				
+				UntypedMeshBufferIterator NormalIter = MeshNormalIter + Delta.SourceIdx;
+
+				const FVector3f BaseNormal = NormalIter.GetAsVec3f();
+
+				const FVector3f Normal = (BaseNormal + Delta.TangentZDelta * Factor).GetSafeNormal();
+
+				// Leave the tangent basis sign untouched for packed normals formats.
+				for (int32 C = 0; C < NormalComps && C < 3; ++C)
+				{
+					ConvertData(C, NormalIter.ptr(), NormalFormat, &Normal, EMeshBufferFormat::Float32);
+				}
+
+				// Tangent
+				if (MeshTangentIter.ptr())
+				{
+					UntypedMeshBufferIterator TangentIter = MeshTangentIter + Delta.SourceIdx;
+
+					const FVector3f BaseTangent = TangentIter.GetAsVec3f();
+
+					// Orthogonalize Tangent based on new Normal. This assumes Normal and BaseTangent are normalized and different.
+					const FVector3f Tangent = (BaseTangent - FVector3f::DotProduct(Normal, BaseTangent) * Normal).GetSafeNormal();
+
+					for (int32 C = 0; C < TangentComps && C < 3; ++C)
+					{
+						ConvertData(C, TangentIter.ptr(), TangentFormat, &Tangent, EMeshBufferFormat::Float32);
+					}
+
+					// BiNormal
+					if (MeshBiNormalIter.ptr())
+					{
+						UntypedMeshBufferIterator BiNormalIter = MeshBiNormalIter + Delta.SourceIdx;
+
+						const FVector3f& N = BaseNormal;
+						const FVector3f& T = BaseTangent;
+						const FVector3f  B = BiNormalIter.GetAsVec3f();
+
+						const float BaseTangentBasisDeterminant =
+							B.X * T.Y * N.Z + B.Z * T.X * N.Y + B.Y * T.Z * N.Y -
+							B.Z * T.Y * N.X - B.Y * T.X * N.Z - B.X * T.Z * N.Y;
+
+						const float BaseTangentBasisDeterminantSign = BaseTangentBasisDeterminant >= 0 ? 1.0f : -1.0f;
+
+						const FVector3f BiNormal = FVector3f::CrossProduct(Tangent, Normal) * BaseTangentBasisDeterminantSign;
+
+						for (int32 C = 0; C < BiNormalComps && C < 3; ++C)
+						{
+							ConvertData(C, BiNormalIter.ptr(), BiNormalFormat, &BiNormal, EMeshBufferFormat::Float32);
+						}
+					}
+				}
+			}
+		}
+	}
+}
+
+void mu::MeshMorph(FMesh* BaseMesh, const FMesh* MaxMesh, const float Factor)
+{
+    MUTABLE_CPUPROFILER_SCOPE(MeshMorph_Compiled);
+
+	if (!BaseMesh)
+	{
+		return;
+	}
+
+	const auto MakeIndexMap = [](
+		MeshVertexIdIteratorConst BaseIdIter, int32 BaseNum,
+		MeshVertexIdIteratorConst MorphIdIter, int32 MorphNum)
+	-> SparseIndexMapSet
+	{	
+		MUTABLE_CPUPROFILER_SCOPE(MakeIndexMap);
+		TArray<SparseIndexMapSet::FRangeDesc> RangeDescs;
+
+		// Detect all ranges and their limits
+		{
+			for (int32 Index = 0; Index < BaseNum; ++Index, ++BaseIdIter)
+			{
+				const uint64 BaseId = BaseIdIter.Get();
+
+				uint32 Prefix = BaseId >> 32;
+				uint32 Id = BaseId & 0xffffffff;
+				bool bFound = false;
+				for (SparseIndexMapSet::FRangeDesc& Range : RangeDescs)
+				{
+					if (Range.Prefix == Prefix)
+					{
+						Range.MinIndex = FMath::Min(Id, Range.MinIndex);
+						Range.MaxIndex = FMath::Max(Id, Range.MaxIndex);
+						bFound = true;
+						break;
+					}
+				}
+
+				if (!bFound)
+				{
+					RangeDescs.Add({Prefix, Id, Id});
+				}
+			}
+		}
+
+		SparseIndexMapSet IndexMap(RangeDescs);
+       
+        for (int32 Index = 0; Index < MorphNum; ++Index, ++MorphIdIter)
+        {
+            const uint64 MorphId = MorphIdIter.Get();
+            
+            IndexMap.Insert(MorphId, Index);
+        }
+
+		return IndexMap;
+	};
+
+	// Number of vertices to modify
+	const int32 MaxNum = MaxMesh ? MaxMesh->GetVertexBuffers().GetElementCount() : 0;
+	const int32 BaseNum = BaseMesh ? BaseMesh->GetVertexBuffers().GetElementCount() : 0;
+	const FMesh* RefTarget = MaxMesh;
+
+	if (BaseNum == 0 || MaxNum == 0)
+	{
+		return;
+	}
+
+	if (!RefTarget)
+	{
+		return;
+	}
+	
+	constexpr int32 MorphBufferDataChannel = 0;
+	const int32 ChannelsNum = RefTarget->GetVertexBuffers().GetBufferChannelCount(MorphBufferDataChannel);
+
+	TArray<UntypedMeshBufferIterator> BaseGenericIters;
+	BaseGenericIters.SetNum(ChannelsNum);
+	TArray<UntypedMeshBufferIteratorConst> MinGenericIters;
+	MinGenericIters.SetNum(ChannelsNum);
+	TArray<UntypedMeshBufferIteratorConst> MaxGenericIters;
+	MaxGenericIters.SetNum(ChannelsNum);
+
+	// Get pointers to vertex position data
+	MeshBufferIterator<EMeshBufferFormat::Float32, float, 3> BasePositionIter(BaseMesh->VertexBuffers, EMeshBufferSemantic::Position, 0);	
+	MeshBufferIteratorConst<EMeshBufferFormat::Float32, float, 3> MaxPositionIter(RefTarget->VertexBuffers, EMeshBufferSemantic::Position, 0);
+
+	// {BiNormal, Tangent, Normal}
+	TStaticArray<UntypedMeshBufferIterator, 3> BaseTangentFrameChannelsIters;
+	UntypedMeshBufferIteratorConst MaxNormalChannelIter;
+
+	const bool bBaseHasNormals = UntypedMeshBufferIteratorConst(BaseMesh->GetVertexBuffers(), EMeshBufferSemantic::Normal, 0).ptr() != nullptr;
+	for (int32 ChannelIndex = 0; ChannelIndex < ChannelsNum; ++ChannelIndex)
+	{
+		const FMeshBufferSet& MBSPriv = RefTarget->GetVertexBuffers();
+		const FMeshBufferChannel& Channel = MBSPriv.Buffers[MorphBufferDataChannel].Channels[ChannelIndex];
+		EMeshBufferSemantic Sem = Channel.Semantic;
+		int32 SemIndex = Channel.SemanticIndex;
+		if (Sem == EMeshBufferSemantic::Normal && bBaseHasNormals)
+		{
+			BaseTangentFrameChannelsIters[2] = UntypedMeshBufferIterator(BaseMesh->GetVertexBuffers(), Sem, SemIndex);
+			MaxNormalChannelIter = UntypedMeshBufferIteratorConst(MaxMesh->GetVertexBuffers(), Sem, SemIndex);
+		}
+		else if (Sem == EMeshBufferSemantic::Tangent && bBaseHasNormals)
+		{
+			BaseTangentFrameChannelsIters[1] = UntypedMeshBufferIterator(BaseMesh->GetVertexBuffers(), Sem, SemIndex);
+		}
+		else if (Sem == EMeshBufferSemantic::Binormal && bBaseHasNormals)
+		{
+			BaseTangentFrameChannelsIters[0] = UntypedMeshBufferIterator(BaseMesh->GetVertexBuffers(), Sem, SemIndex);
+		}
+		else if(Sem != EMeshBufferSemantic::Position)
+		{
+			BaseGenericIters[ChannelIndex] = UntypedMeshBufferIterator(BaseMesh->GetVertexBuffers(), Sem, SemIndex);
+			MaxGenericIters[ChannelIndex] = UntypedMeshBufferIteratorConst(MaxMesh->GetVertexBuffers(), Sem, SemIndex);
+		}
+	}
+
+	MeshVertexIdIteratorConst BaseIdIter(BaseMesh);
+	
+	MeshVertexIdIteratorConst MaxIdIter(MaxMesh);
+	SparseIndexMapSet IndexMap = MakeIndexMap(BaseIdIter, BaseNum, MaxIdIter, MaxNum);
+    
+	MUTABLE_CPUPROFILER_SCOPE(ApplyMorph);
+
+	const UntypedMeshBufferIterator& BaseNormalIter = BaseTangentFrameChannelsIters[2];
+	const UntypedMeshBufferIterator& BaseTangentIter = BaseTangentFrameChannelsIters[1];
+	const UntypedMeshBufferIterator& BaseBiNormalIter = BaseTangentFrameChannelsIters[0];
+
+	const EMeshBufferFormat NormalFormat = BaseNormalIter.GetFormat();
+	const int32 NormalComps = BaseNormalIter.GetComponents();
+
+	const EMeshBufferFormat TangentFormat = BaseTangentIter.GetFormat();
+	const int32 TangentComps = BaseTangentIter.GetComponents();
+
+	const EMeshBufferFormat BiNormalFormat = BaseBiNormalIter.GetFormat();
+	const int32 BiNormalComps = BaseBiNormalIter.GetComponents();
+
+	const EMeshBufferFormat MorphNormalFormat = MaxNormalChannelIter.GetFormat();
+
+	const bool bHasPositions = BasePositionIter.ptr() && MaxPositionIter.ptr();
+	check((BasePositionIter.GetFormat() == EMeshBufferFormat::Float32 && BasePositionIter.GetComponents() == 3) || !bHasPositions);
+
+	const bool bHasOptimizedNormals = NormalFormat == EMeshBufferFormat::PackedDirS8_W_TangentSign && MorphNormalFormat == EMeshBufferFormat::Float32
+		&& (!BaseTangentIter.ptr() || TangentFormat == EMeshBufferFormat::PackedDirS8) && !BaseBiNormalIter.ptr();
+	
+	bool bHasGenericMorphs = false;
+	const int32 ChannelNum = MaxGenericIters.Num();
+	for (int32 ChannelIndex = 0; ChannelIndex < ChannelNum; ++ChannelIndex)
+	{
+		if (!(BaseGenericIters[ChannelIndex].ptr() && MaxGenericIters[ChannelIndex].ptr()))
+		{
+			continue;
+		}
+
+		bHasGenericMorphs = true;
+	}
+
+	for (int32 VertexIndex = 0; VertexIndex < BaseNum; ++VertexIndex)
+	{
+		const uint64 BaseId = (BaseIdIter + VertexIndex).Get();
+		const uint32 MorphIndex = IndexMap.Find(BaseId);
+
+		if (MorphIndex == SparseIndexMap::NotFoundValue)
+		{
+			continue;
+		}
+
+		// Find consecutive run.
+		MeshVertexIdIteratorConst RunBaseIter = BaseIdIter + VertexIndex;
+		MeshVertexIdIteratorConst RunMorphIter = MaxIdIter + MorphIndex;
+
+		int32 RunSize = 0;
+		for (; VertexIndex + RunSize < BaseNum && int32(MorphIndex) + RunSize < MaxNum && RunBaseIter.Get() == RunMorphIter.Get();
+			++RunSize, ++RunBaseIter, ++RunMorphIter);
+
+		// Positions
+		if (bHasPositions)
+		{
+			for (int32 RunIndex = 0; RunIndex < RunSize; ++RunIndex)
+			{
+				FVector3f& Position = *reinterpret_cast<FVector3f*>(*(BasePositionIter + VertexIndex + RunIndex));
+				const FVector3f& MorphPosition = *reinterpret_cast<const FVector3f*>(*(MaxPositionIter + MorphIndex + RunIndex));
+				Position += MorphPosition * Factor;
+			}
+		}
+
+		// Normals
+		if (bHasOptimizedNormals)
+		{
+			for (int32 RunIndex = 0; RunIndex < RunSize; ++RunIndex)
+			{
+				// Normal
+				FPackedNormal* PackedNormal = reinterpret_cast<FPackedNormal*>((BaseNormalIter + VertexIndex + RunIndex).ptr());
+				int8 W = PackedNormal->Vector.W;
+				const FVector3f BaseNormal = PackedNormal->ToFVector3f();
+
+				const FVector3f* MorphNormal = reinterpret_cast<const FVector3f*>((MaxNormalChannelIter + MorphIndex + RunIndex).ptr());
+
+				const FVector3f Normal = (BaseNormal + *MorphNormal * Factor).GetSafeNormal();
+
+				*PackedNormal = *reinterpret_cast<const FVector3f*>(&Normal);
+				PackedNormal->Vector.W = W;
+
+				// Tangent
+				if (BaseTangentIter.ptr())
+				{
+					FPackedNormal* PackedTangent = reinterpret_cast<FPackedNormal*>((BaseTangentIter + (VertexIndex + RunIndex)).ptr());
+					const FVector3f BaseTangent = PackedTangent->ToFVector3f();
+
+					// Orthogonalize Tangent based on new Normal. This assumes Normal and BaseTangent are normalized and different.
+					const FVector3f Tangent = (BaseTangent - FVector3f::DotProduct(Normal, BaseTangent) * Normal).GetSafeNormal();
+
+					*PackedTangent = *reinterpret_cast<const FVector3f*>(&Tangent);
+				}
+			}
+		}
+		else if (BaseNormalIter.ptr())
+		{
+			// When normal is packed, binormal channel is not expected. It is not a big deal if it's there but we would be doing extra unused work in that case. 
+			ensure(!(NormalFormat == EMeshBufferFormat::PackedDir8_W_TangentSign || NormalFormat == EMeshBufferFormat::PackedDirS8_W_TangentSign) || !BaseBiNormalIter.ptr());
+			
+			MUTABLE_CPUPROFILER_SCOPE(ApplyNormalMorph_SlowPath);
+
+			for (int32 RunIndex = 0; RunIndex < RunSize; ++RunIndex)
+			{
+				UntypedMeshBufferIterator NormalIter = BaseNormalIter + (VertexIndex + RunIndex);
+
+				const FVector3f BaseNormal = NormalIter.GetAsVec3f();
+				const FVector3f MorphNormal = (MaxNormalChannelIter + (MorphIndex + RunIndex)).GetAsVec3f();
+
+				const FVector3f Normal = (BaseNormal + MorphNormal * Factor).GetSafeNormal();
+
+				// Leave the tangent basis sign untouched for packed normals formats.
+				for (int32 C = 0; C < NormalComps && C < 3; ++C)
+				{
+					ConvertData(C, NormalIter.ptr(), NormalFormat, &Normal, EMeshBufferFormat::Float32);
+				}
+
+				// Tangent
+				if (BaseTangentIter.ptr())
+				{
+					UntypedMeshBufferIterator TangentIter = BaseTangentIter + (VertexIndex + RunIndex);
+
+					const FVector3f BaseTangent = TangentIter.GetAsVec3f();
+
+					// Orthogonalize Tangent based on new Normal. This assumes Normal and BaseTangent are normalized and different.
+					const FVector3f Tangent = (BaseTangent - FVector3f::DotProduct(Normal, BaseTangent) * Normal).GetSafeNormal();
+
+					for (int32 C = 0; C < TangentComps && C < 3; ++C)
+					{
+						ConvertData(C, TangentIter.ptr(), TangentFormat, &Tangent, EMeshBufferFormat::Float32);
+					}
+
+					// BiNormal
+					if (BaseBiNormalIter.ptr())
+					{
+						UntypedMeshBufferIterator BiNormalIter = BaseBiNormalIter + (VertexIndex + RunIndex);
+
+						const FVector3f& N = BaseNormal;
+						const FVector3f& T = BaseTangent;
+						const FVector3f  B = BiNormalIter.GetAsVec3f();
+
+						const float BaseTangentBasisDeterminant =
+							B.X * T.Y * N.Z + B.Z * T.X * N.Y + B.Y * T.Z * N.Y -
+							B.Z * T.Y * N.X - B.Y * T.X * N.Z - B.X * T.Z * N.Y;
+
+						const float BaseTangentBasisDeterminantSign = BaseTangentBasisDeterminant >= 0 ? 1.0f : -1.0f;
+
+						const FVector3f BiNormal = FVector3f::CrossProduct(Tangent, Normal) * BaseTangentBasisDeterminantSign;
+
+						for (int32 C = 0; C < BiNormalComps && C < 3; ++C)
+						{
+							ConvertData(C, BiNormalIter.ptr(), BiNormalFormat, &BiNormal, EMeshBufferFormat::Float32);
+						}
+					}
+				}
+			}
+		}
+		
+		// Generic Morphs
+		if (bHasGenericMorphs)
+		{
+			MUTABLE_CPUPROFILER_SCOPE(ApplyNormalMorph_Generic);
+			for (int32 ChannelIndex = 0; ChannelIndex < ChannelNum; ++ChannelIndex)
+			{
+				if (!(BaseGenericIters[ChannelIndex].ptr() && MaxGenericIters[ChannelIndex].ptr()))
+				{
+					continue;
+				}
+
+				UntypedMeshBufferIterator ChannelBaseIter = BaseGenericIters[ChannelIndex] + VertexIndex;
+				UntypedMeshBufferIteratorConst ChannelMorphIter = MaxGenericIters[ChannelIndex] + MorphIndex;
+
+				const EMeshBufferFormat DestChannelFormat = BaseGenericIters[ChannelIndex].GetFormat();
+				const int32 DestChannelComps = BaseGenericIters[ChannelIndex].GetComponents();
+
+				// Apply Morph to range found above.
+				for (int32 R = 0; R < RunSize; ++R, ++ChannelBaseIter, ++ChannelMorphIter)
+				{
+					const FVector4f Value = ChannelBaseIter.GetAsVec4f() + ChannelMorphIter.GetAsVec4f() * Factor;
+
+					// TODO: Optimize this for the specific components.
+					// Max 4 components
+					for (int32 Comp = 0; Comp < DestChannelComps && Comp < 4; ++Comp)
+					{
+						ConvertData(Comp, ChannelBaseIter.ptr(), DestChannelFormat, &Value, EMeshBufferFormat::Float32);
+					}
+				}
+			}
+		}
+
+		VertexIndex += FMath::Max(RunSize - 1, 0);
+	}
+}
diff --git a/Engine/Plugins/Mutable/Source/MutableRuntime/Private/MuR/OpMeshMorph.h b/Engine/Plugins/Mutable/Source/MutableRuntime/Private/MuR/OpMeshMorph.h
--- a/Engine/Plugins/Mutable/Source/MutableRuntime/Private/MuR/OpMeshMorph.h
+++ b/Engine/Plugins/Mutable/Source/MutableRuntime/Private/MuR/OpMeshMorph.h
@@ -2,384 +2,20 @@
 
 #pragma once
 
-#include "MuR/MeshPrivate.h"
-#include "MuR/ConvertData.h"
-#include "MuR/MutableTrace.h"
-#include "MuR/Platform.h"
-#include "MuR/SparseIndexMap.h"
-
-#include "PackedNormal.h"
+class FName;
 
 namespace mu
 {
-	/**
-	 * Optimized linear factor version for morphing 2 targets
-	 */
-    inline void MeshMorph2(FMesh* BaseMesh, const FMesh* MinMesh, const FMesh* MaxMesh, const float Factor)
-    {
-        MUTABLE_CPUPROFILER_SCOPE(MeshMorph2);
-
-		if (!BaseMesh)
-		{
-			return;
-		}
-
-		const auto MakeIndexMap = [](
-			MeshVertexIdIteratorConst BaseIdIter, int32 BaseNum,
-			MeshVertexIdIteratorConst MorphIdIter, int32 MorphNum)
-		-> SparseIndexMapSet
-		{	
-			MUTABLE_CPUPROFILER_SCOPE(MakeIndexMap);
-			TArray<SparseIndexMapSet::FRangeDesc> RangeDescs;
-
-			// Detect all ranges and their limits
-			{
-				for (int32 Index = 0; Index < BaseNum; ++Index, ++BaseIdIter)
-				{
-					const uint64 BaseId = BaseIdIter.Get();
-
-					uint32 Prefix = BaseId >> 32;
-					uint32 Id = BaseId & 0xffffffff;
-					bool bFound = false;
-					for (SparseIndexMapSet::FRangeDesc& Range : RangeDescs)
-					{
-						if (Range.Prefix == Prefix)
-						{
-							Range.MinIndex = FMath::Min(Id, Range.MinIndex);
-							Range.MaxIndex = FMath::Max(Id, Range.MaxIndex);
-							bFound = true;
-							break;
-						}
-					}
-
-					if (!bFound)
-					{
-						RangeDescs.Add({Prefix, Id, Id});
-					}
-				}
-			}
-
-			SparseIndexMapSet IndexMap(RangeDescs);
-           
-            for (int32 Index = 0; Index < MorphNum; ++Index, ++MorphIdIter)
-            {
-                const uint64 MorphId = MorphIdIter.Get();
-                
-                IndexMap.Insert(MorphId, Index);
-            }
-
-			return IndexMap;
-		};
-
-		const auto ApplyMorph = [](
-			MeshVertexIdIteratorConst BaseIdIter, const MeshBufferIterator<EMeshBufferFormat::Float32, float, 3>& BasePositionIter,
-			const TStaticArray<UntypedMeshBufferIterator, 3>& BaseTangentFrameIters, const TArray<UntypedMeshBufferIterator>& BaseGenericIters, const int32 BaseNum,
-			MeshVertexIdIteratorConst MorphIdIter, MeshBufferIteratorConst<EMeshBufferFormat::Float32, float, 3>& MorphPositionIter,
-			UntypedMeshBufferIteratorConst& MorphNormalIter, const TArray<UntypedMeshBufferIteratorConst>& MorphGenericIters, const int32 MorphNum,
-			const SparseIndexMapSet& IndexMap, const float Factor)
-			-> void
-			{
-				MUTABLE_CPUPROFILER_SCOPE(ApplyMorph);
-
-				const UntypedMeshBufferIterator& BaseNormalIter = BaseTangentFrameIters[2];
-				const UntypedMeshBufferIterator& BaseTangentIter = BaseTangentFrameIters[1];
-				const UntypedMeshBufferIterator& BaseBiNormalIter = BaseTangentFrameIters[0];
-
-				const EMeshBufferFormat NormalFormat = BaseNormalIter.GetFormat();
-				const int32 NormalComps = BaseNormalIter.GetComponents();
-
-				const EMeshBufferFormat TangentFormat = BaseTangentIter.GetFormat();
-				const int32 TangentComps = BaseTangentIter.GetComponents();
-
-				const EMeshBufferFormat BiNormalFormat = BaseBiNormalIter.GetFormat();
-				const int32 BiNormalComps = BaseBiNormalIter.GetComponents();
-
-				const EMeshBufferFormat MorphNormalFormat = MorphNormalIter.GetFormat();
-
-				const bool bHasPositions = BasePositionIter.ptr() && MorphPositionIter.ptr();
-				check((BasePositionIter.GetFormat() == EMeshBufferFormat::Float32 && BasePositionIter.GetComponents() == 3) || !bHasPositions);
-
-				const bool bHasOptimizedNormals = NormalFormat == EMeshBufferFormat::PackedDirS8_W_TangentSign && MorphNormalFormat == EMeshBufferFormat::Float32
-					&& (!BaseTangentIter.ptr() || TangentFormat == EMeshBufferFormat::PackedDirS8) && !BaseBiNormalIter.ptr();
-				
-				bool bHasGenericMorphs = false;
-				const int32 ChannelNum = MorphGenericIters.Num();
-				for (int32 ChannelIndex = 0; ChannelIndex < ChannelNum; ++ChannelIndex)
-				{
-					if (!(BaseGenericIters[ChannelIndex].ptr() && MorphGenericIters[ChannelIndex].ptr()))
-					{
-						continue;
-					}
-
-					bHasGenericMorphs = true;
-				}
-
-				for (int32 VertexIndex = 0; VertexIndex < BaseNum; ++VertexIndex)
-				{
-					const uint64 BaseId = (BaseIdIter + VertexIndex).Get();
-					const uint32 MorphIndex = IndexMap.Find(BaseId);
-
-					if (MorphIndex == SparseIndexMap::NotFoundValue)
-					{
-						continue;
-					}
-
-					// Find consecutive run.
-					MeshVertexIdIteratorConst RunBaseIter = BaseIdIter + VertexIndex;
-					MeshVertexIdIteratorConst RunMorphIter = MorphIdIter + MorphIndex;
-
-					int32 RunSize = 0;
-					for (; VertexIndex + RunSize < BaseNum && int32(MorphIndex) + RunSize < MorphNum && RunBaseIter.Get() == RunMorphIter.Get();
-						++RunSize, ++RunBaseIter, ++RunMorphIter);
-
-					// Positions
-					if (bHasPositions)
-					{
-						for (int32 RunIndex = 0; RunIndex < RunSize; ++RunIndex)
-						{
-							FVector3f& Position = *reinterpret_cast<FVector3f*>(*(BasePositionIter + VertexIndex + RunIndex));
-							const FVector3f& MorphPosition = *reinterpret_cast<const FVector3f*>(*(MorphPositionIter + MorphIndex + RunIndex));
-							Position += MorphPosition * Factor;
-						}
-					}
-
-					// Normals
-					if (bHasOptimizedNormals)
-					{
-						for (int32 RunIndex = 0; RunIndex < RunSize; ++RunIndex)
-						{
-							// Normal
-							FPackedNormal* PackedNormal = reinterpret_cast<FPackedNormal*>((BaseNormalIter + VertexIndex + RunIndex).ptr());
-							int8 W = PackedNormal->Vector.W;
-							const FVector3f BaseNormal = PackedNormal->ToFVector3f();
-
-							const FVector3f* MorphNormal = reinterpret_cast<const FVector3f*>((MorphNormalIter + MorphIndex + RunIndex).ptr());
-
-							const FVector3f Normal = (BaseNormal + *MorphNormal * Factor).GetSafeNormal();
-
-							*PackedNormal = *reinterpret_cast<const FVector3f*>(&Normal);
-							PackedNormal->Vector.W = W;
-
-							// Tangent
-							if (BaseTangentIter.ptr())
-							{
-								FPackedNormal* PackedTangent = reinterpret_cast<FPackedNormal*>((BaseTangentIter + (VertexIndex + RunIndex)).ptr());
-								const FVector3f BaseTangent = PackedTangent->ToFVector3f();
-
-								// Orthogonalize Tangent based on new Normal. This assumes Normal and BaseTangent are normalized and different.
-								const FVector3f Tangent = (BaseTangent - FVector3f::DotProduct(Normal, BaseTangent) * Normal).GetSafeNormal();
-
-								*PackedTangent = *reinterpret_cast<const FVector3f*>(&Tangent);
-							}
-						}
-					}
-					else if (BaseNormalIter.ptr())
-					{
-						// When normal is packed, binormal channel is not expected. It is not a big deal if it's there but we would be doing extra unused work in that case. 
-						ensure(!(NormalFormat == EMeshBufferFormat::PackedDir8_W_TangentSign || NormalFormat == EMeshBufferFormat::PackedDirS8_W_TangentSign) || !BaseBiNormalIter.ptr());
-						
-						MUTABLE_CPUPROFILER_SCOPE(ApplyNormalMorph_SlowPath);
-
-						for (int32 RunIndex = 0; RunIndex < RunSize; ++RunIndex)
-						{
-							UntypedMeshBufferIterator NormalIter = BaseNormalIter + (VertexIndex + RunIndex);
-
-							const FVector3f BaseNormal = NormalIter.GetAsVec3f();
-							const FVector3f MorphNormal = (MorphNormalIter + (MorphIndex + RunIndex)).GetAsVec3f();
-
-							const FVector3f Normal = (BaseNormal + MorphNormal * Factor).GetSafeNormal();
-
-							// Leave the tangent basis sign untouched for packed normals formats.
-							for (int32 C = 0; C < NormalComps && C < 3; ++C)
-							{
-								ConvertData(C, NormalIter.ptr(), NormalFormat, &Normal, EMeshBufferFormat::Float32);
-							}
-
-							// Tangent
-							if (BaseTangentIter.ptr())
-							{
-								UntypedMeshBufferIterator TangentIter = BaseTangentIter + (VertexIndex + RunIndex);
-
-								const FVector3f BaseTangent = TangentIter.GetAsVec3f();
-
-								// Orthogonalize Tangent based on new Normal. This assumes Normal and BaseTangent are normalized and different.
-								const FVector3f Tangent = (BaseTangent - FVector3f::DotProduct(Normal, BaseTangent) * Normal).GetSafeNormal();
-
-								for (int32 C = 0; C < TangentComps && C < 3; ++C)
-								{
-									ConvertData(C, TangentIter.ptr(), TangentFormat, &Tangent, EMeshBufferFormat::Float32);
-								}
-
-								// BiNormal
-								if (BaseBiNormalIter.ptr())
-								{
-									UntypedMeshBufferIterator BiNormalIter = BaseBiNormalIter + (VertexIndex + RunIndex);
-
-									const FVector3f& N = BaseNormal;
-									const FVector3f& T = BaseTangent;
-									const FVector3f  B = BiNormalIter.GetAsVec3f();
-
-									const float BaseTangentBasisDeterminant =
-										B.X * T.Y * N.Z + B.Z * T.X * N.Y + B.Y * T.Z * N.Y -
-										B.Z * T.Y * N.X - B.Y * T.X * N.Z - B.X * T.Z * N.Y;
-
-									const float BaseTangentBasisDeterminantSign = BaseTangentBasisDeterminant >= 0 ? 1.0f : -1.0f;
-
-									const FVector3f BiNormal = FVector3f::CrossProduct(Tangent, Normal) * BaseTangentBasisDeterminantSign;
-
-									for (int32 C = 0; C < BiNormalComps && C < 3; ++C)
-									{
-										ConvertData(C, BiNormalIter.ptr(), BiNormalFormat, &BiNormal, EMeshBufferFormat::Float32);
-									}
-								}
-							}
-						}
-					}
-					
-					// Generic Morphs
-					if (bHasGenericMorphs)
-					{
-						MUTABLE_CPUPROFILER_SCOPE(ApplyNormalMorph_Generic);
-						for (int32 ChannelIndex = 0; ChannelIndex < ChannelNum; ++ChannelIndex)
-						{
-							if (!(BaseGenericIters[ChannelIndex].ptr() && MorphGenericIters[ChannelIndex].ptr()))
-							{
-								continue;
-							}
-
-							UntypedMeshBufferIterator ChannelBaseIter = BaseGenericIters[ChannelIndex] + VertexIndex;
-							UntypedMeshBufferIteratorConst ChannelMorphIter = MorphGenericIters[ChannelIndex] + MorphIndex;
-
-							const EMeshBufferFormat DestChannelFormat = BaseGenericIters[ChannelIndex].GetFormat();
-							const int32 DestChannelComps = BaseGenericIters[ChannelIndex].GetComponents();
-
-							// Apply Morph to range found above.
-							for (int32 R = 0; R < RunSize; ++R, ++ChannelBaseIter, ++ChannelMorphIter)
-							{
-								const FVector4f Value = ChannelBaseIter.GetAsVec4f() + ChannelMorphIter.GetAsVec4f() * Factor;
-
-								// TODO: Optimize this for the specific components.
-								// Max 4 components
-								for (int32 Comp = 0; Comp < DestChannelComps && Comp < 4; ++Comp)
-								{
-									ConvertData(Comp, ChannelBaseIter.ptr(), DestChannelFormat, &Value, EMeshBufferFormat::Float32);
-								}
-							}
-						}
-					}
-
-					VertexIndex += FMath::Max(RunSize - 1, 0);
-				}
-			};
-
-		// Number of vertices to modify
-		const int32 MinNum = MinMesh ? MinMesh->GetVertexBuffers().GetElementCount() : 0;
-		const int32 MaxNum = MaxMesh ? MaxMesh->GetVertexBuffers().GetElementCount() : 0;
-		const int32 BaseNum = BaseMesh ? BaseMesh->GetVertexBuffers().GetElementCount() : 0;
-		const FMesh* RefTarget = MinNum > 0 ? MinMesh : MaxMesh;
-
-		if (BaseNum == 0 || (MinNum + MaxNum) == 0)
-		{
-			return;
-		}
-
-		if (RefTarget)
-		{
-			constexpr int32 MorphBufferDataChannel = 0;
-			const int32 ChannelsNum = RefTarget->GetVertexBuffers().GetBufferChannelCount(MorphBufferDataChannel);
-
-			TArray<UntypedMeshBufferIterator> BaseGenericIters;
-			BaseGenericIters.SetNum(ChannelsNum);
-			TArray<UntypedMeshBufferIteratorConst> MinGenericIters;
-			MinGenericIters.SetNum(ChannelsNum);
-			TArray<UntypedMeshBufferIteratorConst> MaxGenericIters;
-			MaxGenericIters.SetNum(ChannelsNum);
-
-			// Get pointers to vertex position data
-			MeshBufferIterator<EMeshBufferFormat::Float32, float, 3> BasePositionIter(BaseMesh->VertexBuffers, EMeshBufferSemantic::Position, 0);	
-			MeshBufferIteratorConst<EMeshBufferFormat::Float32, float, 3> MinPositionIter(RefTarget->VertexBuffers, EMeshBufferSemantic::Position, 0);
-			MeshBufferIteratorConst<EMeshBufferFormat::Float32, float, 3> MaxPositionIter(RefTarget->VertexBuffers, EMeshBufferSemantic::Position, 0);
-
-			// {BiNormal, Tangent, Normal}
-			TStaticArray<UntypedMeshBufferIterator, 3> BaseTangentFrameChannelsIters;
-			UntypedMeshBufferIteratorConst MinNormalChannelIter;
-			UntypedMeshBufferIteratorConst MaxNormalChannelIter;
-
-			const bool bBaseHasNormals = UntypedMeshBufferIteratorConst(BaseMesh->GetVertexBuffers(), EMeshBufferSemantic::Normal, 0).ptr() != nullptr;
-			for (int32 ChannelIndex = 0; ChannelIndex < ChannelsNum; ++ChannelIndex)
-			{
-				const FMeshBufferSet& MBSPriv = RefTarget->GetVertexBuffers();
-				const FMeshBufferChannel& Channel = MBSPriv.Buffers[MorphBufferDataChannel].Channels[ChannelIndex];
-				EMeshBufferSemantic Sem = Channel.Semantic;
-				int32 SemIndex = Channel.SemanticIndex;
-				if (Sem == EMeshBufferSemantic::Normal && bBaseHasNormals)
-				{
-					BaseTangentFrameChannelsIters[2] = UntypedMeshBufferIterator(BaseMesh->GetVertexBuffers(), Sem, SemIndex);
-					if (MinNum > 0)
-					{
-						MinNormalChannelIter = UntypedMeshBufferIteratorConst(MinMesh->GetVertexBuffers(), Sem, SemIndex);
-					}
-
-					if (MaxNum > 0)
-					{
-						MaxNormalChannelIter = UntypedMeshBufferIteratorConst(MaxMesh->GetVertexBuffers(), Sem, SemIndex);
-					}
-				}
-				else if (Sem == EMeshBufferSemantic::Tangent && bBaseHasNormals)
-				{
-					BaseTangentFrameChannelsIters[1] = UntypedMeshBufferIterator(BaseMesh->GetVertexBuffers(), Sem, SemIndex);
-				}
-				else if (Sem == EMeshBufferSemantic::Binormal && bBaseHasNormals)
-				{
-					BaseTangentFrameChannelsIters[0] = UntypedMeshBufferIterator(BaseMesh->GetVertexBuffers(), Sem, SemIndex);
-				}
-				else if(Sem != EMeshBufferSemantic::Position)
-				{
-					BaseGenericIters[ChannelIndex] = UntypedMeshBufferIterator(BaseMesh->GetVertexBuffers(), Sem, SemIndex);
-					if (MinNum > 0)
-					{
-						MinGenericIters[ChannelIndex] = UntypedMeshBufferIteratorConst(MinMesh->GetVertexBuffers(), Sem, SemIndex);
-					}
-
-					if (MaxNum > 0)
-					{
-						MaxGenericIters[ChannelIndex] = UntypedMeshBufferIteratorConst(MaxMesh->GetVertexBuffers(), Sem, SemIndex);
-					}
-				}
-			}
-
-			MeshVertexIdIteratorConst BaseIdIter(BaseMesh);
-
-			if (MinNum > 0)
-			{
-				MeshVertexIdIteratorConst MinIdIter(MinMesh);
-				SparseIndexMapSet IndexMap = MakeIndexMap(BaseIdIter, BaseNum, MinIdIter, MinNum);
-
-				ApplyMorph(BaseIdIter, BasePositionIter, BaseTangentFrameChannelsIters, BaseGenericIters, BaseNum,
-					MinIdIter, MinPositionIter, MinNormalChannelIter, MinGenericIters, MinNum, IndexMap, 1.0f - Factor);
-			}
-
-			if (MaxNum > 0)
-			{
-				MeshVertexIdIteratorConst MaxIdIter(MaxMesh);
-				SparseIndexMapSet IndexMap = MakeIndexMap(BaseIdIter, BaseNum, MaxIdIter, MaxNum);
+	class FMesh;
+}
 
-				ApplyMorph(BaseIdIter, BasePositionIter, BaseTangentFrameChannelsIters, BaseGenericIters, BaseNum,
-					MaxIdIter, MaxPositionIter, MaxNormalChannelIter, MaxGenericIters, MaxNum, IndexMap, Factor);
-			}
-		}
-    }
 
-	//TODO Optimized linear factor version
-	inline void MeshMorph(FMesh* BaseMesh, const FMesh* MorphMesh, float Factor)
-	{
-		MeshMorph2(BaseMesh, nullptr, MorphMesh, Factor);
-	}
+namespace mu
+{
+	/** Morph Parameter Meshes. Mesh contains the deltas in morph buffer. Deltas are compressed using the engine format. */
+	void MeshMorph(FMesh* Mesh, const FName& MorphName, float Factor);
 
-    //TODO Optimized Factor-less version
-	inline void MeshMorph(FMesh* BaseMesh, const FMesh* MorphMesh)
-    {
-        // Trust the compiler to remove the factor
-		MeshMorph(BaseMesh, MorphMesh, 1.0f);
-    }
+	
+	/** Morph Compiled Meshes. MaxMesh contains de deltas in the position buffer. */
+	void MeshMorph(FMesh* BaseMesh, const FMesh* MaxMesh, float Factor); // TODO Use the engine compressed format. This means changing how be compile morphs.
 }
diff --git a/Engine/Plugins/Mutable/Source/MutableRuntime/Private/MuR/OpMeshRemove.cpp b/Engine/Plugins/Mutable/Source/MutableRuntime/Private/MuR/OpMeshRemove.cpp
--- a/Engine/Plugins/Mutable/Source/MutableRuntime/Private/MuR/OpMeshRemove.cpp
+++ b/Engine/Plugins/Mutable/Source/MutableRuntime/Private/MuR/OpMeshRemove.cpp
@@ -7,9 +7,183 @@
 #include "MuR/ParallelExecutionUtils.h"
 #include "MuR/Platform.h"
 
+#include "Templates/UnrealTemplate.h"
+#include "Rendering/MorphTargetVertexCodec.h"
 
 namespace mu
 {
+	void MeshRemoveMeshMorphs(FMesh& Result, TConstArrayView<int32> VertexIndexRemap)
+	{
+		MUTABLE_CPUPROFILER_SCOPE(CompressedMorphs_MeshRemoveMesh)
+
+		using namespace UE::MorphTargetVertexCodec;
+
+		if (!Result.HasMorphs())
+		{
+			return;
+		}
+
+		auto GetBatchHeader = [](TArrayView<uint32> HeadersDataView, int32 BatchIndex) -> FDeltaBatchHeader
+		{
+			TArrayView<uint32> HeaderDataView = TArrayView<uint32>(
+						HeadersDataView.GetData() + BatchIndex*NumBatchHeaderDwords,
+						NumBatchHeaderDwords);
+
+			FDeltaBatchHeader BatchHeader;
+			ReadHeader(BatchHeader, HeaderDataView);
+
+			return BatchHeader;
+		};
+
+		auto StoreBatchHeader = [](TArrayView<uint32> HeadersDataView, int32 BatchIndex, const FDeltaBatchHeader& BatchHeader) -> void
+		{
+			TArrayView<uint32> HeaderDataView = TArrayView<uint32>(
+						HeadersDataView.GetData() + BatchIndex*NumBatchHeaderDwords,
+						NumBatchHeaderDwords);
+
+			WriteHeader(BatchHeader, HeaderDataView);
+		};
+
+		int32 NumMorphs = Result.Morph.Names.Num();
+		for (int32 MorphIndex = 0; MorphIndex < NumMorphs; ++MorphIndex)
+		{
+			auto ProcessBatchData =
+			[
+				VertexIndexRemap
+			](
+				FDeltaBatchHeader& InOutBatchHeader,
+				TArrayView<uint32> DataView,
+				int32 BatchVertexRangeBegin,
+				int32 BatchVertexRangeEnd
+			)
+			{
+				TArray<FQuantizedDelta, TInlineAllocator<UE::MorphTargetVertexCodec::BatchSize*2>> QuantizedDeltasStorage;
+
+				// Only check the vertices beforehand if the range is within a small range, otherwise 
+				// read the quantized deltas and only check the vertices in the batch.
+				int32 NumRemovedVerticesInRange = -1;
+				constexpr int32 MaxNumVerticesToCheck = 256;
+				if (BatchVertexRangeEnd - BatchVertexRangeBegin < MaxNumVerticesToCheck)
+				{
+					NumRemovedVerticesInRange = 0;
+					for (int32 Index = BatchVertexRangeBegin; Index < BatchVertexRangeEnd; ++Index)
+					{
+						NumRemovedVerticesInRange += static_cast<int32>(VertexIndexRemap[Index] < 0);
+					}
+				}
+
+				if (NumRemovedVerticesInRange == 0)
+				{
+					InOutBatchHeader.IndexMin = VertexIndexRemap[InOutBatchHeader.IndexMin];
+				}
+				else if (NumRemovedVerticesInRange == BatchVertexRangeEnd - BatchVertexRangeBegin)
+				{
+					InOutBatchHeader.IndexMin    = 0;
+					InOutBatchHeader.NumElements = 0;
+				}
+				else
+				{
+					TArrayView<uint32> BatchDataView = TArrayView<uint32>(
+							DataView.GetData() + InOutBatchHeader.DataOffset/sizeof(uint32), 
+							CalculateBatchDwords(InOutBatchHeader)); 
+
+					TArrayView<FQuantizedDelta> BatchQuantizedDeltas(
+							QuantizedDeltasStorage.GetData(), InOutBatchHeader.NumElements);
+
+					ReadQuantizedDeltas(BatchQuantizedDeltas, InOutBatchHeader, BatchDataView); 
+					
+					const int32 NumDeltas = InOutBatchHeader.NumElements;
+		
+					int32 NumDeltaVerticesRemovedInRange = 0;
+					for (int32 DeltaIndex = 0; DeltaIndex < NumDeltas; ++DeltaIndex)
+					{
+						NumDeltaVerticesRemovedInRange += static_cast<int32>(VertexIndexRemap[BatchQuantizedDeltas[DeltaIndex].Index] < 0);
+					}
+
+					if (NumDeltaVerticesRemovedInRange == 0)
+					{
+						InOutBatchHeader.IndexMin = (uint32)VertexIndexRemap[InOutBatchHeader.IndexMin];
+					}
+					else if (NumDeltaVerticesRemovedInRange == BatchVertexRangeEnd - BatchVertexRangeBegin)
+					{
+						InOutBatchHeader.IndexMin	= 0;
+						InOutBatchHeader.NumElements = 0;
+					}
+					else
+					{
+						check(InOutBatchHeader.NumElements <= UE::MorphTargetVertexCodec::BatchSize);
+						TArrayView<FQuantizedDelta> KeptQuantizedDeltas(
+								QuantizedDeltasStorage.GetData() + UE::MorphTargetVertexCodec::BatchSize, 
+								InOutBatchHeader.NumElements);
+
+						int32 DeltaIndex = 0;
+						for ( ; DeltaIndex < NumDeltas; ++DeltaIndex)
+						{
+							int32 RemappedDeltaVertexIdx = VertexIndexRemap[BatchQuantizedDeltas[DeltaIndex].Index];
+							if (RemappedDeltaVertexIdx >= 0)
+							{
+								InOutBatchHeader.IndexMin = (uint32)RemappedDeltaVertexIdx;
+								break;
+							}
+						}
+
+						uint32 DestDeltaIndex = 0;
+						for ( ; DeltaIndex < NumDeltas; ++DeltaIndex)
+						{
+							int32 RemappedDeltaVertexIdx = VertexIndexRemap[BatchQuantizedDeltas[DeltaIndex].Index];
+							if (RemappedDeltaVertexIdx >= 0)
+							{
+								KeptQuantizedDeltas[DestDeltaIndex++] = FQuantizedDelta
+								{
+									.Position = BatchQuantizedDeltas[DeltaIndex].Position,
+									.TangentZ = BatchQuantizedDeltas[DeltaIndex].TangentZ,
+									.Index    = (uint32)RemappedDeltaVertexIdx,
+								};	
+							}
+						}
+
+						InOutBatchHeader.NumElements = DestDeltaIndex;
+						WriteQuantizedDeltas(KeptQuantizedDeltas, InOutBatchHeader, BatchDataView);
+					}
+				}
+			};
+
+			const int32 NumBatches = Result.Morph.BatchesPerMorph[MorphIndex];
+			
+			if (NumBatches == 0)
+			{
+				continue;
+			}
+
+			int32 BatchStartOffset = Result.Morph.BatchStartOffsetPerMorph[MorphIndex];
+
+			TArrayView<uint32> DataView = MakeArrayView(Result.MorphDataBuffer); 
+
+			FDeltaBatchHeader BatchHeaderStorage[2];
+
+			FDeltaBatchHeader* CurrBatchHeader = &BatchHeaderStorage[0];
+			FDeltaBatchHeader* NextBatchHeader = &BatchHeaderStorage[1];
+			
+			uint32* MorphDataPtr = DataView.GetData();
+
+			ReadHeader(*CurrBatchHeader, TArrayView<uint32>(MorphDataPtr + (BatchStartOffset + 0) * NumBatchHeaderDwords, NumBatchHeaderDwords));
+
+			for (int32 BatchIndex = 0; BatchIndex < NumBatches - 1; ++BatchIndex)
+			{
+				ReadHeader(*NextBatchHeader, TArrayView<uint32>(MorphDataPtr + (BatchStartOffset + BatchIndex + 1) * NumBatchHeaderDwords, NumBatchHeaderDwords));
+				
+				ProcessBatchData(*CurrBatchHeader, DataView, CurrBatchHeader->IndexMin, NextBatchHeader->IndexMin);
+				
+				WriteHeader(*CurrBatchHeader, TArrayView<uint32>(MorphDataPtr + (BatchStartOffset + BatchIndex) * NumBatchHeaderDwords, NumBatchHeaderDwords));
+				::Swap(CurrBatchHeader, NextBatchHeader);
+			}
+
+			const int32 NumVertices = VertexIndexRemap.Num();
+			ProcessBatchData(*CurrBatchHeader, DataView, CurrBatchHeader->IndexMin, NumVertices);
+			WriteHeader(*CurrBatchHeader, TArrayView<uint32>(MorphDataPtr + (BatchStartOffset + NumBatches - 1) * NumBatchHeaderDwords, NumBatchHeaderDwords));
+		}
+	}
+
     void MeshRemoveRecreateSurface(FMesh* Result, const TBitArray<>& UsedVertices, const TBitArray<>& UsedFaces)
     {
         MUTABLE_CPUPROFILER_SCOPE(MeshRemoveRecreateSurface);
@@ -435,6 +609,9 @@ namespace mu
 		Result->GetIndexBuffers().SetElementCount(LastFreeFaceIndex*3);
 
 		MeshRemoveRecreateSurface(Result, UsedVertices, UsedFaces);
+
+		MeshRemoveMeshMorphs(*Result, UsedVerticesMap);
+
     }
 
 	void MeshRemoveMaskInline(FMesh* Mesh, const FMesh* Mask, bool bRemoveIfAllVerticesCulled)
diff --git a/Engine/Plugins/Mutable/Source/MutableRuntime/Public/MuR/Mesh.h b/Engine/Plugins/Mutable/Source/MutableRuntime/Public/MuR/Mesh.h
--- a/Engine/Plugins/Mutable/Source/MutableRuntime/Public/MuR/Mesh.h
+++ b/Engine/Plugins/Mutable/Source/MutableRuntime/Public/MuR/Mesh.h
@@ -93,6 +93,21 @@ namespace mu
 	};
 
 
+	struct FMeshMorph 
+	{
+		TArray<FName> Names;
+		
+		TArray<FVector4f> MaximumValuePerMorph;
+		TArray<FVector4f> MinimumValuePerMorph;
+		TArray<uint32> BatchStartOffsetPerMorph;
+		TArray<uint32> BatchesPerMorph;
+		
+		uint32 NumTotalBatches = 0;
+		float PositionPrecision = 0.0f; // TODO MORPHS Due to merging, this must be per block
+		float TangentZPrecision = 0.0f; // TODO MORPHS Due to merging, this must be per block
+	};
+
+	
 	/** Fill an array with the indices of all triangles belonging to the same UV island as InFirstTriangle.
 	*/
 	MUTABLERUNTIME_API void GetUVIsland(TArray<FTriangleInfo>& InTriangles,
@@ -168,6 +183,7 @@ namespace mu
 		WithSkeletonIDs = 1 << 14,
 		WithAdditionalPhysics = 1 << 15,
 		WithStreamedResources = 1 << 16,
+		WithMorphData = 1 << 17,
 
 		AllFlags = 0xFFFFFFFF
 	};
@@ -410,6 +426,7 @@ namespace mu
 		 */
         UE_API uint32 GetId() const;
 
+    	UE_API bool HasMorphs() const;
     
 	public:
 
@@ -443,12 +460,16 @@ namespace mu
 		FMeshBufferSet VertexBuffers;
 
 		FMeshBufferSet IndexBuffers;
+    	
+    	TArray<uint32> MorphDataBuffer;
 
 		/** Additional buffers used for temporary or custom data in different algorithms. */
 		TArray<TPair<EMeshBufferType, FMeshBufferSet>> AdditionalBuffers;
 
 		TArray<FMeshSurface> Surfaces;
 
+    	FMeshMorph Morph;
+    	
 		/** Externally provided SkeletonIDs of the skeletons required by this mesh. */
 		TArray<uint32> SkeletonIDs;
 
@@ -458,7 +479,7 @@ namespace mu
 		 */
 		TSharedPtr<const FSkeleton> Skeleton;
 		TSharedPtr<const FPhysicsBody> PhysicsBody;
-
+    	
 		/** Additional physics bodies referenced by the mesh that don't merge. */
 		TArray<TSharedPtr<const FPhysicsBody>> AdditionalPhysicsBodies;
 
diff --git a/Engine/Plugins/Mutable/Source/MutableTools/Private/MuT/ASTOpMeshMorph.cpp b/Engine/Plugins/Mutable/Source/MutableTools/Private/MuT/ASTOpMeshMorph.cpp
--- a/Engine/Plugins/Mutable/Source/MutableTools/Private/MuT/ASTOpMeshMorph.cpp
+++ b/Engine/Plugins/Mutable/Source/MutableTools/Private/MuT/ASTOpMeshMorph.cpp
@@ -43,6 +43,7 @@ namespace mu
 	mu::Ptr<ASTOp> ASTOpMeshMorph::Clone(MapChildFuncRef mapChild) const
 	{
 		Ptr<ASTOpMeshMorph> n = new ASTOpMeshMorph();
+		n->Name = Name;
 		n->Factor = mapChild(Factor.child());
 		n->Base = mapChild(Base.child());
 		n->Target = mapChild(Target.child());
@@ -62,7 +63,8 @@ namespace mu
 	//---------------------------------------------------------------------------------------------
 	uint64 ASTOpMeshMorph::Hash() const
 	{
-		uint64 res = std::hash<ASTOp*>()(Factor.child().get());
+		uint64 res = GetTypeHash(Name);
+		hash_combine(res, Factor.child().get());
 		hash_combine(res, Base.child().get());
 		hash_combine(res, Target.child().get());
 		return res;
@@ -80,6 +82,9 @@ namespace mu
 			program.OpAddress.Add((uint32_t)program.ByteCode.Num());
 			AppendCode(program.ByteCode, EOpType::ME_MORPH);
 
+			OP::ADDRESS NameAt = program.AddConstant(Name.ToString());
+			AppendCode(program.ByteCode, NameAt);
+			
 			OP::ADDRESS FactorAt = Factor ? Factor->linkedAddress : 0;
 			AppendCode(program.ByteCode, FactorAt);
 
diff --git a/Engine/Plugins/Mutable/Source/MutableTools/Private/MuT/ASTOpMeshMorph.h b/Engine/Plugins/Mutable/Source/MutableTools/Private/MuT/ASTOpMeshMorph.h
--- a/Engine/Plugins/Mutable/Source/MutableTools/Private/MuT/ASTOpMeshMorph.h
+++ b/Engine/Plugins/Mutable/Source/MutableTools/Private/MuT/ASTOpMeshMorph.h
@@ -18,7 +18,8 @@ namespace mu
 	class ASTOpMeshMorph final : public ASTOp
 	{
 	public:
-
+		FName Name;
+		
 		//! Factor selecting what morphs to apply and with what weight.
 		ASTChild Factor;
 
diff --git a/Engine/Plugins/Mutable/Source/MutableTools/Private/MuT/CodeGenerator.cpp b/Engine/Plugins/Mutable/Source/MutableTools/Private/MuT/CodeGenerator.cpp
--- a/Engine/Plugins/Mutable/Source/MutableTools/Private/MuT/CodeGenerator.cpp
+++ b/Engine/Plugins/Mutable/Source/MutableTools/Private/MuT/CodeGenerator.cpp
@@ -2158,7 +2158,7 @@ namespace mu
 		FMeshTask BaseTask,
 		FGuid BaseSurfaceGuid,
 		const void* ErrorContext,
-		const NodeMeshConstant* OriginalMeshNode)
+		const NodeMesh* OriginalMeshNode)
 	{
 		FMeshTask LastMeshTask = BaseTask;
 		FMeshTask PreModifiersTask = BaseTask;
@@ -2422,9 +2422,11 @@ namespace mu
 
 				// Request morphed skeletal mesh
 				TSharedPtr< TSharedPtr<FMesh> > ResolveMorphedMesh = MakeShared<TSharedPtr<FMesh>>();
+
 				UE::Tasks::FTask TargetMeshTask;
+				if (OriginalMeshNode->GetType() == NodeMeshConstant::GetStaticType())
 				{
-					FMesh* OriginalMesh = OriginalMeshNode->Value.Get();
+					FMesh* OriginalMesh = static_cast<const NodeMeshConstant*>(OriginalMeshNode)->Value.Get();
 					check(OriginalMesh->IsReference());
 					uint32 OriginalMeshID = OriginalMesh->GetReferencedMesh();
 
@@ -2435,6 +2437,13 @@ namespace mu
 						ResolveMorphedMesh, 
 						bRunImmediatlyIfPossible );
 				}
+				else
+				{
+					check(OriginalMeshNode->GetType() == NodeMeshParameter::GetStaticType())
+					TargetMeshTask = UE::Tasks::MakeCompletedTask<void>();
+				}
+
+				check(TargetMeshTask.IsValid())
 
 				// Factor
 				Ptr<ASTOp> FactorOp;
@@ -2454,59 +2463,77 @@ namespace mu
 					FactorOp = ChildResult.op;
 				}
 
-				Ptr<const NodeMeshConstant> OriginalMeshCopy = OriginalMeshNode;
+				Ptr<const NodeMesh> OriginalMeshCopy = OriginalMeshNode;
 				LastMeshTask = UE::Tasks::Launch(TEXT("MutableMeshMorphModifier"),
 					[TargetMeshTask, LastMeshTask, PreModifiersTask, Edit, Modifier, ResolveMorphedMesh, FactorOp, CompilerOptions=CompilerOptions, OriginalMeshCopy]() mutable
 					{
 						FMeshGenerationResult LastMeshResults = LastMeshTask.GetResult();
 						FMeshGenerationResult BaseResults = PreModifiersTask.GetResult();
 
-						TSharedPtr<FMesh> TargetMesh = *ResolveMorphedMesh;
-						if (!TargetMesh)
-						{
-							return LastMeshResults;
-						}
-
-						// Target mesh
-						Ptr<ASTOpConstantResource> TargetMeshOp = new ASTOpConstantResource;
-						TargetMeshOp->Type = EOpType::ME_CONSTANT;
-						TargetMeshOp->SetValue(TargetMesh->Clone(), CompilerOptions->OptimisationOptions.DiskCacheContext);
-						TargetMeshOp->SourceDataDescriptor = OriginalMeshCopy->SourceDataDescriptor;
-
-						// Morph generation through mesh diff
-						Ptr<ASTOpMeshDifference> DiffOp;
 						{
-							DiffOp = new ASTOpMeshDifference();
-							DiffOp->Base = BaseResults.BaseMeshOp;
-							DiffOp->Target = TargetMeshOp;
-
-							// Morphing tex coords here is not supported:
-							// Generating the homogoneous UVs is difficult since we don't have the base
-							// layout yet.                       
-							DiffOp->bIgnoreTextureCoords = true;
-						}
+							Ptr<ASTOpMeshMorph> MorphOp; 
+							if (OriginalMeshCopy->GetType() == NodeMeshConstant::GetStaticType())
+							{
+								// Target mesh
+								TSharedPtr<FMesh> TargetMesh = *ResolveMorphedMesh;
+								if (!TargetMesh)
+								{
+									return LastMeshResults;
+								}
 
-						// Morph operation
-						Ptr<ASTOpMeshMorph> MorphOp = new ASTOpMeshMorph();
-						{
-							MorphOp->Base = LastMeshResults.MeshOp;
-							MorphOp->Target = DiffOp;
-							MorphOp->Factor = FactorOp;
-						}
+								// Morph generation through mesh diff
+								Ptr<ASTOpConstantResource> TargetMeshOp = new ASTOpConstantResource;
+								TargetMeshOp->Type = EOpType::ME_CONSTANT;
+								TargetMeshOp->SetValue(TargetMesh->Clone(), CompilerOptions->OptimisationOptions.DiskCacheContext);
+								TargetMeshOp->SourceDataDescriptor = static_cast<const NodeMeshConstant*>(OriginalMeshCopy.get())->SourceDataDescriptor;
+				
+								Ptr<ASTOpMeshDifference> DiffOp;
+								{
+									DiffOp = new ASTOpMeshDifference();
+									DiffOp->Base = BaseResults.BaseMeshOp;
+									DiffOp->Target = TargetMeshOp;
+
+									// Morphing tex coords here is not supported:
+									// Generating the homogoneous UVs is difficult since we don't have the base
+									// layout yet.                       
+									DiffOp->bIgnoreTextureCoords = true;
+								}
+								
+								MorphOp = new ASTOpMeshMorph();
+								// Morph operation
+								{
+									MorphOp->Base = LastMeshResults.MeshOp;
+									MorphOp->Target = DiffOp;
+									MorphOp->Factor = FactorOp;
+									//MorphOp->Name = FName(*Edit->MeshMorph);
+								}
+							}
+							else
+							{
+								check(OriginalMeshCopy->GetType() == NodeMeshParameter::GetStaticType());
+
+								MorphOp = new ASTOpMeshMorph();
+								
+								MorphOp->Base = LastMeshResults.MeshOp;
+								MorphOp->Target = nullptr;
+								MorphOp->Factor = FactorOp;
+								MorphOp->Name = FName(*Edit->MeshMorph);
+							}
 
-						// Condition to apply the morph
-						if (Modifier.FinalCondition)
-						{
-							Ptr<ASTOpConditional> ConditionalOp = new ASTOpConditional();
-							ConditionalOp->type = EOpType::ME_CONDITIONAL;
-							ConditionalOp->no = LastMeshResults.MeshOp;
-							ConditionalOp->yes = MorphOp;
-							ConditionalOp->condition = Modifier.FinalCondition;
-							LastMeshResults.MeshOp = ConditionalOp;
-						}
-						else
-						{
-							LastMeshResults.MeshOp = MorphOp;
+							// Condition to apply the morph
+							if (Modifier.FinalCondition)
+							{
+								Ptr<ASTOpConditional> ConditionalOp = new ASTOpConditional();
+								ConditionalOp->type = EOpType::ME_CONDITIONAL;
+								ConditionalOp->no = LastMeshResults.MeshOp;
+								ConditionalOp->yes = MorphOp;
+								ConditionalOp->condition = Modifier.FinalCondition;
+								LastMeshResults.MeshOp = ConditionalOp;
+							}
+							else
+							{
+								LastMeshResults.MeshOp = MorphOp;
+							}
 						}
 
 						return LastMeshResults;
@@ -2526,7 +2553,7 @@ namespace mu
 				continue;
 			}
 
-			if (Modifier.Node->GetType()== NodeModifierMeshClipWithMesh::GetStaticType())
+			if (Modifier.Node->GetType() == NodeModifierMeshClipWithMesh::GetStaticType())
 			{
 				const NodeModifierMeshClipWithMesh* TypedClipNode = static_cast<const NodeModifierMeshClipWithMesh*>(Modifier.Node);
 
@@ -2586,9 +2613,11 @@ namespace mu
 				check(OriginalMeshNode);
 
 				TSharedPtr< TSharedPtr<FMesh> > ResolveOriginalMesh = MakeShared<TSharedPtr<FMesh>>();
+
 				UE::Tasks::FTask TargetMeshTask;
+				if (OriginalMeshNode->GetType() == NodeMeshConstant::GetStaticType())
 				{
-					FMesh* OriginalMesh = OriginalMeshNode->Value.Get();
+					FMesh* OriginalMesh = static_cast<const NodeMeshConstant*>(OriginalMeshNode)->Value.Get();
 					check(OriginalMesh->IsReference());
 					uint32 OriginalMeshID = OriginalMesh->GetReferencedMesh();
 
@@ -2600,6 +2629,13 @@ namespace mu
 						ResolveOriginalMesh,
 						bRunImmediatlyIfPossible);
 				}
+				else
+				{
+					check(OriginalMeshNode->GetType() == NodeMeshParameter::GetStaticType())
+					TargetMeshTask = UE::Tasks::MakeCompletedTask<void>();
+				}
+
+				check(TargetMeshTask.IsValid())
 
 				const NodeModifierMeshClipWithUVMask* TypedClipNode = static_cast<const NodeModifierMeshClipWithUVMask*>(Modifier.Node);
 
@@ -2624,63 +2660,72 @@ namespace mu
 					Layout = GenerateLayout(TypedClipNode->ClipLayout, 0);
 				}
 
-				Ptr<const NodeMeshConstant> OriginalMeshCopy = OriginalMeshNode;
+				Ptr<const NodeMesh> OriginalMeshCopy = OriginalMeshNode;
 				LastMeshTask = UE::Tasks::Launch(TEXT("MutableModifier"),
 					[Layout, MaskImage, ResolveOriginalMesh, LastMeshTask, PreModifiersTask, Modifier, TypedClipNode, CompilerOptions = CompilerOptions, ErrorLog = ErrorLog, ErrorContext, OriginalMeshCopy]() mutable
 					{
-						TSharedPtr<FMesh> TargetMesh = *ResolveOriginalMesh;
-						if (!TargetMesh)
-						{
-							return FMeshGenerationResult();
-						}
 						FMeshGenerationResult LastMeshResults = LastMeshTask.GetResult();
 						FMeshGenerationResult BaseResults = PreModifiersTask.GetResult();
+						
+						if (OriginalMeshCopy->GetType() == NodeMeshConstant::GetStaticType())
+						{
+							TSharedPtr<FMesh> TargetMesh = *ResolveOriginalMesh;
+							if (!TargetMesh)
+							{
+								return FMeshGenerationResult();
+							}
 
-						Ptr<ASTOpConstantResource> UVMeshOp = new ASTOpConstantResource;
-						UVMeshOp->Type = EOpType::ME_CONSTANT;
-						UVMeshOp->SetValue(TargetMesh->Clone(), CompilerOptions->OptimisationOptions.DiskCacheContext);
-						UVMeshOp->SourceDataDescriptor = OriginalMeshCopy->SourceDataDescriptor;
+							Ptr<ASTOpConstantResource> UVMeshOp = new ASTOpConstantResource;
+							UVMeshOp->Type = EOpType::ME_CONSTANT;
+							UVMeshOp->SetValue(TargetMesh->Clone(), CompilerOptions->OptimisationOptions.DiskCacheContext);
+							UVMeshOp->SourceDataDescriptor = static_cast<const NodeMeshConstant*>(OriginalMeshCopy.get())->SourceDataDescriptor;
 
-						Ptr<ASTOpMeshMaskClipUVMask> MeshMaskAt = new ASTOpMeshMaskClipUVMask();
-						MeshMaskAt->Source = BaseResults.BaseMeshOp;
-						MeshMaskAt->UVSource = UVMeshOp;
-						MeshMaskAt->LayoutIndex = TypedClipNode->LayoutIndex;
+							Ptr<ASTOpMeshMaskClipUVMask> MeshMaskAt = new ASTOpMeshMaskClipUVMask();
+							MeshMaskAt->Source = BaseResults.BaseMeshOp;
+							MeshMaskAt->UVSource = UVMeshOp;
+							MeshMaskAt->LayoutIndex = TypedClipNode->LayoutIndex;
 
-						if (TypedClipNode->ClipMask)
-						{
-							MeshMaskAt->MaskImage = MaskImage;
-							if (!MeshMaskAt->MaskImage)
+							if (TypedClipNode->ClipMask)
+							{
+								MeshMaskAt->MaskImage = MaskImage;
+								if (!MeshMaskAt->MaskImage)
+								{
+									ErrorLog->Add("Clip UV mask has not been generated", ELMT_ERROR, ErrorContext);
+									return LastMeshResults;
+								}
+							}
+
+							else if (TypedClipNode->ClipLayout)
 							{
-								ErrorLog->Add("Clip UV mask has not been generated", ELMT_ERROR, ErrorContext);
+								Ptr<ASTOpConstantResource> LayoutOp = new ASTOpConstantResource();
+								LayoutOp->Type = EOpType::LA_CONSTANT;
+								LayoutOp->SetValue(Layout, CompilerOptions->OptimisationOptions.DiskCacheContext);
+								MeshMaskAt->MaskLayout = LayoutOp;
+							}
+
+							else
+							{
+								// No mask or layout specified to clip. Don't clip anything.
 								return LastMeshResults;
 							}
-						}
 
-						else if (TypedClipNode->ClipLayout)
-						{
-							Ptr<ASTOpConstantResource> LayoutOp = new ASTOpConstantResource();
-							LayoutOp->Type = EOpType::LA_CONSTANT;
-							LayoutOp->SetValue(Layout, CompilerOptions->OptimisationOptions.DiskCacheContext);
-							MeshMaskAt->MaskLayout = LayoutOp;
-						}
+							if (MeshMaskAt)
+							{
+								Ptr<ASTOpMeshRemoveMask> RemoveOp = new ASTOpMeshRemoveMask();
+								RemoveOp->source = LastMeshResults.MeshOp;
+								RemoveOp->FaceCullStrategy = TypedClipNode->FaceCullStrategy;
+								RemoveOp->AddRemove(Modifier.FinalCondition, MeshMaskAt);
+
+								LastMeshResults.MeshOp = RemoveOp;
+							}
 
-						else
-						{
-							// No mask or layout specified to clip. Don't clip anything.
 							return LastMeshResults;
 						}
-
-						if (MeshMaskAt)
+						else
 						{
-							Ptr<ASTOpMeshRemoveMask> RemoveOp = new ASTOpMeshRemoveMask();
-							RemoveOp->source = LastMeshResults.MeshOp;
-							RemoveOp->FaceCullStrategy = TypedClipNode->FaceCullStrategy;
-							RemoveOp->AddRemove(Modifier.FinalCondition, MeshMaskAt);
-
-							LastMeshResults.MeshOp = RemoveOp;
+							ErrorLog->Add("NodeModifierMeshClipWithUVMask is not supported with mesh parameters, ignoring operation.", ELMT_WARNING, ErrorContext);
+							return LastMeshResults;
 						}
-
-						return LastMeshResults;
 					},
 					UE::Tasks::Prerequisites(TargetMeshTask, LastMeshTask)
 				);
@@ -2713,7 +2758,7 @@ namespace mu
 						// Morph to an ellipse
 						{
 							FShape MorphShape;
-							MorphShape.type = (uint8_t)FShape::Type::Ellipse;
+							MorphShape.type = (uint8)FShape::Type::Ellipse;
 							MorphShape.position = TypedNode->Parameters.Origin;
 							MorphShape.up = TypedNode->Parameters.Normal;
 							// TODO: Move rotation to ellipse rotation reference base instead of passing it directly
diff --git a/Engine/Plugins/Mutable/Source/MutableTools/Private/MuT/CodeGenerator.h b/Engine/Plugins/Mutable/Source/MutableTools/Private/MuT/CodeGenerator.h
--- a/Engine/Plugins/Mutable/Source/MutableTools/Private/MuT/CodeGenerator.h
+++ b/Engine/Plugins/Mutable/Source/MutableTools/Private/MuT/CodeGenerator.h
@@ -933,7 +933,7 @@ namespace mu
 			FMeshTask BaseTask,
 			FGuid SharedSurfaceId,
 			const void* ErrorContext,
-			const NodeMeshConstant* OriginalMeshNode);
+			const NodeMesh* OriginalMeshNode);
 
 		Ptr<ASTOp> ApplyImageBlockModifiers(
 			const TArray<FirstPassGenerator::FModifier>&, 
diff --git a/Engine/Plugins/Mutable/Source/MutableTools/Private/MuT/CodeGenerator_Mesh.cpp b/Engine/Plugins/Mutable/Source/MutableTools/Private/MuT/CodeGenerator_Mesh.cpp
--- a/Engine/Plugins/Mutable/Source/MutableTools/Private/MuT/CodeGenerator_Mesh.cpp
+++ b/Engine/Plugins/Mutable/Source/MutableTools/Private/MuT/CodeGenerator_Mesh.cpp
@@ -135,7 +135,7 @@ namespace mu
 			return UE::Tasks::MakeCompletedTask<FMeshGenerationResult>();
         }        
 
-		if (!Node->Morph)
+		if (Node->Name.IsNone() && !Node->Morph)
 		{
 			return BaseTask;
 		}
@@ -170,6 +170,7 @@ namespace mu
 				FMeshGenerationResult TargetResult = TargetTask.GetResult();
 
 				Ptr<ASTOpMeshMorph> OpMorph = new ASTOpMeshMorph();
+				OpMorph->Name = Node->Name;
 				OpMorph->Factor = FactorOp;
 				OpMorph->Base = BaseTask.GetResult().MeshOp;
 				OpMorph->Target = TargetTask.GetResult().MeshOp;
@@ -1387,7 +1388,15 @@ namespace mu
 			InOptions
 		);
 
-		return MeshParamTask;
+		// Apply the modifier for the pre-normal operations stage.
+		TArray<FirstPassGenerator::FModifier> Modifiers;
+		constexpr bool bModifiersForBeforeOperations = true;
+		this->GetModifiersFor(StaticOptions.ComponentId, StaticOptions.ActiveTags, bModifiersForBeforeOperations, Modifiers);
+
+		// This task does the necessary processing of the mesh when it is available
+		FMeshTask ApplyModifiersTask = ApplyMeshModifiers(Modifiers, StaticOptions, InOptions, MeshParamTask, FGuid(), Node->GetMessageContext(), Node);
+
+		return ApplyModifiersTask;
 	}
 
 }
diff --git a/Engine/Plugins/Mutable/Source/MutableTools/Public/MuT/NodeMeshMorph.h b/Engine/Plugins/Mutable/Source/MutableTools/Public/MuT/NodeMeshMorph.h
--- a/Engine/Plugins/Mutable/Source/MutableTools/Public/MuT/NodeMeshMorph.h
+++ b/Engine/Plugins/Mutable/Source/MutableTools/Public/MuT/NodeMeshMorph.h
@@ -19,6 +19,7 @@ namespace mu
 	{
 	public:
 
+		FName Name;
 		Ptr<NodeScalar> Factor;
 		Ptr<NodeMesh> Base;
 		Ptr<NodeMesh> Morph;
diff --git a/Engine/Source/Runtime/Engine/Private/Rendering/MorphTargetVertexCodec.cpp b/Engine/Source/Runtime/Engine/Private/Rendering/MorphTargetVertexCodec.cpp
--- a/Engine/Source/Runtime/Engine/Private/Rendering/MorphTargetVertexCodec.cpp
+++ b/Engine/Source/Runtime/Engine/Private/Rendering/MorphTargetVertexCodec.cpp
@@ -11,8 +11,19 @@
 namespace UE::MorphTargetVertexCodec
 {
 
-static constexpr uint32 BatchSizeBits = 6;
-static constexpr uint32 BatchSize = 1 << BatchSizeBits;
+constexpr uint32 IndexMaxBits = 31u;
+constexpr uint32 PositionMaxBits = 28u;				// Probably more than we need, but let's just allow it to go this high to be safe for now.
+// For larger deltas this can even be more precision than what was in the float input data!
+// Maybe consider float-like or exponential encoding of large values?
+constexpr float PositionMinValue = -134217728.0f;	// -2^(MaxBits-1)
+constexpr float PositionMaxValue = 134217720.0f;	// Largest float smaller than 2^(MaxBits-1)-1
+// Using 134217727.0f would NOT work as it would be rounded up to 134217728.0f, which is
+// outside the range.
+
+constexpr uint32 TangentZMaxBits = 16u;
+constexpr float TangentZMinValue = -32768.0f;		// -2^(MaxBits-1)
+constexpr float TangentZMaxValue = 32767.0f;		//  2^(MaxBits-1)-1
+
 
 class FDwordBitWriter
 {
@@ -52,6 +63,49 @@ private:
 	int32 			NumPendingBits;
 };
 
+class FDwordViewBitWriter
+{
+public:
+	FDwordViewBitWriter(TArrayView<uint32> Buffer) :
+		Buffer(Buffer),
+		PendingBits(0ull),
+		NumPendingBits(0),
+		NumElements(0)
+	{
+	}
+
+	void PutBits(uint32 Bits, uint32 NumBits)
+	{
+		check(static_cast<uint64>(Bits) < (1ull << NumBits));
+		PendingBits |= static_cast<uint64>(Bits) << NumPendingBits;
+		NumPendingBits += NumBits;
+
+		while (NumPendingBits >= 32)
+		{
+			Buffer[NumElements++] = static_cast<uint32>(PendingBits);
+			PendingBits >>= 32;
+			NumPendingBits -= 32;
+		}
+	}
+
+	void Flush()
+	{
+		if (NumPendingBits > 0)
+		{
+			Buffer[NumElements++] = (static_cast<uint32>(PendingBits));
+		}
+		PendingBits = 0;
+		NumPendingBits = 0;
+	}
+
+private:
+	TArrayView<uint32> Buffer;
+	uint64 			PendingBits;
+	int32 			NumPendingBits;
+	int32           NumElements;
+};
+
+
 class FDwordBitReader
 {
 public:
@@ -106,7 +160,6 @@ private:
 	uint32 Offset = 0;
 };
 
-
 void Encode(
 	TConstArrayView<FMorphTargetDelta> InMorphDeltas,
 	const TBitArray<>* InVertexNeedsTangents,
@@ -131,25 +184,6 @@ void Encode(
 	const float RcpPositionPrecision = 1.0f / InPositionPrecision;
 	const float RcpTangentZPrecision = 1.0f / InTangentZPrecision;
 
-	constexpr uint32 IndexMaxBits = 31u;
-	constexpr uint32 PositionMaxBits = 28u;				// Probably more than we need, but let's just allow it to go this high to be safe for now.
-	// For larger deltas this can even be more precision than what was in the float input data!
-	// Maybe consider float-like or exponential encoding of large values?
-	constexpr float PositionMinValue = -134217728.0f;	// -2^(MaxBits-1)
-	constexpr float PositionMaxValue = 134217720.0f;	// Largest float smaller than 2^(MaxBits-1)-1
-	// Using 134217727.0f would NOT work as it would be rounded up to 134217728.0f, which is
-	// outside the range.
-
-	constexpr uint32 TangentZMaxBits = 16u;
-	constexpr float TangentZMinValue = -32768.0f;		// -2^(MaxBits-1)
-	constexpr float TangentZMaxValue = 32767.0f;		//  2^(MaxBits-1)-1
-
-	struct FQuantizedDelta
-	{
-		FIntVector	Position;
-		FIntVector	TangentZ;
-		uint32		Index;
-	};
 	TArray<FQuantizedDelta> QuantizedDeltas;
 	QuantizedDeltas.Reserve(InMorphDeltas.Num());
 
@@ -285,7 +319,6 @@ void Encode(
 	}
 }
 
-
 static FMorphTargetDelta DecodeMorphTargetDelta(
 	const FDeltaBatchHeader& InHeader,
 	FDwordBitReader& InReader,
@@ -397,4 +430,171 @@ bool IterativeDecode(
 	return true;	
 }
 
+
+void WriteHeader(const FDeltaBatchHeader& InBatchHeader, TArrayView<uint32> OutData)
+{
+	FDwordViewBitWriter BitWriter(OutData);
+				
+	BitWriter.PutBits(InBatchHeader.DataOffset,    32);
+	BitWriter.PutBits(InBatchHeader.IndexBits,      5);
+	BitWriter.PutBits(InBatchHeader.PositionBits.X, 5);
+	BitWriter.PutBits(InBatchHeader.PositionBits.Y, 5);
+	BitWriter.PutBits(InBatchHeader.PositionBits.Z, 5);
+	BitWriter.PutBits(InBatchHeader.bTangents,      1);
+	BitWriter.PutBits(InBatchHeader.NumElements,   11);
+	BitWriter.PutBits(InBatchHeader.IndexMin,      32);
+	BitWriter.PutBits(InBatchHeader.PositionMin.X, 32);
+	BitWriter.PutBits(InBatchHeader.PositionMin.Y, 32);
+	BitWriter.PutBits(InBatchHeader.PositionMin.Z, 32);
+	BitWriter.PutBits(InBatchHeader.TangentZBits.X, 5);
+	BitWriter.PutBits(InBatchHeader.TangentZBits.Y, 5);
+	BitWriter.PutBits(InBatchHeader.TangentZBits.Z, 5);
+	BitWriter.PutBits(0,                           17); // Padding
+	BitWriter.PutBits(InBatchHeader.TangentZMin.X, 32);
+	BitWriter.PutBits(InBatchHeader.TangentZMin.Y, 32);
+	BitWriter.PutBits(InBatchHeader.TangentZMin.Z, 32);
+	
+	BitWriter.Flush();
+}
+
+void ReadHeader(FDeltaBatchHeader& OutBatchHeader, TConstArrayView<uint32> InData)
+{
+	FDwordBitReader BitReader(InData);
+	
+	OutBatchHeader.DataOffset     = BitReader.GetBits(32);
+	OutBatchHeader.IndexBits      = BitReader.GetBits(5);
+	OutBatchHeader.PositionBits.X = BitReader.GetBits(5);
+	OutBatchHeader.PositionBits.Y = BitReader.GetBits(5);
+	OutBatchHeader.PositionBits.Z = BitReader.GetBits(5);
+	OutBatchHeader.bTangents      = static_cast<bool>(BitReader.GetBits(1));
+	OutBatchHeader.NumElements    = BitReader.GetBits(11);
+	OutBatchHeader.IndexMin       = BitReader.GetBits(32);
+	OutBatchHeader.PositionMin.X  = BitReader.GetBits(32);
+	OutBatchHeader.PositionMin.Y  = BitReader.GetBits(32);
+	OutBatchHeader.PositionMin.Z  = BitReader.GetBits(32);
+	OutBatchHeader.TangentZBits.X = BitReader.GetBits(5);
+	OutBatchHeader.TangentZBits.Y = BitReader.GetBits(5);
+	OutBatchHeader.TangentZBits.Z = BitReader.GetBits(5);
+									BitReader.GetBits(17); // Padding
+	OutBatchHeader.TangentZMin.X  = BitReader.GetBits(32);
+	OutBatchHeader.TangentZMin.Y  = BitReader.GetBits(32);
+	OutBatchHeader.TangentZMin.Z  = BitReader.GetBits(32);
+}
+
+void WriteQuantizedDeltas(
+		TConstArrayView<FQuantizedDelta> InQuantizedDeltas, 
+		const FDeltaBatchHeader& InBatchHeader, 
+		TArrayView<uint32> OutData)
+{ 
+	FDwordViewBitWriter BitWriter(OutData);
+	
+	const int32 NumElements = static_cast<int32>(InBatchHeader.NumElements);
+
+	check(NumElements <= InQuantizedDeltas.Num());
+	for (int32 LocalElementIndex = 0; LocalElementIndex < NumElements; ++LocalElementIndex)
+	{
+		const FQuantizedDelta& Delta = InQuantizedDeltas[LocalElementIndex];
+		
+		BitWriter.PutBits(Delta.Index - LocalElementIndex - InBatchHeader.IndexMin, InBatchHeader.IndexBits);
+		BitWriter.PutBits(uint32(Delta.Position.X - InBatchHeader.PositionMin.X), InBatchHeader.PositionBits.X);
+		BitWriter.PutBits(uint32(Delta.Position.Y - InBatchHeader.PositionMin.Y), InBatchHeader.PositionBits.Y);
+		BitWriter.PutBits(uint32(Delta.Position.Z - InBatchHeader.PositionMin.Z), InBatchHeader.PositionBits.Z);
+
+		if (InBatchHeader.bTangents)
+		{
+			BitWriter.PutBits(uint32(Delta.TangentZ.X - InBatchHeader.TangentZMin.X), InBatchHeader.TangentZBits.X);
+			BitWriter.PutBits(uint32(Delta.TangentZ.Y - InBatchHeader.TangentZMin.Y), InBatchHeader.TangentZBits.Y);
+			BitWriter.PutBits(uint32(Delta.TangentZ.Z - InBatchHeader.TangentZMin.Z), InBatchHeader.TangentZBits.Z);
+		}
+	}
+
+	BitWriter.Flush();
+}
+
+void ReadQuantizedDeltas(
+		TArrayView<FQuantizedDelta> OutQuantizedDeltas, 
+		const FDeltaBatchHeader& InBatchHeader, 
+		TConstArrayView<uint32> InData)
+{
+	check(OutQuantizedDeltas.Num() >= (int32)InBatchHeader.NumElements)	
+
+	FDwordBitReader BitReader(InData);
+	
+	const int32 NumBatchElems = (int32)InBatchHeader.NumElements;
+	for (int32 LocalElementIndex = 0; LocalElementIndex < NumBatchElems; ++LocalElementIndex)
+	{
+		FQuantizedDelta& QuantizedDelta = OutQuantizedDeltas[LocalElementIndex];
+		
+		QuantizedDelta.Index = BitReader.GetBits(InBatchHeader.IndexBits) + LocalElementIndex + InBatchHeader.IndexMin;
+		QuantizedDelta.Position.X = BitReader.GetBits(InBatchHeader.PositionBits.X) + InBatchHeader.PositionMin.X;
+		QuantizedDelta.Position.Y = BitReader.GetBits(InBatchHeader.PositionBits.Y) + InBatchHeader.PositionMin.Y;
+		QuantizedDelta.Position.Z = BitReader.GetBits(InBatchHeader.PositionBits.Z) + InBatchHeader.PositionMin.Z;
+
+		if (InBatchHeader.bTangents)
+		{
+			QuantizedDelta.TangentZ.X = BitReader.GetBits(InBatchHeader.TangentZBits.X) + InBatchHeader.TangentZMin.X;
+			QuantizedDelta.TangentZ.Y = BitReader.GetBits(InBatchHeader.TangentZBits.Y) + InBatchHeader.TangentZMin.Y;
+			QuantizedDelta.TangentZ.Z = BitReader.GetBits(InBatchHeader.TangentZBits.Z) + InBatchHeader.TangentZMin.Z;
+		}
+
+		OutQuantizedDeltas[LocalElementIndex] = QuantizedDelta;
+	}
+}
+
+uint32 CalculateBatchDwords(const FDeltaBatchHeader& InBatchHeader)
+{
+	const uint32 ElementSize = InBatchHeader.IndexBits +
+		InBatchHeader.PositionBits.X + InBatchHeader.PositionBits.Y + InBatchHeader.PositionBits.Z +
+		(InBatchHeader.bTangents ? InBatchHeader.TangentZBits.X + InBatchHeader.TangentZBits.Y + InBatchHeader.TangentZBits.Z : 0);
+
+	return FMath::DivideAndRoundUp(ElementSize * InBatchHeader.NumElements, 32u);
+}
+
+void QuantizeDelta(
+	const FMorphTargetDelta& InDelta, 
+	const bool bInNeedsTangent, 
+	FQuantizedDelta& OutQuantizedDelta, 
+	const float InPositionPrecision, 
+	const float InTangentZPrecision)
+{
+	const float RcpPositionPrecision = 1.0f / InPositionPrecision;
+	const float RcpTangentZPrecision = 1.0f / InTangentZPrecision;
+
+	const FVector3f& PositionDelta = InDelta.PositionDelta;
+	OutQuantizedDelta.Position.X = FMath::RoundToInt(FMath::Clamp(PositionDelta.X * RcpPositionPrecision, PositionMinValue, PositionMaxValue));
+	OutQuantizedDelta.Position.Y = FMath::RoundToInt(FMath::Clamp(PositionDelta.Y * RcpPositionPrecision, PositionMinValue, PositionMaxValue));
+	OutQuantizedDelta.Position.Z = FMath::RoundToInt(FMath::Clamp(PositionDelta.Z * RcpPositionPrecision, PositionMinValue, PositionMaxValue));
+
+	const FVector3f TangentZDelta = bInNeedsTangent ? InDelta.TangentZDelta : FVector3f::ZeroVector;
+	OutQuantizedDelta.TangentZ.X = FMath::RoundToInt(FMath::Clamp(TangentZDelta.X * RcpTangentZPrecision, TangentZMinValue, TangentZMaxValue));
+	OutQuantizedDelta.TangentZ.Y = FMath::RoundToInt(FMath::Clamp(TangentZDelta.Y * RcpTangentZPrecision, TangentZMinValue, TangentZMaxValue));
+	OutQuantizedDelta.TangentZ.Z = FMath::RoundToInt(FMath::Clamp(TangentZDelta.Z * RcpTangentZPrecision, TangentZMinValue, TangentZMaxValue));
+
+	OutQuantizedDelta.Index = InDelta.SourceIdx;
+}
+
+void DequantizeDelta(
+	FMorphTargetDelta& OutDelta, 
+	const bool bInNeedsTangent, 
+	const FQuantizedDelta& InQuantizedDelta, 
+	const float InPositionPrecision, 
+	const float InTangentZPrecision)
+{
+	OutDelta.SourceIdx = InQuantizedDelta.Index;
+	OutDelta.PositionDelta.X = InQuantizedDelta.Position.X * InPositionPrecision;
+	OutDelta.PositionDelta.Y = InQuantizedDelta.Position.Y * InPositionPrecision;
+	OutDelta.PositionDelta.Z = InQuantizedDelta.Position.Z * InPositionPrecision;
+
+	if (bInNeedsTangent)
+	{
+		OutDelta.TangentZDelta.X = InQuantizedDelta.TangentZ.X * InTangentZPrecision;
+		OutDelta.TangentZDelta.Y = InQuantizedDelta.TangentZ.Y * InTangentZPrecision;
+		OutDelta.TangentZDelta.Z = InQuantizedDelta.TangentZ.Z * InTangentZPrecision;
+	}
+	else
+	{
+		OutDelta.TangentZDelta = FVector3f::ZeroVector;
+	}
+}
+
 } // namespace UE::MorphTargetVertexCodec
diff --git a/Engine/Source/Runtime/Engine/Private/Rendering/MorphTargetVertexInfoBuffers.cpp b/Engine/Source/Runtime/Engine/Private/Rendering/MorphTargetVertexInfoBuffers.cpp
--- a/Engine/Source/Runtime/Engine/Private/Rendering/MorphTargetVertexInfoBuffers.cpp
+++ b/Engine/Source/Runtime/Engine/Private/Rendering/MorphTargetVertexInfoBuffers.cpp
@@ -4,6 +4,7 @@
 #include "ProfilingDebugging/LoadTimeTracker.h"
 #include "Rendering/SkeletalMeshLODRenderData.h"
 #include "Animation/MorphTarget.h"
+#include "Containers/StaticArray.h"
 #include "DataDrivenShaderPlatformInfo.h"
 #include "RHIResourceUtils.h"
 
@@ -128,44 +129,6 @@ FArchive& operator<<(FArchive& Ar, FMorphTargetVertexInfoBuffers& MorphTargetVer
 	return Ar;
 }
 
-class FDwordBitWriter
-{
-public:
-	FDwordBitWriter(TArray<uint32>& Buffer) :
-		Buffer(Buffer),
-		PendingBits(0ull),
-		NumPendingBits(0)
-	{
-	}
-
-	void PutBits(uint32 Bits, uint32 NumBits)
-	{
-		check((uint64)Bits < (1ull << NumBits));
-		PendingBits |= (uint64)Bits << NumPendingBits;
-		NumPendingBits += NumBits;
-
-		while (NumPendingBits >= 32)
-		{
-			Buffer.Add((uint32)PendingBits);
-			PendingBits >>= 32;
-			NumPendingBits -= 32;
-		}
-	}
-
-	void Flush()
-	{
-		if (NumPendingBits > 0)
-			Buffer.Add((uint32)PendingBits);
-		PendingBits = 0;
-		NumPendingBits = 0;
-	}
-
-private:
-	TArray<uint32>& Buffer;
-	uint64 			PendingBits;
-	int32 			NumPendingBits;
-};
-
 void FMorphTargetVertexInfoBuffers::InitMorphResources(EShaderPlatform ShaderPlatform, const TArray<FSkelMeshRenderSection>& RenderSections, const TArray<UMorphTarget*>& MorphTargets, int32 NumVertices, int32 LODIndex, float TargetPositionErrorTolerance)
 {
 	check(!IsRHIInitialized());
@@ -217,15 +180,13 @@ void FMorphTargetVertexInfoBuffers::InitMorphResourcesStreaming(const TArray<FSk
 
 void FMorphTargetVertexInfoBuffers::Compress(const TArray<FSkelMeshRenderSection>& RenderSections, const TArray<const FMorphTargetLODModel*>& MorphTargets, const TBitArray<>& UsesBuiltinMorphTargetCompression, int32 NumVertices, float TargetPositionErrorTolerance)
 {
-	constexpr uint32 NumBatchHeaderDwords = 10u;
-	
 	using namespace UE::MorphTargetVertexCodec;
 	
 	PositionPrecision = ComputePositionPrecision(TargetPositionErrorTolerance);
 	TangentZPrecision = ComputeTangentPrecision();
 
 	MorphData.Empty();
-
+	
 	BatchStartOffsetPerMorph.Empty(MorphTargets.Num());
 	BatchesPerMorph.Empty(MorphTargets.Num());
 	MaximumValuePerMorph.Empty(MorphTargets.Num());
@@ -246,7 +207,7 @@ void FMorphTargetVertexInfoBuffers::Compress(const TArray<FSkelMeshRenderSection
 			}
 		}
 	}
-
+	
 	// Populate the arrays to be filled in later in the render thread
 	TArray<FDeltaBatchHeader> BatchHeaders;
 	TArray<uint32> BitstreamData;
@@ -275,10 +236,11 @@ void FMorphTargetVertexInfoBuffers::Compress(const TArray<FSkelMeshRenderSection
 		else
 		{
 			const FMorphTargetDelta* MorphDeltas = MorphModel->Vertices.GetData();
-
-			for (int32 DeltaIndex = 0; DeltaIndex < NumSrcDeltas; DeltaIndex++)
+			
+			for (int32 DeltaIndex = 0; DeltaIndex < NumSrcDeltas; ++DeltaIndex)
 			{
-				const auto& MorphDelta = MorphDeltas[DeltaIndex];
+				const FMorphTargetDelta& MorphDelta = MorphDeltas[DeltaIndex];
+			
 				const FVector3f TangentZDelta = (VertexNeedsTangents.IsValidIndex(MorphDelta.SourceIdx) && VertexNeedsTangents[MorphDelta.SourceIdx]) ? MorphDelta.TangentZDelta : FVector3f::ZeroVector;
 
 				// when import, we do check threshold, and also when adding weight, we do have threshold for how smaller weight can fit in
@@ -308,30 +270,19 @@ void FMorphTargetVertexInfoBuffers::Compress(const TArray<FSkelMeshRenderSection
 	NumTotalBatches = BatchHeaders.Num();
 
 	// Write packed batch headers
-	for (const FDeltaBatchHeader& BatchHeader : BatchHeaders)
+	for (FDeltaBatchHeader& BatchHeader : BatchHeaders)
 	{
-		const uint32 DataOffset = BatchHeader.DataOffset + BatchHeaders.Num() * NumBatchHeaderDwords * sizeof(uint32);
-
-		MorphData.Add(DataOffset);
-		MorphData.Add(BatchHeader.IndexBits |
-			(BatchHeader.PositionBits.X << 5) | (BatchHeader.PositionBits.Y << 10) | (BatchHeader.PositionBits.Z << 15) |
-			(BatchHeader.bTangents ? (1u << 20) : 0u) |
-			(BatchHeader.NumElements << 21));
-		MorphData.Add(BatchHeader.IndexMin);
-		MorphData.Add(BatchHeader.PositionMin.X);
-		MorphData.Add(BatchHeader.PositionMin.Y);
-		MorphData.Add(BatchHeader.PositionMin.Z);
-
-		MorphData.Add(BatchHeader.TangentZBits.X | (BatchHeader.TangentZBits.Y << 5) | (BatchHeader.TangentZBits.Z << 10));
-		MorphData.Add(BatchHeader.TangentZMin.X);
-		MorphData.Add(BatchHeader.TangentZMin.Y);
-		MorphData.Add(BatchHeader.TangentZMin.Z);
+		BatchHeader.DataOffset += BatchHeaders.Num() * NumBatchHeaderDwords * sizeof(uint32);
+
+		TStaticArray<uint32, NumBatchHeaderDwords> HeaderData;
+		WriteHeader(BatchHeader, HeaderData);
+		MorphData.Append(HeaderData);
 	}
 
 	// Append bitstream data
 	MorphData.Append(BitstreamData);
 
-	if(MorphData.Num() > 0)
+	if (MorphData.Num() > 0)
 	{
 		// Pad to make sure it is always safe to access the data with load4s.
 		MorphData.Add(0u);
diff --git a/Engine/Source/Runtime/Engine/Public/Rendering/MorphTargetVertexCodec.h b/Engine/Source/Runtime/Engine/Public/Rendering/MorphTargetVertexCodec.h
--- a/Engine/Source/Runtime/Engine/Public/Rendering/MorphTargetVertexCodec.h
+++ b/Engine/Source/Runtime/Engine/Public/Rendering/MorphTargetVertexCodec.h
@@ -14,6 +14,21 @@ class FDefaultBitArrayAllocator;
 namespace UE::MorphTargetVertexCodec
 {
 
+constexpr uint32 BatchSizeBits = 6u;
+constexpr uint32 BatchSize = 1u << BatchSizeBits;
+
+constexpr uint32 NumBatchHeaderDwords = 10u;
+
+
+struct FQuantizedDelta
+{
+	FIntVector	Position;
+	FIntVector	TangentZ;
+	uint32		Index;
+
+	bool operator==(const FQuantizedDelta&) const = default; 
+};
+
 struct FDeltaBatchHeader
 {
 	uint32		DataOffset;
@@ -28,6 +43,8 @@ struct FDeltaBatchHeader
 	FIntVector	PositionMin;
 	FIntVector	TangentZMin;
 
+	bool operator==(const FDeltaBatchHeader&) const = default; 
+
 	friend FArchive& operator<<(FArchive& Ar, FDeltaBatchHeader& InHeader)
 	{
 		Ar << InHeader.DataOffset;
@@ -83,15 +100,6 @@ ENGINE_API void Encode(
 	TArray<uint32>& OutCompressedVertices
 	);
 
-ENGINE_API void Encode(
-	TConstArrayView<FMorphTargetDelta> InMorphDeltas,
-	const TBitArray<>* InVertexNeedsTangents,
-	const float InPositionPrecision,
-	const float InTangentZPrecision,
-	TArray<FDeltaBatchHeader>& OutBatchHeaders,
-	TArray<uint32>& OutCompressedVertices
-	);
-
 /** Decode all the vertices given by the two arrays. */ 
 ENGINE_API void Decode(
 	TConstArrayView<FDeltaBatchHeader> InBatchHeaders,
@@ -112,4 +120,49 @@ ENGINE_API bool IterativeDecode(
 	FMorphTargetDelta& OutMorphTargetDelta
 	);
 
+/** Reads and decodes OutBatchHeader bits from InData. */
+ENGINE_API void ReadHeader(FDeltaBatchHeader& OutBatchHeader, const TArrayView<const uint32> InData);
+
+/** Encodes and writes InBatchHeader bits to OutData. */
+ENGINE_API void WriteHeader(const FDeltaBatchHeader& InBatchHeader, TArrayView<uint32> OutData);
+
+/** Encodes and writes the batch quantized data from InQuantizeDeletas to OutData. 
+ *  OutData must be sufficiently large to hold the number of Dwords specified by InBatchHeader. 
+ */
+ENGINE_API void WriteQuantizedDeltas(
+	TConstArrayView<FQuantizedDelta> InQuantizedDeltas, 
+	const FDeltaBatchHeader& InBatchHeader, 
+	TArrayView<uint32> OutData
+	);
+
+/** Reads and decodes the batch quantized data from InData to OutQuantizeDeletas.
+ *  InData must be sufficiently large to hold the number of Dwords specified by InBatchHeader. 
+ */
+ENGINE_API void ReadQuantizedDeltas(
+	TArrayView<FQuantizedDelta> OutQuantizedDeltas, 
+	const FDeltaBatchHeader& InBatchHeader, 
+	TConstArrayView<uint32> InData
+	);
+
+/** Calculates the number of Dwords the batch data has from BatchHeader. */
+ENGINE_API uint32 CalculateBatchDwords(const FDeltaBatchHeader& BatchHeader);
+
+/** Quantize FMorphTargetDelta using the given precision. */
+ENGINE_API void QuantizeDelta(
+	const FMorphTargetDelta& InDelta, 
+	const bool bInNeedsTangent, 
+	FQuantizedDelta& OutQuantizedDelta, 
+	const float InPositionPrecision, 
+	const float InTangentZPrecision
+	);
+
+/* Dequantize FMorphTargetDelta using the given precision. */
+ENGINE_API void DequantizeDelta(
+	FMorphTargetDelta& OutDelta, 
+	const bool bInNeedsTangent, 
+	const FQuantizedDelta& InQuantizedDelta, 
+	const float InPositionPrecision, 
+	const float InTangentZPrecision
+	);
+
 }
diff --git a/Engine/Source/Runtime/Engine/Public/Rendering/MorphTargetVertexInfoBuffers.h b/Engine/Source/Runtime/Engine/Public/Rendering/MorphTargetVertexInfoBuffers.h
--- a/Engine/Source/Runtime/Engine/Public/Rendering/MorphTargetVertexInfoBuffers.h
+++ b/Engine/Source/Runtime/Engine/Public/Rendering/MorphTargetVertexInfoBuffers.h
@@ -72,11 +72,16 @@ public:
 		return TangentZPrecision;
 	}
 
+	const uint32* GetData() const
+	{
+		return MorphData.GetData();
+	}
+	
 	uint64 GetMorphDataSizeInBytes() const 
 	{ 
 		return MorphData.Num() * sizeof(uint32);
 	}
-
+	
 	static ENGINE_API bool IsPlatformShaderSupported(EShaderPlatform ShaderPlatform);
 
 	FBufferRHIRef MorphDataBuffer;
@@ -89,11 +94,12 @@ protected:
 	ENGINE_API void Serialize(FArchive& Ar);
 
 	void Compress(const TArray<FSkelMeshRenderSection>& RenderSections, const TArray<const FMorphTargetLODModel*>& MorphTargets, const TBitArray<>& UsesBuiltinMorphTargetCompression, int32 NumVertices, float TargetPositionErrorTolerance);
-	
+
 	// Transient data. Gets deleted as soon as the GPU resource has been initialized (unless excplicitly disabled by bEmptyMorphCPUDataOnInitRHI).
 	TResourceArray<uint32> MorphData;
 
-	//x,y,y separate for position and shared w for tangent
+	//x,y,z separate for position and shared w for tangent
+
 	TArray<FVector4f> MaximumValuePerMorph;
 	TArray<FVector4f> MinimumValuePerMorph;
 	TArray<uint32> BatchStartOffsetPerMorph;
