[Mutable] Fix crash when merging two compressed morphs in ME_MERGE Ops if the second mesh has extra morphs.

#rnx

diff --git a/Engine/Plugins/Mutable/Source/MutableRuntime/Internal/MuR/OpMeshMerge.h b/Engine/Plugins/Mutable/Source/MutableRuntime/Internal/MuR/OpMeshMerge.h
--- a/Engine/Plugins/Mutable/Source/MutableRuntime/Internal/MuR/OpMeshMerge.h
+++ b/Engine/Plugins/Mutable/Source/MutableRuntime/Internal/MuR/OpMeshMerge.h
@@ -209,6 +209,8 @@ namespace mu
 					WriteHeader(BatchHeader, HeaderDataView);
 
 					const int32 NumBatchDwords = CalculateBatchDwords(BatchHeader);
+
+					check(Result.MorphDataBuffer.Num() >= BatchHeader.DataOffset / sizeof(uint32) + NumBatchDwords);
 					FMemory::Memcpy(
 							Result.MorphDataBuffer.GetData() + BatchHeader.DataOffset / sizeof(uint32),
 							MeshB.MorphDataBuffer.GetData() + SourceDataOffset / sizeof(uint32),
@@ -223,12 +225,12 @@ namespace mu
 		for (; MorphIndex < ResultNumMorphs; ++MorphIndex)
 		{
 			TArrayView<uint32> MorphHeadersDataView = TArrayView<uint32>(
-					Result.MorphDataBuffer.GetData() + Result.Morph.BatchStartOffsetPerMorph[MorphIndex],
-					Result.Morph.BatchesPerMorph[MorphIndex]);
+					Result.MorphDataBuffer.GetData() + Result.Morph.BatchStartOffsetPerMorph[MorphIndex]*NumBatchHeaderDwords,
+					Result.Morph.BatchesPerMorph[MorphIndex]*NumBatchHeaderDwords);
 
 			int32 BatchIndex = 0;
 			{
-				int32 NumBatches = MeshB.Morph.BatchesPerMorph[MorphIndex];
+				int32 NumBatches = Result.Morph.BatchesPerMorph[MorphIndex];
 				for (; BatchIndex < NumBatches; ++BatchIndex)
 				{
 					TArrayView<uint32> HeaderDataView = TArrayView<uint32>(
@@ -245,9 +247,11 @@ namespace mu
 					WriteHeader(BatchHeader, HeaderDataView);
 
 					const int32 NumBatchDwords = CalculateBatchDwords(BatchHeader);
+				
+					check(Result.MorphDataBuffer.Num() >= BatchHeader.DataOffset / sizeof(uint32) + NumBatchDwords);
 					FMemory::Memcpy(
 							Result.MorphDataBuffer.GetData() + BatchHeader.DataOffset / sizeof(uint32),
-							MeshA.MorphDataBuffer.GetData() + SourceDataOffset / sizeof(uint32),
+							MeshB.MorphDataBuffer.GetData() + SourceDataOffset / sizeof(uint32),
 							NumBatchDwords * sizeof(uint32));
 
 					BatchAccumulatedOffsetInDwords += NumBatchDwords;
