[mutable] Improved handling of the compilation AR search when running in the context of a commandlet.
- Made the mutable compilation AR search be SYNCHRONOUS if a non cooking commandlet is being run as, in that context, we have no guarantee that a ASYNC AR search can be completed as no engine/AR tick is guaranteed to happen.

#rnx

diff --git a/Engine/Plugins/Mutable/Source/CustomizableObjectEditor/Private/MuCOE/CustomizableObjectCompiler.cpp b/Engine/Plugins/Mutable/Source/CustomizableObjectEditor/Private/MuCOE/CustomizableObjectCompiler.cpp
--- a/Engine/Plugins/Mutable/Source/CustomizableObjectEditor/Private/MuCOE/CustomizableObjectCompiler.cpp
+++ b/Engine/Plugins/Mutable/Source/CustomizableObjectEditor/Private/MuCOE/CustomizableObjectCompiler.cpp
@@ -440,9 +440,14 @@ void FCustomizableObjectCompiler::EnqueueCompileRequest(const TSharedRef<FCompil
 		}
 		else
 		{
+			// The engine and AR system might not get ticked when running a non-cook commandlet.
+			// In that case, force the search to be synchronous to ensure the compilation can proceed.
+			const bool bShouldSearchAssetsSync = (IsRunningCommandlet() && !IsRunningCookCommandlet()) ? true : !CompileRequest->bAsync;
+			
+			UE_LOG(LogMutable, Display, TEXT("Performing full %s Asset Registry search as required by the CO compilation process."), bShouldSearchAssetsSync ? *FString(TEXT("SYNC")) : *FString(TEXT("ASYNC")));
+			
 			// Scan for all the assets async/sync based on the compilation options as the AR has not yet been scanned
-			UE_LOG(LogMutable, Display, TEXT("Performing full Asset Registry search as required by the CO compilation process."));
-			AssetRegistry.SearchAllAssets(!CompileRequest->bAsync);
+			AssetRegistry.SearchAllAssets(bShouldSearchAssetsSync);
 		}
 	}
 
