[mutable] Reduced overriden materials log verbosity
- The log is not relevant enought to grant itself the category of Log

#rnx

diff --git a/Engine/Plugins/Mutable/Source/CustomizableObject/Private/MuCO/CustomizableObjectSystem.cpp b/Engine/Plugins/Mutable/Source/CustomizableObject/Private/MuCO/CustomizableObjectSystem.cpp
--- a/Engine/Plugins/Mutable/Source/CustomizableObject/Private/MuCO/CustomizableObjectSystem.cpp
+++ b/Engine/Plugins/Mutable/Source/CustomizableObject/Private/MuCO/CustomizableObjectSystem.cpp
@@ -1500,7 +1500,7 @@ void UpdateSkeletalMesh(const TSharedRef<FUpdateContextPrivate>& Context)
 
 	if (ParentsWithOverrideMaterialsUpdatedCount > 0)
 	{
-		UE_LOG(LogMutable, Log, TEXT("A total of %u Skeletal Mesh Components got their override materials updated."), ParentsWithOverrideMaterialsUpdatedCount);
+		UE_LOG(LogMutable, Verbose, TEXT("A total of %u Skeletal Mesh Components got their override materials updated."), ParentsWithOverrideMaterialsUpdatedCount);
 	}
 }
 
