[mutable] Fixed access to nullptr FLogBenchmarkUtil.cpp
- The code expected the disk file to be created and ready when logging the image streaming data, but that may not be the case if we started the benchmark after the mesh got updated but before the image data was reported.

#rnx

diff --git a/Engine/Plugins/Mutable/Source/CustomizableObject/Internal/MuCO/LogBenchmarkUtil.h b/Engine/Plugins/Mutable/Source/CustomizableObject/Internal/MuCO/LogBenchmarkUtil.h
--- a/Engine/Plugins/Mutable/Source/CustomizableObject/Internal/MuCO/LogBenchmarkUtil.h
+++ b/Engine/Plugins/Mutable/Source/CustomizableObject/Internal/MuCO/LogBenchmarkUtil.h
@@ -138,7 +138,7 @@ public:
 	/** Gathers update stats when it has finished. */
 	void FinishUpdateMesh(const TSharedRef<FUpdateContextPrivate>& Context);
 
-	void FinishUpdateImage(const FString& CustomizableObjectPathName, const FString& InstancePathName, const FString& InstanceDescriptor, const bool bDidLevelBeginPlay, double TaskUpdateImageTime, const int64 TaskUpdateImageMemoryPeak, const int64 TaskUpdateImageRealMemoryPeak) const;
+	void FinishUpdateImage(const FString& CustomizableObjectPathName, const FString& InstancePathName, const FString& InstanceDescriptor, const bool bDidLevelBeginPlay, double TaskUpdateImageTime, const int64 TaskUpdateImageMemoryPeak, const int64 TaskUpdateImageRealMemoryPeak);
 	
 	DECLARE_BENCHMARK_STAT(NumAllocatedTextures, uint32);
 	DECLARE_BENCHMARK_STAT(TextureGPUSize, uint64);
diff --git a/Engine/Plugins/Mutable/Source/CustomizableObject/Private/MuCO/LogBenchmarkUtil.cpp b/Engine/Plugins/Mutable/Source/CustomizableObject/Private/MuCO/LogBenchmarkUtil.cpp
--- a/Engine/Plugins/Mutable/Source/CustomizableObject/Private/MuCO/LogBenchmarkUtil.cpp
+++ b/Engine/Plugins/Mutable/Source/CustomizableObject/Private/MuCO/LogBenchmarkUtil.cpp
@@ -373,7 +373,7 @@ void FLogBenchmarkUtil::FinishUpdateMesh(const TSharedRef<FUpdateContextPrivate>
 }
 
 
-void FLogBenchmarkUtil::FinishUpdateImage(const FString& CustomizableObjectPathName, const FString& InstancePathName, const FString& InstanceDescriptor, const bool bDidLevelBeginPlay, const double TaskUpdateImageTime, const int64 TaskUpdateImageMemoryPeak, const int64 TaskUpdateImageRealMemoryPeak) const
+void FLogBenchmarkUtil::FinishUpdateImage(const FString& CustomizableObjectPathName, const FString& InstancePathName, const FString& InstanceDescriptor, const bool bDidLevelBeginPlay, const double TaskUpdateImageTime, const int64 TaskUpdateImageMemoryPeak, const int64 TaskUpdateImageRealMemoryPeak)
 {
 	check(IsInGameThread());
 
@@ -382,6 +382,11 @@ void FLogBenchmarkUtil::FinishUpdateImage(const FString& CustomizableObjectPathN
 		return;
 	}
 
+	if (!Archive) // Create file if it does not exist.
+	{
+		Archive = CreateFile();
+	}
+	
 	FInstanceUpdateStats MipsUpdateData;
 	MipsUpdateData.UpdateType = TEXT("Image");
 	MipsUpdateData.CustomizableObjectPathName = CustomizableObjectPathName;
