[Mutable] Fixed function VisitNodes for macros. Now all comopents and participating object will be found even if they are inside nested mutable macros.

#rnx

diff --git a/Engine/Plugins/Mutable/Source/CustomizableObjectEditor/Private/MuCOE/GenerateMutableSource/GenerateMutableSourceComponent.cpp b/Engine/Plugins/Mutable/Source/CustomizableObjectEditor/Private/MuCOE/GenerateMutableSource/GenerateMutableSourceComponent.cpp
--- a/Engine/Plugins/Mutable/Source/CustomizableObjectEditor/Private/MuCOE/GenerateMutableSource/GenerateMutableSourceComponent.cpp
+++ b/Engine/Plugins/Mutable/Source/CustomizableObjectEditor/Private/MuCOE/GenerateMutableSource/GenerateMutableSourceComponent.cpp
@@ -5,6 +5,7 @@
 #include "Factories/FbxMeshImportData.h"
 #include "MuCOE/GenerateMutableSource/GenerateMutableSource.h"
 
+#include "MuCOE/CustomizableObjectMacroLibrary/CustomizableObjectMacroLibrary.h"
 #include "MuCOE/GenerateMutableSource/GenerateMutableSourceMesh.h"
 #include "MuCOE/GenerateMutableSource/GenerateMutableSourceSurface.h"
 #include "MuCOE/GenerateMutableSource/GenerateMutableSourceFloat.h"
@@ -484,7 +485,19 @@ void FirstPass(UCustomizableObjectNodeComponentMesh& Node, FMutableGraphGenerati
 		return Element.ComponentName == ComponentName;
 	}))
 	{
-		FText Msg = FText::Format(LOCTEXT("ComponentNodeWithSameNameExists", "Error! A Mesh Component node with the same name already exists in the Customizable Object [{0}]"), FText::FromString(GraphTraversal::GetObject(*Result->Node)->GetName()));
+		FText Msg;
+		if (Result->Node->IsInMacro())
+		{
+			const UCustomizableObjectNodeMacroInstance* CurrentMacro = GenerationContext.MacroNodesStack.Last();
+			check(CurrentMacro);
+			
+			Msg = FText::Format(LOCTEXT("ComponentNodeWithSameNameExists_Macro", "Error! A Mesh Component node with the same name already exists in a Macro Instance [{0}]"), FText::FromName(CurrentMacro->ParentMacro->Name));
+		}
+		else
+		{
+			Msg = FText::Format(LOCTEXT("ComponentNodeWithSameNameExists_CO", "Error! A Mesh Component node with the same name already exists in the Customizable Object [{0}]"), FText::FromString(GraphTraversal::GetObject(*Result->Node)->GetName()));
+		}
+		
 		GenerationContext.Log(Msg, &Node, EMessageSeverity::Error);
 		return;
 	}
diff --git a/Engine/Plugins/Mutable/Source/CustomizableObjectEditor/Private/MuCOE/GraphTraversal.cpp b/Engine/Plugins/Mutable/Source/CustomizableObjectEditor/Private/MuCOE/GraphTraversal.cpp
--- a/Engine/Plugins/Mutable/Source/CustomizableObjectEditor/Private/MuCOE/GraphTraversal.cpp
+++ b/Engine/Plugins/Mutable/Source/CustomizableObjectEditor/Private/MuCOE/GraphTraversal.cpp
@@ -404,16 +404,16 @@ const UCustomizableObject* GraphTraversal::GetRootObject(const UCustomizableObje
 
 
 void RecursiveVisitNodes(UCustomizableObjectNode& CurrentNode, const TFunction<void(UCustomizableObjectNode&)>& VisitFunction,
-	TArray<const UCustomizableObjectNodeMacroInstance*>& MacroContext, TSet<UCustomizableObjectNode*>& VisitedNodes, const TMultiMap<FGuid, UCustomizableObjectNodeObject*>& ObjectGroupMap)
+	TArray<const UCustomizableObjectNodeMacroInstance*>& MacroContext, TSet<TPair<UCustomizableObjectNode*, TArray<const UCustomizableObjectNodeMacroInstance*>>>& VisitedNodes, const TMultiMap<FGuid, UCustomizableObjectNodeObject*>& ObjectGroupMap)
 {
 	// Check if we already visited this node
-	if (VisitedNodes.Contains(&CurrentNode))
+	if (VisitedNodes.Contains({&CurrentNode, MacroContext}))
 	{
 		return;
 	}
 
 	// Add it to the cache of already visited nodes
-	VisitedNodes.Add(&CurrentNode);
+	VisitedNodes.Add({ &CurrentNode, MacroContext});
 
 	// Visit the current Node
 	VisitFunction(CurrentNode);
@@ -447,73 +447,19 @@ void RecursiveVisitNodes(UCustomizableObjectNode& CurrentNode, const TFunction<v
 
 			else if (UCustomizableObjectNodeMacroInstance* MacroInstanceNode = Cast<UCustomizableObjectNodeMacroInstance>(ConnectedNode))
 			{
-				// We visit this node here because we need to pass through it multiple times
-				if (!VisitedNodes.Contains(MacroInstanceNode))
-				{
-					VisitFunction(*MacroInstanceNode);
-					VisitedNodes.Add(MacroInstanceNode);
-				}
-
+				// First visit the graph of the Macro
 				if (const UEdGraphPin* OutputPin = MacroInstanceNode->GetMacroIOPin(ECOMacroIOType::COMVT_Output, ConnectedPin->PinName))
 				{
 					if (UCustomizableObjectNode* OutputNode = Cast<UCustomizableObjectNode>(OutputPin->GetOwningNode()))
 					{
-						// We visit this node here because we need to pass through it multiple times
-						if (!VisitedNodes.Contains(OutputNode))
-						{
-							VisitFunction(*OutputNode);
-							VisitedNodes.Add(OutputNode);
-						}
-					}
-
-					// Going to the node linked to the OutputPin
-					if (const UEdGraphPin* OutputLinkedPin = FollowInputPin(*OutputPin))
-					{
-						// We go directly to the node linked to the output node to ensure that is linked and we need to visit it.
-						if (UCustomizableObjectNode* OutputLinkedNode = Cast<UCustomizableObjectNode>(OutputLinkedPin->GetOwningNode()))
-						{
-							MacroContext.Push(MacroInstanceNode);
-							RecursiveVisitNodes(*OutputLinkedNode, VisitFunction, MacroContext, VisitedNodes, ObjectGroupMap);
-							MacroContext.Pop();
-						}
+						MacroContext.Push(MacroInstanceNode);
+						RecursiveVisitNodes(*OutputNode, VisitFunction, MacroContext, VisitedNodes, ObjectGroupMap);
+						MacroContext.Pop();
 					}
 				}
-			}
-
-			else if (UCustomizableObjectNodeTunnel* IOMacroNode = Cast<UCustomizableObjectNodeTunnel>(ConnectedNode))
-			{
-				// Output Nodes are already visited
-				if (IOMacroNode->bIsInputNode)
-				{
-					check(MacroContext.Num());
-
-					// We visit this node here because we need to pass through it multiple times
-					if (!VisitedNodes.Contains(IOMacroNode))
-					{
-						VisitFunction(*IOMacroNode);
-						VisitedNodes.Add(IOMacroNode);
-					}
-
-					// Going to the parent's graph where the current macro was instantiated
-					const UCustomizableObjectNodeMacroInstance* ParentMacroInstanceNode = MacroContext.Pop();
 
-					if (ParentMacroInstanceNode)
-					{
-						if (const UEdGraphPin* InputPin = ParentMacroInstanceNode->FindPin(ConnectedPin->PinName, EEdGraphPinDirection::EGPD_Input))
-						{
-							// Going to the node linked to the Macro Instance
-							if (const UEdGraphPin* InputLinkedPin = FollowInputPin(*InputPin))
-							{
-								if (UCustomizableObjectNode* InputLinkedNode = Cast<UCustomizableObjectNode>(InputLinkedPin->GetOwningNode()))
-								{
-									RecursiveVisitNodes(*InputLinkedNode, VisitFunction, MacroContext, VisitedNodes, ObjectGroupMap);
-								}
-							}
-						}
-					}
-
-					MacroContext.Push(ParentMacroInstanceNode);
-				}
+				// Continue visiting from the Macro
+				RecursiveVisitNodes(*MacroInstanceNode, VisitFunction, MacroContext, VisitedNodes, ObjectGroupMap);
 			}
 
 			else if (UCustomizableObjectNode* Node = Cast<UCustomizableObjectNode>(ConnectedNode))
@@ -527,7 +473,8 @@ void RecursiveVisitNodes(UCustomizableObjectNode& CurrentNode, const TFunction<v
 
 void GraphTraversal::VisitNodes(UCustomizableObjectNode& StartNode, const TFunction<void(UCustomizableObjectNode&)>& VisitFunction, const TMultiMap<FGuid, UCustomizableObjectNodeObject*>* ObjectGroupMap, TArray<const UCustomizableObjectNodeMacroInstance*>* MacroContext)
 {
-	TSet<UCustomizableObjectNode*> VisitedNodes;
+	// Set to keep tracking of the visited nodes. We use the macro context to re-visit macros that can be instantiated more than once.
+	TSet<TPair<UCustomizableObjectNode*, TArray<const UCustomizableObjectNodeMacroInstance*>>> VisitedNodes;
 	
 	TMultiMap<FGuid, UCustomizableObjectNodeObject*> LocalObjectGroupMap;
 	TArray<const UCustomizableObjectNodeMacroInstance*> LocalMacroContext;
