[Mutable] Remove mu::FInstance private.

diff --git a/Engine/Plugins/Mutable/Source/CustomizableObjectEditor/Private/MuCOE/SMutableInstanceViewer.cpp b/Engine/Plugins/Mutable/Source/CustomizableObjectEditor/Private/MuCOE/SMutableInstanceViewer.cpp
--- a/Engine/Plugins/Mutable/Source/CustomizableObjectEditor/Private/MuCOE/SMutableInstanceViewer.cpp
+++ b/Engine/Plugins/Mutable/Source/CustomizableObjectEditor/Private/MuCOE/SMutableInstanceViewer.cpp
@@ -389,14 +389,10 @@ TSharedRef<SWidget> SMutableInstanceViewer::GenerateDataTableSlates()
 {
 	// Formatting
 	constexpr int32 IndentationSpace = 16;
-	constexpr int32 SimpleSpacing = 1;
 
 	constexpr int32 AfterTitleSpacing = 4;
-	constexpr int32 EndOfSectionSpacing = 12;
 
 	// Naming
-	const FText GeneralDataTitle = LOCTEXT("GeneralDataTitle", "General Data");
-	const FText InstanceIdTitle = LOCTEXT("InstanceIdTitle", "Instance Id : ");
 	const FText ComponentsTitle = LOCTEXT("ComponentsTitle", "Components");
 	
 	return SNew(SScrollBox)
@@ -404,49 +400,8 @@ TSharedRef<SWidget> SMutableInstanceViewer::GenerateDataTableSlates()
 		[
 			SNew(SVerticalBox)
 
-			// General data ----------------------------------------------------------------
-			+ SVerticalBox::Slot()
-			.AutoHeight()
-			[
-				SNew(STextBlock).
-				Text(GeneralDataTitle)
-			]
-
-			+ SVerticalBox::Slot().
-			  Padding(IndentationSpace, AfterTitleSpacing).
-			  AutoHeight()
-			[
-				SNew(SVerticalBox)
-
-				// Instace ID
-				+ SVerticalBox::Slot().
-				  Padding(0, SimpleSpacing).
-				  AutoHeight()
-				[
-					SNew(SHorizontalBox)
-
-					// Vertices title
-					+ SHorizontalBox::Slot()
-					.AutoWidth()
-					[
-						SNew(STextBlock).
-						Text(InstanceIdTitle)
-					]
-
-					// Vertices value
-					+ SHorizontalBox::Slot()
-					.AutoWidth()
-					[
-						SNew(STextBlock).
-						Text(this, &SMutableInstanceViewer::GetInstanceID)
-					]
-
-				]
-			]
-
 			// Components Data --------------------------------------------------------------
 			+ SVerticalBox::Slot()
-			  .Padding(0, EndOfSectionSpacing)
 			  .AutoHeight()
 			[
 				SNew(SVerticalBox)
@@ -592,12 +547,6 @@ TSharedRef<SWidget> SMutableInstanceViewer::GenerateSurfaceListView(const TShare
 }
 
 
-FText SMutableInstanceViewer::GetInstanceID() const
-{
-	return FText::AsNumber(MutableInstance ? MutableInstance->GetId() : 0);
-}
-
-
 TSharedRef<ITableRow> SMutableInstanceViewer::OnGenerateComponentRow(TSharedPtr<FMutableInstanceViewerComponentElement> In, const TSharedRef<STableViewBase>& OwnerTable)
 {
 	return SNew(SMutableInstanceComponentListRow, OwnerTable, In, SharedThis(this) );
diff --git a/Engine/Plugins/Mutable/Source/CustomizableObjectEditor/Private/MuCOE/SMutableInstanceViewer.h b/Engine/Plugins/Mutable/Source/CustomizableObjectEditor/Private/MuCOE/SMutableInstanceViewer.h
--- a/Engine/Plugins/Mutable/Source/CustomizableObjectEditor/Private/MuCOE/SMutableInstanceViewer.h
+++ b/Engine/Plugins/Mutable/Source/CustomizableObjectEditor/Private/MuCOE/SMutableInstanceViewer.h
@@ -90,12 +90,7 @@ private:
 
 	/** Generates the tables showing the buffer data on the Instance alongside with the bone tree found on the mutable Instance */
 	TSharedRef<SWidget> GenerateDataTableSlates();
-
-private:
 	
-	/** */
-	FText GetInstanceID() const;
-
 	/** */	
 	TArray<TSharedPtr<FMutableInstanceViewerComponentElement>> Components;
 	
diff --git a/Engine/Plugins/Mutable/Source/MutableRuntime/Private/MuR/CodeRunner.cpp b/Engine/Plugins/Mutable/Source/MutableRuntime/Private/MuR/CodeRunner.cpp
--- a/Engine/Plugins/Mutable/Source/MutableRuntime/Private/MuR/CodeRunner.cpp
+++ b/Engine/Plugins/Mutable/Source/MutableRuntime/Private/MuR/CodeRunner.cpp
@@ -13,7 +13,6 @@
 #include "Misc/AssertionMacros.h"
 #include "MuR/ImagePrivate.h"
 #include "MuR/Instance.h"
-#include "MuR/InstancePrivate.h"
 #include "MuR/Mesh.h"
 #include "MuR/MeshBufferSet.h"
 #include "MuR/Model.h"
@@ -548,7 +547,7 @@ namespace mu
                     check(  nameAd < (uint32)Program.ConstantStrings.Num() );
                     const FString& Name = Program.ConstantStrings[ nameAd ];
 
-                    Result->GetPrivate()->AddVector( 0, 0, 0, value, FName(Name) );
+                    Result->AddVector( 0, 0, 0, value, FName(Name) );
                 }
                 StoreInstance( Item, Result );
                 break;
@@ -593,7 +592,7 @@ namespace mu
                     check(  nameAd < (uint32)Program.ConstantStrings.Num() );
                     const FString& Name = Program.ConstantStrings[ nameAd ];
 
-                    Result->GetPrivate()->AddScalar( 0, 0, 0, value, FName(Name));
+                    Result->AddScalar( 0, 0, 0, value, FName(Name));
                 }
                 StoreInstance( Item, Result );
                 break;
@@ -638,7 +637,7 @@ namespace mu
                     check( nameAd < (uint32)Program.ConstantStrings.Num() );
                     const FString& Name = Program.ConstantStrings[nameAd];
 
-                    Result->GetPrivate()->AddString( 0, 0, 0, value->GetValue(), FName(Name) );
+                    Result->AddString( 0, 0, 0, value->GetValue(), FName(Name) );
                 }
                 StoreInstance( Item, Result );
                 break;
@@ -679,14 +678,14 @@ namespace mu
                 {
                     TSharedPtr<const FInstance> pComp = LoadInstance( FCacheAddress(Args.value,Item) );
 
-                    int32 NewComponentIndex = Result->GetPrivate()->AddComponent();
+                    int32 NewComponentIndex = Result->AddComponent();
 
-                    if ( !pComp->GetPrivate()->Components.IsEmpty() )
+                    if ( !pComp->Components.IsEmpty() )
                     {
-                        Result->GetPrivate()->Components[NewComponentIndex] = pComp->GetPrivate()->Components[0];
+                        Result->Components[NewComponentIndex] = pComp->Components[0];
 
 						// Id
-                    	Result->GetPrivate()->Components[NewComponentIndex].Id = Args.ExternalId;
+                    	Result->Components[NewComponentIndex].Id = Args.ExternalId;
                     }
                 }
                 StoreInstance( Item, Result );
@@ -731,31 +730,31 @@ namespace mu
                 {
 					TSharedPtr<const FInstance> pSurf = LoadInstance( FCacheAddress(Args.value,Item) );
 
-                    int32 sindex = Result->GetPrivate()->AddSurface( 0, 0 );
+                    int32 sindex = Result->AddSurface( 0, 0 );
 
                     // Surface data
                     if (pSurf
                             &&
-                            pSurf->GetPrivate()->Components.Num()
+                            pSurf->Components.Num()
                             &&
-                            pSurf->GetPrivate()->Components[0].LODs.Num()
+                            pSurf->Components[0].LODs.Num()
                             &&
-                            pSurf->GetPrivate()->Components[0].LODs[0].Surfaces.Num())
+                            pSurf->Components[0].LODs[0].Surfaces.Num())
                     {
-                        Result->GetPrivate()->Components[0].LODs[0].Surfaces[sindex] =
-                            pSurf->GetPrivate()->Components[0].LODs[0].Surfaces[0];
+                        Result->Components[0].LODs[0].Surfaces[sindex] =
+                            pSurf->Components[0].LODs[0].Surfaces[0];
                     }
 
                     // Name
                     OP::ADDRESS nameAd = Args.name;
                     check( nameAd < (uint32)Program.ConstantStrings.Num() );
                     const FString& Name = Program.ConstantStrings[ nameAd ];
-                    Result->GetPrivate()->SetSurfaceName( 0, 0, sindex, FName(Name) );
+                    Result->SetSurfaceName( 0, 0, sindex, FName(Name) );
 
                     // IDs
-                    Result->GetPrivate()->Components[0].LODs[0].Surfaces[sindex].InternalId = Args.id;
-                    Result->GetPrivate()->Components[0].LODs[0].Surfaces[sindex].ExternalId = Args.ExternalId;
-                    Result->GetPrivate()->Components[0].LODs[0].Surfaces[sindex].SharedId = Args.SharedSurfaceId;
+                    Result->Components[0].LODs[0].Surfaces[sindex].InternalId = Args.id;
+                    Result->Components[0].LODs[0].Surfaces[sindex].ExternalId = Args.ExternalId;
+                    Result->Components[0].LODs[0].Surfaces[sindex].SharedId = Args.SharedSurfaceId;
                 }
                 StoreInstance( Item, Result );
                 break;
@@ -807,7 +806,7 @@ namespace mu
             {
                 // Assemble result
 				TSharedPtr<FInstance> Result = MakeShared<FInstance>();
-				int32 ComponentIndex = Result->GetPrivate()->AddComponent();
+				int32 ComponentIndex = Result->AddComponent();
 
                 for (uint8 LODIndex = 0; LODIndex < LODCount; ++LODIndex)
                 {
@@ -820,18 +819,18 @@ namespace mu
 						bool bIsSelectedLod = ( (1<<LODIndex) & lodMask ) != 0;
 
 						// Add an empty LOD even if not selected.
-						int32 InstanceLODIndex = Result->GetPrivate()->AddLOD(ComponentIndex);
+						int32 InstanceLODIndex = Result->AddLOD(ComponentIndex);
 						
 						if (bIsSelectedLod)
                         {
 							TSharedPtr<const FInstance> pLOD = LoadInstance( FCacheAddress(LODAddress,Item) );
 
                             // In a degenerated case, the returned pLOD may not have an LOD inside
- 							if (!pLOD->GetPrivate()->Components.IsEmpty()
+ 							if (!pLOD->Components.IsEmpty()
 								&&
-								!pLOD->GetPrivate()->Components[0].LODs.IsEmpty())
+								!pLOD->Components[0].LODs.IsEmpty())
 							{
-								Result->GetPrivate()->Components[ComponentIndex].LODs[InstanceLODIndex] = pLOD->GetPrivate()->Components[0].LODs[0];
+								Result->Components[ComponentIndex].LODs[InstanceLODIndex] = pLOD->Components[0].LODs[0];
 							}
 						}
                     }
@@ -879,7 +878,7 @@ namespace mu
 						check(NameAddress < (uint32)Program.ConstantStrings.Num());
 						const FString& NameString = Program.ConstantStrings[NameAddress];
 
-						Result->GetPrivate()->AddExtensionData(ExtensionData, FName(NameString) );
+						Result->AddExtensionData(ExtensionData, FName(NameString) );
 					}
 
 					StoreInstance(Item, Result);
@@ -922,7 +921,7 @@ namespace mu
 					Result = mu::CloneOrTakeOver<FInstance>(Base);
 				}
 
-				Result->GetPrivate()->AddOverlayMaterial(0, LoadScalar(FCacheAddress(Args.OverlayMaterialId, Item)));
+				Result->AddOverlayMaterial(0, LoadScalar(FCacheAddress(Args.OverlayMaterialId, Item)));
 
 				StoreInstance(Item, Result);
 				break;
@@ -986,7 +985,7 @@ namespace mu
 					OP::ADDRESS NameAd = Args.name;
 					check(NameAd < (uint32)Program.ConstantStrings.Num());
 					const FString& Name = Program.ConstantStrings[NameAd];
-					Result->GetPrivate()->SetMesh(0, 0, MeshId, FName(Name));
+					Result->SetMesh(0, 0, MeshId, FName(Name));
                 }
                 StoreInstance( Item, Result );
                 break;
@@ -1027,7 +1026,7 @@ namespace mu
 					OP::ADDRESS NameAd = Args.name;
 					check(NameAd < (uint32)Program.ConstantStrings.Num());
 					const FString& Name = Program.ConstantStrings[NameAd];
-					Result->GetPrivate()->AddImage(0, 0, 0, ImageId, FName(Name) );
+					Result->AddImage(0, 0, 0, ImageId, FName(Name) );
                 }
                 StoreInstance( Item, Result );
                 break;
diff --git a/Engine/Plugins/Mutable/Source/MutableRuntime/Private/MuR/Instance.cpp b/Engine/Plugins/Mutable/Source/MutableRuntime/Private/MuR/Instance.cpp
--- a/Engine/Plugins/Mutable/Source/MutableRuntime/Private/MuR/Instance.cpp
+++ b/Engine/Plugins/Mutable/Source/MutableRuntime/Private/MuR/Instance.cpp
@@ -5,90 +5,49 @@
 
 #include "HAL/LowLevelMemTracker.h"
 #include "Misc/AssertionMacros.h"
-#include "MuR/InstancePrivate.h"
 #include "MuR/MutableMath.h"
 
 namespace mu
 {
-
-	//---------------------------------------------------------------------------------------------
-	//---------------------------------------------------------------------------------------------
-	//---------------------------------------------------------------------------------------------
-	FInstance::FInstance()
-	{
-		LLM_SCOPE_BYNAME(TEXT("MutableRuntime"));
-		m_pD = new Private();
-	}
-
-
-	//---------------------------------------------------------------------------------------------
-	FInstance::~FInstance()
-	{
-		LLM_SCOPE_BYNAME(TEXT("MutableRuntime"));
-        check( m_pD );
-		delete m_pD;
-		m_pD = nullptr;
-	}
-
-
-	//---------------------------------------------------------------------------------------------
-	FInstance::Private* FInstance::GetPrivate() const
-	{
-		return m_pD;
-	}
-
-
-    //---------------------------------------------------------------------------------------------
     TSharedPtr<FInstance> FInstance::Clone() const
     {
 		LLM_SCOPE_BYNAME(TEXT("MutableRuntime"));
 		TSharedPtr<FInstance> Result = MakeShared<FInstance>();
 
-        *Result->GetPrivate() = *m_pD;
+        *Result = *this;
 
         return Result;
     }
-
-
-	//---------------------------------------------------------------------------------------------
+	
+	
 	int32 FInstance::GetDataSize() const
 	{
-		return 16 + sizeof(Private) + m_pD->Components.GetAllocatedSize() + m_pD->ExtensionData.GetAllocatedSize();
+		return 16 + sizeof(*this) + Components.GetAllocatedSize() + ExtensionData.GetAllocatedSize();
 	}
-
-
-    //---------------------------------------------------------------------------------------------
-    FInstance::FID FInstance::GetId() const
-    {
-        return m_pD->Id;
-    }
-
-
-    //---------------------------------------------------------------------------------------------
+	
+    
     int32 FInstance::GetComponentCount() const
     {
-		return m_pD->Components.Num();
+		return Components.Num();
     }
-
-
-	//---------------------------------------------------------------------------------------------
+	
+	
 	int32 FInstance::GetLODCount( int32 ComponentIndex ) const
 	{
-		if (m_pD->Components.IsValidIndex(ComponentIndex))
+		if (Components.IsValidIndex(ComponentIndex))
 		{
-			return m_pD->Components[ComponentIndex].LODs.Num();
+			return Components[ComponentIndex].LODs.Num();
 		}
 		check(false);
 		return 0;
 	}
-
 	
-	//---------------------------------------------------------------------------------------------
+	
 	uint16 FInstance::GetComponentId( int32 ComponentIndex ) const
 	{
-		if (m_pD->Components.IsValidIndex(ComponentIndex))
+		if (Components.IsValidIndex(ComponentIndex))
 		{
-			return m_pD->Components[ComponentIndex].Id;
+			return Components[ComponentIndex].Id;
 		}
 		else
 		{
@@ -98,14 +57,13 @@ namespace mu
 		return 0;
 	}
 
-
-    //---------------------------------------------------------------------------------------------
+	
     int32 FInstance::GetSurfaceCount( int32 ComponentIndex, int32 LODIndex ) const
     {
-		if (m_pD->Components.IsValidIndex(ComponentIndex) &&
-			m_pD->Components[ComponentIndex].LODs.IsValidIndex(LODIndex))
+		if (Components.IsValidIndex(ComponentIndex) &&
+			Components[ComponentIndex].LODs.IsValidIndex(LODIndex))
 		{
-			return m_pD->Components[ComponentIndex].LODs[LODIndex].Surfaces.Num();
+			return Components[ComponentIndex].LODs[LODIndex].Surfaces.Num();
 		}
 		else
 		{
@@ -114,16 +72,15 @@ namespace mu
 
 		return 0;
 	}
-
-
-    //---------------------------------------------------------------------------------------------
+	
+    
     uint32 FInstance::GetSurfaceId(int32 ComponentIndex, int32 LODIndex, int32 SurfaceIndex ) const
     {
-        if (m_pD->Components.IsValidIndex(ComponentIndex) &&
-			m_pD->Components[ComponentIndex].LODs.IsValidIndex(LODIndex) &&
-			m_pD->Components[ComponentIndex].LODs[LODIndex].Surfaces.IsValidIndex(SurfaceIndex) )
+        if (Components.IsValidIndex(ComponentIndex) &&
+			Components[ComponentIndex].LODs.IsValidIndex(LODIndex) &&
+			Components[ComponentIndex].LODs[LODIndex].Surfaces.IsValidIndex(SurfaceIndex) )
         {
-            return m_pD->Components[ComponentIndex].LODs[LODIndex].Surfaces[SurfaceIndex].InternalId;
+            return Components[ComponentIndex].LODs[LODIndex].Surfaces[SurfaceIndex].InternalId;
         }
 		else
 		{
@@ -133,16 +90,15 @@ namespace mu
         return 0;
     }
 
-
-    //---------------------------------------------------------------------------------------------
+    
     int32 FInstance::FindSurfaceById(int32 ComponentIndex, int32 LODIndex, uint32 id ) const
     {
-		if (m_pD->Components.IsValidIndex(ComponentIndex) &&
-			m_pD->Components[ComponentIndex].LODs.IsValidIndex(LODIndex))
+		if (Components.IsValidIndex(ComponentIndex) &&
+			Components[ComponentIndex].LODs.IsValidIndex(LODIndex))
 		{
-			for (int32 i = 0; i < m_pD->Components[ComponentIndex].LODs[LODIndex].Surfaces.Num(); ++i)
+			for (int32 i = 0; i < Components[ComponentIndex].LODs[LODIndex].Surfaces.Num(); ++i)
 			{
-				if (m_pD->Components[ComponentIndex].LODs[LODIndex].Surfaces[i].InternalId == id)
+				if (Components[ComponentIndex].LODs[LODIndex].Surfaces[i].InternalId == id)
 				{
 					return i;
 				}
@@ -155,16 +111,15 @@ namespace mu
 
         return -1;
     }
-
 	
-	//---------------------------------------------------------------------------------------------
+	
 	void FInstance::FindBaseSurfaceBySharedId(int32 CompIndex, uint32 SharedId, int32& OutSurfaceIndex, int32& OutLODIndex) const
 	{
-		if (m_pD->Components.IsValidIndex(CompIndex))
+		if (Components.IsValidIndex(CompIndex))
 		{
-			for (int32 LodIndex = 0; LodIndex < m_pD->Components[CompIndex].LODs.Num(); LodIndex++)
+			for (int32 LodIndex = 0; LodIndex < Components[CompIndex].LODs.Num(); LodIndex++)
 			{
-				FInstanceLOD& LOD = m_pD->Components[CompIndex].LODs[LodIndex];
+				const FInstanceLOD& LOD = Components[CompIndex].LODs[LodIndex];
 				for (int32 SurfaceIndex = 0; SurfaceIndex < LOD.Surfaces.Num(); ++SurfaceIndex)
 				{
 					if (LOD.Surfaces[SurfaceIndex].SharedId == SharedId)
@@ -182,15 +137,14 @@ namespace mu
 		OutLODIndex = INDEX_NONE;
 	}
 
-
-	//---------------------------------------------------------------------------------------------
+	
 	uint32 FInstance::GetSharedSurfaceId(int32 ComponentIndex, int32 LODIndex, int32 SurfaceIndex) const
 	{
-		if (m_pD->Components.IsValidIndex(ComponentIndex) &&
-			m_pD->Components[ComponentIndex].LODs.IsValidIndex(LODIndex) &&
-			m_pD->Components[ComponentIndex].LODs[LODIndex].Surfaces.IsValidIndex(SurfaceIndex))
+		if (Components.IsValidIndex(ComponentIndex) &&
+			Components[ComponentIndex].LODs.IsValidIndex(LODIndex) &&
+			Components[ComponentIndex].LODs[LODIndex].Surfaces.IsValidIndex(SurfaceIndex))
 		{
-			return m_pD->Components[ComponentIndex].LODs[LODIndex].Surfaces[SurfaceIndex].SharedId;
+			return Components[ComponentIndex].LODs[LODIndex].Surfaces[SurfaceIndex].SharedId;
 		}
 		else
 		{
@@ -200,15 +154,14 @@ namespace mu
 		return 0;
 	}
 	
-
-    //---------------------------------------------------------------------------------------------
+    
     uint32 FInstance::GetSurfaceCustomId(int32 ComponentIndex, int32 LODIndex, int32 SurfaceIndex ) const
     {
-        if (m_pD->Components.IsValidIndex(ComponentIndex) &&
-			m_pD->Components[ComponentIndex].LODs.IsValidIndex(LODIndex) &&
-			m_pD->Components[ComponentIndex].LODs[LODIndex].Surfaces.IsValidIndex(SurfaceIndex))
+        if (Components.IsValidIndex(ComponentIndex) &&
+			Components[ComponentIndex].LODs.IsValidIndex(LODIndex) &&
+			Components[ComponentIndex].LODs[LODIndex].Surfaces.IsValidIndex(SurfaceIndex))
         {
-            return m_pD->Components[ComponentIndex].LODs[LODIndex].Surfaces[SurfaceIndex].ExternalId;
+            return Components[ComponentIndex].LODs[LODIndex].Surfaces[SurfaceIndex].ExternalId;
         }
 		else
 		{
@@ -217,202 +170,187 @@ namespace mu
 
         return 0;
     }
-
-
-	//---------------------------------------------------------------------------------------------
+	
+	
     int32 FInstance::GetImageCount(int32 ComponentIndex, int32 LODIndex, int32 SurfaceIndex ) const
 	{
-		check(m_pD->Components.IsValidIndex(ComponentIndex));
-		check(m_pD->Components[ComponentIndex].LODs.IsValidIndex(LODIndex));
-		check(m_pD->Components[ComponentIndex].LODs[LODIndex].Surfaces.IsValidIndex(SurfaceIndex));
+		check(Components.IsValidIndex(ComponentIndex));
+		check(Components[ComponentIndex].LODs.IsValidIndex(LODIndex));
+		check(Components[ComponentIndex].LODs[LODIndex].Surfaces.IsValidIndex(SurfaceIndex));
 
-		return m_pD->Components[ComponentIndex].LODs[LODIndex].Surfaces[SurfaceIndex].Images.Num();
+		return Components[ComponentIndex].LODs[LODIndex].Surfaces[SurfaceIndex].Images.Num();
 	}
 
-
-	//---------------------------------------------------------------------------------------------
+	
     int32 FInstance::GetVectorCount(int32 ComponentIndex, int32 LODIndex, int32 SurfaceIndex ) const
 	{
-		check(m_pD->Components.IsValidIndex(ComponentIndex));
-		check(m_pD->Components[ComponentIndex].LODs.IsValidIndex(LODIndex));
-		check(m_pD->Components[ComponentIndex].LODs[LODIndex].Surfaces.IsValidIndex(SurfaceIndex));
+		check(Components.IsValidIndex(ComponentIndex));
+		check(Components[ComponentIndex].LODs.IsValidIndex(LODIndex));
+		check(Components[ComponentIndex].LODs[LODIndex].Surfaces.IsValidIndex(SurfaceIndex));
 
-		return m_pD->Components[ComponentIndex].LODs[LODIndex].Surfaces[SurfaceIndex].Vectors.Num();
+		return Components[ComponentIndex].LODs[LODIndex].Surfaces[SurfaceIndex].Vectors.Num();
 	}
 
-
-	//---------------------------------------------------------------------------------------------
+	
     int32 FInstance::GetScalarCount(int32 ComponentIndex, int32 LODIndex, int32 SurfaceIndex ) const
 	{
-		check(m_pD->Components.IsValidIndex(ComponentIndex));
-		check(m_pD->Components[ComponentIndex].LODs.IsValidIndex(LODIndex));
-		check(m_pD->Components[ComponentIndex].LODs[LODIndex].Surfaces.IsValidIndex(SurfaceIndex));
+		check(Components.IsValidIndex(ComponentIndex));
+		check(Components[ComponentIndex].LODs.IsValidIndex(LODIndex));
+		check(Components[ComponentIndex].LODs[LODIndex].Surfaces.IsValidIndex(SurfaceIndex));
 
-		return m_pD->Components[ComponentIndex].LODs[LODIndex].Surfaces[SurfaceIndex].Scalars.Num();
+		return Components[ComponentIndex].LODs[LODIndex].Surfaces[SurfaceIndex].Scalars.Num();
 	}
 
-
-    //---------------------------------------------------------------------------------------------
+    
     int32 FInstance::GetStringCount(int32 ComponentIndex, int32 LODIndex, int32 SurfaceIndex ) const
     {
-		check(m_pD->Components.IsValidIndex(ComponentIndex));
-		check(m_pD->Components[ComponentIndex].LODs.IsValidIndex(LODIndex));
-		check(m_pD->Components[ComponentIndex].LODs[LODIndex].Surfaces.IsValidIndex(SurfaceIndex));
+		check(Components.IsValidIndex(ComponentIndex));
+		check(Components[ComponentIndex].LODs.IsValidIndex(LODIndex));
+		check(Components[ComponentIndex].LODs[LODIndex].Surfaces.IsValidIndex(SurfaceIndex));
 
-		return m_pD->Components[ComponentIndex].LODs[LODIndex].Surfaces[SurfaceIndex].Strings.Num();
+		return Components[ComponentIndex].LODs[LODIndex].Surfaces[SurfaceIndex].Strings.Num();
 	}
-
-
-    //---------------------------------------------------------------------------------------------
+	
+    
 	FResourceID FInstance::GetMeshId(int32 ComponentIndex, int32 LODIndex ) const
     {
-        check(m_pD->Components.IsValidIndex(ComponentIndex));
-        check(m_pD->Components[ComponentIndex].LODs.IsValidIndex(LODIndex));
+        check(Components.IsValidIndex(ComponentIndex));
+        check(Components[ComponentIndex].LODs.IsValidIndex(LODIndex));
 
-		return m_pD->Components[ComponentIndex].LODs[LODIndex].MeshId;
+		return Components[ComponentIndex].LODs[LODIndex].MeshId;
     }
 
-
-	//---------------------------------------------------------------------------------------------
+	
 	FResourceID FInstance::GetImageId(int32 ComponentIndex, int32 LODIndex, int32 SurfaceIndex, int32 ImageIndex) const
 	{
-		check(m_pD->Components.IsValidIndex(ComponentIndex));
-		check(m_pD->Components[ComponentIndex].LODs.IsValidIndex(LODIndex));
-        check(m_pD->Components[ComponentIndex].LODs[LODIndex].Surfaces.IsValidIndex(SurfaceIndex));
-        check(m_pD->Components[ComponentIndex].LODs[LODIndex].Surfaces[SurfaceIndex].Images.IsValidIndex(ImageIndex));
+		check(Components.IsValidIndex(ComponentIndex));
+		check(Components[ComponentIndex].LODs.IsValidIndex(LODIndex));
+        check(Components[ComponentIndex].LODs[LODIndex].Surfaces.IsValidIndex(SurfaceIndex));
+        check(Components[ComponentIndex].LODs[LODIndex].Surfaces[SurfaceIndex].Images.IsValidIndex(ImageIndex));
 
-		return m_pD->Components[ComponentIndex].LODs[LODIndex].Surfaces[SurfaceIndex].Images[ImageIndex].Id;
+		return Components[ComponentIndex].LODs[LODIndex].Surfaces[SurfaceIndex].Images[ImageIndex].Id;
 	}
 
-
-	//---------------------------------------------------------------------------------------------
+	
     FName FInstance::GetImageName(int32 ComponentIndex, int32 LODIndex, int32 SurfaceIndex, int32 ImageIndex ) const
 	{
-		check(m_pD->Components.IsValidIndex(ComponentIndex));
-		check(m_pD->Components[ComponentIndex].LODs.IsValidIndex(LODIndex));
-        check(m_pD->Components[ComponentIndex].LODs[LODIndex].Surfaces.IsValidIndex(SurfaceIndex));
-		check(m_pD->Components[ComponentIndex].LODs[LODIndex].Surfaces[SurfaceIndex].Images.IsValidIndex(ImageIndex));
+		check(Components.IsValidIndex(ComponentIndex));
+		check(Components[ComponentIndex].LODs.IsValidIndex(LODIndex));
+        check(Components[ComponentIndex].LODs[LODIndex].Surfaces.IsValidIndex(SurfaceIndex));
+		check(Components[ComponentIndex].LODs[LODIndex].Surfaces[SurfaceIndex].Images.IsValidIndex(ImageIndex));
 
-        return m_pD->Components[ComponentIndex].LODs[LODIndex].Surfaces[SurfaceIndex].Images[ImageIndex].Name;
+        return Components[ComponentIndex].LODs[LODIndex].Surfaces[SurfaceIndex].Images[ImageIndex].Name;
 	}
-
-
-	//---------------------------------------------------------------------------------------------
+	
+	
 	FVector4f FInstance::GetVector(int32 ComponentIndex, int32 LODIndex, int32 SurfaceIndex, int32 VectorIndex) const
 	{
-		check(m_pD->Components.IsValidIndex(ComponentIndex));
-		check(m_pD->Components[ComponentIndex].LODs.IsValidIndex(LODIndex));
-        check(m_pD->Components[ComponentIndex].LODs[LODIndex].Surfaces.IsValidIndex(SurfaceIndex));
-		check(m_pD->Components[ComponentIndex].LODs[LODIndex].Surfaces[SurfaceIndex].Vectors.IsValidIndex(VectorIndex));
+		check(Components.IsValidIndex(ComponentIndex));
+		check(Components[ComponentIndex].LODs.IsValidIndex(LODIndex));
+        check(Components[ComponentIndex].LODs[LODIndex].Surfaces.IsValidIndex(SurfaceIndex));
+		check(Components[ComponentIndex].LODs[LODIndex].Surfaces[SurfaceIndex].Vectors.IsValidIndex(VectorIndex));
 
-        return m_pD->Components[ComponentIndex].LODs[LODIndex].Surfaces[SurfaceIndex].Vectors[VectorIndex].Value;
+        return Components[ComponentIndex].LODs[LODIndex].Surfaces[SurfaceIndex].Vectors[VectorIndex].Value;
 	}
-
-
-	//---------------------------------------------------------------------------------------------
+	
+	
 	FName FInstance::GetVectorName(int32 ComponentIndex, int32 LODIndex, int32 SurfaceIndex, int32 VectorIndex) const
 	{
-		check(m_pD->Components.IsValidIndex(ComponentIndex));
-		check(m_pD->Components[ComponentIndex].LODs.IsValidIndex(LODIndex));
-        check(m_pD->Components[ComponentIndex].LODs[LODIndex].Surfaces.IsValidIndex(SurfaceIndex));
-		check(m_pD->Components[ComponentIndex].LODs[LODIndex].Surfaces[SurfaceIndex].Vectors.IsValidIndex(VectorIndex));
+		check(Components.IsValidIndex(ComponentIndex));
+		check(Components[ComponentIndex].LODs.IsValidIndex(LODIndex));
+        check(Components[ComponentIndex].LODs[LODIndex].Surfaces.IsValidIndex(SurfaceIndex));
+		check(Components[ComponentIndex].LODs[LODIndex].Surfaces[SurfaceIndex].Vectors.IsValidIndex(VectorIndex));
 
-        return m_pD->Components[ComponentIndex].LODs[LODIndex].Surfaces[SurfaceIndex].Vectors[VectorIndex].Name;
+        return Components[ComponentIndex].LODs[LODIndex].Surfaces[SurfaceIndex].Vectors[VectorIndex].Name;
 	}
-
-
-	//---------------------------------------------------------------------------------------------
+	
+	
     float FInstance::GetScalar(int32 ComponentIndex, int32 LODIndex, int32 SurfaceIndex, int32 ScalarIndex ) const
 	{
-		check(m_pD->Components.IsValidIndex(ComponentIndex));
-		check(m_pD->Components[ComponentIndex].LODs.IsValidIndex(LODIndex));
-        check(m_pD->Components[ComponentIndex].LODs[LODIndex].Surfaces.IsValidIndex(SurfaceIndex));
-		check(m_pD->Components[ComponentIndex].LODs[LODIndex].Surfaces[SurfaceIndex].Scalars.IsValidIndex(ScalarIndex));
+		check(Components.IsValidIndex(ComponentIndex));
+		check(Components[ComponentIndex].LODs.IsValidIndex(LODIndex));
+        check(Components[ComponentIndex].LODs[LODIndex].Surfaces.IsValidIndex(SurfaceIndex));
+		check(Components[ComponentIndex].LODs[LODIndex].Surfaces[SurfaceIndex].Scalars.IsValidIndex(ScalarIndex));
 
-        return m_pD->Components[ComponentIndex].LODs[LODIndex].Surfaces[SurfaceIndex].Scalars[ScalarIndex].Value;
+        return Components[ComponentIndex].LODs[LODIndex].Surfaces[SurfaceIndex].Scalars[ScalarIndex].Value;
 	}
 
-
-	//---------------------------------------------------------------------------------------------
+	
 	FName FInstance::GetScalarName(int32 ComponentIndex, int32 LODIndex, int32 SurfaceIndex, int32 ScalarIndex) const
 	{
-		check(m_pD->Components.IsValidIndex(ComponentIndex));
-		check(m_pD->Components[ComponentIndex].LODs.IsValidIndex(LODIndex));
-        check(m_pD->Components[ComponentIndex].LODs[LODIndex].Surfaces.IsValidIndex(SurfaceIndex));
-		check(m_pD->Components[ComponentIndex].LODs[LODIndex].Surfaces[SurfaceIndex].Scalars.IsValidIndex(ScalarIndex));
+		check(Components.IsValidIndex(ComponentIndex));
+		check(Components[ComponentIndex].LODs.IsValidIndex(LODIndex));
+        check(Components[ComponentIndex].LODs[LODIndex].Surfaces.IsValidIndex(SurfaceIndex));
+		check(Components[ComponentIndex].LODs[LODIndex].Surfaces[SurfaceIndex].Scalars.IsValidIndex(ScalarIndex));
 
-        return m_pD->Components[ComponentIndex].LODs[LODIndex].Surfaces[SurfaceIndex].Scalars[ScalarIndex].Name;
+        return Components[ComponentIndex].LODs[LODIndex].Surfaces[SurfaceIndex].Scalars[ScalarIndex].Name;
 	}
 
-
-    //---------------------------------------------------------------------------------------------
+	
     FString FInstance::GetString(int32 ComponentIndex, int32 LODIndex, int32 SurfaceIndex, int32 StringIndex ) const
     {
-        check(m_pD->Components.IsValidIndex(ComponentIndex));
-        check(m_pD->Components[ComponentIndex].LODs.IsValidIndex(LODIndex));
-        check(m_pD->Components[ComponentIndex].LODs[LODIndex].Surfaces.IsValidIndex(SurfaceIndex));
-		check(m_pD->Components[ComponentIndex].LODs[LODIndex].Surfaces[SurfaceIndex].Strings.IsValidIndex(StringIndex));
+        check(Components.IsValidIndex(ComponentIndex));
+        check(Components[ComponentIndex].LODs.IsValidIndex(LODIndex));
+        check(Components[ComponentIndex].LODs[LODIndex].Surfaces.IsValidIndex(SurfaceIndex));
+		check(Components[ComponentIndex].LODs[LODIndex].Surfaces[SurfaceIndex].Strings.IsValidIndex(StringIndex));
 
-		bool bValid = m_pD->Components[ComponentIndex].LODs[LODIndex].Surfaces[SurfaceIndex].Strings.IsValidIndex(StringIndex);
+		bool bValid = Components[ComponentIndex].LODs[LODIndex].Surfaces[SurfaceIndex].Strings.IsValidIndex(StringIndex);
 		if (bValid)
         {
-            return m_pD->Components[ComponentIndex].LODs[LODIndex].Surfaces[SurfaceIndex].Strings[StringIndex].Value;
+            return Components[ComponentIndex].LODs[LODIndex].Surfaces[SurfaceIndex].Strings[StringIndex].Value;
         }
 
         return "";
     }
-
-
-    //---------------------------------------------------------------------------------------------
+	
+    
 	FName FInstance::GetStringName(int32 ComponentIndex, int32 LODIndex, int32 SurfaceIndex, int32 StringIndex) const
     {
-        check(m_pD->Components.IsValidIndex(ComponentIndex));
-        check(m_pD->Components[ComponentIndex].LODs.IsValidIndex(LODIndex));
-        check(m_pD->Components[ComponentIndex].LODs[LODIndex].Surfaces.IsValidIndex(SurfaceIndex));
+        check(Components.IsValidIndex(ComponentIndex));
+        check(Components[ComponentIndex].LODs.IsValidIndex(LODIndex));
+        check(Components[ComponentIndex].LODs[LODIndex].Surfaces.IsValidIndex(SurfaceIndex));
 
-		bool bValid = m_pD->Components[ComponentIndex].LODs[LODIndex].Surfaces[SurfaceIndex].Strings.IsValidIndex(StringIndex);
+		bool bValid = Components[ComponentIndex].LODs[LODIndex].Surfaces[SurfaceIndex].Strings.IsValidIndex(StringIndex);
 		if (bValid)
 		{
-            return m_pD->Components[ComponentIndex].LODs[LODIndex].Surfaces[SurfaceIndex].Strings[StringIndex].Name;
+            return Components[ComponentIndex].LODs[LODIndex].Surfaces[SurfaceIndex].Strings[StringIndex].Name;
         }
 
         return NAME_None;
     }
 
+	
 	int32 FInstance::GetOverlayMaterial(int32 ComponentIndex) const
 	{
-		check(m_pD->Components.IsValidIndex(ComponentIndex));
-		check(m_pD->Components[ComponentIndex].OverlayMaterialId < int32(1 << 22));
-		return static_cast<int32>(m_pD->Components[ComponentIndex].OverlayMaterialId);
+		check(Components.IsValidIndex(ComponentIndex));
+		check(Components[ComponentIndex].OverlayMaterialId < 1 << 22);
+		return static_cast<int32>(Components[ComponentIndex].OverlayMaterialId);
 	}
 
-    //---------------------------------------------------------------------------------------------
+    
 	int32 FInstance::GetExtensionDataCount() const
 	{
-		return m_pD->ExtensionData.Num();
+		return ExtensionData.Num();
 	}
 
-
-    //---------------------------------------------------------------------------------------------
+	
 	void FInstance::GetExtensionData(int32 Index, TSharedPtr<const FExtensionData>& OutExtensionData, FName& OutName) const
 	{
-		check(m_pD->ExtensionData.IsValidIndex(Index));
+		check(ExtensionData.IsValidIndex(Index));
 
-		OutExtensionData = m_pD->ExtensionData[Index].Data;
-		OutName = m_pD->ExtensionData[Index].Name;
+		OutExtensionData = ExtensionData[Index].Data;
+		OutName = ExtensionData[Index].Name;
 	}
 
-
-    //---------------------------------------------------------------------------------------------
-	int32 FInstance::Private::AddComponent()
+    
+	int32 FInstance::AddComponent()
 	{
 		int32 result = Components.Emplace();
 		return result;
 	}
-
-
-    //---------------------------------------------------------------------------------------------
-    int32 FInstance::Private::AddLOD( int32 ComponentIndex )
+	
+    
+    int32 FInstance::AddLOD( int32 ComponentIndex )
     {
         // Automatically create the necessary lods and components
         while (ComponentIndex >= Components.Num())
@@ -423,9 +361,8 @@ namespace mu
         return Components[ComponentIndex].LODs.Emplace();
     }
 
-
-    //---------------------------------------------------------------------------------------------
-    int32 FInstance::Private::AddSurface( int32 ComponentIndex, int32 LODIndex )
+	
+    int32 FInstance::AddSurface( int32 ComponentIndex, int32 LODIndex )
     {
         // Automatically create the necessary lods and components
         while (ComponentIndex >= Components.Num())
@@ -440,9 +377,8 @@ namespace mu
         return Components[ComponentIndex].LODs[LODIndex].Surfaces.Emplace();
     }
 
-
-    //---------------------------------------------------------------------------------------------
-    void FInstance::Private::SetSurfaceName( int32 ComponentIndex, int32 LODIndex, int32 SurfaceIndex, FName Name)
+    
+    void FInstance::SetSurfaceName( int32 ComponentIndex, int32 LODIndex, int32 SurfaceIndex, FName Name)
     {
         // Automatically create the necessary lods and components
 		while (ComponentIndex >= Components.Num())
@@ -462,9 +398,8 @@ namespace mu
         surface.Name = Name;
     }
 
-
-	//---------------------------------------------------------------------------------------------
-	void FInstance::Private::SetMesh(int32 ComponentIndex, int32 LODIndex, FResourceID meshId, FName Name)
+	
+	void FInstance::SetMesh(int32 ComponentIndex, int32 LODIndex, FResourceID meshId, FName Name)
 	{
 		// Automatically create the necessary lods and components
 		while (ComponentIndex >= Components.Num())
@@ -480,10 +415,9 @@ namespace mu
 		LOD.MeshId = meshId;
 		LOD.MeshName = Name;
 	}
-
-
-	//---------------------------------------------------------------------------------------------
-    int32 FInstance::Private::AddImage( int32 ComponentIndex, int32 LODIndex, int32 SurfaceIndex, FResourceID imageId, FName Name)
+	
+	
+    int32 FInstance::AddImage( int32 ComponentIndex, int32 LODIndex, int32 SurfaceIndex, FResourceID imageId, FName Name)
 	{
 		// Automatically create the necessary lods and components
 		while (ComponentIndex >= Components.Num())
@@ -502,10 +436,9 @@ namespace mu
 		FInstanceSurface& Surface = Components[ComponentIndex].LODs[LODIndex].Surfaces[SurfaceIndex];
 		return Surface.Images.Add({ imageId, Name });
 	}
-
-
-	//---------------------------------------------------------------------------------------------
-    int32 FInstance::Private::AddVector( int32 ComponentIndex, int32 LODIndex, int32 SurfaceIndex, const FVector4f& vec, FName Name)
+	
+	
+    int32 FInstance::AddVector( int32 ComponentIndex, int32 LODIndex, int32 SurfaceIndex, const FVector4f& vec, FName Name)
 	{
 		// Automatically create the necessary lods and components
 		while (ComponentIndex >= Components.Num())
@@ -525,9 +458,8 @@ namespace mu
 		return Surface.Vectors.Add({ vec, Name } );
 	}
 
-
-    //---------------------------------------------------------------------------------------------
-    int32 FInstance::Private::AddScalar( int32 ComponentIndex, int32 LODIndex, int32 SurfaceIndex, float sca, FName Name)
+	
+    int32 FInstance::AddScalar( int32 ComponentIndex, int32 LODIndex, int32 SurfaceIndex, float sca, FName Name)
     {
         // Automatically create the necessary lods and components
 		while (ComponentIndex >= Components.Num())
@@ -546,10 +478,9 @@ namespace mu
 		FInstanceSurface& Surface = Components[ComponentIndex].LODs[LODIndex].Surfaces[SurfaceIndex];
 		return Surface.Scalars.Add({ sca, Name });
     }
-
-
-    //---------------------------------------------------------------------------------------------
-    int32 FInstance::Private::AddString( int32 ComponentIndex, int32 LODIndex, int32 SurfaceIndex, const FString& Value, FName Name)
+	
+    
+    int32 FInstance::AddString( int32 ComponentIndex, int32 LODIndex, int32 SurfaceIndex, const FString& Value, FName Name)
     {
         // Automatically create the necessary lods and components
 		while (ComponentIndex >= Components.Num())
@@ -568,9 +499,9 @@ namespace mu
 		FInstanceSurface& Surface = Components[ComponentIndex].LODs[LODIndex].Surfaces[SurfaceIndex];
 		return Surface.Strings.Add({ Value, Name });
     }
+	
 
-
-	void FInstance::Private::AddOverlayMaterial(int32 ComponentIndex, float OverlayMaterialId)
+	void FInstance::AddOverlayMaterial(int32 ComponentIndex, float OverlayMaterialId)
 	{
 		while (ComponentIndex >= Components.Num())
 		{
@@ -579,10 +510,9 @@ namespace mu
 
 		Components[ComponentIndex].OverlayMaterialId = OverlayMaterialId;
 	}
-
-
-    //---------------------------------------------------------------------------------------------
-	void FInstance::Private::AddExtensionData(const TSharedPtr<const FExtensionData>& Data, FName Name)
+	
+    
+	void FInstance::AddExtensionData(const TSharedPtr<const FExtensionData>& Data, FName Name)
 	{
 		check(Data);
 
diff --git a/Engine/Plugins/Mutable/Source/MutableRuntime/Private/MuR/InstancePrivate.h b/Engine/Plugins/Mutable/Source/MutableRuntime/Private/MuR/InstancePrivate.h
deleted file mode 100644
diff --git a/Engine/Plugins/Mutable/Source/MutableRuntime/Private/MuR/System.cpp b/Engine/Plugins/Mutable/Source/MutableRuntime/Private/MuR/System.cpp
--- a/Engine/Plugins/Mutable/Source/MutableRuntime/Private/MuR/System.cpp
+++ b/Engine/Plugins/Mutable/Source/MutableRuntime/Private/MuR/System.cpp
@@ -15,7 +15,6 @@
 #include "Misc/AssertionMacros.h"
 #include "MuR/CodeRunner.h"
 #include "MuR/CodeVisitor.h"
-#include "MuR/InstancePrivate.h"
 #include "MuR/Mesh.h"
 #include "MuR/Model.h"
 #include "MuR/ModelPrivate.h"
@@ -315,11 +314,8 @@ namespace mu
 		LiveInstance->Cache->CheckHitCountsCleared();
 
 		LiveInstance->Instance = Result;
-		if (Result)
-		{
-			Result->GetPrivate()->Id = LiveInstance->InstanceID;
-		}
-		else
+
+    	if (!Result)
 		{
 			// In case of failure return an empty instance, to prevent following code to have to check it every time
 			Result = MakeShared<FInstance>();
diff --git a/Engine/Plugins/Mutable/Source/MutableRuntime/Public/MuR/Instance.h b/Engine/Plugins/Mutable/Source/MutableRuntime/Public/MuR/Instance.h
--- a/Engine/Plugins/Mutable/Source/MutableRuntime/Public/MuR/Instance.h
+++ b/Engine/Plugins/Mutable/Source/MutableRuntime/Public/MuR/Instance.h
@@ -14,6 +14,81 @@
 namespace mu
 {
 
+	/** Helper functions to make and read FResourceIDs */
+	inline FResourceID MakeResourceID(uint32 RootAddress, uint32 ParameterBlobIndex)
+	{
+		return (uint64(RootAddress) << 32) | uint64(ParameterBlobIndex);
+	}
+
+	
+	struct FInstanceSurface
+	{
+		FName Name;
+		uint32 InternalId=0;
+		uint32 ExternalId =0;
+		uint32 SharedId =0;
+
+		struct FInstanceImage
+		{
+			FResourceID Id;
+			FName Name;
+		};
+
+		TArray<FInstanceImage, TInlineAllocator<4>> Images;
+
+		struct FInstanceVector
+		{
+			FVector4f Value;
+			FName Name;
+		};
+
+		TArray<FInstanceVector> Vectors;
+
+		struct FInstanceScalar
+		{
+			float Value;
+			FName Name;
+		};
+
+		TArray<FInstanceScalar> Scalars;
+
+		struct FInstanceString
+		{
+			FString Value;
+			FName Name;
+		};
+
+		TArray<FInstanceString> Strings;
+	};
+
+
+	struct FInstanceLOD
+	{
+		FResourceID MeshId;
+		FName MeshName;
+
+		// The order must match the meshes surfaces
+		TArray<FInstanceSurface, TInlineAllocator<4>> Surfaces;
+	};
+
+
+	struct FInstanceComponent
+	{
+		uint16 Id;
+
+		float OverlayMaterialId = -1.f;
+
+		TArray<FInstanceLOD, TInlineAllocator<4>> LODs;
+	};
+
+	
+	struct NamedExtensionData
+	{
+		TSharedPtr<const FExtensionData> Data;
+		FName Name;
+	};
+
+	
     /** A customised object created from a model and a set of parameter values.
     * It corresponds to an "engine object" but the contents of its data depends on the Model, and
     * it may contain any number of LODs, components, surfaces, meshes and images, even none.
@@ -25,30 +100,20 @@ namespace mu
 		//-----------------------------------------------------------------------------------------
 		// Life cycle
 		//-----------------------------------------------------------------------------------------
-
-		UE_API FInstance();
-
+		
         //! Clone this instance
 		UE_API TSharedPtr<FInstance> Clone() const;
 
 		// Resource interface
 		UE_API int32 GetDataSize() const override;
 
-		/** */
-		UE_API ~FInstance();
-
 		//-----------------------------------------------------------------------------------------
 		// Own interface
 		//-----------------------------------------------------------------------------------------
 
         //! Type for the instance unique identifiers.
         typedef uint32 FID;
-
-        //! Get a unique identifier for this instance. It doesn't change during the entire
-        //! lifecycle of each instance. This identifier can be used in the System methods to update
-        //! or release the instance.
-        UE_API FInstance::FID GetId() const;
-
+		
 		//! Get the number of components of this instance.
 		UE_API int32 GetComponentCount() const;
 
@@ -197,18 +262,34 @@ namespace mu
 		//! \param OutExtensionData Receives the ExtensionData
 		//! \param OutName Receives the name associated with the ExtensionData. Guaranteed to be a valid string of non-zero length.
 		UE_API void GetExtensionData(int32 Index, TSharedPtr<const FExtensionData>& OutExtensionData, FName& OutName) const;
+		
+		int32 AddComponent();
+		
+		int32 AddLOD(int32 ComponentIndex);
 
-        //-----------------------------------------------------------------------------------------
-		// Interface pattern
-		//-----------------------------------------------------------------------------------------
-		class Private;
+		void SetMesh(int32 ComponentIndex, int32 LODIndex, FResourceID, FName Name);
+
+		int32 AddSurface(int32 ComponentIndex, int32 LODIndex);
+
+		void SetSurfaceName(int32 ComponentIndex, int32 LODIndex, int32 SurfaceIndex, FName Name);
+
+		int32 AddImage(int32 ComponentIndex, int32 LODIndex, int32 SurfaceIndex, FResourceID, FName Name);
+
+		int32 AddVector(int32 ComponentIndex, int32 LODIndex, int32 SurfaceIndex, const FVector4f&, FName Name);
+
+		int32 AddScalar(int32 ComponentIndex, int32 LODIndex, int32 SurfaceIndex, float, FName Name);
 
-		UE_API Private* GetPrivate() const;
+		int32 AddString(int32 ComponentIndex, int32 LODIndex, int32 SurfaceIndex, const FString& Value, FName Name);
 
-	private:
+		void AddOverlayMaterial( int32 ComponentIndex, float OverlayMaterialId);
 
-		Private* m_pD;
+		// Data must be non-null
+		void AddExtensionData(const TSharedPtr<const FExtensionData>& Data, FName Name);
+		
+		TArray<FInstanceComponent,TInlineAllocator<4>> Components;
 
+		// Every entry must have a valid ExtensionData and name
+		TArray<NamedExtensionData> ExtensionData;
 	};
 
 
