[mutable] Fixed edge case where the Private data of the COI has not yet been set.

#rnx

diff --git a/Engine/Plugins/Mutable/Source/CustomizableObject/Private/MuCO/CustomizableObjectSystem.cpp b/Engine/Plugins/Mutable/Source/CustomizableObject/Private/MuCO/CustomizableObjectSystem.cpp
--- a/Engine/Plugins/Mutable/Source/CustomizableObject/Private/MuCO/CustomizableObjectSystem.cpp
+++ b/Engine/Plugins/Mutable/Source/CustomizableObject/Private/MuCO/CustomizableObjectSystem.cpp
@@ -4751,9 +4751,9 @@ void UCustomizableObjectSystemPrivate::UpdateStats()
 {
 	NumSkeletalMeshes = 0;
 	
-	for (TObjectIterator<UCustomizableObjectInstance> Instance; Instance; ++Instance)
+	for (TObjectIterator<const UCustomizableObjectInstance> Instance; Instance; ++Instance)
 	{
-		if (!IsValid(*Instance))
+		if (!IsValid(*Instance) || !Instance->GetPrivate())
 		{
 			continue;
 		}
