[Mutable] Remove editor-only properties marked as exports from cooked builds.

#rnx

diff --git a/Engine/Plugins/Mutable/Source/CustomizableObjectEditor/Private/MuCOE/CustomizableObjectGraph.cpp b/Engine/Plugins/Mutable/Source/CustomizableObjectEditor/Private/MuCOE/CustomizableObjectGraph.cpp
--- a/Engine/Plugins/Mutable/Source/CustomizableObjectEditor/Private/MuCOE/CustomizableObjectGraph.cpp
+++ b/Engine/Plugins/Mutable/Source/CustomizableObjectEditor/Private/MuCOE/CustomizableObjectGraph.cpp
@@ -145,6 +145,12 @@ void UCustomizableObjectGraph::PostRename(UObject* OldOuter, const FName OldName
 }
 
 
+bool UCustomizableObjectGraph::IsEditorOnly() const
+{
+	return true;
+}
+
+
 void UCustomizableObjectGraph::PostDuplicate(bool bDuplicateForPIE)
 {
 	// In BeginPostDuplicate, nodes can call RequestNotificationForNodeIdChange
diff --git a/Engine/Plugins/Mutable/Source/CustomizableObjectEditor/Private/MuCOE/CustomizableObjectGraph.h b/Engine/Plugins/Mutable/Source/CustomizableObjectEditor/Private/MuCOE/CustomizableObjectGraph.h
--- a/Engine/Plugins/Mutable/Source/CustomizableObjectEditor/Private/MuCOE/CustomizableObjectGraph.h
+++ b/Engine/Plugins/Mutable/Source/CustomizableObjectEditor/Private/MuCOE/CustomizableObjectGraph.h
@@ -22,6 +22,8 @@ public:
 	UE_API virtual void PostLoad() override;
 	UE_API void PostRename(UObject * OldOuter, const FName OldName) override;
 
+	UE_API virtual bool IsEditorOnly() const override;
+
 	// Own Interface
 	UE_API void NotifyNodeIdChanged(const FGuid& OldGuid, const FGuid& NewGuid);
 
