[Mutable] Added a new compilation error. It's triggered when mutable tries to compile a data table without rows.

#rnx

diff --git a/Engine/Plugins/Mutable/Source/CustomizableObjectEditor/Private/MuCOE/GenerateMutableSource/GenerateMutableSourceTable.cpp b/Engine/Plugins/Mutable/Source/CustomizableObjectEditor/Private/MuCOE/GenerateMutableSource/GenerateMutableSourceTable.cpp
--- a/Engine/Plugins/Mutable/Source/CustomizableObjectEditor/Private/MuCOE/GenerateMutableSource/GenerateMutableSourceTable.cpp
+++ b/Engine/Plugins/Mutable/Source/CustomizableObjectEditor/Private/MuCOE/GenerateMutableSource/GenerateMutableSourceTable.cpp
@@ -1063,6 +1063,27 @@ mu::Ptr<mu::FTable> GenerateMutableSourceTable(const UDataTable* DataTable, cons
 		TArray<uint32> RowIds;
 		TArray<FName> RowNames = GetRowsToCompile(*DataTable, *TableNode, GenerationContext, RowIds);
 
+		// Check if the generated table will contain anything
+		if (RowNames.IsEmpty())
+		{
+			FText EmptyTableMessage;
+
+			if (DataTable->GetRowNames().IsEmpty())
+			{
+				EmptyTableMessage = FText::Format(LOCTEXT("EmptyDataTable_Error",
+					"Node Table [{0}] has no rows to compile."), FText::FromString(TableNode->ParameterName));
+			}
+			else
+			{
+				EmptyTableMessage = FText::Format(LOCTEXT("AllDataTableRowsFiltered_Error",
+					"All data table rows from the node [{0}] have been filtered and nothing will be compiled. Please review the compilation filters of the node."), FText::FromString(TableNode->ParameterName));
+			}
+
+			GenerationContext.Log(EmptyTableMessage, TableNode);
+
+			return nullptr;
+		}
+
 		// Adding the Name Column
 		MutableTable->AddColumn("Name", mu::ETableColumnType::String);
 
