[Mutable] Add new Material Parameter type, API, and infrastructure.

#rnx

diff --git a/Engine/Plugins/Mutable/Source/CustomizableObject/Internal/MuCO/CustomizableObjectPrivate.h b/Engine/Plugins/Mutable/Source/CustomizableObject/Internal/MuCO/CustomizableObjectPrivate.h
--- a/Engine/Plugins/Mutable/Source/CustomizableObject/Internal/MuCO/CustomizableObjectPrivate.h
+++ b/Engine/Plugins/Mutable/Source/CustomizableObject/Internal/MuCO/CustomizableObjectPrivate.h
@@ -1642,6 +1642,8 @@ enum class ECustomizableObjectVersions : int32
 
 	MorphMeshParameters,
 
+	MaterialParameter,
+
 	// -----<new versions can be added above this line>--------
 	LastCustomizableObjectVersion
 };
diff --git a/Engine/Plugins/Mutable/Source/CustomizableObject/Internal/MuCO/CustomizableObjectSystemPrivate.h b/Engine/Plugins/Mutable/Source/CustomizableObject/Internal/MuCO/CustomizableObjectSystemPrivate.h
--- a/Engine/Plugins/Mutable/Source/CustomizableObject/Internal/MuCO/CustomizableObjectSystemPrivate.h
+++ b/Engine/Plugins/Mutable/Source/CustomizableObject/Internal/MuCO/CustomizableObjectSystemPrivate.h
@@ -18,6 +18,7 @@
 #include "MuR/System.h"
 #include "MuR/Skeleton.h"
 #include "MuR/Image.h"
+#include "MuR/Material.h"
 #include "UObject/GCObject.h"
 #include "WorldCollision.h"
 #include "Engine/StreamableManager.h"
@@ -384,6 +385,7 @@ struct FGeneratedMaterial
 {
 	GENERATED_USTRUCT_BODY();
 
+	/* Output material */
 	UPROPERTY(Category = NoCategory, VisibleAnywhere)
 	TObjectPtr<UMaterialInterface> MaterialInterface;
 
@@ -393,8 +395,8 @@ struct FGeneratedMaterial
 	// Surface or SharedSurface Id
 	uint32 SurfaceId = 0;
 
-	// Index of the material to instantiate (UCustomizableObject::ReferencedMaterials)
-	uint32 MaterialIndex = 0;
+	// Core Material. Strong object pointer to a material or a material index
+	mu::FMaterial Material;
 
 #if WITH_EDITORONLY_DATA
 	// Name of the component that contains this material
@@ -402,7 +404,7 @@ struct FGeneratedMaterial
 	FName ComponentName;
 #endif
 
-	bool operator==(const FGeneratedMaterial& Other) const { return SurfaceId == Other.SurfaceId && MaterialIndex == Other.MaterialIndex; };
+	bool operator==(const FGeneratedMaterial& Other) const { return SurfaceId == Other.SurfaceId && Material == Other.Material; };
 };
 
 
@@ -466,18 +468,14 @@ struct FInstanceUpdateData
 		uint16 FirstScalar = 0;
 		uint16 ScalarCount = 0;
 
-		/** Index of the material in the referenced materials array of the CO.
-		* A negative value means that the material of this surface slot of the mesh doesn't need to be changed.
-		* This is valid for pass-through meshes.
-		*/
-		int32 MaterialIndex = -1;
-
 		/** Id of the surface in the mutable core instance. */
 		uint32 SurfaceId = 0;
 		uint32 SharedSurfaceId = 0;
 
 		/** Id of the metadata associated with this surface. */ 
 		uint32 SurfaceMetadataId = 0;
+
+		mu::FMaterial Material;
 	};
 
 	struct FLOD
@@ -507,7 +505,7 @@ struct FInstanceUpdateData
 		uint16 LODCount = 0;
 
 		/** Overlay material. Only one per mesh, if any. */
-		int32 OverlayMaterial = INDEX_NONE;
+		mu::FMaterial OverlayMaterial;
 	};
 
 	TArray<FComponent> Components;
diff --git a/Engine/Plugins/Mutable/Source/CustomizableObject/Private/MuCO/CustomizableObject.cpp b/Engine/Plugins/Mutable/Source/CustomizableObject/Private/MuCO/CustomizableObject.cpp
--- a/Engine/Plugins/Mutable/Source/CustomizableObject/Private/MuCO/CustomizableObject.cpp
+++ b/Engine/Plugins/Mutable/Source/CustomizableObject/Private/MuCO/CustomizableObject.cpp
@@ -1447,6 +1447,12 @@ void UCustomizableObjectPrivate::UpdateParameterPropertiesFromModel(const TShare
 				Data.Type = EMutableParameterType::SkeletalMesh;
 				break;
 			}
+			
+			case mu::EParameterType::Material:
+			{
+				Data.Type = EMutableParameterType::Material;
+				break;
+			}
 
 			default:
 				// Unhandled type?
diff --git a/Engine/Plugins/Mutable/Source/CustomizableObject/Private/MuCO/CustomizableObjectInstance.cpp b/Engine/Plugins/Mutable/Source/CustomizableObject/Private/MuCO/CustomizableObjectInstance.cpp
--- a/Engine/Plugins/Mutable/Source/CustomizableObject/Private/MuCO/CustomizableObjectInstance.cpp
+++ b/Engine/Plugins/Mutable/Source/CustomizableObject/Private/MuCO/CustomizableObjectInstance.cpp
@@ -4906,8 +4906,8 @@ UE::Tasks::FTask UCustomizableInstancePrivate::LoadAdditionalAssetsAndData(const
 	
 	for (const FInstanceUpdateData::FSurface& Surface : OperationData->InstanceUpdateData.Surfaces)
 	{
-		const int32 MaterialIndex = Surface.MaterialIndex;
-		if (MaterialIndex<0 || ObjectToInstanceIndexMap.Contains(MaterialIndex))
+		const int32 MaterialIndex = Surface.Material.MaterialId;
+		if (MaterialIndex == INDEX_NONE || ObjectToInstanceIndexMap.Contains(MaterialIndex))
 		{
 			continue;
 		}
@@ -4926,16 +4926,16 @@ UE::Tasks::FTask UCustomizableInstancePrivate::LoadAdditionalAssetsAndData(const
 
 	for (const FInstanceUpdateData::FComponent& Component : OperationData->InstanceUpdateData.Components)
 	{
-		if (Component.OverlayMaterial == INDEX_NONE || ObjectToInstanceIndexMap.Contains(Component.OverlayMaterial))
+		if (Component.OverlayMaterial.MaterialId == INDEX_NONE || ObjectToInstanceIndexMap.Contains(Component.OverlayMaterial.MaterialId))
 		{
 			continue;
 		}
 
-		TSoftObjectPtr<UMaterialInterface> AssetPtr = ModelResources.Materials.IsValidIndex(Component.OverlayMaterial) ? ModelResources.Materials[Component.OverlayMaterial] : nullptr;
+		TSoftObjectPtr<UMaterialInterface> AssetPtr = ModelResources.Materials.IsValidIndex(Component.OverlayMaterial.MaterialId) ? ModelResources.Materials[Component.OverlayMaterial.MaterialId] : nullptr;
 		UMaterialInterface* LoadedMaterial = AssetPtr.Get();
 
 		const int32 ReferencedMaterialsIndex = ReferencedMaterials.Add(LoadedMaterial);
-		ObjectToInstanceIndexMap.Add(Component.OverlayMaterial, ReferencedMaterialsIndex);
+		ObjectToInstanceIndexMap.Add(Component.OverlayMaterial.MaterialId, ReferencedMaterialsIndex);
 
 		if (!LoadedMaterial && !AssetPtr.IsNull())
 		{
@@ -5801,14 +5801,22 @@ void UCustomizableInstancePrivate::BuildMaterials(const TSharedRef<FUpdateContex
 			ComponentData->OverrideMaterials.Reset();
 			ComponentData->OverlayMaterial = nullptr;
 
-			if (const uint32* ReferencedMaterialIndex = ObjectToInstanceIndexMap.Find(Component.OverlayMaterial))
+			if (Component.OverlayMaterial.MaterialId != INDEX_NONE)
 			{
-				if (ReferencedMaterials.IsValidIndex(*ReferencedMaterialIndex))
+				if (const uint32* ReferencedMaterialIndex = ObjectToInstanceIndexMap.Find(Component.OverlayMaterial.MaterialId))
 				{
-					ComponentData->OverlayMaterial = ReferencedMaterials[*ReferencedMaterialIndex];
-					OverlayMaterial = ComponentData->OverlayMaterial;
+					if (ReferencedMaterials.IsValidIndex(*ReferencedMaterialIndex))
+					{
+						ComponentData->OverlayMaterial = ReferencedMaterials[*ReferencedMaterialIndex];
+						OverlayMaterial = ComponentData->OverlayMaterial;
+					}
 				}
 			}
+			else
+			{
+				ComponentData->OverlayMaterial = Component.OverlayMaterial.Material.Get();
+				OverlayMaterial = ComponentData->OverlayMaterial;
+			}
 		}
 
 		if (!bUseOverrideMaterialsOnly)
@@ -5852,7 +5860,7 @@ void UCustomizableInstancePrivate::BuildMaterials(const TSharedRef<FUpdateContex
 				const FInstanceUpdateData::FSurface& Surface = OperationData->InstanceUpdateData.Surfaces[LOD.FirstSurface + SurfaceIndex];
 
 				// Is this a material in a passthrough mesh that we don't modify?
-				if (Surface.MaterialIndex<0)
+				if (Surface.Material.MaterialId == INDEX_NONE && !Surface.Material.Material)
 				{
 					Materials.Emplace();
 #if WITH_EDITOR
@@ -5874,8 +5882,18 @@ void UCustomizableInstancePrivate::BuildMaterials(const TSharedRef<FUpdateContex
 					continue;
 				}
 
-				const uint32 ReferencedMaterialIndex = ObjectToInstanceIndexMap[Surface.MaterialIndex];
-				UMaterialInterface* MaterialTemplate = ReferencedMaterials[ReferencedMaterialIndex];
+				UMaterialInterface* MaterialTemplate = nullptr;
+
+				if (Surface.Material.MaterialId == INDEX_NONE)
+				{
+					MaterialTemplate = Surface.Material.Material.Get();
+				}
+				else
+				{
+					const uint32 ReferencedMaterialIndex = ObjectToInstanceIndexMap[Surface.Material.MaterialId];
+					MaterialTemplate = ReferencedMaterials[ReferencedMaterialIndex];
+				}
+				
 				if (!MaterialTemplate)
 				{
 					// Missing MaterialTemplate. Use DefaultMaterial instead. 
@@ -6320,7 +6338,7 @@ void UCustomizableInstancePrivate::BuildMaterials(const TSharedRef<FUpdateContex
 
 					FGeneratedMaterial& Material = GeneratedMaterials.AddDefaulted_GetRef();
 					Material.SurfaceId = Surface.SurfaceId;
-					Material.MaterialIndex = Surface.MaterialIndex;
+					Material.Material = Surface.Material;
 					Material.MaterialInterface = MaterialInterface;
 
 #if WITH_EDITORONLY_DATA
diff --git a/Engine/Plugins/Mutable/Source/CustomizableObject/Private/MuCO/CustomizableObjectInstanceDescriptor.cpp b/Engine/Plugins/Mutable/Source/CustomizableObject/Private/MuCO/CustomizableObjectInstanceDescriptor.cpp
--- a/Engine/Plugins/Mutable/Source/CustomizableObject/Private/MuCO/CustomizableObjectInstanceDescriptor.cpp
+++ b/Engine/Plugins/Mutable/Source/CustomizableObject/Private/MuCO/CustomizableObjectInstanceDescriptor.cpp
@@ -817,6 +817,32 @@ TSharedPtr<mu::FParameters> FCustomizableObjectInstanceDescriptor::GetParameters
  
 			break;
 		}
+
+		case mu::EParameterType::Material:
+		{
+			for (const FCustomizableObjectMaterialParameterValue& MaterialParameter : MaterialParameters)
+			{
+				if (MaterialParameter.ParameterName == Name || (Uid.IsValid() && MaterialParameter.Id == Uid))
+				{
+					if (TSharedPtr<mu::FRangeIndex> RangeIdxPtr = MutableParameters->NewRangeIndex(ParamIndex))
+					{
+						for (int32 RangeIndex = 0; RangeIndex < MaterialParameter.ParameterRangeValues.Num(); ++RangeIndex)
+						{
+							RangeIdxPtr->SetPosition(0, RangeIndex);
+							MutableParameters->SetMaterialValue(ParamIndex, MaterialParameter.ParameterRangeValues[RangeIndex], RangeIdxPtr.Get());
+						}
+					}
+					else
+					{
+						MutableParameters->SetMaterialValue(ParamIndex, MaterialParameter.ParameterValue);
+					}
+
+					break;
+				}
+			}
+
+			break;
+		}
  
 		default:
 			check(false); // Missing case.
@@ -866,6 +892,7 @@ void FCustomizableObjectInstanceDescriptor::ReloadParameters()
 	TArray<FCustomizableObjectFloatParameterValue> OldFloatParameters = FloatParameters;
 	TArray<FCustomizableObjectTextureParameterValue> OldTextureParameters = TextureParameters;
 	TArray<FCustomizableObjectSkeletalMeshParameterValue> OldSkeletalMeshParameters = SkeletalMeshParameters;
+	TArray<FCustomizableObjectMaterialParameterValue> OldMaterialParameters = MaterialParameters;
 	TArray<FCustomizableObjectVectorParameterValue> OldVectorParameters = VectorParameters;
 	TArray<FCustomizableObjectProjectorParameterValue> OldProjectorParameters = ProjectorParameters;
 	TArray<FCustomizableObjectTransformParameterValue> OldTransformParameters = TransformParameters;
@@ -875,6 +902,7 @@ void FCustomizableObjectInstanceDescriptor::ReloadParameters()
 	FloatParameters.Reset();
 	TextureParameters.Reset();
 	SkeletalMeshParameters.Reset();
+	MaterialParameters.Reset();
 	VectorParameters.Reset();
 	ProjectorParameters.Reset();
 	TransformParameters.Reset();
@@ -1338,6 +1366,63 @@ void FCustomizableObjectInstanceDescriptor::ReloadParameters()
 			break;
 		}
  
+		case mu::EParameterType::Material:
+		{
+			FCustomizableObjectMaterialParameterValue Param;
+			Param.ParameterName = Name;
+			Param.Id = Uid;
+
+			auto FindByNameAndUid = [&](const FCustomizableObjectMaterialParameterValue& P)
+				{
+					return P.ParameterName == Name || (Uid.IsValid() && P.Id == Uid);
+				};
+
+			if (FCustomizableObjectMaterialParameterValue* Result = OldMaterialParameters.FindByPredicate(FindByNameAndUid))
+			{
+				if (TSharedPtr<mu::FRangeIndex> RangeIdxPtr = MutableParameters->NewRangeIndex(ParamIndex))
+				{
+					Param.ParameterRangeValues = Result->ParameterRangeValues;
+				}
+				else
+				{
+					Param.ParameterValue = Result->ParameterValue;
+				}
+			}
+			else // Not found in Instance Parameters. Use Mutable Parameters.
+			{
+				if (TSharedPtr<mu::FRangeIndex> RangeIdxPtr = MutableParameters->NewRangeIndex(ParamIndex))
+				{
+					const int32 ValueCount = MutableParameters->GetValueCount(ParamIndex);
+
+					for (int32 ValueIndex = 0; ValueIndex < ValueCount; ++ValueIndex)
+					{
+						TSharedPtr<mu::FRangeIndex> RangeValueIdxPtr = MutableParameters->GetValueIndex(ParamIndex, ValueIndex);
+						int32 RangeIndex = RangeValueIdxPtr->GetPosition(0);
+
+						if (!Param.ParameterRangeValues.IsValidIndex(RangeIndex))
+						{
+							const int32 PreviousNum = Param.ParameterRangeValues.Num();
+							Param.ParameterRangeValues.AddUninitialized(RangeIndex + 1 - Param.ParameterRangeValues.Num());
+
+							for (int32 Index = PreviousNum; Index < Param.ParameterRangeValues.Num(); ++Index)
+							{
+								Param.ParameterRangeValues[Index] = nullptr;
+							}
+						}
+
+						Param.ParameterRangeValues[RangeIndex] = MutableParameters->GetMaterialValue(ParamIndex, RangeValueIdxPtr.Get());
+					}
+				}
+				else
+				{
+					Param.ParameterValue = MutableParameters->GetMaterialValue(ParamIndex);
+				}
+			}
+
+			MaterialParameters.Add(Param);
+
+			break;
+		}
 		default:
 			check(false); // Missing case.
 			break;
diff --git a/Engine/Plugins/Mutable/Source/CustomizableObject/Private/MuCO/CustomizableObjectSystem.cpp b/Engine/Plugins/Mutable/Source/CustomizableObject/Private/MuCO/CustomizableObjectSystem.cpp
--- a/Engine/Plugins/Mutable/Source/CustomizableObject/Private/MuCO/CustomizableObjectSystem.cpp
+++ b/Engine/Plugins/Mutable/Source/CustomizableObject/Private/MuCO/CustomizableObjectSystem.cpp
@@ -22,6 +22,7 @@
 #include "MuCO/UnrealPortabilityHelpers.h"
 #include "MuR/Model.h"
 #include "MuR/Settings.h"
+#include "MuR/Material.h"
 #include "UObject/UObjectIterator.h"
 #include "Widgets/Notifications/SNotificationList.h"
 #include "ContentStreaming.h"
@@ -2615,6 +2616,7 @@ namespace impl
 
 						Surface.SurfaceId = MutableInstance->GetSurfaceId(InstanceComponentIndex, BaseLODIndex, BaseSurfaceIndex);
 						Surface.SurfaceMetadataId = MutableInstance->GetSurfaceCustomId(InstanceComponentIndex, BaseLODIndex, BaseSurfaceIndex);
+						Surface.Material = MutableInstance->GetMaterial(InstanceComponentIndex, BaseLODIndex, BaseSurfaceIndex);
 
 						// Vectors
 						Surface.FirstVector = OperationData->InstanceUpdateData.Vectors.Num();
@@ -2638,20 +2640,7 @@ namespace impl
 								const FName ScalarName = MutableInstance->GetScalarName(InstanceComponentIndex, BaseLODIndex, BaseSurfaceIndex, ScalarIndex);
 							const float ScalarValue = MutableInstance->GetScalar(InstanceComponentIndex, BaseLODIndex, BaseSurfaceIndex, ScalarIndex);
 
-							FString EncodingMaterialIdString = "__MutableMaterialId";
-
-							// Decoding Material Switch from Mutable parameter name
-							if (ScalarName.ToString().Contains(EncodingMaterialIdString))
-							{
-								Surface.MaterialIndex = static_cast<uint32>(ScalarValue);
-
-								// This parameter is not needed in the final material instance
-								Surface.ScalarCount -= 1;
-							}
-							else
-							{
-								OperationData->InstanceUpdateData.Scalars.Push({ ScalarName, ScalarValue });
-							}
+							OperationData->InstanceUpdateData.Scalars.Push({ ScalarName, ScalarValue });
 						}
 
 						// Images
diff --git a/Engine/Plugins/Mutable/Source/CustomizableObjectEditor/Private/MuCOE/CustomizableObjectCompiler.cpp b/Engine/Plugins/Mutable/Source/CustomizableObjectEditor/Private/MuCOE/CustomizableObjectCompiler.cpp
--- a/Engine/Plugins/Mutable/Source/CustomizableObjectEditor/Private/MuCOE/CustomizableObjectCompiler.cpp
+++ b/Engine/Plugins/Mutable/Source/CustomizableObjectEditor/Private/MuCOE/CustomizableObjectCompiler.cpp
@@ -985,6 +985,7 @@ void FCustomizableObjectCompiler::CompileInternal(bool bAsync)
 
 		ModelResources->TextureParameterDefaultValues = MoveTemp(GenerationContext.TextureParameterDefaultValues);
 		ModelResources->SkeletalMeshParameterDefaultValues = MoveTemp(GenerationContext.SkeletalMeshParameterDefaultValues);
+		ModelResources->MaterialParameterDefaultValues = MoveTemp(GenerationContext.MaterialParameterDefaultValues);
 
 		ModelResources->AnimBPs = MoveTemp(GenerationContext.AnimBPAssets);
 
diff --git a/Engine/Plugins/Mutable/Source/CustomizableObjectEditor/Private/MuCOE/CustomizableObjectEditorModule.cpp b/Engine/Plugins/Mutable/Source/CustomizableObjectEditor/Private/MuCOE/CustomizableObjectEditorModule.cpp
--- a/Engine/Plugins/Mutable/Source/CustomizableObjectEditor/Private/MuCOE/CustomizableObjectEditorModule.cpp
+++ b/Engine/Plugins/Mutable/Source/CustomizableObjectEditor/Private/MuCOE/CustomizableObjectEditorModule.cpp
@@ -38,6 +38,7 @@
 #include "MuCOE/Nodes/CustomizableObjectNodeMacroInstance.h"
 #include "MuCOE/Nodes/CustomizableObjectNodeMacroInstanceDetails.h"
 #include "MuCOE/Nodes/CustomizableObjectNodeMaterial.h"
+#include "MuCOE/Nodes/CustomizableObjectNodeMaterialParameter.h"
 #include "MuCOE/Nodes/CustomizableObjectNodeMeshSectionDetails.h"
 #include "MuCOE/Nodes/CustomizableObjectNodeModifierEditMeshSection.h"
 #include "MuCOE/Nodes/CustomizableObjectNodeModifierEditMeshSectionDetails.h"
@@ -130,7 +131,7 @@ constexpr int32 MaxAccessTimespan = 30;
 
 TAutoConsoleVariable<bool> CVarMutableDerivedDataCacheUsage(
 	TEXT("mutable.DerivedDataCacheUsage"),
-	true,
+	false,
 	TEXT("Derived data cache access for compiled data.")
 	TEXT("true - Allows DDC access according to the DDC policies specified in the plugin settings."),
 	ECVF_Default);
@@ -1013,6 +1014,12 @@ TMap<FName, FGuid> FCustomizableObjectEditorModule::GetParticipatingObjects(cons
 			AddObjectSoft(TextureParameterNode->ReferenceValue);
 		}
 
+		else if (UCustomizableObjectNodeMaterialParameter* MaterialParamNode = Cast<UCustomizableObjectNodeMaterialParameter>(&Node))
+		{
+			AddObjectSoft(MaterialParamNode->DefaultValue.ToSoftObjectPath());
+			AddObjectSoft(MaterialParamNode->ReferenceValue.ToSoftObjectPath());
+		}
+
 		else if (UCustomizableObjectNodeMaterial* MatNode = Cast<UCustomizableObjectNodeMaterial>(&Node))
 		{
 			AddObjectSoft(MatNode->GetMaterial());
diff --git a/Engine/Plugins/Mutable/Source/CustomizableObjectEditor/Private/MuCOE/GenerateMutableSource/GenerateMutableSource.h b/Engine/Plugins/Mutable/Source/CustomizableObjectEditor/Private/MuCOE/GenerateMutableSource/GenerateMutableSource.h
--- a/Engine/Plugins/Mutable/Source/CustomizableObjectEditor/Private/MuCOE/GenerateMutableSource/GenerateMutableSource.h
+++ b/Engine/Plugins/Mutable/Source/CustomizableObjectEditor/Private/MuCOE/GenerateMutableSource/GenerateMutableSource.h
@@ -853,8 +853,12 @@ struct FMutableGraphGenerationContext
 	// Current material parameter name to find the corresponding column in a mutable table
 	FString CurrentMaterialTableParameter;
 
+	/** */
+	mu::ETableColumnType CurrentTableColumnType = mu::ETableColumnType::None;
+
 	// Current material parameter id to find the corresponding column in a mutable table
-	FString CurrentMaterialTableParameterId;
+	FString CurrentMaterialParameterId;
+
 
 	// Material to SharedSurfaceId
 	TMap<FSharedSurfaceKey, FGuid> SurfaceGuids;
@@ -883,6 +887,7 @@ struct FMutableGraphGenerationContext
 	
 	TMap<FName, TObjectPtr<UTexture>> TextureParameterDefaultValues;
 	TMap<FName, TObjectPtr<USkeletalMesh>> SkeletalMeshParameterDefaultValues;
+	TMap<FName, TObjectPtr<UMaterialInterface>> MaterialParameterDefaultValues;
 };
 
 
diff --git a/Engine/Plugins/Mutable/Source/CustomizableObjectEditor/Private/MuCOE/GenerateMutableSource/GenerateMutableSourceComponent.cpp b/Engine/Plugins/Mutable/Source/CustomizableObjectEditor/Private/MuCOE/GenerateMutableSource/GenerateMutableSourceComponent.cpp
--- a/Engine/Plugins/Mutable/Source/CustomizableObjectEditor/Private/MuCOE/GenerateMutableSource/GenerateMutableSourceComponent.cpp
+++ b/Engine/Plugins/Mutable/Source/CustomizableObjectEditor/Private/MuCOE/GenerateMutableSource/GenerateMutableSourceComponent.cpp
@@ -10,6 +10,7 @@
 #include "MuCOE/GenerateMutableSource/GenerateMutableSourceSurface.h"
 #include "MuCOE/GenerateMutableSource/GenerateMutableSourceFloat.h"
 #include "MuCOE/GenerateMutableSource/GenerateMutableSourceMacro.h"
+#include "MuCOE/GenerateMutableSource/GenerateMutableSourceMaterial.h"
 #include "MuCOE/GraphTraversal.h"
 #include "MuCOE/Nodes/CustomizableObjectNodeComponentMesh.h"
 #include "MuCOE/Nodes/CustomizableObjectNodeComponentMeshAddTo.h"
@@ -29,6 +30,7 @@
 #include "MuT/NodeScalarConstant.h"
 #include "MuT/NodeSurfaceNew.h"
 #include "MuT/NodeMeshConstant.h"
+#include "MuT/NodeMaterialConstant.h"
 
 #include "Rendering/SkeletalMeshLODModel.h"
 #include "Rendering/SkeletalMeshModel.h"
@@ -151,15 +153,15 @@ mu::Ptr<mu::NodeComponent> GenerateMutableSourceComponent(const UEdGraphPin* Pin
 		UEdGraphPin* MaterialAssetPin = TypedComponentMesh->GetOverlayMaterialAssetPin();
 		if (const UEdGraphPin* ConnectedPin = MaterialAssetPin ? FollowInputPin(*MaterialAssetPin) : nullptr)
 		{
-			GenerationContext.CurrentMaterialTableParameterId = TEXT("__MutableMaterialId") + ConnectedPin->PinId.ToString();
-			mu::Ptr<mu::NodeScalar> ScalarNode = GenerateMutableSourceFloat(ConnectedPin, GenerationContext);
-			NodeComponentNew->OverlayMaterial = ScalarNode;
+			GenerationContext.CurrentMaterialParameterId = ConnectedPin->PinId.ToString();
+			mu::Ptr<mu::NodeMaterial> MaterialNode = GenerateMutableSourceMaterial(ConnectedPin, GenerationContext);
+			NodeComponentNew->OverlayMaterial = MaterialNode;
 		}
 		else if(UMaterialInterface* OverlayMaterial = GenerationContext.LoadObject(TypedComponentMesh->GetOverlayMaterial()))
 		{
-			mu::Ptr<mu::NodeScalarConstant> ScalarNode = new mu::NodeScalarConstant();
-			ScalarNode->Value = GenerationContext.ReferencedMaterials.AddUnique(OverlayMaterial);
-			NodeComponentNew->OverlayMaterial = ScalarNode;
+			mu::Ptr<mu::NodeMaterialConstant> MaterialNode = new mu::NodeMaterialConstant();
+			MaterialNode->MaterialId = GenerationContext.ReferencedMaterials.AddUnique(OverlayMaterial);
+			NodeComponentNew->OverlayMaterial = MaterialNode;
 		}
 
 		Result = NodeComponentNew;
diff --git a/Engine/Plugins/Mutable/Source/CustomizableObjectEditor/Private/MuCOE/GenerateMutableSource/GenerateMutableSourceFloat.cpp b/Engine/Plugins/Mutable/Source/CustomizableObjectEditor/Private/MuCOE/GenerateMutableSource/GenerateMutableSourceFloat.cpp
--- a/Engine/Plugins/Mutable/Source/CustomizableObjectEditor/Private/MuCOE/GenerateMutableSource/GenerateMutableSourceFloat.cpp
+++ b/Engine/Plugins/Mutable/Source/CustomizableObjectEditor/Private/MuCOE/GenerateMutableSource/GenerateMutableSourceFloat.cpp
@@ -346,12 +346,6 @@ mu::Ptr<mu::NodeScalar> GenerateMutableSourceFloat(const UEdGraphPin* Pin, FMuta
 				{
 					mu::Ptr<mu::NodeScalarTable> ScalarTableNode = new mu::NodeScalarTable();
 
-					if (Pin->PinType.PinCategory == Schema->PC_Material)
-					{
-						// Materials use the parameter id as column names
-						ColumnName = GenerationContext.CurrentMaterialTableParameterId;
-					}
-
 					// Generating a new Float column if not exists
 					if (Table->FindColumn(ColumnName) == INDEX_NONE)
 					{
diff --git a/Engine/Plugins/Mutable/Source/CustomizableObjectEditor/Private/MuCOE/GenerateMutableSource/GenerateMutableSourceImage.cpp b/Engine/Plugins/Mutable/Source/CustomizableObjectEditor/Private/MuCOE/GenerateMutableSource/GenerateMutableSourceImage.cpp
--- a/Engine/Plugins/Mutable/Source/CustomizableObjectEditor/Private/MuCOE/GenerateMutableSource/GenerateMutableSourceImage.cpp
+++ b/Engine/Plugins/Mutable/Source/CustomizableObjectEditor/Private/MuCOE/GenerateMutableSource/GenerateMutableSourceImage.cpp
@@ -1065,7 +1065,7 @@ mu::Ptr<mu::NodeImage> GenerateMutableSourceImage(const UEdGraphPin* Pin, FMutab
 						if (Pin->PinType.PinCategory == Schema->PC_Material)
 						{
 							// Material parameters use the Data Table Column Name + Parameter id as mutable column Name to aboid duplicated names (i.e. two MI columns with the same parents but different values).
-							ColumnName = Property->GetDisplayNameText().ToString() + GenerationContext.CurrentMaterialTableParameterId;
+							ColumnName = Property->GetDisplayNameText().ToString() + GenerationContext.CurrentMaterialParameterId;
 						}
 						else
 						{
@@ -1101,6 +1101,7 @@ mu::Ptr<mu::NodeImage> GenerateMutableSourceImage(const UEdGraphPin* Pin, FMutab
 						{
 							int32 Dummy = -1; // TODO MTBL-1512
 							bool Dummy2 = false;
+							GenerationContext.CurrentTableColumnType = mu::ETableColumnType::Image;
 							bSuccess = GenerateTableColumn(TypedNodeTable, Pin, GeneratedTable, ColumnName, Property, 
 								FMutableSourceMeshData(), Dummy, Dummy, GenerationContext.CurrentLOD, Dummy, Dummy2, GenerationContext);
 
diff --git a/Engine/Plugins/Mutable/Source/CustomizableObjectEditor/Private/MuCOE/GenerateMutableSource/GenerateMutableSourceMaterial.cpp b/Engine/Plugins/Mutable/Source/CustomizableObjectEditor/Private/MuCOE/GenerateMutableSource/GenerateMutableSourceMaterial.cpp
new file mode 100644
--- /dev/null
+++ b/Engine/Plugins/Mutable/Source/CustomizableObjectEditor/Private/MuCOE/GenerateMutableSource/GenerateMutableSourceMaterial.cpp
@@ -0,0 +1,214 @@
+// Copyright Epic Games, Inc. All Rights Reserved.
+
+#include "MuCOE/GenerateMutableSource/GenerateMutableSourceMaterial.h"
+
+#include "MuCOE/CustomizableObjectCompiler.h"
+#include "MuCOE/CustomizableObjectGraph.h"
+#include "MuCOE/EdGraphSchema_CustomizableObject.h"
+#include "MuCOE/GenerateMutableSource/GenerateMutableSourceMacro.h"
+#include "MuCOE/GenerateMutableSource/GenerateMutableSourceTable.h"
+#include "MuCOE/GraphTraversal.h"
+#include "MuCOE/Nodes/CustomizableObjectNodeMaterial.h"
+#include "MuCOE/Nodes/CustomizableObjectNodeMaterialParameter.h"
+#include "MuCOE/Nodes/CustomizableObjectNodeMacroInstance.h"
+#include "MuCOE/Nodes/CustomizableObjectNodeTable.h"
+#include "MuCOE/Nodes/CustomizableObjectNodeTunnel.h"
+#include "MuCO/LoadUtils.h"
+#include "MuT/NodeMaterialConstant.h"
+#include "MuT/NodeMaterialSwitch.h"
+#include "MuT/NodeMaterialTable.h"
+#include "MuT/NodeMaterialVariation.h"
+#include "MuT/NodeMaterialParameter.h"
+
+#include "Materials/MaterialInterface.h"
+#include "Materials/MaterialInstance.h"
+
+
+#define LOCTEXT_NAMESPACE "CustomizableObjectEditor"
+
+
+mu::Ptr<mu::NodeMaterial> GenerateMutableSourceMaterial(const UEdGraphPin* Pin, FMutableGraphGenerationContext& GenerationContext)
+{
+	check(Pin)
+	RETURN_ON_CYCLE(*Pin, GenerationContext)
+
+	CheckNumOutputs(*Pin, GenerationContext);
+	
+	const UEdGraphSchema_CustomizableObject* Schema = GetDefault<UEdGraphSchema_CustomizableObject>();
+
+	UCustomizableObjectNode* Node = CastChecked<UCustomizableObjectNode>(Pin->GetOwningNode());
+
+	const FGeneratedKey Key(reinterpret_cast<void*>(&GenerateMutableSourceMaterial), *Pin, *Node, GenerationContext);
+	if (const FGeneratedData* Generated = GenerationContext.Generated.Find(Key))
+	{
+		return static_cast<mu::NodeMaterial*>(Generated->Node.get());
+	}
+
+	if (Node->IsNodeOutDatedAndNeedsRefresh())
+	{
+		Node->SetRefreshNodeWarning();
+	}
+
+	// Bool that determines if a node can be added to the cache of nodes.
+	// Most nodes need to be added to the cache but there are some that don't. For exampel, MacroInstanceNodes
+	bool bCacheNode = true;
+
+	mu::Ptr<mu::NodeMaterial> Result;
+	
+	if (const UCustomizableObjectNodeMaterialParameter* MaterialParameterNode = Cast<UCustomizableObjectNodeMaterialParameter>(Node))
+	{
+		mu::Ptr<mu::NodeMaterialParameter> MaterialNode = new mu::NodeMaterialParameter();
+		Result = MaterialNode;
+
+		MaterialNode->Name = MaterialParameterNode->GetParameterName(&GenerationContext.MacroNodesStack);
+		MaterialNode->UID = GenerationContext.GetNodeIdUnique(Node).ToString();
+
+		if (MaterialParameterNode->DefaultValue)
+		{
+			GenerationContext.MaterialParameterDefaultValues.Add(FName(MaterialParameterNode->GetParameterName()), MutablePrivate::LoadObject(MaterialParameterNode->DefaultValue));
+		}
+
+		GenerationContext.ParameterUIDataMap.Add(MaterialParameterNode->GetParameterName(&GenerationContext.MacroNodesStack), FMutableParameterData(
+			MaterialParameterNode->ParamUIMetadata,
+			EMutableParameterType::Material));
+	}
+
+	
+	else if (const UCustomizableObjectNodeTable* TypedNodeTable = Cast<UCustomizableObjectNodeTable>(Node))
+	{
+		//This node will add a default value in case of error
+		mu::Ptr<mu::NodeMaterialConstant> ConstantValue = new mu::NodeMaterialConstant();
+		ConstantValue->MaterialId = GenerationContext.CurrentReferencedMaterialIndex;
+
+		Result = ConstantValue;
+
+		if (Pin->PinType.PinCategory == Schema->PC_Material)
+		{
+			// Material pins have to skip the cache of nodes or they will return always the same column node
+			bCacheNode = false;
+		}
+
+		bool bSuccess = true;
+		UDataTable* DataTable = GetDataTable(TypedNodeTable, GenerationContext);
+
+		if (DataTable)
+		{
+			FString ColumnName = TypedNodeTable->GetPinColumnName(Pin);
+			FProperty* Property = TypedNodeTable->FindPinProperty(*Pin);
+
+			if (!Property)
+			{
+				FString Msg = FString::Printf(TEXT("Couldn't find the column [%s] in the data table's struct."), *ColumnName);
+				GenerationContext.Log(FText::FromString(Msg), Node);
+
+				bSuccess = false;
+			}
+
+			if (bSuccess)
+			{
+				if (UMaterialInstance* TableMaterial = TypedNodeTable->GetColumnDefaultAssetByType<UMaterialInstance>(Pin))
+				{
+					check(GenerationContext.ReferencedMaterials.IsValidIndex(GenerationContext.CurrentReferencedMaterialIndex));
+					UMaterialInterface* BaseMaterial = GenerationContext.ReferencedMaterials[GenerationContext.CurrentReferencedMaterialIndex];
+
+					// Checking if the reference material of the Table Node has the same parent as the material of the Material Node 
+					if (!BaseMaterial || TableMaterial->GetMaterial() != BaseMaterial->GetMaterial())
+					{
+						GenerationContext.Log(LOCTEXT("DifferentParentMaterial", "The Default Material of the Data Table and the Mesh Section must have the same Parent Material."), Node);
+						bSuccess = false;
+					}
+				}
+				else
+				{
+					FString TableColumnName = TypedNodeTable->GetPinColumnName(Pin);
+					FText Msg = FText::Format(LOCTEXT("DefaultValueNotFound", "Couldn't find a default value in the data table's struct for the column {0}. The default value is null or not a Material Instance."), FText::FromString(TableColumnName));
+					GenerationContext.Log(Msg, Node);
+					bSuccess = false;
+				}
+			}
+
+			if (bSuccess)
+			{
+				// Generating a new data table if not exists
+				mu::Ptr<mu::FTable> Table;
+				Table = GenerateMutableSourceTable(DataTable, TypedNodeTable, GenerationContext);
+
+				if (Table)
+				{
+					mu::Ptr<mu::NodeMaterialTable> MaterialTableNode = new mu::NodeMaterialTable();
+
+					// Materials use the parameter id as column names
+					ColumnName = GenerationContext.CurrentMaterialParameterId;
+
+					// Generating a new Float column if not exists
+					if (Table->FindColumn(ColumnName) == INDEX_NONE)
+					{
+						int32 Dummy = -1; // TODO MTBL-1512
+						bool Dummy2 = false;
+						GenerationContext.CurrentTableColumnType = mu::ETableColumnType::Material;
+						bSuccess = GenerateTableColumn(TypedNodeTable, Pin, Table, ColumnName, Property, FMutableSourceMeshData(), Dummy, Dummy, GenerationContext.CurrentLOD, Dummy, Dummy2, GenerationContext);
+
+						if (!bSuccess)
+						{
+							FString Msg = FString::Printf(TEXT("Failed to generate the mutable table column [%s]"), *ColumnName);
+							GenerationContext.Log(FText::FromString(Msg), Node);
+						}
+					}
+
+					if (bSuccess)
+					{
+						Result = MaterialTableNode;
+
+						MaterialTableNode->Table = Table;
+						MaterialTableNode->ColumnName = ColumnName;
+						MaterialTableNode->ParameterName = TypedNodeTable->ParameterName;
+						MaterialTableNode->bNoneOption = TypedNodeTable->bAddNoneOption;
+						MaterialTableNode->DefaultRowName = TypedNodeTable->DefaultRowName.ToString();
+					}
+				}
+				else
+				{
+					FString Msg = FString::Printf(TEXT("Couldn't generate a mutable table."), *ColumnName);
+					GenerationContext.Log(FText::FromString(Msg), Node);
+				}
+			}
+		}
+		else
+		{
+			GenerationContext.Log(LOCTEXT("MaterialTableError", "Couldn't find the data table of the node."), Node);
+		}
+	}
+
+	else if (const UCustomizableObjectNodeMacroInstance* TypedNodeMacro = Cast<UCustomizableObjectNodeMacroInstance>(Node))
+	{
+		bCacheNode = false;
+		Result = GenerateMutableSourceMacro<mu::NodeMaterial>(*Pin, GenerationContext, GenerateMutableSourceMaterial);
+	}
+
+	else if (const UCustomizableObjectNodeTunnel* TypedNodeTunnel = Cast<UCustomizableObjectNodeTunnel>(Node))
+	{
+		bCacheNode = false;
+		Result = GenerateMutableSourceMacro<mu::NodeMaterial>(*Pin, GenerationContext, GenerateMutableSourceMaterial);
+	}
+
+	else
+	{
+		GenerationContext.Log(LOCTEXT("UnimplementedNode", "Node type not implemented yet."), Node);
+	}
+
+	if (bCacheNode)
+	{
+		GenerationContext.Generated.Add(Key, FGeneratedData(Node, Result));
+		GenerationContext.GeneratedNodes.Add(Node);
+	}
+
+	if (Result)
+	{
+		Result->SetMessageContext(Node);
+	}
+
+	return Result;
+}
+
+#undef LOCTEXT_NAMESPACE
+
diff --git a/Engine/Plugins/Mutable/Source/CustomizableObjectEditor/Private/MuCOE/GenerateMutableSource/GenerateMutableSourceMaterial.h b/Engine/Plugins/Mutable/Source/CustomizableObjectEditor/Private/MuCOE/GenerateMutableSource/GenerateMutableSourceMaterial.h
new file mode 100644
--- /dev/null
+++ b/Engine/Plugins/Mutable/Source/CustomizableObjectEditor/Private/MuCOE/GenerateMutableSource/GenerateMutableSourceMaterial.h
@@ -0,0 +1,11 @@
+// Copyright Epic Games, Inc. All Rights Reserved.
+
+#pragma once
+
+#include "MuT/NodeMaterial.h"
+
+class UEdGraphPin;
+struct FMutableGraphGenerationContext;
+
+/** Convert a CustomizableObject Source Graph into a mutable source graph. */
+mu::Ptr<mu::NodeMaterial> GenerateMutableSourceMaterial(const UEdGraphPin* Pin, FMutableGraphGenerationContext& GenerationContext);
diff --git a/Engine/Plugins/Mutable/Source/CustomizableObjectEditor/Private/MuCOE/GenerateMutableSource/GenerateMutableSourceSurface.cpp b/Engine/Plugins/Mutable/Source/CustomizableObjectEditor/Private/MuCOE/GenerateMutableSource/GenerateMutableSourceSurface.cpp
--- a/Engine/Plugins/Mutable/Source/CustomizableObjectEditor/Private/MuCOE/GenerateMutableSource/GenerateMutableSourceSurface.cpp
+++ b/Engine/Plugins/Mutable/Source/CustomizableObjectEditor/Private/MuCOE/GenerateMutableSource/GenerateMutableSourceSurface.cpp
@@ -28,6 +28,7 @@
 #include "MuCOE/GenerateMutableSource/GenerateMutableSourceTable.h"
 #include "MuCOE/GenerateMutableSource/GenerateMutableSourceLayout.h"
 #include "MuCOE/GenerateMutableSource/GenerateMutableSourceMacro.h"
+#include "MuCOE/GenerateMutableSource/GenerateMutableSourceMaterial.h"
 #include "MuCOE/GraphTraversal.h"
 #include "MuCOE/MutableUtils.h"
 #include "MuCOE/Nodes/CustomizableObjectNodeCopyMaterial.h"
@@ -53,7 +54,7 @@
 #include "MuT/NodeMeshConstant.h"
 #include "MuT/NodeMeshFormat.h"
 #include "MuT/NodeMeshFragment.h"
-#include "MuT/NodeScalarConstant.h"
+#include "MuT/NodeMaterialConstant.h"
 #include "MuT/NodeSurfaceSwitch.h"
 #include "MuT/NodeSurfaceVariation.h"
 #include "MuT/UnrealPixelFormatOverride.h"
@@ -218,42 +219,31 @@ mu::Ptr<mu::NodeSurface> GenerateMutableSourceSurface(const UEdGraphPin * Pin, F
 			}
 		}
 
-		bool bTableMaterialPinLinked = TypedNodeMat->GetMaterialAssetPin() && FollowInputPin(*TypedNodeMat->GetMaterialAssetPin()) != nullptr;
+		bool bMaterialPinLinked = false;
+		const UEdGraphPin* ConnectedMaterialAssetPin = nullptr;
 
 		// Checking if we should not use the material of the table node even if it is linked to the material node
-		if (TypedNodeMat->GetMaterialAssetPin())
+		if (const UEdGraphPin* MaterialAssetPin = TypedNodeMat->GetMaterialAssetPin())
 		{
-			if (const UEdGraphPin* ConnectedPin = FollowInputPin(*TypedNodeMat->GetMaterialAssetPin()))
+			ConnectedMaterialAssetPin = FollowInputPin(*MaterialAssetPin);
+
+			if (ConnectedMaterialAssetPin)
 			{
 				//Check if the pin goes through a macro or tunnel node
-				if (const UEdGraphPin* MaterialAssetPin = GraphTraversal::FindIOPinSourceThroughMacroContext(*ConnectedPin, &GenerationContext.MacroNodesStack))
-				{
-					if (const UCustomizableObjectNodeTable* TypedNodeTable = Cast< UCustomizableObjectNodeTable >(MaterialAssetPin->GetOwningNode()))
-					{
-						if (UMaterialInstance* TableMaterial = TypedNodeTable->GetColumnDefaultAssetByType<UMaterialInstance>(MaterialAssetPin))
-						{
-							// Checking if the reference material of the Table Node has the same parent as the material of the Material Node 
-							if (!Material || TableMaterial->GetMaterial() != Material->GetMaterial())
-							{
-								bTableMaterialPinLinked = false;
-
-								GenerationContext.Log(LOCTEXT("DifferentParentMaterial", "The Default Material Instance of the Data Table must have the same Parent Material."), TypedNodeMat->GetMaterialNode());
-							}
-						}
-						else
-						{
-							FString TableColumnName = TypedNodeTable->GetPinColumnName(MaterialAssetPin);
+				ConnectedMaterialAssetPin = GraphTraversal::FindIOPinSourceThroughMacroContext(*ConnectedMaterialAssetPin, &GenerationContext.MacroNodesStack);
 
-							FText Msg = FText::Format(LOCTEXT("DefaultValueNotFound", "Couldn't find a default value in the data table's struct for the column {0}. The default value is null or not a Material Instance."), FText::FromString(TableColumnName));
-							GenerationContext.Log(Msg, Node);
-
-							bTableMaterialPinLinked = false;
-						}
-					}
-				}
+				GenerationContext.CurrentMaterialParameterId = MaterialAssetPin->PinId.ToString();
+				SurfNode->Material = GenerateMutableSourceMaterial(ConnectedMaterialAssetPin, GenerationContext);
 			}
 		}
 
+		if (!SurfNode->Material)
+		{
+			mu::Ptr<mu::NodeMaterialConstant> MaterialNode = new mu::NodeMaterialConstant();
+			MaterialNode->MaterialId = GenerationContext.CurrentReferencedMaterialIndex;
+			SurfNode->Material = MaterialNode;
+		}
+
 		int32 NumImages = TypedNodeMat->GetNumParameters(EMaterialParameterType::Texture);
 		SurfNode->Images.SetNum(NumImages);
 
@@ -319,7 +309,7 @@ mu::Ptr<mu::NodeSurface> GenerateMutableSourceSurface(const UEdGraphPin * Pin, F
 						{
 							//TODO(Max) UE-220247: Add support for multilayer materials
 							GenerationContext.CurrentMaterialTableParameter = ImageName;
-							GenerationContext.CurrentMaterialTableParameterId = ImageId.ParameterId.ToString();
+							GenerationContext.CurrentMaterialParameterId = ImageId.ParameterId.ToString();
 
 							ReferenceTexture = GroupProjectionImg.get() ? GroupProjectionReferenceTexture : nullptr;
 						
@@ -337,12 +327,10 @@ mu::Ptr<mu::NodeSurface> GenerateMutableSourceSurface(const UEdGraphPin * Pin, F
 								}
 							}
 
-							if (!ReferenceTexture && bTableMaterialPinLinked)
+							if (!ReferenceTexture && 
+								ConnectedMaterialAssetPin && Cast<UCustomizableObjectNodeTable>(ConnectedMaterialAssetPin->GetOwningNode()) != nullptr)
 							{
-								if (const UEdGraphPin* ConnectedPin = FollowInputPin(*TypedNodeMat->GetMaterialAssetPin()))
-								{
-									ReferenceTexture = FindReferenceImage(ConnectedPin, GenerationContext);
-								}
+								ReferenceTexture = FindReferenceImage(ConnectedMaterialAssetPin, GenerationContext);
 							}	
 
 							if (!ReferenceTexture)
@@ -416,14 +404,11 @@ mu::Ptr<mu::NodeSurface> GenerateMutableSourceSurface(const UEdGraphPin * Pin, F
 									}
 								}
 
-								if (bTableMaterialPinLinked)
+								// If the table material pin is linked to a table node, get all the textures of the current material parameter (CurrentMaterialTableParameter) from the Material Instances of the specified data table column.
+								// Then Generate a mutable table column with all these textures.
+								if (ConnectedMaterialAssetPin && Cast<UCustomizableObjectNodeTable>(ConnectedMaterialAssetPin->GetOwningNode()) != nullptr)
 								{
-									// If the table material pin is linked to a table node, get all the textures of the current material parameter (CurrentMaterialTableParameter) from the Material Instances of the specified data table column.
-									// Then Generate a mutable table column with all these textures.
-									if (const UEdGraphPin* ConnectedPin = FollowInputPin(*TypedNodeMat->GetMaterialAssetPin()))
-									{
-										return GenerateMutableSourceImage(ConnectedPin, GenerationContext, Props.TextureSize);
-									}
+									return GenerateMutableSourceImage(ConnectedMaterialAssetPin, GenerationContext, Props.TextureSize);
 								}
 
 								// Else
@@ -728,7 +713,7 @@ mu::Ptr<mu::NodeSurface> GenerateMutableSourceSurface(const UEdGraphPin * Pin, F
 			FNodeMaterialParameterId VectorId = TypedNodeMat->GetParameterId(EMaterialParameterType::Vector, VectorIndex);
 
 			GenerationContext.CurrentMaterialTableParameter = VectorName;
-			GenerationContext.CurrentMaterialTableParameterId = VectorId.ParameterId.ToString();
+			GenerationContext.CurrentMaterialParameterId = VectorId.ParameterId.ToString();
 
 			if (bVectorPinConnected)
 			{				
@@ -759,7 +744,7 @@ mu::Ptr<mu::NodeSurface> GenerateMutableSourceSurface(const UEdGraphPin * Pin, F
 			FNodeMaterialParameterId ScalarId = TypedNodeMat->GetParameterId(EMaterialParameterType::Scalar, ScalarIndex);
 
 			GenerationContext.CurrentMaterialTableParameter = ScalarName;
-			GenerationContext.CurrentMaterialTableParameterId = ScalarId.ParameterId.ToString();
+			GenerationContext.CurrentMaterialParameterId = ScalarId.ParameterId.ToString();
 
 			if (bScalarPinConnected)
 			{
@@ -779,43 +764,7 @@ mu::Ptr<mu::NodeSurface> GenerateMutableSourceSurface(const UEdGraphPin * Pin, F
 			}
 		}
 
-		// New method to pass the surface id as a scalar parameter
-		{
-			int32 MaterialIndex = NumScalar;
-			SurfNode->Scalars.SetNum(NumScalar + 1);
-
-			const UEdGraphPin* MaterialPin = TypedNodeMat->GetMaterialAssetPin();
-			check(MaterialPin);
-
-			//Encoding name for material material id parameter
-			FString MaterialName = "__MutableMaterialId";
-
-			if (bTableMaterialPinLinked)
-			{
-				if (const UEdGraphPin* ConnectedPin = FollowInputPin(*MaterialPin))
-				{					
-					// The 'None' option will use the CurrentReferencedMaterialIndex.
-					check(GenerationContext.CurrentReferencedMaterialIndex != INDEX_NONE);
-					
-					// Adding pin ID to support multiple material columns
-					MaterialName += MaterialPin->PinId.ToString();
-
-					GenerationContext.CurrentMaterialTableParameterId = MaterialName;
-					mu::Ptr<mu::NodeScalar> ScalarNode = GenerateMutableSourceFloat(ConnectedPin, GenerationContext);
-
-					SurfNode->Scalars[MaterialIndex].Scalar = ScalarNode;
-					SurfNode->Scalars[MaterialIndex].Name = MaterialName;
-				}
-			}
-			else
-			{
-				mu::Ptr<mu::NodeScalarConstant> ScalarNode = new mu::NodeScalarConstant();
-				ScalarNode->Value = GenerationContext.CurrentReferencedMaterialIndex;
-
-				SurfNode->Scalars[MaterialIndex].Scalar = ScalarNode;
-				SurfNode->Scalars[MaterialIndex].Name = MaterialName;
-			}
-		}
+		
 
 		for (const FString& Tag : TypedNodeMat->GetEnableTags(&GenerationContext.MacroNodesStack))
 		{
diff --git a/Engine/Plugins/Mutable/Source/CustomizableObjectEditor/Private/MuCOE/GenerateMutableSource/GenerateMutableSourceTable.cpp b/Engine/Plugins/Mutable/Source/CustomizableObjectEditor/Private/MuCOE/GenerateMutableSource/GenerateMutableSourceTable.cpp
--- a/Engine/Plugins/Mutable/Source/CustomizableObjectEditor/Private/MuCOE/GenerateMutableSource/GenerateMutableSourceTable.cpp
+++ b/Engine/Plugins/Mutable/Source/CustomizableObjectEditor/Private/MuCOE/GenerateMutableSource/GenerateMutableSourceTable.cpp
@@ -468,18 +468,17 @@ bool FillTableColumn(const UCustomizableObjectNodeTable* TableNode, mu::Ptr<mu::
 				LogRowGenerationMessage(TableNode, DataTablePtr, GenerationContext, Warning.ToString(), RowName);
 			}
 			
-			FString EncodedSwitchParameterName = "__MutableMaterialId";
-			if (ColumnName.Contains(EncodedSwitchParameterName))
+			if (GenerationContext.CurrentTableColumnType == mu::ETableColumnType::Material)
 			{
 				CurrentColumn = MutableTable.get()->FindColumn(ColumnName);
 
 				if (CurrentColumn == -1)
 				{
-					CurrentColumn = MutableTable->AddColumn(ColumnName, mu::ETableColumnType::Scalar);
+					CurrentColumn = MutableTable->AddColumn(ColumnName, mu::ETableColumnType::Material);
 				}
 
 				int32 ReferenceMaterialId = GenerationContext.ReferencedMaterials.AddUnique(MaterialInstance);
-				MutableTable->SetCell(CurrentColumn, RowId, (float)ReferenceMaterialId);
+				MutableTable->SetCell(CurrentColumn, RowId, ReferenceMaterialId);
 
 				return true;
 			}
@@ -496,7 +495,7 @@ bool FillTableColumn(const UCustomizableObjectNodeTable* TableNode, mu::Ptr<mu::
 
 			MaterialInstance->GetMaterial()->GetAllParameterInfoOfType(EMaterialParameterType::Texture, ParameterInfos, ParameterGuids);
 			
-			FGuid ParameterId(GenerationContext.CurrentMaterialTableParameterId);
+			FGuid ParameterId(GenerationContext.CurrentMaterialParameterId);
 			int32 ParameterIndex = ParameterGuids.Find(ParameterId);
 
 			if (ParameterIndex != INDEX_NONE && ParameterInfos[ParameterIndex].Name == GenerationContext.CurrentMaterialTableParameter)
@@ -1313,15 +1312,12 @@ bool GenerateNoneRow(const UCustomizableObjectNodeTable& TableNode, const UEdGra
 	const UEdGraphSchema_CustomizableObject* Schema = GetDefault<UEdGraphSchema_CustomizableObject>();
 	check(Schema);
 
-	//Ensure that we are generating a material column and not a Texture column
-	FString EncodedSwitchParameterName = "__MutableMaterialId";
-
-	if (Pin->PinType.PinCategory == Schema->PC_Material && ColumnName.Contains(EncodedSwitchParameterName))
+	if (Pin->PinType.PinCategory == Schema->PC_Material && GenerationContext.CurrentTableColumnType == mu::ETableColumnType::Material)
 	{
 		// Add a scalar column if it has not been added yet
 		if (ColumnId == INDEX_NONE)
 		{
-			ColumnId = MutableTable->AddColumn(ColumnName, mu::ETableColumnType::Scalar);
+			ColumnId = MutableTable->AddColumn(ColumnName, mu::ETableColumnType::Material);
 		}
 
 		MutableTable->SetCell(ColumnId, RowId, GenerationContext.CurrentReferencedMaterialIndex);
diff --git a/Engine/Plugins/Mutable/Source/CustomizableObjectEditor/Private/MuCOE/GraphTraversal.cpp b/Engine/Plugins/Mutable/Source/CustomizableObjectEditor/Private/MuCOE/GraphTraversal.cpp
--- a/Engine/Plugins/Mutable/Source/CustomizableObjectEditor/Private/MuCOE/GraphTraversal.cpp
+++ b/Engine/Plugins/Mutable/Source/CustomizableObjectEditor/Private/MuCOE/GraphTraversal.cpp
@@ -11,6 +11,7 @@
 #include "MuCOE/Nodes/CustomizableObjectNodeModifierExtendMeshSection.h"
 #include "MuCOE/Nodes/CustomizableObjectNodeExternalPin.h"
 #include "MuCOE/Nodes/CustomizableObjectNodeMaterial.h"
+#include "MuCOE/Nodes/CustomizableObjectNodeMaterialParameter.h"
 #include "MuCOE/Nodes/CustomizableObjectNodeMaterialVariation.h"
 #include "MuCOE/Nodes/CustomizableObjectNodeMeshMorph.h"
 #include "MuCOE/Nodes/CustomizableObjectNodeMeshMorphStackApplication.h"
diff --git a/Engine/Plugins/Mutable/Source/CustomizableObjectEditor/Private/MuCOE/SMutableGraphViewer.cpp b/Engine/Plugins/Mutable/Source/CustomizableObjectEditor/Private/MuCOE/SMutableGraphViewer.cpp
--- a/Engine/Plugins/Mutable/Source/CustomizableObjectEditor/Private/MuCOE/SMutableGraphViewer.cpp
+++ b/Engine/Plugins/Mutable/Source/CustomizableObjectEditor/Private/MuCOE/SMutableGraphViewer.cpp
@@ -204,6 +204,13 @@ static const char* MutableNodeNames[] =
 	"Matrix",
 	"MatrixConstant",
 	"MatrixParameter",
+
+	"Material",
+	"MaterialConstant",
+	"MaterialTable",
+	"MaterialSwitch",
+	"MaterialVariation",
+	"MaterialParameter",
 };
 
 static_assert(UE_ARRAY_COUNT(MutableNodeNames) == SIZE_T(mu::Node::EType::Count));
diff --git a/Engine/Plugins/Mutable/Source/MutableRuntime/Internal/MuR/Operations.h b/Engine/Plugins/Mutable/Source/MutableRuntime/Internal/MuR/Operations.h
--- a/Engine/Plugins/Mutable/Source/MutableRuntime/Internal/MuR/Operations.h
+++ b/Engine/Plugins/Mutable/Source/MutableRuntime/Internal/MuR/Operations.h
@@ -39,6 +39,7 @@ namespace mu
         ST_CONSTANT,
 		ED_CONSTANT,
     	MA_CONSTANT,
+		MI_CONSTANT,
 
         //! User parameter
         BO_PARAMETER,
@@ -50,6 +51,7 @@ namespace mu
 		ME_PARAMETER,
 		ST_PARAMETER,
     	MA_PARAMETER,
+    	MI_PARAMETER,
 
 		//! A referenced, but opaque engine resource
 		IM_REFERENCE,
@@ -64,6 +66,7 @@ namespace mu
         LA_CONDITIONAL,
         IN_CONDITIONAL,
 		ED_CONDITIONAL,
+		MI_CONDITIONAL,
 
         //! Select one of several values depending on an int input
         NU_SWITCH,
@@ -74,6 +77,7 @@ namespace mu
         LA_SWITCH,
         IN_SWITCH,
 		ED_SWITCH,
+		MI_SWITCH,
 
         //-----------------------------------------------------------------------------------------
         // Boolean operations
@@ -312,6 +316,9 @@ namespace mu
 		//! Add overlay material to an instance
 		IN_ADDOVERLAYMATERIAL,
 
+		//! Add a material to an instance
+		IN_ADDMATERIAL,
+
         //-----------------------------------------------------------------------------------------
         // Layout operations
         //-----------------------------------------------------------------------------------------
@@ -1030,12 +1037,12 @@ namespace mu
 			ADDRESS ExtensionDataName;
 		};
 
-		struct InstanceAddOverlayMaterialArgs
+		struct InstanceAddMaterialArgs
 		{
 			// This is a reference to an op that produces the Instance that the material will
 			// be added to.
 			ADDRESS Instance;
-			ADDRESS OverlayMaterialId;
+			ADDRESS Material;
 		};
 
         struct LayoutPackArgs
@@ -1090,6 +1097,7 @@ namespace mu
 		String,
 		ExtensionData,
 		Matrix,
+		Material,
 
 		// Supporting data types : Never returned as an actual data type for any operation.
 		Shape,
diff --git a/Engine/Plugins/Mutable/Source/MutableRuntime/Internal/MuR/ParametersPrivate.h b/Engine/Plugins/Mutable/Source/MutableRuntime/Internal/MuR/ParametersPrivate.h
--- a/Engine/Plugins/Mutable/Source/MutableRuntime/Internal/MuR/ParametersPrivate.h
+++ b/Engine/Plugins/Mutable/Source/MutableRuntime/Internal/MuR/ParametersPrivate.h
@@ -13,6 +13,7 @@
 
 #include "Engine/Texture.h"
 #include "Engine/SkeletalMesh.h"
+#include "Materials/MaterialInterface.h"
 
 #define UE_API MUTABLERUNTIME_API
 
@@ -258,9 +259,10 @@ namespace Private
 	using FParamSkeletalMeshType = TStrongObjectPtr<USkeletalMesh>;
 	using FParamStringType = Private::TIndirectObject<FString>;
 	using FParamMatrixType = Private::TIndirectObject<FMatrix44f>;
+	using FParamMaterialType = TStrongObjectPtr<UMaterialInterface>;
 	
 	using FParameterValue = TVariant<
-            FParamBoolType, FParamIntType, FParamFloatType, FParamColorType, FParamProjectorType, FParamTextureType, FParamSkeletalMeshType, FParamStringType, FParamMatrixType>;
+            FParamBoolType, FParamIntType, FParamFloatType, FParamColorType, FParamProjectorType, FParamTextureType, FParamSkeletalMeshType, FParamStringType, FParamMatrixType, FParamMaterialType>;
 
     // static_assert to track ParameterValue size changes. It is ok to change if needed.
     static_assert(sizeof(FParameterValue) == 8*4, "ParameterValue size has changed.");
diff --git a/Engine/Plugins/Mutable/Source/MutableRuntime/Internal/MuR/SystemPrivate.h b/Engine/Plugins/Mutable/Source/MutableRuntime/Internal/MuR/SystemPrivate.h
--- a/Engine/Plugins/Mutable/Source/MutableRuntime/Internal/MuR/SystemPrivate.h
+++ b/Engine/Plugins/Mutable/Source/MutableRuntime/Internal/MuR/SystemPrivate.h
@@ -11,6 +11,7 @@
 #include "MuR/ModelPrivate.h"
 #include "MuR/Mesh.h"
 #include "MuR/Instance.h"
+#include "MuR/Material.h"
 #include "MuR/ParametersPrivate.h"
 #include "MuR/MutableTrace.h"
 
@@ -513,6 +514,7 @@ namespace mu
 		TMemoryTrackedArray<TSharedPtr<const String>> StringResults;
 		TMemoryTrackedArray<TSharedPtr<const FExtensionData>> ExtensionDataResults;
 		TMemoryTrackedArray<FMatrix44f> MatrixResults;
+		TMemoryTrackedArray<FMaterial> MaterialResults;
 
 		/** */
 		inline const ExecutionIndex& GetRangeIndex(uint32_t i)
@@ -571,6 +573,7 @@ namespace mu
 				StringResults.Add(nullptr);
 				MatrixResults.Add(FMatrix44f());
 				ExtensionDataResults.Add(nullptr);
+				MaterialResults.Add(FMaterial());
 			}
 		}
 
@@ -736,6 +739,16 @@ namespace mu
 			return MatrixResults[Data->DataTypeIndex];
 		}
 
+		FMaterial GetMaterial(FCacheAddress at)
+		{
+			if (!at.At) return {};
+			const FOpExecutionData* Data = OpExecutionData.get_ptr(at);
+			if (!Data) return {};
+
+			check(Data->DataType == EDataType::Material);
+			return MaterialResults[Data->DataTypeIndex];
+		}
+
 		FProjector GetProjector(FCacheAddress at)
 		{
 			if (!at.At) return FProjector();
@@ -929,6 +942,26 @@ namespace mu
 			}
 			check(Data->DataTypeIndex != 0);
 		}
+
+		void SetMaterial(FCacheAddress at, const FMaterial& v)
+		{
+			check(at.At < OpExecutionData.size_code());
+			FOpExecutionData* Data = OpExecutionData.get_ptr(at);
+			check(Data->DataType == EDataType::Material || Data->DataType == EDataType::None);
+			Data->DataType = EDataType::Material;
+			Data->IsValueValid = true;
+
+			if (!Data->DataTypeIndex)
+			{
+				Data->DataTypeIndex = MaterialResults.Num();
+				MaterialResults.Add(v);
+			}
+			else
+			{
+				MaterialResults[Data->DataTypeIndex] = v;
+			}
+			check(Data->DataTypeIndex != 0);
+		}
 		
 		void SetProjector(FCacheAddress at, const FProjector& v)
 		{
diff --git a/Engine/Plugins/Mutable/Source/MutableRuntime/Private/MuR/CodeRunner.cpp b/Engine/Plugins/Mutable/Source/MutableRuntime/Private/MuR/CodeRunner.cpp
--- a/Engine/Plugins/Mutable/Source/MutableRuntime/Private/MuR/CodeRunner.cpp
+++ b/Engine/Plugins/Mutable/Source/MutableRuntime/Private/MuR/CodeRunner.cpp
@@ -13,6 +13,7 @@
 #include "Misc/AssertionMacros.h"
 #include "MuR/ImagePrivate.h"
 #include "MuR/Instance.h"
+#include "MuR/Material.h"
 #include "MuR/Mesh.h"
 #include "MuR/MeshBufferSet.h"
 #include "MuR/Model.h"
@@ -345,6 +346,7 @@ namespace mu
             case EDataType::Layout:			StoreLayout( cat, LoadLayout(rat) ); break;
             case EDataType::Instance:		StoreInstance( cat, LoadInstance(rat) ); break;
 			case EDataType::ExtensionData:	StoreExtensionData( cat, LoadExtensionData(rat) ); break;
+			case EDataType::Material:		StoreMaterial( cat, LoadMaterial(rat) ); break;
             default:
                 // Not implemented
                 check( false );
@@ -403,6 +405,7 @@ namespace mu
 				case EDataType::Layout:			StoreLayout(Item, nullptr); break;
 				case EDataType::Instance:		StoreInstance(Item, nullptr); break;
 				case EDataType::ExtensionData:	StoreExtensionData(Item, MakeShared<FExtensionData>()); break;
+				case EDataType::Material:		StoreMaterial(Item, {}); break;
 				default:
 					// Not implemented
 					check(false);
@@ -489,6 +492,7 @@ namespace mu
             case EDataType::Layout:			StoreLayout(ItemAddress, LoadLayout(ResultAddress)); break;
             case EDataType::Instance:		StoreInstance(ItemAddress, LoadInstance(ResultAddress)); break;
 			case EDataType::ExtensionData:	StoreExtensionData(ItemAddress, LoadExtensionData(ResultAddress)); break;
+			case EDataType::Material:		StoreMaterial(ItemAddress, LoadMaterial(ResultAddress)); break;
             default:
                 // Not implemented
                 check(false);
@@ -893,17 +897,17 @@ namespace mu
 		}
 		case EOpType::IN_ADDOVERLAYMATERIAL:
 		{
-			OP::InstanceAddOverlayMaterialArgs Args = Program.GetOpArgs<OP::InstanceAddOverlayMaterialArgs>(Item.At);
+			OP::InstanceAddMaterialArgs Args = Program.GetOpArgs<OP::InstanceAddMaterialArgs>(Item.At);
 			switch (Item.Stage)
 			{
 			case 0:
 			{
 				// Must pass in a Material op
-				check(Args.OverlayMaterialId);
+				check(Args.Material);
 
 				AddOp(FScheduledOp(Item.At, Item, 1),
 					FScheduledOp(Args.Instance, Item),
-					FScheduledOp(Args.OverlayMaterialId, Item));
+					FScheduledOp(Args.Material, Item));
 
 				break;
 			}
@@ -921,7 +925,7 @@ namespace mu
 					Result = mu::CloneOrTakeOver<FInstance>(Base);
 				}
 
-				Result->AddOverlayMaterial(0, LoadScalar(FCacheAddress(Args.OverlayMaterialId, Item)));
+				Result->AddOverlayMaterial(0, LoadMaterial(FCacheAddress(Args.Material, Item)));
 
 				StoreInstance(Item, Result);
 				break;
@@ -934,6 +938,44 @@ namespace mu
 			break;
 		}
 
+		case EOpType::IN_ADDMATERIAL:
+		{
+			OP::InstanceAddMaterialArgs Args = Program.GetOpArgs<OP::InstanceAddMaterialArgs>(Item.At);
+			switch (Item.Stage)
+			{
+			case 0:
+				AddOp(FScheduledOp(Item.At, Item, 1),
+					FScheduledOp(Args.Instance, Item),
+					FScheduledOp(Args.Material, Item));
+				break;
+
+			case 1:
+			{
+				TSharedPtr<const FInstance> Base = LoadInstance(FCacheAddress(Args.Instance, Item));
+
+				TSharedPtr<FInstance> Result;
+				if (Base)
+				{
+					Result = mu::CloneOrTakeOver<FInstance>(Base);
+				}
+				else
+				{
+					Result = MakeShared<FInstance>();
+				}
+
+				// Material
+				Result->AddMaterial(0, 0, 0, LoadMaterial(FCacheAddress(Args.Material, Item)));
+
+				StoreInstance(Item, Result);
+				break;
+			}
+
+			default:
+				check(false);
+			}
+
+			break;
+		}
         default:
                 check(false);
         }
@@ -1113,6 +1155,17 @@ namespace mu
             break;
 		}
 
+		case EOpType::MI_CONSTANT:
+		{
+			OP::ResourceConstantArgs Args = Program.GetOpArgs<OP::ResourceConstantArgs>(Item.At);
+
+			FMaterial Source;
+			Source.MaterialId = Args.value;
+
+			StoreMaterial(Item, Source);
+			break;
+		}
+
         default:
             if (type!=EOpType::NONE)
             {
@@ -6098,6 +6151,35 @@ namespace mu
 		}
     }
 
+	//---------------------------------------------------------------------------------------------
+	void CodeRunner::RunCode_Material(const FScheduledOp& Item, const FParameters* InParams, const FModel* InModel)
+	{
+		MUTABLE_CPUPROFILER_SCOPE(RunCode_Material);
+
+		const FProgram& Program = Model->GetPrivate()->Program;
+
+		EOpType type = Program.GetOpType(Item.At);
+
+		switch (type)
+		{
+		case EOpType::MI_PARAMETER:
+		{
+			MUTABLE_CPUPROFILER_SCOPE(RunCodeImageDesc_MI_PARAMETER);
+			check(Item.Stage == 0);
+
+			OP::ParameterArgs Args = Program.GetOpArgs<OP::ParameterArgs>(Item.At);
+			TSharedPtr<FRangeIndex> Index = BuildCurrentOpRangeIndex(Item, InParams, InModel, Args.variable);
+
+			FMaterial Material;
+			Material.Material = TStrongObjectPtr(InParams->GetMaterialValue(Args.variable, Index.Get()));
+			
+			StoreMaterial(Item, Material);
+			break;
+		}
+		default: unimplemented();
+		}
+	}
+
 
     //---------------------------------------------------------------------------------------------
     void CodeRunner::RunCode_Layout(const FScheduledOp& Item, const FModel* InModel )
@@ -6304,12 +6386,14 @@ namespace mu
         case EOpType::LA_CONDITIONAL:
         case EOpType::IN_CONDITIONAL:
 		case EOpType::ED_CONDITIONAL:
+		case EOpType::MI_CONDITIONAL:
             RunCode_Conditional(Item, InModel.Get());
             break;
 
         case EOpType::ME_CONSTANT:
 		case EOpType::IM_CONSTANT:
 		case EOpType::ED_CONSTANT:
+		case EOpType::MI_CONSTANT:
             RunCode_ConstantResource(Item, InModel.Get());
             break;
 
@@ -6321,6 +6405,7 @@ namespace mu
         case EOpType::LA_SWITCH:
         case EOpType::IN_SWITCH:
 		case EOpType::ED_SWITCH:
+		case EOpType::MI_SWITCH:
             RunCode_Switch(Item, InModel.Get());
             break;
 
@@ -6377,7 +6462,10 @@ namespace mu
 			case EDataType::Matrix:
 				RunCode_Matrix(Item, pParams, InModel.Get());
 				break;
-				
+
+			case EDataType::Material:
+				RunCode_Material(Item, pParams, InModel.Get());
+				break;
 			default:
 				check(false);
 				break;
diff --git a/Engine/Plugins/Mutable/Source/MutableRuntime/Private/MuR/CodeRunner.h b/Engine/Plugins/Mutable/Source/MutableRuntime/Private/MuR/CodeRunner.h
--- a/Engine/Plugins/Mutable/Source/MutableRuntime/Private/MuR/CodeRunner.h
+++ b/Engine/Plugins/Mutable/Source/MutableRuntime/Private/MuR/CodeRunner.h
@@ -129,6 +129,7 @@ namespace  mu
         void RunCode_Colour(const FScheduledOp&, const FParameters*, const FModel* );
         void RunCode_Projector(const FScheduledOp&, const FParameters*, const FModel* );
         void RunCode_Matrix(const FScheduledOp&, const FParameters*, const FModel* );
+        void RunCode_Material(const FScheduledOp&, const FParameters*, const FModel* );
 
 		void RunCodeImageDesc(const FScheduledOp&, const FParameters*, const FModel*, uint32 LodMask);
 	public:
@@ -457,6 +458,11 @@ namespace  mu
 			return System->WorkingMemoryManager.CurrentInstanceCache->GetMatrix(From);
 		}
 
+		inline FMaterial LoadMaterial(const FCacheAddress& From)
+		{
+			return System->WorkingMemoryManager.CurrentInstanceCache->GetMaterial(From);
+		}
+
 		inline TSharedPtr<const String> LoadString(const FCacheAddress& From)
 		{
 			return System->WorkingMemoryManager.CurrentInstanceCache->GetString(From);
@@ -526,6 +532,11 @@ namespace  mu
 		{
 			System->WorkingMemoryManager.CurrentInstanceCache->SetMatrix(To, Value);
 		}
+
+		inline void StoreMaterial(const FCacheAddress& To, const FMaterial& Value)
+		{
+			System->WorkingMemoryManager.CurrentInstanceCache->SetMaterial(To, Value);
+		}
     	
 		inline void StoreProjector(const FCacheAddress& To, const FProjector& Value)
 		{
diff --git a/Engine/Plugins/Mutable/Source/MutableRuntime/Private/MuR/Instance.cpp b/Engine/Plugins/Mutable/Source/MutableRuntime/Private/MuR/Instance.cpp
--- a/Engine/Plugins/Mutable/Source/MutableRuntime/Private/MuR/Instance.cpp
+++ b/Engine/Plugins/Mutable/Source/MutableRuntime/Private/MuR/Instance.cpp
@@ -319,12 +319,19 @@ namespace mu
         return NAME_None;
     }
 
-	
-	int32 FInstance::GetOverlayMaterial(int32 ComponentIndex) const
+	const FMaterial& FInstance::GetOverlayMaterial(int32 ComponentIndex) const
 	{
 		check(Components.IsValidIndex(ComponentIndex));
-		check(Components[ComponentIndex].OverlayMaterialId < 1 << 22);
-		return static_cast<int32>(Components[ComponentIndex].OverlayMaterialId);
+		return Components[ComponentIndex].OverlayMaterial;
+	}
+
+	const FMaterial& FInstance::GetMaterial(int32 ComponentIndex, int32 LODIndex, int32 SurfaceIndex) const
+	{
+		check(Components.IsValidIndex(ComponentIndex));
+		check(Components[ComponentIndex].LODs.IsValidIndex(LODIndex));
+		check(Components[ComponentIndex].LODs[LODIndex].Surfaces.IsValidIndex(SurfaceIndex));
+
+		return Components[ComponentIndex].LODs[LODIndex].Surfaces[SurfaceIndex].Material;
 	}
 
     
@@ -501,14 +508,34 @@ namespace mu
     }
 	
 
-	void FInstance::AddOverlayMaterial(int32 ComponentIndex, float OverlayMaterialId)
+	void FInstance::AddMaterial(int32 ComponentIndex, int32 LODIndex, int32 SurfaceIndex, FMaterial Material)
+	{
+		// Automatically create the necessary lods and components
+		while (ComponentIndex >= Components.Num())
+		{
+			AddComponent();
+		}
+		while (LODIndex >= Components[ComponentIndex].LODs.Num())
+		{
+			AddLOD(ComponentIndex);
+		}
+		while (SurfaceIndex >= Components[ComponentIndex].LODs[LODIndex].Surfaces.Num())
+		{
+			AddSurface(LODIndex, ComponentIndex);
+		}
+		
+		Components[ComponentIndex].LODs[LODIndex].Surfaces[SurfaceIndex].Material = Material;
+	}
+
+
+	void FInstance::AddOverlayMaterial(int32 ComponentIndex, FMaterial OverlayMaterial)
 	{
 		while (ComponentIndex >= Components.Num())
 		{
 			AddComponent();
 		}
 
-		Components[ComponentIndex].OverlayMaterialId = OverlayMaterialId;
+		Components[ComponentIndex].OverlayMaterial = OverlayMaterial;
 	}
 	
     
diff --git a/Engine/Plugins/Mutable/Source/MutableRuntime/Private/MuR/Material.cpp b/Engine/Plugins/Mutable/Source/MutableRuntime/Private/MuR/Material.cpp
new file mode 100644
--- /dev/null
+++ b/Engine/Plugins/Mutable/Source/MutableRuntime/Private/MuR/Material.cpp
@@ -0,0 +1,28 @@
+// Copyright Epic Games, Inc. All Rights Reserved.
+
+
+#include "MuR/Material.h"
+
+#include "Containers/UnrealString.h"
+#include "HAL/LowLevelMemTracker.h"
+#include "HAL/PlatformCrt.h"
+#include "HAL/UnrealMemory.h"
+#include "MuR/MutableTrace.h"
+
+namespace mu
+{
+	TSharedPtr<FMaterial> FMaterial::Clone() const
+	{
+		TSharedPtr<FMaterial> Result = MakeShared<FMaterial>();
+
+		Result->Material = Material;
+		Result->MaterialId = MaterialId;
+
+		return Result;
+	}
+
+	bool FMaterial::operator==(const FMaterial& Other) const
+	{
+		return MaterialId == Other.MaterialId && Material == Other.Material;
+	};
+}
diff --git a/Engine/Plugins/Mutable/Source/MutableRuntime/Private/MuR/Model.cpp b/Engine/Plugins/Mutable/Source/MutableRuntime/Private/MuR/Model.cpp
--- a/Engine/Plugins/Mutable/Source/MutableRuntime/Private/MuR/Model.cpp
+++ b/Engine/Plugins/Mutable/Source/MutableRuntime/Private/MuR/Model.cpp
@@ -21,6 +21,7 @@
 
 class UTexture;
 class USkeletalMesh;
+class UMaterialInterface;
 
 
 namespace mu
@@ -515,6 +516,18 @@ namespace mu
     		}
     	}
 
+		if (MaterialDefaults)
+		{
+			for (const TPair<FName, TObjectPtr<UMaterialInterface>>& Pair : *MaterialDefaults)
+			{
+				int32 Index = pRes->GetPrivate()->Find(Pair.Key.ToString());
+				check(Index != INDEX_NONE);
+				check(pRes->GetType(Index) == EParameterType::Material);
+
+				pRes->SetMaterialValue(Index, Pair.Value);
+			}
+		}
+
         return pRes;
     }
 
diff --git a/Engine/Plugins/Mutable/Source/MutableRuntime/Private/MuR/Operations.cpp b/Engine/Plugins/Mutable/Source/MutableRuntime/Private/MuR/Operations.cpp
--- a/Engine/Plugins/Mutable/Source/MutableRuntime/Private/MuR/Operations.cpp
+++ b/Engine/Plugins/Mutable/Source/MutableRuntime/Private/MuR/Operations.cpp
@@ -28,6 +28,7 @@ namespace mu
 		{ EDataType::String,		},	// ST_CONSTANT
 		{ EDataType::ExtensionData,	},	// ED_CONSTANT
 		{ EDataType::Matrix,		},	// MA_CONSTANT
+		{ EDataType::Material,		},	// MI_CONSTANT
 
 		{ EDataType::Bool,			},	// BO_PARAMETER
 		{ EDataType::Int,			},	// NU_PARAMETER
@@ -38,6 +39,7 @@ namespace mu
 		{ EDataType::Mesh,			},	// ME_PARAMETER
 		{ EDataType::String,		},	// ST_PARAMETER
 		{ EDataType::Matrix,		},	// MA_PARAMETER
+		{ EDataType::Material,		},	// MI_PARAMETER
 
 		{ EDataType::Image,			},	// IM_REFERENCE
 		{ EDataType::Mesh,			},	// ME_REFERENCE
@@ -50,6 +52,7 @@ namespace mu
 		{ EDataType::Layout,		},	// LA_CONDITIONAL
 		{ EDataType::Instance,		},	// IN_CONDITIONAL
 		{ EDataType::ExtensionData,	},	// ED_CONDITIONAL
+		{ EDataType::Material,		},	// MI_CONDITIONAL
 
 		{ EDataType::Int,			},	// NU_SWITCH
 		{ EDataType::Scalar,		},	// SC_SWITCH
@@ -59,6 +62,7 @@ namespace mu
 		{ EDataType::Layout,		},	// LA_SWITCH
 		{ EDataType::Instance,		},	// IN_SWITCH
 		{ EDataType::ExtensionData,	},	// ED_SWITCH
+		{ EDataType::Material,		},	// MI_SWITCH
 
 		{ EDataType::Bool,			},	// BO_EQUAL_SC_CONST
 		{ EDataType::Bool,			},	// BO_AND
@@ -134,6 +138,7 @@ namespace mu
 		{ EDataType::Instance,		},	// IN_ADDLOD
 		{ EDataType::Instance,		},	// IN_ADDEXTENSIONDATA
 		{ EDataType::Instance,		},	// IN_ADDOVERLAYMATERIAL
+		{ EDataType::Instance,		},	// IN_ADDMATERIAL
 
 		{ EDataType::Layout,		},	// LA_PACK
 		{ EDataType::Layout,		},	// LA_MERGE
@@ -169,6 +174,7 @@ namespace mu
         case EOpType::PR_CONSTANT:
 		case EOpType::ED_CONSTANT:
 		case EOpType::MA_CONSTANT:
+		case EOpType::MI_CONSTANT:
         case EOpType::BO_PARAMETER:
         case EOpType::NU_PARAMETER:
         case EOpType::SC_PARAMETER:
@@ -177,6 +183,7 @@ namespace mu
 		case EOpType::IM_PARAMETER:
 		case EOpType::ME_PARAMETER:
 		case EOpType::MA_PARAMETER:
+		case EOpType::MI_PARAMETER:
 		case EOpType::IM_REFERENCE:
 		case EOpType::ME_REFERENCE:
 			break;
@@ -196,6 +203,7 @@ namespace mu
         case EOpType::LA_CONDITIONAL:
         case EOpType::IN_CONDITIONAL:
 		case EOpType::ED_CONDITIONAL:
+		case EOpType::MI_CONDITIONAL:
         {
 			OP::ConditionalArgs args = program.GetOpArgs<OP::ConditionalArgs>(at);
             f(args.condition );
@@ -212,6 +220,7 @@ namespace mu
         case EOpType::LA_SWITCH:
         case EOpType::IN_SWITCH:
 		case EOpType::ED_SWITCH:
+		case EOpType::MI_SWITCH:
         {
 			const uint8* Data = program.GetOpArgsPointer(at);
 			
@@ -822,11 +831,12 @@ namespace mu
 		}
 
 		case EOpType::IN_ADDOVERLAYMATERIAL:
+		case EOpType::IN_ADDMATERIAL:
 		{
-			const OP::InstanceAddOverlayMaterialArgs Args = program.GetOpArgs<OP::InstanceAddOverlayMaterialArgs>(at);
+			const OP::InstanceAddMaterialArgs Args = program.GetOpArgs<OP::InstanceAddMaterialArgs>(at);
 
 			f(Args.Instance);
-			f(Args.OverlayMaterialId);
+			f(Args.Material);
 
 			break;
 		}
diff --git a/Engine/Plugins/Mutable/Source/MutableRuntime/Private/MuR/Parameters.cpp b/Engine/Plugins/Mutable/Source/MutableRuntime/Private/MuR/Parameters.cpp
--- a/Engine/Plugins/Mutable/Source/MutableRuntime/Private/MuR/Parameters.cpp
+++ b/Engine/Plugins/Mutable/Source/MutableRuntime/Private/MuR/Parameters.cpp
@@ -9,6 +9,7 @@
 #include "Misc/AssertionMacros.h"
 #include "MuR/Model.h"
 #include "MuR/ModelPrivate.h"
+#include "MuR/Material.h"
 #include "MuR/ParametersPrivate.h"
 #include "MuR/Serialisation.h"
 #include "MuR/SerialisationPrivate.h"
@@ -892,6 +893,86 @@ namespace mu
     }
 
 
+	UMaterialInterface* FParameters::GetMaterialValue(int32 Index, const FRangeIndex* RangeIndex) const
+	{
+		check(Index >= 0 && Index < (int32)m_pD->Values.Num());
+		check(GetType(Index) == EParameterType::Material);
+
+		// Early out in case of invalid parameters
+		if (Index < 0 || Index >= (int32)m_pD->Values.Num() ||
+			GetType(Index) != EParameterType::Material)
+		{
+			return nullptr;
+		}
+
+		// Single value case
+		if (!RangeIndex)
+		{
+			// Return the single value
+			return m_pD->Values[Index].Get<FParamMaterialType>().Get();
+		}
+
+		// Multivalue case
+		check(RangeIndex->Parameter == Index);
+
+		if (Index < int32(m_pD->MultiValues.Num()))
+		{
+			const TMap< TArray<int32>, FParameterValue >& m = m_pD->MultiValues[Index];
+			const FParameterValue* it = m.Find(RangeIndex->Values);
+			if (it)
+			{
+				return it->Get<FParamMaterialType>().Get();
+			}
+		}
+
+		// Multivalue parameter, but no multivalue set. Return single value.
+		return m_pD->Values[Index].Get<FParamMaterialType>().Get();
+	}
+
+
+	void FParameters::SetMaterialValue(int32 Index, UMaterialInterface* Material, const FRangeIndex* RangeIndex)
+	{
+		LLM_SCOPE_BYNAME(TEXT("MutableRuntime"));
+
+		check(Index >= 0 && Index < (int32)m_pD->Values.Num());
+		check(GetType(Index) == EParameterType::Material);
+
+		// Early out in case of invalid parameters
+		if (Index < 0 || Index >= (int32)m_pD->Values.Num() ||
+			GetType(Index) != EParameterType::Material)
+		{
+			return;
+		}
+
+		// Single value case
+		if (!RangeIndex)
+		{
+			// Clear multivalue, if set.
+			if (Index < int32(m_pD->MultiValues.Num()))
+			{
+				m_pD->MultiValues[Index].Empty();
+			}
+
+			m_pD->Values[Index].Set<FParamMaterialType>(TStrongObjectPtr(Material));
+		}
+
+		// Multivalue case
+		else
+		{
+			check(RangeIndex->Parameter == Index);
+
+			if (Index >= int32(m_pD->MultiValues.Num()))
+			{
+				m_pD->MultiValues.SetNum(Index + 1);
+			}
+
+			TMap< TArray<int32>, FParameterValue >& m = m_pD->MultiValues[Index];
+			FParameterValue& it = m.FindOrAdd(RangeIndex->Values);
+			it.Set<FParamMaterialType>(TStrongObjectPtr(Material));
+		}
+	}
+
+
     FProjector FParameters::Private::GetProjectorValue( int32 Index, const FRangeIndex* Pos ) const
     {
         const FProgram& Program = Model->GetPrivate()->Program;
diff --git a/Engine/Plugins/Mutable/Source/MutableRuntime/Private/MuR/System.cpp b/Engine/Plugins/Mutable/Source/MutableRuntime/Private/MuR/System.cpp
--- a/Engine/Plugins/Mutable/Source/MutableRuntime/Private/MuR/System.cpp
+++ b/Engine/Plugins/Mutable/Source/MutableRuntime/Private/MuR/System.cpp
@@ -739,6 +739,7 @@ namespace mu
 			case EOpType::IM_PARAMETER:
 			case EOpType::ME_PARAMETER:
 			case EOpType::MA_PARAMETER:
+			case EOpType::MI_PARAMETER:
             {
 				OP::ParameterArgs args = Program.GetOpArgs<OP::ParameterArgs>(at);
 				OP::ADDRESS ParamIndex = args.variable;
diff --git a/Engine/Plugins/Mutable/Source/MutableRuntime/Public/MuR/Instance.h b/Engine/Plugins/Mutable/Source/MutableRuntime/Public/MuR/Instance.h
--- a/Engine/Plugins/Mutable/Source/MutableRuntime/Public/MuR/Instance.h
+++ b/Engine/Plugins/Mutable/Source/MutableRuntime/Public/MuR/Instance.h
@@ -7,6 +7,7 @@
 #include "MuR/Ptr.h"
 #include "MuR/RefCounted.h"
 #include "MuR/Types.h"
+#include "MuR/Material.h"
 
 #define UE_API MUTABLERUNTIME_API
 
@@ -28,6 +29,8 @@ namespace mu
 		uint32 ExternalId =0;
 		uint32 SharedId =0;
 
+		FMaterial Material;
+
 		struct FInstanceImage
 		{
 			FResourceID Id;
@@ -76,7 +79,7 @@ namespace mu
 	{
 		uint16 Id;
 
-		float OverlayMaterialId = -1.f;
+		FMaterial OverlayMaterial;
 
 		TArray<FInstanceLOD, TInlineAllocator<4>> LODs;
 	};
@@ -252,7 +255,9 @@ namespace mu
 
 		//! Get the index of the overlay material.
 		//! \param ComponentIndex Index of the component, from 0 to GetComponentCount()-1
-		UE_API int32 GetOverlayMaterial( int32 ComponentIndex ) const;
+		UE_API const FMaterial& GetOverlayMaterial( int32 ComponentIndex ) const;
+
+		UE_API const FMaterial& GetMaterial(int32 ComponentIndex, int32 LODIndex, int32 SurfaceIndex) const;
 
 		//! Get the number of ExtensionData values in a component
 		UE_API int32 GetExtensionDataCount() const;
@@ -281,7 +286,9 @@ namespace mu
 
 		int32 AddString(int32 ComponentIndex, int32 LODIndex, int32 SurfaceIndex, const FString& Value, FName Name);
 
-		void AddOverlayMaterial( int32 ComponentIndex, float OverlayMaterialId);
+		void AddMaterial(int32 ComponentIndex, int32 LODIndex, int32 SurfaceIndex, FMaterial OverlayMaterial);
+
+		void AddOverlayMaterial(int32 ComponentIndex, FMaterial OverlayMaterial);
 
 		// Data must be non-null
 		void AddExtensionData(const TSharedPtr<const FExtensionData>& Data, FName Name);
diff --git a/Engine/Plugins/Mutable/Source/MutableRuntime/Public/MuR/Material.h b/Engine/Plugins/Mutable/Source/MutableRuntime/Public/MuR/Material.h
new file mode 100644
--- /dev/null
+++ b/Engine/Plugins/Mutable/Source/MutableRuntime/Public/MuR/Material.h
@@ -0,0 +1,46 @@
+// Copyright Epic Games, Inc. All Rights Reserved.
+
+#pragma once
+
+#include "HAL/PlatformMath.h"
+#include "MuR/Ptr.h"
+#include "MuR/RefCounted.h"
+#include "MuR/Types.h"
+
+#include "UObject/StrongObjectPtr.h"
+#include "Materials/MaterialInterface.h"
+
+#define UE_API MUTABLERUNTIME_API
+
+namespace mu
+{
+
+    /** Material type resource */
+	class FMaterial : public FResource
+	{
+	public:
+
+		//-----------------------------------------------------------------------------------------
+		// Life cycle
+		//-----------------------------------------------------------------------------------------
+
+        //! Clone this material
+		UE_API TSharedPtr<FMaterial> Clone() const;
+
+		// Resource interface
+		UE_API int32 GetDataSize() const override { return 0; };
+
+		//-----------------------------------------------------------------------------------------
+		// Own interface
+		//-----------------------------------------------------------------------------------------
+		UE_API bool operator==(const FMaterial& Other) const;
+
+	public:
+
+		TStrongObjectPtr<UMaterialInterface> Material;
+
+		int32 MaterialId = INDEX_NONE;
+	};
+}
+
+#undef UE_API
diff --git a/Engine/Plugins/Mutable/Source/MutableRuntime/Public/MuR/Parameters.h b/Engine/Plugins/Mutable/Source/MutableRuntime/Public/MuR/Parameters.h
--- a/Engine/Plugins/Mutable/Source/MutableRuntime/Public/MuR/Parameters.h
+++ b/Engine/Plugins/Mutable/Source/MutableRuntime/Public/MuR/Parameters.h
@@ -10,6 +10,7 @@
 
 class UTexture;
 class USkeletalMesh;
+class UMaterialInterface;
 
 #define UE_API MUTABLERUNTIME_API
 
@@ -50,6 +51,9 @@ namespace mu
 
 		/** An externally provided mesh. */
 		Mesh,
+		
+		/** An externally provided material*/
+    	Material,
 
         /** A text string. */
         String,
@@ -352,6 +356,19 @@ namespace mu
 		//! \param pos Only for multidimensional parameters: relevant position to set in the ranges
 		UE_API void SetMatrixValue( int32 Index, const FMatrix44f& Value, const FRangeIndex* Pos = nullptr );
 
+		//! Return the value of a Material parameter.
+		//! \pre The parameter specified by index is a T_MATERIAL.
+		//! \param index Index of the parameter from 0 to GetCount()-1
+		//! \param pos Only for multidimensional parameters: relevant position to set in the ranges
+		//! \return The Material interface specified when setting the image value (\see SetImageValue)
+		UE_API UMaterialInterface* GetMaterialValue(int32 Index, const FRangeIndex* RangeIndex = nullptr) const;
+
+		//! If the parameter is of the material type, set its value.
+		//! \param index Index of the parameter from 0 to GetCount()-1
+		//! \param ExternalResource Material interface to set as value.
+		//! \param pos Only for multidimensional parameters: relevant position to set in the ranges
+		UE_API void SetMaterialValue(int32 Index, UMaterialInterface* ExternalResource, const FRangeIndex* RangeIndex = nullptr);
+
         //! Utility method to compare the values of a specific parameter with the values of another
         //! FParameters object. It returns false if type or values are different.
         UE_API bool HasSameValue( int32 ThisParamIndex, const TSharedPtr<const FParameters>& Other, int32 OtherParamIndex ) const;
diff --git a/Engine/Plugins/Mutable/Source/MutableTools/Internal/MuT/TablePrivate.h b/Engine/Plugins/Mutable/Source/MutableTools/Internal/MuT/TablePrivate.h
--- a/Engine/Plugins/Mutable/Source/MutableTools/Internal/MuT/TablePrivate.h
+++ b/Engine/Plugins/Mutable/Source/MutableTools/Internal/MuT/TablePrivate.h
@@ -22,6 +22,7 @@ namespace mu
 	struct FTableValue
 	{
 		// TODO: Union
+		int32 Int;
 		float Scalar;
 		FVector4f Color;
 		Ptr<TResourceProxy<FImage>> ProxyImage;
diff --git a/Engine/Plugins/Mutable/Source/MutableTools/Private/MuT/ASTOpAddOverlayMaterial.cpp b/Engine/Plugins/Mutable/Source/MutableTools/Private/MuT/ASTOpAddOverlayMaterial.cpp
--- a/Engine/Plugins/Mutable/Source/MutableTools/Private/MuT/ASTOpAddOverlayMaterial.cpp
+++ b/Engine/Plugins/Mutable/Source/MutableTools/Private/MuT/ASTOpAddOverlayMaterial.cpp
@@ -16,7 +16,7 @@ namespace mu
 //-------------------------------------------------------------------------------------------------
 	ASTOpAddOverlayMaterial::ASTOpAddOverlayMaterial()
 	: Instance(this)
-	, OverlayMaterialId(this)
+	, OverlayMaterial(this)
 {
 }
 
@@ -33,7 +33,7 @@ bool ASTOpAddOverlayMaterial::IsEqual(const ASTOp& OtherUntyped) const
 	{
 		const ASTOpAddOverlayMaterial* Other = static_cast<const ASTOpAddOverlayMaterial*>(&OtherUntyped);
 		return Instance == Other->Instance
-			&& OverlayMaterialId == Other->OverlayMaterialId;
+			&& OverlayMaterial == Other->OverlayMaterial;
 	}
 
 	return false;
@@ -44,7 +44,7 @@ uint64 ASTOpAddOverlayMaterial::Hash() const
 	uint64 Result = std::hash<uint64>()(uint64(EOpType::IN_ADDOVERLAYMATERIAL));
 	
 	hash_combine(Result, Instance.child().get());
-	hash_combine(Result, OverlayMaterialId.child().get());
+	hash_combine(Result, OverlayMaterial.child().get());
 	
 	return Result;
 }
@@ -54,7 +54,7 @@ mu::Ptr<ASTOp> ASTOpAddOverlayMaterial::Clone(MapChildFuncRef MapChild) const
 	Ptr<ASTOpAddOverlayMaterial> NewInstance = new ASTOpAddOverlayMaterial();
 
 	NewInstance->Instance = MapChild(Instance.child());
-	NewInstance->OverlayMaterialId = MapChild(OverlayMaterialId.child());
+	NewInstance->OverlayMaterial = MapChild(OverlayMaterial.child());
 
 	return NewInstance;
 }
@@ -62,7 +62,7 @@ mu::Ptr<ASTOp> ASTOpAddOverlayMaterial::Clone(MapChildFuncRef MapChild) const
 void ASTOpAddOverlayMaterial::ForEachChild(const TFunctionRef<void(ASTChild&)> F)
 {
 	F(Instance);
-	F(OverlayMaterialId);
+	F(OverlayMaterial);
 }
 
 void ASTOpAddOverlayMaterial::Link(FProgram& Program, FLinkerOptions*)
@@ -73,7 +73,7 @@ void ASTOpAddOverlayMaterial::Link(FProgram& Program, FLinkerOptions*)
 		return;
 	}
 
-	OP::InstanceAddOverlayMaterialArgs Args;
+	OP::InstanceAddMaterialArgs Args;
 	FMemory::Memzero(Args);
 
 	if (Instance)
@@ -81,9 +81,9 @@ void ASTOpAddOverlayMaterial::Link(FProgram& Program, FLinkerOptions*)
 		Args.Instance = Instance->linkedAddress;
 	}
 
-	if (OverlayMaterialId)
+	if (OverlayMaterial)
 	{
-		Args.OverlayMaterialId = OverlayMaterialId->linkedAddress;
+		Args.Material = OverlayMaterial->linkedAddress;
 	}
 	
 	linkedAddress = (OP::ADDRESS)Program.OpAddress.Num();
diff --git a/Engine/Plugins/Mutable/Source/MutableTools/Private/MuT/ASTOpAddOverlayMaterial.h b/Engine/Plugins/Mutable/Source/MutableTools/Private/MuT/ASTOpAddOverlayMaterial.h
--- a/Engine/Plugins/Mutable/Source/MutableTools/Private/MuT/ASTOpAddOverlayMaterial.h
+++ b/Engine/Plugins/Mutable/Source/MutableTools/Private/MuT/ASTOpAddOverlayMaterial.h
@@ -22,7 +22,7 @@ class ASTOpAddOverlayMaterial final : public ASTOp
 public:
 
 	ASTChild Instance;
-	ASTChild OverlayMaterialId;
+	ASTChild OverlayMaterial;
 
 public:
 
diff --git a/Engine/Plugins/Mutable/Source/MutableTools/Private/MuT/ASTOpConstantMaterial.cpp b/Engine/Plugins/Mutable/Source/MutableTools/Private/MuT/ASTOpConstantMaterial.cpp
new file mode 100644
--- /dev/null
+++ b/Engine/Plugins/Mutable/Source/MutableTools/Private/MuT/ASTOpConstantMaterial.cpp
@@ -0,0 +1,47 @@
+// Copyright Epic Games, Inc. All Rights Reserved.
+
+#include "MuT/ASTOpConstantMaterial.h"
+
+mu::ASTOpConstantMaterial::ASTOpConstantMaterial(int32 InMaterialId) : MaterialId(InMaterialId)
+{
+}
+
+uint64 mu::ASTOpConstantMaterial::Hash() const
+{
+	uint64 res = std::hash<uint64>()(uint64(EOpType::MI_CONSTANT));
+	hash_combine(res, (uint64)MaterialId);
+	return res;
+}
+
+bool mu::ASTOpConstantMaterial::IsEqual(const ASTOp& OtherUntyped) const
+{
+	if (OtherUntyped.GetOpType() == GetOpType())
+	{
+		const ASTOpConstantMaterial* Other = static_cast<const ASTOpConstantMaterial*>(&OtherUntyped);
+		return MaterialId == Other->MaterialId;
+	}
+	
+	return false;
+}
+
+mu::Ptr<mu::ASTOp> mu::ASTOpConstantMaterial::Clone(MapChildFuncRef mapChild) const
+{
+	Ptr<ASTOpConstantMaterial> MaterialOp = new ASTOpConstantMaterial();
+	MaterialOp->MaterialId = MaterialId;
+	return MaterialOp;
+}
+
+void mu::ASTOpConstantMaterial::Link(FProgram& program, FLinkerOptions* Options)
+{
+	if (!linkedAddress)
+	{
+		OP::ResourceConstantArgs Args;
+		FMemory::Memzero(Args);
+		Args.value = MaterialId;
+
+		linkedAddress = static_cast<OP::ADDRESS>(program.OpAddress.Num());
+		program.OpAddress.Add(static_cast<uint32_t>(program.ByteCode.Num()));
+		AppendCode(program.ByteCode, EOpType::MI_CONSTANT);
+		AppendCode(program.ByteCode, Args);
+	}
+}
diff --git a/Engine/Plugins/Mutable/Source/MutableTools/Private/MuT/ASTOpConstantMaterial.h b/Engine/Plugins/Mutable/Source/MutableTools/Private/MuT/ASTOpConstantMaterial.h
new file mode 100644
--- /dev/null
+++ b/Engine/Plugins/Mutable/Source/MutableTools/Private/MuT/ASTOpConstantMaterial.h
@@ -0,0 +1,38 @@
+﻿// Copyright Epic Games, Inc. All Rights Reserved.
+
+#pragma once
+
+#include "HAL/Platform.h"
+#include "MuR/Operations.h"
+#include "MuR/Ptr.h"
+#include "MuT/AST.h"
+
+
+namespace mu
+{
+
+struct FProgram;
+
+//---------------------------------------------------------------------------------------------
+//!
+//---------------------------------------------------------------------------------------------
+class ASTOpConstantMaterial final : public ASTOp
+{
+public:
+
+	//!
+	int32 MaterialId = INDEX_NONE;
+
+public:
+
+	ASTOpConstantMaterial(int32 MaterialId = INDEX_NONE);
+
+	EOpType GetOpType() const override { return EOpType::MI_CONSTANT; }
+	uint64 Hash() const override;
+	void ForEachChild(const TFunctionRef<void(ASTChild&)>) override {}
+	bool IsEqual(const ASTOp& otherUntyped) const override;
+	Ptr<ASTOp> Clone(MapChildFuncRef mapChild) const override;
+	void Link(FProgram& program, FLinkerOptions* Options) override;
+};
+
+}
diff --git a/Engine/Plugins/Mutable/Source/MutableTools/Private/MuT/CodeGenerator.cpp b/Engine/Plugins/Mutable/Source/MutableTools/Private/MuT/CodeGenerator.cpp
--- a/Engine/Plugins/Mutable/Source/MutableTools/Private/MuT/CodeGenerator.cpp
+++ b/Engine/Plugins/Mutable/Source/MutableTools/Private/MuT/CodeGenerator.cpp
@@ -609,15 +609,16 @@ namespace mu
 
 				Ptr<ASTOp> LastInstOp = LODsOp;
 
-				if (Ptr<NodeScalar> pScalarNode = InNode->OverlayMaterial)
+				if (Ptr<NodeMaterial> MaterialNode = InNode->OverlayMaterial)
 				{
 					Ptr<ASTOpAddOverlayMaterial> AddOverlayMaterialOp = new ASTOpAddOverlayMaterial();
 					AddOverlayMaterialOp->Instance = LastInstOp;
 
-					// Scalar
-					FScalarGenerationResult ScalarResult;
-					GenerateScalar(ScalarResult, Options, pScalarNode);
-					AddOverlayMaterialOp->OverlayMaterialId = ScalarResult.op;
+					// Material
+					FMaterialGenerationOptions MaterialOptions;
+					FMaterialGenerationResult MaterialResult;
+					GenerateMaterial(MaterialOptions, MaterialResult, MaterialNode);
+					AddOverlayMaterialOp->OverlayMaterial = MaterialResult.op;
 
 					LastInstOp = AddOverlayMaterialOp;
 				}
@@ -1645,6 +1646,21 @@ namespace mu
 							LastSurfOp = op;
 						}
 					}
+
+					// Create the expression for the Material
+					//------------------------------------------------------------------------
+					if (SurfaceNode->Material)
+					{
+						Ptr<ASTOpInstanceAdd> Op = new ASTOpInstanceAdd();
+						Op->type = EOpType::IN_ADDMATERIAL;
+						Op->instance = LastSurfOp;
+
+						FMaterialGenerationOptions MaterialOptions;
+						FMaterialGenerationResult MaterialResult;
+						GenerateMaterial(MaterialOptions, MaterialResult, SurfaceNode->Material);
+						Op->value = MaterialResult.op;
+						LastSurfOp = Op;
+					}
 				}
 
 				FSurfaceGenerationResult SurfaceResult;
diff --git a/Engine/Plugins/Mutable/Source/MutableTools/Private/MuT/CodeGenerator.h b/Engine/Plugins/Mutable/Source/MutableTools/Private/MuT/CodeGenerator.h
--- a/Engine/Plugins/Mutable/Source/MutableTools/Private/MuT/CodeGenerator.h
+++ b/Engine/Plugins/Mutable/Source/MutableTools/Private/MuT/CodeGenerator.h
@@ -74,6 +74,12 @@ namespace mu
 	class NodeImageTable;
 	class NodeImageTransform;
 	class NodeImageVariation;
+	class NodeMaterial;
+	class NodeMaterialConstant;
+	class NodeMaterialTable;
+	class NodeMaterialSwitch;
+	class NodeMaterialVariation;
+	class NodeMaterialParameter;
 	class NodeMeshApplyPose;
 	class NodeMeshClipDeform;
 	class NodeMeshClipMorphPlane;
@@ -102,6 +108,7 @@ namespace mu
 	class NodeMatrix;
 	class NodeMatrixConstant;
 	class NodeMatrixParameter;
+	class NodeMaterialParameter;
 	struct FObjectState;
 	struct FProgram;
 
@@ -416,6 +423,28 @@ namespace mu
 		void GenerateTableSwitchNode(const NODE_TABLE& node, Ptr<NODE_SWITCH_TYPE>& OutResult, F&& GenerateOption);
 
 
+		//-----------------------------------------------------------------------------------------
+		// Materials
+
+		/** Options that affect the generation of Materials. It is like list of what required data we want while parsing down the material node graph. */
+		struct FMaterialGenerationOptions : public FGenericGenerationOptions
+		{
+			
+		};
+
+		/** */
+		struct FMaterialGenerationResult
+		{
+			Ptr<ASTOp> op;
+		};
+
+		void GenerateMaterial(const FMaterialGenerationOptions&, FMaterialGenerationResult& result, const Ptr<const NodeMaterial>& node);
+		void GenerateMaterial_Constant(const FMaterialGenerationOptions&, FMaterialGenerationResult& result, const NodeMaterialConstant* node);
+		void GenerateMaterial_Switch(const FMaterialGenerationOptions&, FMaterialGenerationResult& result, const NodeMaterialSwitch* node);
+		void GenerateMaterial_Table(const FMaterialGenerationOptions&, FMaterialGenerationResult& result, const NodeMaterialTable* node);
+		void GenerateMaterial_Variation(const FMaterialGenerationOptions&, FMaterialGenerationResult& result, const NodeMaterialVariation* node);
+		void GenerateMaterial_Parameter(const FMaterialGenerationOptions&, FMaterialGenerationResult& result, const NodeMaterialParameter* node);
+
 		//-----------------------------------------------------------------------------------------
 		// Images
 		
diff --git a/Engine/Plugins/Mutable/Source/MutableTools/Private/MuT/CodeGenerator_FirstPass.h b/Engine/Plugins/Mutable/Source/MutableTools/Private/MuT/CodeGenerator_FirstPass.h
--- a/Engine/Plugins/Mutable/Source/MutableTools/Private/MuT/CodeGenerator_FirstPass.h
+++ b/Engine/Plugins/Mutable/Source/MutableTools/Private/MuT/CodeGenerator_FirstPass.h
@@ -24,6 +24,7 @@
 #include "MuT/NodeSurfaceVariation.h"
 #include "MuT/NodeSurfaceSwitch.h"
 #include "MuT/NodeMeshParameter.h"
+#include "MuT/NodeMaterialParameter.h"
 #include "MuT/ASTOpParameter.h"
 #include "MuT/ErrorLog.h"
 #include "Async/Mutex.h"
diff --git a/Engine/Plugins/Mutable/Source/MutableTools/Private/MuT/CodeGenerator_Material.cpp b/Engine/Plugins/Mutable/Source/MutableTools/Private/MuT/CodeGenerator_Material.cpp
new file mode 100644
--- /dev/null
+++ b/Engine/Plugins/Mutable/Source/MutableTools/Private/MuT/CodeGenerator_Material.cpp
@@ -0,0 +1,213 @@
+// Copyright Epic Games, Inc. All Rights Reserved.
+
+#include "MuR/Operations.h"
+#include "MuR/ParametersPrivate.h"
+#include "MuR/Platform.h"
+#include "MuR/Ptr.h"
+#include "MuR/RefCounted.h"
+#include "MuR/System.h"
+#include "MuR/Material.h"
+#include "MuT/AST.h"
+#include "MuT/ASTOpConditional.h"
+#include "MuT/ASTOpConstantMaterial.h"
+#include "MuT/ASTOpSwitch.h"
+#include "MuT/CodeGenerator.h"
+#include "MuT/CodeGenerator_FirstPass.h"
+#include "MuT/CompilerPrivate.h"
+#include "MuT/ErrorLog.h"
+#include "MuT/NodeScalar.h"
+#include "MuT/NodeMaterial.h"
+#include "MuT/NodeMaterialConstant.h"
+#include "MuT/NodeMaterialSwitch.h"
+#include "MuT/NodeMaterialVariation.h"
+#include "MuT/NodeMaterialTable.h"
+#include "MuT/NodeMaterialParameter.h"
+#include "MuT/Table.h"
+#include "MuT/TablePrivate.h"
+
+
+namespace mu
+{
+	class Node;
+
+	void CodeGenerator::GenerateMaterial(const FMaterialGenerationOptions& InOptions, FMaterialGenerationResult& Result , const Ptr<const NodeMaterial>& InUntypedNode)
+	{
+		if (!InUntypedNode)
+		{
+			return;
+		}
+
+		const NodeMaterial* Node = InUntypedNode.get();
+
+		// Generate for each different type of node
+		switch (Node->GetType()->Type)
+		{
+		case Node::EType::MaterialConstant: GenerateMaterial_Constant(InOptions, Result, static_cast<const NodeMaterialConstant*>(Node)); break;
+		case Node::EType::MaterialSwitch: GenerateMaterial_Switch(InOptions, Result, static_cast<const NodeMaterialSwitch*>(Node)); break;
+		case Node::EType::MaterialVariation: GenerateMaterial_Variation(InOptions, Result, static_cast<const NodeMaterialVariation*>(Node)); break;
+		case Node::EType::MaterialTable: GenerateMaterial_Table(InOptions, Result, static_cast<const NodeMaterialTable*>(Node)); break;
+		case Node::EType::MaterialParameter: GenerateMaterial_Parameter(InOptions, Result, static_cast<const NodeMaterialParameter*>(Node)); break;
+		default: check(false);
+		}
+	}
+
+	void CodeGenerator::GenerateMaterial_Constant(const FMaterialGenerationOptions& Options, FMaterialGenerationResult& Result, const NodeMaterialConstant* Node)
+	{
+		Result.op = new ASTOpConstantMaterial(Node->MaterialId);
+	}
+
+	void CodeGenerator::GenerateMaterial_Switch(const FMaterialGenerationOptions& Options, FMaterialGenerationResult& Result, const NodeMaterialSwitch* Node)
+	{
+		if (Node->Options.Num() == 0)
+		{
+			return;
+		}
+
+		Ptr<ASTOp> Variable;
+		if (Node->Parameter)
+		{
+			FScalarGenerationResult ParamResult;
+			GenerateScalar(ParamResult, Options, Node->Parameter.get());
+			Variable = ParamResult.op;
+		}
+		else
+		{
+			// This argument is required
+			Variable = GenerateMissingScalarCode(TEXT("Switch variable"), 0.0f, Node->GetMessageContext());
+		}
+
+		Ptr<ASTOpSwitch> SwitchOp = new ASTOpSwitch();
+		SwitchOp->Type = EOpType::MI_SWITCH;
+		SwitchOp->Variable = Variable;
+
+		// Options
+		SwitchOp->Cases.Reserve(Node->Options.Num());
+
+		for (int32 OptionIndex = 0; OptionIndex < Node->Options.Num(); ++OptionIndex)
+		{
+			FMaterialGenerationResult OptionResult;
+			const Ptr<NodeMaterial>& OptionMaterial = Node->Options[OptionIndex];
+			if (OptionMaterial)
+			{
+				GenerateMaterial(Options, OptionResult, OptionMaterial);
+			}
+			else
+			{
+				// This argument is required. Empty material
+				OptionResult.op = new ASTOpConstantMaterial();
+			}
+
+			SwitchOp->Cases.Emplace(int16(OptionIndex), SwitchOp, OptionResult.op);
+		}
+
+		Result.op = SwitchOp;
+	}
+
+	void CodeGenerator::GenerateMaterial_Variation(const FMaterialGenerationOptions& Options, FMaterialGenerationResult& Result, const NodeMaterialVariation* Node)
+	{
+		Ptr<ASTOp> LastOp;
+
+		// Default case
+		if (Node->DefaultMaterial)
+		{
+			FMaterialGenerationResult DefaultResult;
+			GenerateMaterial(Options, DefaultResult, Node->DefaultMaterial);
+			LastOp = DefaultResult.op;
+		}
+
+		// Process variations in reverse order, since conditionals are built bottom-up.
+		for (int32 VariationIndex = int32(Node->Variations.Num()) - 1; VariationIndex >= 0; --VariationIndex)
+		{
+			int32 TagIndex = INDEX_NONE;
+			const FString& tag = Node->Variations[VariationIndex].Tag;
+			for (int32 Index = 0; Index < FirstPass.Tags.Num(); ++Index)
+			{
+				if (FirstPass.Tags[Index].Tag == tag)
+				{
+					TagIndex = Index;
+				}
+			}
+
+			if (TagIndex == INDEX_NONE)
+			{
+				FString Msg = FString::Printf(TEXT("Unknown tag found in material variation [%s]."), *tag);
+
+				ErrorLog->Add(Msg, ELMT_WARNING, Node->GetMessageContext());
+				continue;
+			}
+
+			Ptr<ASTOp> VariationOp;
+			if (Node->Variations[VariationIndex].Material)
+			{
+				FMaterialGenerationResult VariationResult;
+				GenerateMaterial(Options, VariationResult, Node->Variations[VariationIndex].Material);
+				VariationOp = VariationResult.op;
+			}
+			else
+			{
+				// This argument is required
+				VariationOp = new ASTOpConstantMaterial();
+			}
+
+			Ptr<ASTOpConditional> conditional = new ASTOpConditional;
+			conditional->type = EOpType::MI_CONDITIONAL;
+			conditional->no = LastOp;
+			conditional->yes = VariationOp;
+			conditional->condition = FirstPass.Tags[TagIndex].GenericCondition;
+
+			LastOp = conditional;
+		}
+
+		Result.op = LastOp;
+	}
+
+	void CodeGenerator::GenerateMaterial_Table(const FMaterialGenerationOptions& Options, FMaterialGenerationResult& Result, const NodeMaterialTable* Node)
+	{
+		Result.op = GenerateTableSwitch<NodeMaterialTable, ETableColumnType::Material, EOpType::MI_SWITCH>(*Node,
+			[this](const NodeMaterialTable& Node, int32 ColumnIndex, int32 RowIndex, FErrorLog* pErrorLog)
+			{
+				Ptr<ASTOp> Op;
+
+				int32 MaterialId = Node.Table->GetPrivate()->Rows[RowIndex].Values[ColumnIndex].Int;
+				if (MaterialId != INDEX_NONE)
+				{
+					Op = new ASTOpConstantMaterial(MaterialId);
+				}
+
+				return Op;
+			});
+	}
+
+	void CodeGenerator::GenerateMaterial_Parameter(const FMaterialGenerationOptions& InOptions, FMaterialGenerationResult& Result, const NodeMaterialParameter* InNode)
+	{
+		{
+			UE::TUniqueLock Lock(FirstPass.ParameterNodes.Mutex);
+			if (Ptr<ASTOpParameter>* Found = FirstPass.ParameterNodes.GenericParametersCache.Find(InNode))
+			{
+				Result.op = *Found;
+				return;
+			}
+		}
+			
+		Ptr<ASTOpParameter> MaterialOp = new ASTOpParameter();
+		MaterialOp->Type = EOpType::MI_PARAMETER;
+
+		MaterialOp->Parameter.Name = InNode->Name;
+		bool bParseOk = FGuid::Parse(InNode->UID, MaterialOp->Parameter.UID);
+		check(bParseOk);
+		MaterialOp->Parameter.Type = EParameterType::Material;
+		MaterialOp->Parameter.DefaultValue.Set<FParamMaterialType>(nullptr);
+
+		// Generate the code for the ranges
+		for (const Ptr<NodeRange>& RangeNode : InNode->Ranges)
+		{
+			FRangeGenerationResult RangeResult;
+			GenerateRange(RangeResult, InOptions, RangeNode);
+			MaterialOp->Ranges.Emplace(MaterialOp.get(), RangeResult.sizeOp, RangeResult.rangeName, RangeResult.rangeUID);
+		}
+
+		Result.op = MaterialOp;
+
+		FirstPass.ParameterNodes.GenericParametersCache.Add(InNode, MaterialOp);
+	}
+}
diff --git a/Engine/Plugins/Mutable/Source/MutableTools/Private/MuT/Compiler.cpp b/Engine/Plugins/Mutable/Source/MutableTools/Private/MuT/Compiler.cpp
--- a/Engine/Plugins/Mutable/Source/MutableTools/Private/MuT/Compiler.cpp
+++ b/Engine/Plugins/Mutable/Source/MutableTools/Private/MuT/Compiler.cpp
@@ -67,6 +67,7 @@ namespace mu
 		{ false,	{ 0, 0, 0, 0, 0, 0, 0, 0, 0, 0, 0, 0, 0, 0, 0, 0, 0, 0, 0, 0, 0, 0, 0, 0, 0, 0, 0, 0, 0 }		},	// ST_CONSTANT
 		{ true,		{ 0, 0, 0, 0, 0, 0, 0, 0, 0, 0, 0, 0, 0, 0, 0, 0, 0, 0, 0, 0, 0, 0, 0, 0, 0, 0, 0, 0, 0 }		},	// ED_CONSTANT
 		{ false,	{ 0, 0, 0, 0, 0, 0, 0, 0, 0, 0, 0, 0, 0, 0, 0, 0, 0, 0, 0, 0, 0, 0, 0, 0, 0, 0, 0, 0, 0 }		},	// MA_CONSTANT
+		{ false,	{ 0, 0, 0, 0, 0, 0, 0, 0, 0, 0, 0, 0, 0, 0, 0, 0, 0, 0, 0, 0, 0, 0, 0, 0, 0, 0, 0, 0, 0 }		},	// MI_CONSTANT
 
 		{ false,	{ 0, 0, 0, 0, 0, 0, 0, 0, 0, 0, 0, 0, 0, 0, 0, 0, 0, 0, 0, 0, 0, 0, 0, 0, 0, 0, 0, 0, 0 }		},	// BO_PARAMETER
 		{ false,	{ 0, 0, 0, 0, 0, 0, 0, 0, 0, 0, 0, 0, 0, 0, 0, 0, 0, 0, 0, 0, 0, 0, 0, 0, 0, 0, 0, 0, 0 }		},	// NU_PARAMETER
@@ -77,6 +78,7 @@ namespace mu
 		{ true,		{ 0, 0, 0, 0, 0, 0, 0, 0, 0, 0, 0, 0, 0, 0, 0, 0, 0, 0, 0, 0, 0, 0, 0, 0, 0, 0, 0, 0, 0 }		},	// ME_PARAMETER
 		{ false,	{ 0, 0, 0, 0, 0, 0, 0, 0, 0, 0, 0, 0, 0, 0, 0, 0, 0, 0, 0, 0, 0, 0, 0, 0, 0, 0, 0, 0, 0 }		},	// ST_PARAMETER
 		{ false,	{ 0, 0, 0, 0, 0, 0, 0, 0, 0, 0, 0, 0, 0, 0, 0, 0, 0, 0, 0, 0, 0, 0, 0, 0, 0, 0, 0, 0, 0 }		},	// MA_PARAMETER
+		{ false,	{ 0, 0, 0, 0, 0, 0, 0, 0, 0, 0, 0, 0, 0, 0, 0, 0, 0, 0, 0, 0, 0, 0, 0, 0, 0, 0, 0, 0, 0 }		},	// MI_PARAMETER
 
 		{ true,		{ 0, 0, 0, 0, 0, 0, 0, 0, 0, 0, 0, 0, 0, 0, 0, 0, 0, 0, 0, 0, 0, 0, 0, 0, 0, 0, 0, 0, 0 }		},	// IM_REFERENCE
 		{ true,		{ 0, 0, 0, 0, 0, 0, 0, 0, 0, 0, 0, 0, 0, 0, 0, 0, 0, 0, 0, 0, 0, 0, 0, 0, 0, 0, 0, 0, 0 }		},	// ME_REFERENCE
@@ -89,6 +91,7 @@ namespace mu
 		{ false,	{ 0, 0, 0, 0, 0, 0, 0, 0, 0, 0, 0, 0, 0, 0, 0, 0, 0, 0, 0, 0, 0, 0, 0, 0, 0, 0, 0, 0, 0 }		},	// LA_CONDITIONAL
 		{ false,	{ 0, 0, 0, 0, 0, 0, 0, 0, 0, 0, 0, 0, 0, 0, 0, 0, 0, 0, 0, 0, 0, 0, 0, 0, 0, 0, 0, 0, 0 }		},	// IN_CONDITIONAL
 		{ false,	{ 0, 0, 0, 0, 0, 0, 0, 0, 0, 0, 0, 0, 0, 0, 0, 0, 0, 0, 0, 0, 0, 0, 0, 0, 0, 0, 0, 0, 0 }		},	// ED_CONDITIONAL
+		{ false,	{ 0, 0, 0, 0, 0, 0, 0, 0, 0, 0, 0, 0, 0, 0, 0, 0, 0, 0, 0, 0, 0, 0, 0, 0, 0, 0, 0, 0, 0 }		},	// MI_CONDITIONAL
 
 		{ false,	{ 0, 0, 0, 0, 0, 0, 0, 0, 0, 0, 0, 0, 0, 0, 0, 0, 0, 0, 0, 0, 0, 0, 0, 0, 0, 0, 0, 0, 0 }		},	// NU_SWITCH
 		{ false,	{ 0, 0, 0, 0, 0, 0, 0, 0, 0, 0, 0, 0, 0, 0, 0, 0, 0, 0, 0, 0, 0, 0, 0, 0, 0, 0, 0, 0, 0 }		},	// SC_SWITCH
@@ -98,6 +101,7 @@ namespace mu
 		{ false,	{ 0, 0, 0, 0, 0, 0, 0, 0, 0, 0, 0, 0, 0, 0, 0, 0, 0, 0, 0, 0, 0, 0, 0, 0, 0, 0, 0, 0, 0 }		},	// LA_SWITCH
 		{ false,	{ 0, 0, 0, 0, 0, 0, 0, 0, 0, 0, 0, 0, 0, 0, 0, 0, 0, 0, 0, 0, 0, 0, 0, 0, 0, 0, 0, 0, 0 }		},	// IN_SWITCH
 		{ false,	{ 0, 0, 0, 0, 0, 0, 0, 0, 0, 0, 0, 0, 0, 0, 0, 0, 0, 0, 0, 0, 0, 0, 0, 0, 0, 0, 0, 0, 0 }		},	// ED_SWITCH
+		{ false,	{ 0, 0, 0, 0, 0, 0, 0, 0, 0, 0, 0, 0, 0, 0, 0, 0, 0, 0, 0, 0, 0, 0, 0, 0, 0, 0, 0, 0, 0 }		},	// MI_SWITCH
 
 		{ false,	{ 0, 0, 0, 0, 0, 0, 0, 0, 0, 0, 0, 0, 0, 0, 0, 0, 0, 0, 0, 0, 0, 0, 0, 0, 0, 0, 0, 0, 0 }		},	// BO_EQUAL_SC_CONST
 		{ false,	{ 0, 0, 0, 0, 0, 0, 0, 0, 0, 0, 0, 0, 0, 0, 0, 0, 0, 0, 0, 0, 0, 0, 0, 0, 0, 0, 0, 0, 0 }		},	// BO_AND
@@ -173,6 +177,7 @@ namespace mu
 		{ false,	{ 0, 0, 0, 0, 0, 0, 0, 0, 0, 0, 0, 0, 0, 0, 0, 0, 0, 0, 0, 0, 0, 0, 0, 0, 0, 0, 0, 0, 0 }		},	// IN_ADDLOD
 		{ false,	{ 0, 0, 0, 0, 0, 0, 0, 0, 0, 0, 0, 0, 0, 0, 0, 0, 0, 0, 0, 0, 0, 0, 0, 0, 0, 0, 0, 0, 0 }		},	// IN_ADDEXTENSIONDATA
 		{ false,	{ 0, 0, 0, 0, 0, 0, 0, 0, 0, 0, 0, 0, 0, 0, 0, 0, 0, 0, 0, 0, 0, 0, 0, 0, 0, 0, 0, 0, 0 }		},	// IN_ADDOVERLAYMATERIAL
+		{ false,	{ 0, 0, 0, 0, 0, 0, 0, 0, 0, 0, 0, 0, 0, 0, 0, 0, 0, 0, 0, 0, 0, 0, 0, 0, 0, 0, 0, 0, 0 }		},	// IN_ADDMATERIAL
 
 		{ true,		{ 0, 0, 0, 0, 0, 0, 0, 0, 0, 0, 0, 0, 0, 0, 0, 0, 0, 0, 0, 0, 0, 0, 0, 0, 0, 0, 0, 0, 0 }		},	// LA_PACK
 		{ true,		{ 0, 0, 0, 0, 0, 0, 0, 0, 0, 0, 0, 0, 0, 0, 0, 0, 0, 0, 0, 0, 0, 0, 0, 0, 0, 0, 0, 0, 0 }		},	// LA_MERGE
diff --git a/Engine/Plugins/Mutable/Source/MutableTools/Private/MuT/ErrorLog.cpp b/Engine/Plugins/Mutable/Source/MutableTools/Private/MuT/ErrorLog.cpp
--- a/Engine/Plugins/Mutable/Source/MutableTools/Private/MuT/ErrorLog.cpp
+++ b/Engine/Plugins/Mutable/Source/MutableTools/Private/MuT/ErrorLog.cpp
@@ -212,6 +212,7 @@ namespace mu
 		TEXT("ST_CONSTANT      "),
 		TEXT("ED_CONSTANT      "),
 		TEXT("MA_CONSTANT      "),
+		TEXT("MI_CONSTANT      "),
 
 		TEXT("BO_PARAMETER     "),
 		TEXT("NU_PARAMETER     "),
@@ -222,6 +223,7 @@ namespace mu
 		TEXT("ME_PARAMETER     "),
 		TEXT("ST_PARAMETER     "),
 		TEXT("MA_PARAMETER     "),
+		TEXT("MI_PARAMETER     "),
 		
 		TEXT("IM_REFERENCE     "),
 		TEXT("ME_REFERENCE     "),
@@ -234,6 +236,7 @@ namespace mu
 		TEXT("LA_CONDITIONAL   "),
 		TEXT("IN_CONDITIONAL   "),
 		TEXT("ED_CONDITIONAL   "),
+		TEXT("MI_CONDITIONAL   "),
 		
 		TEXT("NU_SWITCH        "),
 		TEXT("SC_SWITCH        "),
@@ -243,6 +246,7 @@ namespace mu
 		TEXT("LA_SWITCH        "),
 		TEXT("IN_SWITCH        "),
 		TEXT("ED_SWITCH        "),
+		TEXT("MI_SWITCH        "),
 
 		TEXT("BO_EQUAL_SC_CONST"),
 		TEXT("BO_AND           "),
@@ -318,6 +322,7 @@ namespace mu
         TEXT("IN_ADDLOD        "),
 		TEXT("IN_ADDEXTENSIDATA"),
 		TEXT("IN_ADDOVERLAYMATERIAL"),
+		TEXT("IN_ADDMATERIAL   "),
 		
 		TEXT("LA_PACK          "),
 		TEXT("LA_MERGE         "),
diff --git a/Engine/Plugins/Mutable/Source/MutableTools/Private/MuT/Node.cpp b/Engine/Plugins/Mutable/Source/MutableTools/Private/MuT/Node.cpp
--- a/Engine/Plugins/Mutable/Source/MutableTools/Private/MuT/Node.cpp
+++ b/Engine/Plugins/Mutable/Source/MutableTools/Private/MuT/Node.cpp
@@ -92,6 +92,13 @@
 #include "MuT/NodeProjector.h"
 #include "MuT/NodeRange.h"
 #include "MuT/NodeRangeFromScalar.h"
+#include "MuT/NodeMaterial.h"
+#include "MuT/NodeMaterialConstant.h"
+#include "MuT/NodeMaterialTable.h"
+#include "MuT/NodeMaterialSwitch.h"
+#include "MuT/NodeMaterialVariation.h"
+#include "MuT/NodeMaterialParameter.h"
+
 
 namespace mu
 {
@@ -209,6 +216,13 @@ namespace mu
 	FNodeType NodeRange::StaticType = FNodeType(Node::EType::Range, Node::GetStaticType());
 	FNodeType NodeRangeFromScalar::StaticType = FNodeType(Node::EType::RangeFromScalar, NodeRange::GetStaticType());
 
+	FNodeType NodeMaterial::StaticType = FNodeType(Node::EType::Material, NodeMaterial::GetStaticType());
+	FNodeType NodeMaterialConstant::StaticType = FNodeType(Node::EType::MaterialConstant, NodeMaterialConstant::GetStaticType());
+	FNodeType NodeMaterialTable::StaticType = FNodeType(Node::EType::MaterialTable, NodeMaterialTable::GetStaticType());
+	FNodeType NodeMaterialSwitch::StaticType = FNodeType(Node::EType::MaterialSwitch, NodeMaterialSwitch::GetStaticType());
+	FNodeType NodeMaterialVariation::StaticType = FNodeType(Node::EType::MaterialVariation, NodeMaterialVariation::GetStaticType());
+	FNodeType NodeMaterialParameter::StaticType = FNodeType(Node::EType::MaterialParameter, NodeMaterialParameter::GetStaticType());
+
 
 	FNodeType::FNodeType()
 	{
diff --git a/Engine/Plugins/Mutable/Source/MutableTools/Private/MuT/Table.cpp b/Engine/Plugins/Mutable/Source/MutableTools/Private/MuT/Table.cpp
--- a/Engine/Plugins/Mutable/Source/MutableTools/Private/MuT/Table.cpp
+++ b/Engine/Plugins/Mutable/Source/MutableTools/Private/MuT/Table.cpp
@@ -100,7 +100,20 @@ namespace mu
 	}
 
 
-    void FTable::SetCell( int32 Column, uint32 RowId, float Value, const void* ErrorContext)
+	void FTable::SetCell(int32 Column, uint32 RowId, int32 IntValue, const void* ErrorContext)
+	{
+		int32 Row = m_pD->FindRow(RowId);
+		check(Row >= 0);
+		check(Column < m_pD->Columns.Num());
+		check(Column < m_pD->Rows[Row].Values.Num());
+		check(m_pD->Columns[Column].Type == ETableColumnType::Material);
+
+		m_pD->Rows[Row].Values[Column].Int = IntValue;
+		m_pD->Rows[Row].Values[Column].ErrorContext = ErrorContext;
+	}
+
+
+	void FTable::SetCell( int32 Column, uint32 RowId, float Value, const void* ErrorContext)
 	{
 		int32 Row = m_pD->FindRow(RowId);
 		check( Row>=0 );
diff --git a/Engine/Plugins/Mutable/Source/MutableTools/Public/MuT/Node.h b/Engine/Plugins/Mutable/Source/MutableTools/Public/MuT/Node.h
--- a/Engine/Plugins/Mutable/Source/MutableTools/Public/MuT/Node.h
+++ b/Engine/Plugins/Mutable/Source/MutableTools/Public/MuT/Node.h
@@ -153,6 +153,13 @@ namespace mu
 			MatrixConstant,
 			MatrixParameter,
 
+			Material,
+			MaterialConstant,
+			MaterialTable,
+			MaterialSwitch,
+			MaterialVariation,
+			MaterialParameter,
+
 			Count
 		};
 
diff --git a/Engine/Plugins/Mutable/Source/MutableTools/Public/MuT/NodeComponentNew.h b/Engine/Plugins/Mutable/Source/MutableTools/Public/MuT/NodeComponentNew.h
--- a/Engine/Plugins/Mutable/Source/MutableTools/Public/MuT/NodeComponentNew.h
+++ b/Engine/Plugins/Mutable/Source/MutableTools/Public/MuT/NodeComponentNew.h
@@ -7,7 +7,7 @@
 #include "MuR/RefCounted.h"
 #include "MuT/Node.h"
 #include "MuT/NodeComponent.h"
-#include "MuT/NodeScalar.h"
+#include "MuT/NodeMaterial.h"
 
 #define UE_API MUTABLETOOLS_API
 
@@ -21,7 +21,7 @@ namespace mu
 		/** Externally managed id assigned to this component. */
 		uint16 Id = 0;
 
-		Ptr<NodeScalar> OverlayMaterial;
+		Ptr<NodeMaterial> OverlayMaterial;
 
 	public:
 
diff --git a/Engine/Plugins/Mutable/Source/MutableTools/Public/MuT/NodeMaterial.h b/Engine/Plugins/Mutable/Source/MutableTools/Public/MuT/NodeMaterial.h
new file mode 100644
--- /dev/null
+++ b/Engine/Plugins/Mutable/Source/MutableTools/Public/MuT/NodeMaterial.h
@@ -0,0 +1,38 @@
+// Copyright Epic Games, Inc. All Rights Reserved.
+
+#pragma once
+
+#include "HAL/Platform.h"
+#include "MuR/Ptr.h"
+#include "MuR/RefCounted.h"
+#include "MuT/Node.h"
+
+#define UE_API MUTABLETOOLS_API
+
+namespace mu
+{
+
+    /** Base class of any node that outputs a Material. */
+	class NodeMaterial : public Node
+	{
+	public:
+
+		// Node Interface
+		virtual const FNodeType* GetType() const override { return &StaticType; }
+		static const FNodeType* GetStaticType() { return &StaticType; }
+
+	protected:
+
+		// Forbidden. Manage with the Ptr<> template.
+		inline ~NodeMaterial() {}
+
+	private:
+
+		static UE_API FNodeType StaticType;
+
+	};
+
+
+}
+
+#undef UE_API
diff --git a/Engine/Plugins/Mutable/Source/MutableTools/Public/MuT/NodeMaterialConstant.h b/Engine/Plugins/Mutable/Source/MutableTools/Public/MuT/NodeMaterialConstant.h
new file mode 100644
--- /dev/null
+++ b/Engine/Plugins/Mutable/Source/MutableTools/Public/MuT/NodeMaterialConstant.h
@@ -0,0 +1,40 @@
+// Copyright Epic Games, Inc. All Rights Reserved.
+
+#pragma once
+
+#include "HAL/Platform.h"
+#include "MuR/Ptr.h"
+#include "MuR/RefCounted.h"
+#include "MuT/NodeMaterial.h"
+
+#define UE_API MUTABLETOOLS_API
+
+namespace mu
+{
+
+    /** NodeMaterialConstant class. */
+	class NodeMaterialConstant : public NodeMaterial
+	{
+	public:
+
+		// Node Interface
+		virtual const FNodeType* GetType() const override { return &StaticType; }
+		static const FNodeType* GetStaticType() { return &StaticType; }
+
+		int32 MaterialId = INDEX_NONE;
+
+	protected:
+
+		// Forbidden. Manage with the Ptr<> template.
+		inline ~NodeMaterialConstant() {}
+
+	private:
+
+		static UE_API FNodeType StaticType;
+
+	};
+
+
+}
+
+#undef UE_API
diff --git a/Engine/Plugins/Mutable/Source/MutableTools/Public/MuT/NodeMaterialParameter.h b/Engine/Plugins/Mutable/Source/MutableTools/Public/MuT/NodeMaterialParameter.h
new file mode 100644
--- /dev/null
+++ b/Engine/Plugins/Mutable/Source/MutableTools/Public/MuT/NodeMaterialParameter.h
@@ -0,0 +1,50 @@
+// Copyright Epic Games, Inc. All Rights Reserved.
+
+#pragma once
+
+#include "HAL/Platform.h"
+#include "MuR/Ptr.h"
+#include "MuR/RefCounted.h"
+#include "MuT/NodeRange.h"
+#include "MuT/NodeMaterial.h"
+
+#define UE_API MUTABLETOOLS_API
+
+namespace mu
+{
+
+    /** NodeMaterialParameter class. */
+	class NodeMaterialParameter : public NodeMaterial
+	{
+	public:
+
+		/** Name of the parameter */
+		FString Name;
+
+		/** User provided ID to identify the parameter. */
+		FString UID;
+
+		/** Ranges for the parameter in case it is multidimensional. */
+		TArray<Ptr<NodeRange>> Ranges;
+
+	public:
+
+		// Node Interface
+		virtual const FNodeType* GetType() const override { return &StaticType; }
+		static const FNodeType* GetStaticType() { return &StaticType; }
+
+	protected:
+
+		// Forbidden. Manage with the Ptr<> template.
+		inline ~NodeMaterialParameter() {}
+
+	private:
+
+		static UE_API FNodeType StaticType;
+
+	};
+
+
+}
+
+#undef UE_API
diff --git a/Engine/Plugins/Mutable/Source/MutableTools/Public/MuT/NodeMaterialSwitch.h b/Engine/Plugins/Mutable/Source/MutableTools/Public/MuT/NodeMaterialSwitch.h
new file mode 100644
--- /dev/null
+++ b/Engine/Plugins/Mutable/Source/MutableTools/Public/MuT/NodeMaterialSwitch.h
@@ -0,0 +1,45 @@
+// Copyright Epic Games, Inc. All Rights Reserved.
+
+#pragma once
+
+#include "HAL/Platform.h"
+#include "MuR/Ptr.h"
+#include "MuR/RefCounted.h"
+#include "MuT/NodeScalar.h"
+#include "MuT/NodeMaterial.h"
+
+#define UE_API MUTABLETOOLS_API
+
+namespace mu
+{
+
+    /** NodeMaterialSwitch class. */
+	class NodeMaterialSwitch : public NodeMaterial
+	{
+	public:
+
+		Ptr<NodeScalar> Parameter;
+
+		TArray<Ptr<NodeMaterial>> Options;
+
+	public:
+
+		// Node Interface
+		virtual const FNodeType* GetType() const override { return &StaticType; }
+		static const FNodeType* GetStaticType() { return &StaticType; }
+
+	protected:
+
+		// Forbidden. Manage with the Ptr<> template.
+		inline ~NodeMaterialSwitch() {}
+
+	private:
+
+		static UE_API FNodeType StaticType;
+
+	};
+
+
+}
+
+#undef UE_API
diff --git a/Engine/Plugins/Mutable/Source/MutableTools/Public/MuT/NodeMaterialTable.h b/Engine/Plugins/Mutable/Source/MutableTools/Public/MuT/NodeMaterialTable.h
new file mode 100644
--- /dev/null
+++ b/Engine/Plugins/Mutable/Source/MutableTools/Public/MuT/NodeMaterialTable.h
@@ -0,0 +1,48 @@
+// Copyright Epic Games, Inc. All Rights Reserved.
+
+#pragma once
+
+#include "HAL/Platform.h"
+#include "MuR/Ptr.h"
+#include "MuR/RefCounted.h"
+#include "MuT/Table.h"
+#include "MuT/NodeMaterial.h"
+
+#define UE_API MUTABLETOOLS_API
+
+namespace mu
+{
+
+    /** NodeMaterialTable class. */
+	class NodeMaterialTable : public NodeMaterial
+	{
+	public:
+		FString ParameterName;
+		FString ColumnName;
+		FString DefaultRowName;
+		
+		bool bNoneOption = false;
+		
+		Ptr<FTable> Table;
+
+	public:
+
+		// Node Interface
+		virtual const FNodeType* GetType() const override { return &StaticType; }
+		static const FNodeType* GetStaticType() { return &StaticType; }
+
+	protected:
+
+		// Forbidden. Manage with the Ptr<> template.
+		inline ~NodeMaterialTable() {}
+
+	private:
+
+		static UE_API FNodeType StaticType;
+
+	};
+
+
+}
+
+#undef UE_API
diff --git a/Engine/Plugins/Mutable/Source/MutableTools/Public/MuT/NodeMaterialVariation.h b/Engine/Plugins/Mutable/Source/MutableTools/Public/MuT/NodeMaterialVariation.h
new file mode 100644
--- /dev/null
+++ b/Engine/Plugins/Mutable/Source/MutableTools/Public/MuT/NodeMaterialVariation.h
@@ -0,0 +1,48 @@
+// Copyright Epic Games, Inc. All Rights Reserved.
+
+#pragma once
+
+#include "HAL/Platform.h"
+#include "MuR/Ptr.h"
+#include "MuR/RefCounted.h"
+#include "MuT/NodeMaterial.h"
+
+#define UE_API MUTABLETOOLS_API
+
+namespace mu
+{
+
+    /** NodeMaterialVariation class. */
+	class NodeMaterialVariation : public NodeMaterial
+	{
+	public:
+		Ptr<NodeMaterial> DefaultMaterial;
+
+		struct FVariation
+		{
+			Ptr<NodeMaterial> Material;
+			FString Tag;
+		};
+
+		TArray<FVariation> Variations;
+	public:
+
+		// Node Interface
+		virtual const FNodeType* GetType() const override { return &StaticType; }
+		static const FNodeType* GetStaticType() { return &StaticType; }
+
+	protected:
+
+		// Forbidden. Manage with the Ptr<> template.
+		inline ~NodeMaterialVariation() {}
+
+	private:
+
+		static UE_API FNodeType StaticType;
+
+	};
+
+
+}
+
+#undef UE_API
diff --git a/Engine/Plugins/Mutable/Source/MutableTools/Public/MuT/NodeSurfaceNew.h b/Engine/Plugins/Mutable/Source/MutableTools/Public/MuT/NodeSurfaceNew.h
--- a/Engine/Plugins/Mutable/Source/MutableTools/Public/MuT/NodeSurfaceNew.h
+++ b/Engine/Plugins/Mutable/Source/MutableTools/Public/MuT/NodeSurfaceNew.h
@@ -7,6 +7,7 @@
 #include "MuT/Node.h"
 #include "MuT/NodeSurface.h"
 #include "MuT/NodeMesh.h"
+#include "MuT/NodeMaterial.h"
 #include "MuT/NodeImage.h"
 #include "MuT/NodeScalar.h"
 #include "MuT/NodeString.h"
@@ -36,6 +37,8 @@ namespace mu
 
 		Ptr<NodeMesh> Mesh;
 
+		Ptr<NodeMaterial> Material;
+
 		struct FImageData
 		{
 			FString Name;
diff --git a/Engine/Plugins/Mutable/Source/MutableTools/Public/MuT/Table.h b/Engine/Plugins/Mutable/Source/MutableTools/Public/MuT/Table.h
--- a/Engine/Plugins/Mutable/Source/MutableTools/Public/MuT/Table.h
+++ b/Engine/Plugins/Mutable/Source/MutableTools/Public/MuT/Table.h
@@ -24,7 +24,8 @@ namespace mu
 		Color,
 		Image,
 		Mesh,
-		String
+		String,
+		Material
 	};
 
 
@@ -56,6 +57,7 @@ namespace mu
         UE_API void AddRow( uint32 id );
 
 		//!
+        UE_API void SetCell( int32 Column, uint32 RowId, int32 IntValue, const void* ErrorContext = nullptr);
         UE_API void SetCell( int32 Column, uint32 RowId, float Value, const void* ErrorContext = nullptr);
         UE_API void SetCell( int32 Column, uint32 RowId, const FVector4f& Value, const void* ErrorContext = nullptr);
 		UE_API void SetCell( int32 Column, uint32 RowId, TResourceProxy<FImage>* Value, const void* ErrorContext = nullptr);
