[mutable ci] Added additional check to ensure the shaders have been compiled before we start doing mutable work
- Made the GC collection performed by the commandlets more agressive to reduce peak memory usage.

#rnx

diff --git a/Engine/Plugins/Mutable/Source/MutableValidation/Private/MuV/COIBakingTestCommandlet.cpp b/Engine/Plugins/Mutable/Source/MutableValidation/Private/MuV/COIBakingTestCommandlet.cpp
--- a/Engine/Plugins/Mutable/Source/MutableValidation/Private/MuV/COIBakingTestCommandlet.cpp
+++ b/Engine/Plugins/Mutable/Source/MutableValidation/Private/MuV/COIBakingTestCommandlet.cpp
@@ -3,6 +3,7 @@
 #include "MuV/COIBakingTestCommandlet.h"
 #include "CustomizableObjectCompilationUtility.h"
 #include "ScopedLogSection.h"
+#include "ShaderCompiler.h"
 #include "ValidationUtils.h"
 #include "AssetRegistry/AssetData.h"
 #include "HAL/FileManager.h"
@@ -15,7 +16,6 @@
 #include "MuCO/CustomizableObjectPrivate.h"
 #include "MuCO/CustomizableObjectSystemPrivate.h"
 #include "MuCO/LoadUtils.h"
-#include "MuCOE/CustomizableObjectBenchmarkingUtils.h"
 
 #include UE_INLINE_GENERATED_CPP_BY_NAME(COIBakingTestCommandlet)
 
@@ -80,8 +80,15 @@ int32 UCOIBakingTestCommandlet::Main(const FString& Params)
 	PrepareAssetRegistry();
 
 	// Make sure there is nothing else that the engine needs to do before starting our test
-	Wait(60);
+	Wait(60);	// todo: UE-304050 Remove this wait as it may no longer be required due to us calling for GShaderCompilingManager->FinishAllCompilation()
 
+	// Block until async shader compiling is finished before we try to use the shaders for exporting
+	// The code is structured to only block once for all materials, so that shader compiling is able to utilize many cores
+	if (GShaderCompilingManager && GShaderCompilingManager->GetNumRemainingJobs() > 0)
+	{
+		GShaderCompilingManager->FinishAllCompilation();
+	}
+	
 	LogGlobalSettings();
 	
 	// Compile it's CO (using current config)
diff --git a/Engine/Plugins/Mutable/Source/MutableValidation/Private/MuV/COIBulkUpdateTestCommandlet.cpp b/Engine/Plugins/Mutable/Source/MutableValidation/Private/MuV/COIBulkUpdateTestCommandlet.cpp
--- a/Engine/Plugins/Mutable/Source/MutableValidation/Private/MuV/COIBulkUpdateTestCommandlet.cpp
+++ b/Engine/Plugins/Mutable/Source/MutableValidation/Private/MuV/COIBulkUpdateTestCommandlet.cpp
@@ -4,6 +4,7 @@
 
 #include "CustomizableObjectCompilationUtility.h"
 #include "CustomizableObjectInstanceUpdateUtility.h"
+#include "ShaderCompiler.h"
 #include "ValidationUtils.h"
 #include "AssetRegistry/AssetData.h"
 #include "AssetRegistry/AssetRegistryModule.h"
@@ -119,7 +120,14 @@ int32 UCOIBulkUpdateTestCommandlet::Main(const FString& Params)
 	// ------ Execution of the actual mutable operations ------ 
 	
 	// Make sure there is nothing else that the engine needs to do before starting our test
-	Wait(60);
+	Wait(60);	// todo: UE-304050 Remove this wait as it may no longer be required due to us calling for GShaderCompilingManager->FinishAllCompilation()
+	
+	// Block until async shader compiling is finished before we try to use the shaders for exporting
+	// The code is structured to only block once for all materials, so that shader compiling is able to utilize many cores
+	if (GShaderCompilingManager && GShaderCompilingManager->GetNumRemainingJobs() > 0)
+	{
+		GShaderCompilingManager->FinishAllCompilation();
+	}
 
 	// Cache the target compilation platform so we can override the compilation configs of the target COs
 	ITargetPlatformManagerModule& TPM = GetTargetPlatformManagerRef();
@@ -163,7 +171,7 @@ int32 UCOIBulkUpdateTestCommandlet::Main(const FString& Params)
 			Instances.RemoveAt(0);
 			UE_LOG(LogMutable, Display, TEXT("\t( %u / %u ) Processing instance : \"%s\" ."),CurrentInstanceIndex++, TotalAmountOfInstances ,*Instance->GetName());
 
-			CollectGarbage(GARBAGE_COLLECTION_KEEPFLAGS, true);
+			CollectGarbage(RF_NoFlags, true);
 			
 			// Update each one of the instances and notify if the update failed in any manner
 			InstanceUpdatingUtility->UpdateInstance(*Instance);
@@ -173,7 +181,7 @@ int32 UCOIBulkUpdateTestCommandlet::Main(const FString& Params)
 		}
 
 		ResourcesIterator.RemoveCurrent();
-		CollectGarbage(GARBAGE_COLLECTION_KEEPFLAGS, true);
+		CollectGarbage(RF_NoFlags, true);
 	}
 
 	UE_LOG(LogMutable,Display,TEXT("Mutable commandlet finished."));
diff --git a/Engine/Plugins/Mutable/Source/MutableValidation/Private/MuV/CustomizableObjectBulkValidationCommandlet.cpp b/Engine/Plugins/Mutable/Source/MutableValidation/Private/MuV/CustomizableObjectBulkValidationCommandlet.cpp
--- a/Engine/Plugins/Mutable/Source/MutableValidation/Private/MuV/CustomizableObjectBulkValidationCommandlet.cpp
+++ b/Engine/Plugins/Mutable/Source/MutableValidation/Private/MuV/CustomizableObjectBulkValidationCommandlet.cpp
@@ -2,7 +2,7 @@
 
 #include "MuV/CustomizableObjectBulkValidationCommandlet.h"
 
-#include "ScopedLogSection.h"
+#include "ShaderCompiler.h"
 #include "ValidationUtils.h"
 #include "MuCO/CustomizableObject.h"
 #include "MuCO/CustomizableObjectSystemPrivate.h"
@@ -92,7 +92,14 @@ int32 UCustomizableObjectBulkValidationCommandlet::Main(const FString& Params)
 	}
 	
 	// Make sure there is nothing else that the engine needs to do before starting our test
-	Wait(60);
+	Wait(60);	// todo: UE-304050 Remove this wait as it may no longer be required due to us calling for GShaderCompilingManager->FinishAllCompilation()
+
+	// Block until async shader compiling is finished before we try to use the shaders for exporting
+	// The code is structured to only block once for all materials, so that shader compiling is able to utilize many cores
+	if (GShaderCompilingManager && GShaderCompilingManager->GetNumRemainingJobs() > 0)
+	{
+		GShaderCompilingManager->FinishAllCompilation();
+	}
 	
 	LogGlobalSettings();
 	
diff --git a/Engine/Plugins/Mutable/Source/MutableValidation/Private/MuV/CustomizableObjectInstanceUpdateUtility.cpp b/Engine/Plugins/Mutable/Source/MutableValidation/Private/MuV/CustomizableObjectInstanceUpdateUtility.cpp
--- a/Engine/Plugins/Mutable/Source/MutableValidation/Private/MuV/CustomizableObjectInstanceUpdateUtility.cpp
+++ b/Engine/Plugins/Mutable/Source/MutableValidation/Private/MuV/CustomizableObjectInstanceUpdateUtility.cpp
@@ -8,8 +8,8 @@
 #include "Engine/TextureStreamingTypes.h"
 #include "SceneTypes.h"
 #include "ScopedLogSection.h"
-#include "ValidationUtils.h"
 #include "Components/SkeletalMeshComponent.h"
+#include "Engine/SkeletalMesh.h"
 
 
 bool FCustomizableObjectInstanceUpdateUtility::UpdateInstance(UCustomizableObjectInstance& InInstance)
@@ -114,25 +114,22 @@ void FCustomizableObjectInstanceUpdateUtility::OnInstanceUpdateResult(const FUpd
 			return;
 		}
 
+		// Request the streaming in of all the components affected by the update
 		for (const FName& ComponentName : InstanceComponents)
 		{
 			TStrongObjectPtr<USkeletalMeshComponent> SkeletalComponent = TStrongObjectPtr(NewObject<USkeletalMeshComponent>());
 			SkeletalComponent->SetSkeletalMesh(Instance->GetComponentMeshSkeletalMesh(ComponentName));
 
-			ComponentsBeingUpdated.Add(SkeletalComponent);
-		}
-
-		// Request the streaming in of all the components affected by the update
-		for (TStrongObjectPtr<USkeletalMeshComponent>& ComponentBeingUpdated : ComponentsBeingUpdated)
-		{
-			FStreamingTextureLevelContext LevelContext(EMaterialQualityLevel::Num, ComponentBeingUpdated.Get());
+			FStreamingTextureLevelContext LevelContext(EMaterialQualityLevel::Num, SkeletalComponent.Get());
 			TArray<FStreamingRenderAssetPrimitiveInfo> RenderAssetInfoArray;
-			ComponentBeingUpdated->GetStreamingRenderAssetInfo(LevelContext, RenderAssetInfoArray);
+			SkeletalComponent->GetStreamingRenderAssetInfo(LevelContext, RenderAssetInfoArray);
 
 			for (const FStreamingRenderAssetPrimitiveInfo& Info : RenderAssetInfoArray)
 			{
 				Info.RenderAsset->StreamIn(MAX_int32, true);
 			}
+			
+			ComponentsBeingUpdated.Add(SkeletalComponent);
 		}
 
 		UE_LOG(LogMutable, Display, TEXT("Instance %s requesting the streaming of all MIPs."), *InstanceName);
diff --git a/Engine/Plugins/Mutable/Source/MutableValidation/Private/MuV/CustomizableObjectValidationCommandlet.cpp b/Engine/Plugins/Mutable/Source/MutableValidation/Private/MuV/CustomizableObjectValidationCommandlet.cpp
--- a/Engine/Plugins/Mutable/Source/MutableValidation/Private/MuV/CustomizableObjectValidationCommandlet.cpp
+++ b/Engine/Plugins/Mutable/Source/MutableValidation/Private/MuV/CustomizableObjectValidationCommandlet.cpp
@@ -2,23 +2,17 @@
 
 #include "MuV/CustomizableObjectValidationCommandlet.h"
 
-#include "CustomizableObjectCompilationUtility.h"
-#include "CustomizableObjectInstanceUpdateUtility.h"
+#include "ShaderCompiler.h"
 #include "ValidationUtils.h"
 #include "Components/SkeletalMeshComponent.h"
 #include "Containers/Array.h"
 #include "Containers/SpscQueue.h"
 #include "MuCO/CustomizableObject.h"
 #include "MuCO/CustomizableObjectInstancePrivate.h"
-#include "MuCO/CustomizableObjectPrivate.h"
-#include "MuCO/CustomizableObjectSystem.h"
-#include "MuCO/CustomizableObjectCompilerTypes.h"
-#include "MuR/Model.h"
 #include "Interfaces/ITargetPlatformManagerModule.h"
 #include "Interfaces/ITargetPlatform.h"
 #include "MuCO/CustomizableObjectSystemPrivate.h"
 #include "MuCO/LoadUtils.h"
-#include "MuCOE/CustomizableObjectBenchmarkingUtils.h"
 #include "Serialization/MemoryWriter.h"
 
 #include UE_INLINE_GENERATED_CPP_BY_NAME(CustomizableObjectValidationCommandlet)
@@ -79,8 +73,15 @@ int32 UCustomizableObjectValidationCommandlet::Main(const FString& Params)
 	PrepareAssetRegistry();
 
 	// Make sure there is nothing else that the engine needs to do before starting our test
-	Wait(60);
+	Wait(60);	// todo: UE-304050 Remove this wait as it may no longer be required due to us calling for GShaderCompilingManager->FinishAllCompilation()
 
+	// Block until async shader compiling is finished before we try to use the shaders for exporting
+	// The code is structured to only block once for all materials, so that shader compiling is able to utilize many cores
+	if (GShaderCompilingManager && GShaderCompilingManager->GetNumRemainingJobs() > 0)
+	{
+		GShaderCompilingManager->FinishAllCompilation();
+	}
+	
 	LogGlobalSettings();
 
 	// Body of the test ------------------------------------------------------------------------------------------------------------------------------
diff --git a/Engine/Plugins/Mutable/Source/MutableValidation/Private/MuV/ValidationUtils.cpp b/Engine/Plugins/Mutable/Source/MutableValidation/Private/MuV/ValidationUtils.cpp
--- a/Engine/Plugins/Mutable/Source/MutableValidation/Private/MuV/ValidationUtils.cpp
+++ b/Engine/Plugins/Mutable/Source/MutableValidation/Private/MuV/ValidationUtils.cpp
@@ -272,7 +272,7 @@ bool TestCustomizableObject(UCustomizableObject& InTargetCustomizableObject, con
 		TStrongObjectPtr<UCustomizableObjectInstance> InstanceToUpdate;
 		while (InstancesToProcess.Dequeue(InstanceToUpdate))
 		{
-			CollectGarbage(GARBAGE_COLLECTION_KEEPFLAGS, true);
+			CollectGarbage(RF_NoFlags, true);
 			check(InstanceToUpdate);
 
 			if (!InstanceUpdatingUtility->UpdateInstance(*InstanceToUpdate))
