[Mutable] Added new Material nodes: Material Constant. Material Variation and Material Swtch.

#rnx

diff --git a/Engine/Plugins/Mutable/Source/CustomizableObjectEditor/Private/MuCOE/EdGraphSchema_CustomizableObject.cpp b/Engine/Plugins/Mutable/Source/CustomizableObjectEditor/Private/MuCOE/EdGraphSchema_CustomizableObject.cpp
--- a/Engine/Plugins/Mutable/Source/CustomizableObjectEditor/Private/MuCOE/EdGraphSchema_CustomizableObject.cpp
+++ b/Engine/Plugins/Mutable/Source/CustomizableObjectEditor/Private/MuCOE/EdGraphSchema_CustomizableObject.cpp
@@ -16,6 +16,9 @@
 #include "MuCOE/CustomizableObjectGraph.h"
 #include "MuCOE/CustomizableObjectMacroLibrary/CustomizableObjectGraphEditorToolkit.h"
 #include "MuCOE/CustomizableObjectEditor.h"
+#include "MuCOE/Nodes/CONodeMaterialConstant.h"
+#include "MuCOE/Nodes/CONodeMaterialVariation.h"
+#include "MuCOE/Nodes/CONodeMaterialSwitch.h"
 #include "MuCOE/Nodes/CustomizableObjectNodeAnimationPose.h"
 #include "MuCOE/Nodes/CustomizableObjectNodeColorArithmeticOp.h"
 #include "MuCOE/Nodes/CustomizableObjectNodeColorConstant.h"
@@ -293,6 +296,7 @@ PC_PassThroughImage, PC_Projector, PC_GroupProjector, PC_Color, PC_Float, PC_Enu
 
 // Node categories
 const FText UEdGraphSchema_CustomizableObject::NC_Experimental(FText::FromString("Experimental"));
+const FText UEdGraphSchema_CustomizableObject::NC_Material(FText::FromString("Material"));
 
 
 TSharedPtr<FCustomizableObjectSchemaAction_NewNode> AddNewNodeAction(FGraphActionListBuilderBase& ContextMenuBuilder, const FString& Category, const FText& MenuDesc, const FText& Tooltip, const int32 Grouping = 0, const FString& Keywords = FString())
@@ -618,6 +622,17 @@ void UEdGraphSchema_CustomizableObject::GetGraphContextActions(FGraphContextMenu
 		AddNewNodeCategoryActionsFiltered(ExperimentalTemplateNodes, ContextMenuBuilder, NC_Experimental.ToString(), 2, Filter);
 	}
 
+	{
+		UCustomizableObjectNode* MaterialTemplateNodes[] =
+		{
+			ContextMenuBuilder.CreateTemplateNode<UCONodeMaterialConstant>(),
+			ContextMenuBuilder.CreateTemplateNode<UCONodeMaterialVariation>(),
+			ContextMenuBuilder.CreateTemplateNode<UCONodeMaterialSwitch>()
+		};
+
+		AddNewNodeCategoryActionsFiltered(MaterialTemplateNodes, ContextMenuBuilder, NC_Material.ToString(), 2, Filter);
+	}
+
 	{
 		// External Pin Nodes
 		TArray<FName> PinTypes({ PC_MeshSection, PC_Modifier, PC_Mesh, PC_Image, PC_Projector, PC_GroupProjector, PC_Color, PC_Float, PC_Bool, PC_Enum, PC_Transform, PC_Stack, PC_PassThroughImage, PC_Material, PC_PoseAsset, PC_Component });
diff --git a/Engine/Plugins/Mutable/Source/CustomizableObjectEditor/Private/MuCOE/GenerateMutableSource/GenerateMutableSourceMaterial.cpp b/Engine/Plugins/Mutable/Source/CustomizableObjectEditor/Private/MuCOE/GenerateMutableSource/GenerateMutableSourceMaterial.cpp
--- a/Engine/Plugins/Mutable/Source/CustomizableObjectEditor/Private/MuCOE/GenerateMutableSource/GenerateMutableSourceMaterial.cpp
+++ b/Engine/Plugins/Mutable/Source/CustomizableObjectEditor/Private/MuCOE/GenerateMutableSource/GenerateMutableSourceMaterial.cpp
@@ -5,6 +5,7 @@
 #include "MuCOE/CustomizableObjectCompiler.h"
 #include "MuCOE/CustomizableObjectGraph.h"
 #include "MuCOE/EdGraphSchema_CustomizableObject.h"
+#include "MuCOE/GenerateMutableSource/GenerateMutableSourceFloat.h"
 #include "MuCOE/GenerateMutableSource/GenerateMutableSourceMacro.h"
 #include "MuCOE/GenerateMutableSource/GenerateMutableSourceTable.h"
 #include "MuCOE/GraphTraversal.h"
@@ -13,6 +14,9 @@
 #include "MuCOE/Nodes/CustomizableObjectNodeMacroInstance.h"
 #include "MuCOE/Nodes/CustomizableObjectNodeTable.h"
 #include "MuCOE/Nodes/CustomizableObjectNodeTunnel.h"
+#include "MuCOE/Nodes/CONodeMaterialConstant.h"
+#include "MuCOE/Nodes/CONodeMaterialSwitch.h"
+#include "MuCOE/Nodes/CONodeMaterialVariation.h"
 #include "MuCO/LoadUtils.h"
 #include "MuT/NodeMaterialConstant.h"
 #include "MuT/NodeMaterialSwitch.h"
@@ -55,7 +59,120 @@ mu::Ptr<mu::NodeMaterial> GenerateMutableSourceMaterial(const UEdGraphPin* Pin,
 
 	mu::Ptr<mu::NodeMaterial> Result;
 	
-	if (const UCustomizableObjectNodeMaterialParameter* MaterialParameterNode = Cast<UCustomizableObjectNodeMaterialParameter>(Node))
+	if (const UCONodeMaterialConstant* TypedNodeMaterialConstant = Cast<UCONodeMaterialConstant>(Node))
+	{
+		mu::Ptr<mu::NodeMaterialConstant> MaterialNode = new mu::NodeMaterialConstant();
+		Result = MaterialNode;
+
+		if (UMaterialInterface* Material = GenerationContext.LoadObject(TypedNodeMaterialConstant->Material))
+		{
+			GenerationContext.CurrentReferencedMaterialIndex = GenerationContext.ReferencedMaterials.AddUnique(Material);
+			MaterialNode->MaterialId = GenerationContext.CurrentReferencedMaterialIndex;
+		}
+	}
+
+	else if (const UCONodeMaterialSwitch* TypedNodeMaterialSwitch = Cast<UCONodeMaterialSwitch>(Node))
+	{
+		Result = [&]()
+		{
+			const UEdGraphPin* SwitchParameter = TypedNodeMaterialSwitch->SwitchParameter();
+
+			// Check Switch Parameter arity preconditions.
+			if (const UEdGraphPin* ConnectedPin = FollowInputPin(*SwitchParameter))
+			{
+				mu::Ptr<mu::NodeScalar> SwitchParam = GenerateMutableSourceFloat(ConnectedPin, GenerationContext);
+
+				// Switch Param not generated
+				if (!SwitchParam)
+				{
+					const FText Message = LOCTEXT("FailedToGenerateSwitchParam", "Could not generate switch enum parameter. Please refesh the switch node and connect an enum.");
+					GenerationContext.Log(Message, Node);
+
+					return Result;
+				}
+
+				if (SwitchParam->GetType() != mu::NodeScalarEnumParameter::GetStaticType())
+				{
+					const FText Message = LOCTEXT("WrongSwitchParamType", "Switch parameter of incorrect type.");
+					GenerationContext.Log(Message, Node);
+
+					return Result;
+				}
+
+				const int32 NumSwitchOptions = TypedNodeMaterialSwitch->GetNumElements();
+
+				mu::NodeScalarEnumParameter* EnumParameter = static_cast<mu::NodeScalarEnumParameter*>(SwitchParam.get());
+				if (NumSwitchOptions != EnumParameter->Options.Num())
+				{
+					const FText Message = LOCTEXT("MismatchedSwitch", "Switch enum and switch node have different number of options. Please refresh the switch node to make sure the outcomes are labeled properly.");
+					GenerationContext.Log(Message, Node);
+				}
+
+				mu::Ptr<mu::NodeMaterialSwitch> SwitchNode = new mu::NodeMaterialSwitch;
+				SwitchNode->Parameter = SwitchParam;
+				SwitchNode->Options.SetNum(NumSwitchOptions);
+
+				for (int SelectorIndex = 0; SelectorIndex < NumSwitchOptions; ++SelectorIndex)
+				{
+					if (const UEdGraphPin* MaterialPin = FollowInputPin(*TypedNodeMaterialSwitch->GetElementPin(SelectorIndex)))
+					{
+						SwitchNode->Options[SelectorIndex] = GenerateMutableSourceMaterial(MaterialPin, GenerationContext);
+					}
+					else
+					{
+						const FText Message = LOCTEXT("MissingTexture", "Unable to generate material switch node. Required connection not found.");
+						GenerationContext.Log(Message, Node);
+						return Result;
+					}
+				}
+
+				Result = SwitchNode;
+				return Result;
+			}
+			else
+			{
+				GenerationContext.Log(LOCTEXT("NoEnumParamInSwitch", "Switch nodes must have an enum switch parameter. Please connect an enum and refesh the switch node."), Node);
+				return Result;
+			}
+		}(); // invoke lambda.
+	}
+
+	else if (const UCONodeMaterialVariation* TypedNodeMaterialVariation = Cast<UCONodeMaterialVariation>(Node))
+	{
+		mu::Ptr<mu::NodeMaterialVariation> MaterialNode = new mu::NodeMaterialVariation();
+		Result = MaterialNode;
+
+		if (const UEdGraphPin* ConnectedPin = FollowInputPin(*TypedNodeMaterialVariation->DefaultPin()))
+		{
+			mu::Ptr<mu::NodeMaterial> ChildNode = GenerateMutableSourceMaterial(ConnectedPin, GenerationContext);
+			if (ChildNode)
+			{
+				MaterialNode->DefaultMaterial = ChildNode;
+			}
+			else
+			{
+				GenerationContext.Log(LOCTEXT("MaterialVariationFailed", "Material variation generation failed."), Node);
+			}
+		}
+
+		const int32 NumVariations = TypedNodeMaterialVariation->GetNumVariations();
+		MaterialNode->Variations.SetNum(NumVariations);
+		for (int VariationIndex = 0; VariationIndex < NumVariations; ++VariationIndex)
+		{
+			if (const UEdGraphPin* VariationPin = TypedNodeMaterialVariation->VariationPin(VariationIndex))
+			{
+				MaterialNode->Variations[VariationIndex].Tag = TypedNodeMaterialVariation->GetVariationTag(VariationIndex, &GenerationContext.MacroNodesStack);
+
+				if (const UEdGraphPin* ConnectedPin = FollowInputPin(*VariationPin))
+				{
+					mu::Ptr<mu::NodeMaterial> ChildNode = GenerateMutableSourceMaterial(ConnectedPin, GenerationContext);
+					MaterialNode->Variations[VariationIndex].Material = ChildNode;
+				}
+			}
+		}
+	}
+	
+	else if (const UCustomizableObjectNodeMaterialParameter* MaterialParameterNode = Cast<UCustomizableObjectNodeMaterialParameter>(Node))
 	{
 		mu::Ptr<mu::NodeMaterialParameter> MaterialNode = new mu::NodeMaterialParameter();
 		Result = MaterialNode;
diff --git a/Engine/Plugins/Mutable/Source/CustomizableObjectEditor/Private/MuCOE/Nodes/CONodeMaterialConstant.cpp b/Engine/Plugins/Mutable/Source/CustomizableObjectEditor/Private/MuCOE/Nodes/CONodeMaterialConstant.cpp
new file mode 100644
--- /dev/null
+++ b/Engine/Plugins/Mutable/Source/CustomizableObjectEditor/Private/MuCOE/Nodes/CONodeMaterialConstant.cpp
@@ -0,0 +1,43 @@
+// Copyright Epic Games, Inc. All Rights Reserved.
+
+#include "MuCOE/Nodes/CONodeMaterialConstant.h"
+#include "MuCOE/CustomizableObjectGraph.h"
+#include "MuCOE/EdGraphSchema_CustomizableObject.h"
+
+#include UE_INLINE_GENERATED_CPP_BY_NAME(CONodeMaterialConstant)
+
+class UCustomizableObjectNodeRemapPins;
+struct FPropertyChangedEvent;
+struct FSlateBrush;
+
+#define LOCTEXT_NAMESPACE "CustomizableObjectEditor"
+
+void UCONodeMaterialConstant::AllocateDefaultPins(UCustomizableObjectNodeRemapPins* RemapPins)
+{
+	const FName Type = UEdGraphSchema_CustomizableObject::PC_Material;
+	const FName PinName = UEdGraphSchema_CustomizableObject::GetPinCategoryName(Type);
+	const FText PinFriendlyName = UEdGraphSchema_CustomizableObject::GetPinCategoryFriendlyName(Type);
+	
+	UEdGraphPin* ValuePin = CustomCreatePin(EGPD_Output, Type, PinName);
+	ValuePin->PinFriendlyName = PinFriendlyName;
+}
+
+
+FText UCONodeMaterialConstant::GetNodeTitle(ENodeTitleType::Type TitleType) const
+{
+	return LOCTEXT("Material_Constant", "Material Constant");
+}
+
+
+FLinearColor UCONodeMaterialConstant::GetNodeTitleColor() const
+{
+	return UEdGraphSchema_CustomizableObject::GetPinTypeColor(UEdGraphSchema_CustomizableObject::PC_Material);
+}
+
+
+FText UCONodeMaterialConstant::GetTooltipText() const
+{
+	return LOCTEXT("Material_Constant_Tooltip", "Define an Unreal Engine Material that does not change at runtime.");
+}
+
+#undef LOCTEXT_NAMESPACE
diff --git a/Engine/Plugins/Mutable/Source/CustomizableObjectEditor/Private/MuCOE/Nodes/CONodeMaterialConstant.h b/Engine/Plugins/Mutable/Source/CustomizableObjectEditor/Private/MuCOE/Nodes/CONodeMaterialConstant.h
new file mode 100644
--- /dev/null
+++ b/Engine/Plugins/Mutable/Source/CustomizableObjectEditor/Private/MuCOE/Nodes/CONodeMaterialConstant.h
@@ -0,0 +1,34 @@
+// Copyright Epic Games, Inc. All Rights Reserved.
+
+#pragma once
+
+#include "MuCOE/Nodes/CustomizableObjectNode.h"
+#include "CONodeMaterialConstant.generated.h"
+
+#define UE_API CUSTOMIZABLEOBJECTEDITOR_API
+
+namespace ENodeTitleType { enum Type : int; }
+class UCustomizableObjectNodeRemapPins;
+
+UCLASS(MinimalAPI)
+class UCONodeMaterialConstant : public UCustomizableObjectNode
+{
+public:
+
+	GENERATED_BODY()
+
+	/** Value of the Node. */
+	UPROPERTY(EditAnywhere, Category = CustomizableObject)
+	TSoftObjectPtr<UMaterialInterface> Material;
+
+	// Begin EdGraphNode interface
+	UE_API virtual FText GetNodeTitle(ENodeTitleType::Type TitleType) const override;
+	UE_API virtual FLinearColor GetNodeTitleColor() const override;
+	UE_API virtual FText GetTooltipText() const override;
+
+	// UCustomizableObjectNode interface
+	UE_API virtual void AllocateDefaultPins(UCustomizableObjectNodeRemapPins* RemapPins) override;
+
+};
+
+#undef UE_API
diff --git a/Engine/Plugins/Mutable/Source/CustomizableObjectEditor/Private/MuCOE/Nodes/CONodeMaterialSwitch.cpp b/Engine/Plugins/Mutable/Source/CustomizableObjectEditor/Private/MuCOE/Nodes/CONodeMaterialSwitch.cpp
new file mode 100644
--- /dev/null
+++ b/Engine/Plugins/Mutable/Source/CustomizableObjectEditor/Private/MuCOE/Nodes/CONodeMaterialSwitch.cpp
@@ -0,0 +1,17 @@
+// Copyright Epic Games, Inc. All Rights Reserved.
+
+#include "MuCOE/Nodes/CONodeMaterialSwitch.h"
+#include "MuCOE/EdGraphSchema_CustomizableObject.h"
+#include UE_INLINE_GENERATED_CPP_BY_NAME(CONodeMaterialSwitch)
+
+#define LOCTEXT_NAMESPACE "CustomizableObjectEditor"
+
+
+FName UCONodeMaterialSwitch::GetCategory() const
+{
+	return UEdGraphSchema_CustomizableObject::PC_Material;
+}
+
+
+#undef LOCTEXT_NAMESPACE
+
diff --git a/Engine/Plugins/Mutable/Source/CustomizableObjectEditor/Private/MuCOE/Nodes/CONodeMaterialSwitch.h b/Engine/Plugins/Mutable/Source/CustomizableObjectEditor/Private/MuCOE/Nodes/CONodeMaterialSwitch.h
new file mode 100644
--- /dev/null
+++ b/Engine/Plugins/Mutable/Source/CustomizableObjectEditor/Private/MuCOE/Nodes/CONodeMaterialSwitch.h
@@ -0,0 +1,21 @@
+// Copyright Epic Games, Inc. All Rights Reserved.
+
+#pragma once
+
+#include "MuCOE/Nodes/CustomizableObjectNodeSwitchBase.h"
+#include "CONodeMaterialSwitch.generated.h"
+
+#define UE_API CUSTOMIZABLEOBJECTEDITOR_API
+
+
+UCLASS(MinimalAPI)
+class UCONodeMaterialSwitch : public UCustomizableObjectNodeSwitchBase
+{
+public:
+	GENERATED_BODY()
+
+	// UCustomizableObjectNodeSwitchBase interface
+	UE_API virtual FName GetCategory() const override;
+};
+
+#undef UE_API
diff --git a/Engine/Plugins/Mutable/Source/CustomizableObjectEditor/Private/MuCOE/Nodes/CONodeMaterialVariation.cpp b/Engine/Plugins/Mutable/Source/CustomizableObjectEditor/Private/MuCOE/Nodes/CONodeMaterialVariation.cpp
new file mode 100644
--- /dev/null
+++ b/Engine/Plugins/Mutable/Source/CustomizableObjectEditor/Private/MuCOE/Nodes/CONodeMaterialVariation.cpp
@@ -0,0 +1,11 @@
+// Copyright Epic Games, Inc. All Rights Reserved.
+
+#include "MuCOE/Nodes/CONodeMaterialVariation.h"
+#include "MuCOE/EdGraphSchema_CustomizableObject.h"
+
+#include UE_INLINE_GENERATED_CPP_BY_NAME(CONodeMaterialVariation)
+
+FName UCONodeMaterialVariation::GetCategory() const
+{
+	return UEdGraphSchema_CustomizableObject::PC_Material;
+}
diff --git a/Engine/Plugins/Mutable/Source/CustomizableObjectEditor/Private/MuCOE/Nodes/CONodeMaterialVariation.h b/Engine/Plugins/Mutable/Source/CustomizableObjectEditor/Private/MuCOE/Nodes/CONodeMaterialVariation.h
new file mode 100644
--- /dev/null
+++ b/Engine/Plugins/Mutable/Source/CustomizableObjectEditor/Private/MuCOE/Nodes/CONodeMaterialVariation.h
@@ -0,0 +1,20 @@
+// Copyright Epic Games, Inc. All Rights Reserved.
+
+#pragma once
+
+#include "MuCOE/Nodes/CustomizableObjectNodeVariation.h"
+#include "CONodeMaterialVariation.generated.h"
+
+#define UE_API CUSTOMIZABLEOBJECTEDITOR_API
+
+UCLASS(MinimalAPI)
+class UCONodeMaterialVariation : public UCustomizableObjectNodeVariation
+{
+public:
+	GENERATED_BODY()
+	
+	// UCustomizableObjectNodeVariation interface
+	UE_API virtual FName GetCategory() const override;
+};
+
+#undef UE_API
diff --git a/Engine/Plugins/Mutable/Source/CustomizableObjectEditor/Public/MuCOE/EdGraphSchema_CustomizableObject.h b/Engine/Plugins/Mutable/Source/CustomizableObjectEditor/Public/MuCOE/EdGraphSchema_CustomizableObject.h
--- a/Engine/Plugins/Mutable/Source/CustomizableObjectEditor/Public/MuCOE/EdGraphSchema_CustomizableObject.h
+++ b/Engine/Plugins/Mutable/Source/CustomizableObjectEditor/Public/MuCOE/EdGraphSchema_CustomizableObject.h
@@ -52,6 +52,7 @@ public:
 
 	// Node categories
 	static UE_API const FText NC_Experimental;
+	static UE_API const FText NC_Material;
 
 	// EdGraphSchema interface
 	UE_API virtual void GetAssetsGraphHoverMessage(const TArray<FAssetData>& Assets, const UEdGraph* HoverGraph, FString& OutTooltipText, bool& OutOkIcon) const override;
diff --git a/Engine/Plugins/Mutable/Source/MutableTools/Private/MuT/Node.cpp b/Engine/Plugins/Mutable/Source/MutableTools/Private/MuT/Node.cpp
--- a/Engine/Plugins/Mutable/Source/MutableTools/Private/MuT/Node.cpp
+++ b/Engine/Plugins/Mutable/Source/MutableTools/Private/MuT/Node.cpp
@@ -216,12 +216,12 @@ namespace mu
 	FNodeType NodeRange::StaticType = FNodeType(Node::EType::Range, Node::GetStaticType());
 	FNodeType NodeRangeFromScalar::StaticType = FNodeType(Node::EType::RangeFromScalar, NodeRange::GetStaticType());
 
-	FNodeType NodeMaterial::StaticType = FNodeType(Node::EType::Material, NodeMaterial::GetStaticType());
-	FNodeType NodeMaterialConstant::StaticType = FNodeType(Node::EType::MaterialConstant, NodeMaterialConstant::GetStaticType());
-	FNodeType NodeMaterialTable::StaticType = FNodeType(Node::EType::MaterialTable, NodeMaterialTable::GetStaticType());
-	FNodeType NodeMaterialSwitch::StaticType = FNodeType(Node::EType::MaterialSwitch, NodeMaterialSwitch::GetStaticType());
-	FNodeType NodeMaterialVariation::StaticType = FNodeType(Node::EType::MaterialVariation, NodeMaterialVariation::GetStaticType());
-	FNodeType NodeMaterialParameter::StaticType = FNodeType(Node::EType::MaterialParameter, NodeMaterialParameter::GetStaticType());
+	FNodeType NodeMaterial::StaticType = FNodeType(Node::EType::Material, Node::GetStaticType());
+	FNodeType NodeMaterialConstant::StaticType = FNodeType(Node::EType::MaterialConstant, NodeMaterial::GetStaticType());
+	FNodeType NodeMaterialTable::StaticType = FNodeType(Node::EType::MaterialTable, NodeMaterial::GetStaticType());
+	FNodeType NodeMaterialSwitch::StaticType = FNodeType(Node::EType::MaterialSwitch, NodeMaterial::GetStaticType());
+	FNodeType NodeMaterialVariation::StaticType = FNodeType(Node::EType::MaterialVariation, NodeMaterial::GetStaticType());
+	FNodeType NodeMaterialParameter::StaticType = FNodeType(Node::EType::MaterialParameter, NodeMaterial::GetStaticType());
 
 
 	FNodeType::FNodeType()
