[Mutable] Replaced FResourceID with FImageId and FMeshId

- Replaced GC of those IDs with a Smart Pointer system (TRegistry).
- Moved Image Cache from the COSystem to CO.
- Replaced FMeshCache from GC to Smart Pointer system (TRegistry).
- Bake no longer modifies global System settings.
- Removed unsafe Release Textures immediately.

diff --git a/Engine/Plugins/Mutable/Source/CustomizableObject/Internal/MuCO/CustomizableObjectInstancePrivate.h b/Engine/Plugins/Mutable/Source/CustomizableObject/Internal/MuCO/CustomizableObjectInstancePrivate.h
--- a/Engine/Plugins/Mutable/Source/CustomizableObject/Internal/MuCO/CustomizableObjectInstancePrivate.h
+++ b/Engine/Plugins/Mutable/Source/CustomizableObject/Internal/MuCO/CustomizableObjectInstancePrivate.h
@@ -21,7 +21,6 @@ namespace mu
 {
 	class FPhysicsBody;
 	class FMesh;
-	typedef uint64 FResourceID;
 }
 
 struct FMutableModelImageProperties;
@@ -148,7 +147,7 @@ struct FCustomizableInstanceComponentData
 	/** Array of generated MeshIds per each LOD, used to decide if the mesh should be updated or not.
 	 *  Size == NumLODsAvailable
 	 *  LODs without mesh will be set to the maximum value of FResourceID (Max_uint64). */
-	TArray<mu::FResourceID> LastMeshIdPerLOD;
+	TArray<mu::FMeshId> LastMeshIdPerLOD;
 
 	UPROPERTY(Transient)
 	TArray<TObjectPtr<UMaterialInterface>> OverrideMaterials;
@@ -269,9 +268,7 @@ public:
 
 	// Returns true if success (?)
 	UE_API bool UpdateSkeletalMesh_PostBeginUpdate0(UCustomizableObjectInstance* Public, const TSharedRef<FUpdateContextPrivate>& OperationData);
-
-	static UE_API void ReleaseMutableTexture(const FMutableImageCacheKey& MutableTextureKey, UTexture2D* Texture, struct FMutableResourceCache& Cache);
-
+	
 	// Copy data generated in the mutable thread over to the instance and initializes additional data required during the update
 	UE_API void PrepareForUpdate(const TSharedRef<FUpdateContextPrivate>& OperationData);
 
@@ -281,10 +278,7 @@ public:
 	/** Calls ReleaseResources on all SkeletalMeshes generated by this instance and invalidates the generated data.
 	  * It should not be called if the meshes are still in use or shared with other instances. */
 	UE_API void DiscardResources();
-
-	// Releases all the mutable resources this instance holds, should only be called when it is not going to be used any more.
-	UE_API void ReleaseMutableResources(bool bCalledFromBeginDestroy, const UCustomizableObjectInstance& Instance);
-
+	
 	/** Set the reference SkeletalMesh, or an empty mesh, to all actors using this instance. */
 	UE_API void SetReferenceSkeletalMesh() const;
 
@@ -360,10 +354,6 @@ public:
 
 	/** Finds in ProjectorParameters a parameter with name ParamName, returns the index if found, -1 otherwise. */
 	UE_API int32 FindProjectorParameterNameIndex(const FString& ParamName) const;
-
-#if WITH_EDITOR
-	UE_API void UpdateSkeletalMeshAsyncResult(FInstanceUpdateNativeDelegate Callback, bool bIgnoreCloseDist = false, bool bForceHighPriority = false, TSharedPtr<FMutableSystemSettingsOverrides> MutableSystemSettingsOverride = nullptr);
-#endif
 	
 	UE_API UCustomizableObjectInstance* GetPublic() const;
 
@@ -380,9 +370,7 @@ public:
 
 	// Converts a ReferencedMaterials index from the CustomizableObject to an index in the ReferencedMaterials in the Instance
 	TMap<int32, uint32> ObjectToInstanceIndexMap;
-
-	TArray<FGeneratedTexture> TexturesToRelease;
-
+	
 	UPROPERTY(Transient)
 	TArray< TObjectPtr<UPhysicsAsset> > ClothingPhysicsAssets;
 
diff --git a/Engine/Plugins/Mutable/Source/CustomizableObject/Internal/MuCO/CustomizableObjectMipDataProvider.h b/Engine/Plugins/Mutable/Source/CustomizableObject/Internal/MuCO/CustomizableObjectMipDataProvider.h
--- a/Engine/Plugins/Mutable/Source/CustomizableObject/Internal/MuCO/CustomizableObjectMipDataProvider.h
+++ b/Engine/Plugins/Mutable/Source/CustomizableObject/Internal/MuCO/CustomizableObjectMipDataProvider.h
@@ -2,6 +2,7 @@
 
 #pragma once
 
+#include "CustomizableObjectSystemPrivate.h"
 #include "MuCO/FMutableTaskGraph.h"
 #include "Async/TaskGraphInterfaces.h"
 #include "Engine/TextureMipDataProviderFactory.h"
@@ -32,8 +33,15 @@ class FMutableUpdateContext
 public:
 	FMutableUpdateContext() = default;
 	
-	FMutableUpdateContext(const FString& InCustomizableObjectPathName, const FString& InInstancePathName, TSharedPtr<mu::FSystem> InSystem,
-		TSharedPtr<mu::FModel, ESPMode::ThreadSafe> InModel, TSharedPtr<FModelStreamableBulkData> InModelStreamable, TSharedPtr<const mu::FParameters> InParameters, int32 InState);
+	FMutableUpdateContext(const FString& InCustomizableObjectPathName,
+		const FString& InInstancePathName,
+		const TSharedPtr<mu::FSystem>& InSystem,
+		const TSharedPtr<mu::FModel>& InModel,
+		const TSharedPtr<mu::FMeshIdRegistry>& InMeshIdRegistry,
+		const TSharedPtr<mu::FImageIdRegistry>& InImageIdRegistry,
+		const TSharedPtr<FModelStreamableBulkData>& InModelStreamable,
+		const TSharedPtr<const mu::FParameters>& InParameters,
+		int32 InState);
 	
 	const FString& GetCustomizableObjectPathName() const;
 
@@ -62,9 +70,15 @@ private:
 	FName CustomizableObjectName;
 
 	TSharedPtr<mu::FSystem> System;
-	TSharedPtr<mu::FModel, ESPMode::ThreadSafe> Model;
-	TSharedPtr<FModelStreamableBulkData> ModelStreamableBulkData;
+	TSharedPtr<mu::FModel> Model;
+
+public:
+	TSharedPtr<mu::FMeshIdRegistry> MeshIdRegistry;
+	TSharedPtr<mu::FImageIdRegistry> ImageIdRegistry;
 
+private:
+	TSharedPtr<FModelStreamableBulkData> ModelStreamableBulkData;
+	
 	TSharedPtr<const mu::FParameters> Parameters;
 	int32 State = -1;
 
@@ -169,7 +183,7 @@ class UMutableTextureMipDataProviderFactory : public UTextureMipDataProviderFact
 public:
 	virtual FTextureMipDataProvider* AllocateMipDataProvider(UTexture* Asset) override
 	{
-		check(ImageRef.ImageID > 0);
+		check(ImageRef.ImageID);
 		FMutableTextureMipDataProvider* Result = new FMutableTextureMipDataProvider(Asset, CustomizableObjectInstance, ImageRef);
 		Result->UpdateContext = UpdateContext;		
 		return Result;
diff --git a/Engine/Plugins/Mutable/Source/CustomizableObject/Internal/MuCO/CustomizableObjectPrivate.h b/Engine/Plugins/Mutable/Source/CustomizableObject/Internal/MuCO/CustomizableObjectPrivate.h
--- a/Engine/Plugins/Mutable/Source/CustomizableObject/Internal/MuCO/CustomizableObjectPrivate.h
+++ b/Engine/Plugins/Mutable/Source/CustomizableObject/Internal/MuCO/CustomizableObjectPrivate.h
@@ -9,6 +9,7 @@
 #include "MuCO/StateMachine.h"
 #include "MuCO/CustomizableObjectUIData.h"
 #include "MuCO/CustomizableObjectIdentifier.h"
+#include "MuR/ResourceId.h"
 #include "Serialization/BulkData.h"
 #include "Templates/SharedPointer.h"
 #include "UObject/WeakObjectPtr.h"
@@ -22,6 +23,8 @@
 
 #include "CustomizableObjectPrivate.generated.h"
 
+class UCustomizableObjectSkeletalMesh;
+
 namespace mu 
 { 
 	class FModel;
@@ -254,12 +257,43 @@ struct FMutableModelParameterProperties
 class FMeshCache
 {
 public:
-	USkeletalMesh* Get(const TArray<mu::FResourceID>& Key);
+	UCustomizableObjectSkeletalMesh* Get(const TArray<mu::FMeshId>& Key);
 
-	void Add(const TArray<mu::FResourceID>& Key, USkeletalMesh* Value);
+	void Add(const TArray<mu::FMeshId>& Key, UCustomizableObjectSkeletalMesh* Value);
 
 private:
-	TMap<TArray<mu::FResourceID>, TWeakObjectPtr<USkeletalMesh>> GeneratedMeshes;
+	using FMeshCacheRegistry = TRegistry<TArray<mu::FMeshId>, TWeakObjectPtr<UCustomizableObjectSkeletalMesh>>;
+	
+	TSharedPtr<FMeshCacheRegistry> GeneratedMeshes = MakeShared<FMeshCacheRegistry>();
+
+public:
+	typedef FMeshCacheRegistry::FHandle FId;
+};
+
+
+
+class FTextureCache
+{
+public:
+	struct FId
+	{
+		mu::FImageId Resource;
+
+		int32 SkippedMips = 0;
+
+		bool operator==(const FId&) const = default;
+
+		friend uint32 GetTypeHash(const FId& Key);
+	};
+
+	UTexture2D* Get(const FId& Key);
+	
+	void Add(const FId& Key, UTexture2D* Value);
+
+	void Remove(const UTexture2D& Value);
+
+private:
+	TMap<FId, TWeakObjectPtr<UTexture2D>> GeneratedTextures;
 };
 
 
@@ -1842,12 +1876,19 @@ public:
 
 	CUSTOMIZABLEOBJECT_API bool IsParameterMultidimensional(const int32& InParamIndex) const;
 	
-	/** Cache of generated SkeletalMeshes */
+	/** Cache of generated Skeletal Meshes. Passthrough Skeletal Meshes are not contemplated. */
 	FMeshCache MeshCache;
 
+	/** Cache of generated Textures. Passthrough Textures are not contemplated. */
+	FTextureCache TextureCache;
+	
 	/** Cache of merged Skeletons */
 	FSkeletonCache SkeletonCache;
 
+	TSharedRef<mu::FMeshIdRegistry> MeshIdRegistry = MakeShared<mu::FMeshIdRegistry>();
+	
+	TSharedRef<mu::FImageIdRegistry> ImageIdRegistry = MakeShared<mu::FImageIdRegistry>();
+	
 	// See UCustomizableObjectSystem::LockObject. Must only be modified from the game thread
 	bool bLocked = false;
 
diff --git a/Engine/Plugins/Mutable/Source/CustomizableObject/Internal/MuCO/CustomizableObjectSystemPrivate.h b/Engine/Plugins/Mutable/Source/CustomizableObject/Internal/MuCO/CustomizableObjectSystemPrivate.h
--- a/Engine/Plugins/Mutable/Source/CustomizableObject/Internal/MuCO/CustomizableObjectSystemPrivate.h
+++ b/Engine/Plugins/Mutable/Source/CustomizableObject/Internal/MuCO/CustomizableObjectSystemPrivate.h
@@ -19,6 +19,7 @@
 #include "MuR/Skeleton.h"
 #include "MuR/Image.h"
 #include "MuR/Material.h"
+#include "MuR/ResourceId.h"
 #include "UObject/GCObject.h"
 #include "WorldCollision.h"
 #include "Engine/StreamableManager.h"
@@ -29,7 +30,6 @@
 #include "Animation/MorphTarget.h"
 #include "MuCO/MutableStreamableManager.h"
 #include "UObject/StrongObjectPtr.h"
-
 #include "Tasks/Task.h"
 
 #include "CustomizableObjectSystemPrivate.generated.h"
@@ -49,6 +49,26 @@ struct FTexturePlatformData;
 extern TAutoConsoleVariable<bool> CVarMutableHighPriorityLoading;
 
 
+/** Key to identify an image inside a generated mutable runtime instance. */
+struct FMutableImageReference
+{
+	mu::FImageId ImageID;
+
+	uint32 SurfaceId = 0;
+
+	uint8 LOD = 0;
+	
+	/** Object component index. */
+	uint8 Component = 0;
+
+	uint8 Image = 0;
+
+	uint8 BaseMip = 0;
+
+	TArray<int32> ConstantImagesNeededToGenerate;
+};
+
+
 /** Serialized Morph data. */
 struct FMorphTargetMeshData
 {
@@ -109,40 +129,6 @@ void LoadClothing(FMutableStreamRequest& MutableDataStreamer, const TSharedRef<c
 void FinishUpdateGlobal(const TSharedRef<FUpdateContextPrivate>& Context);
 
 
-#if WITH_EDITOR
-
-/**
- * Class used to hold some MutableSystem settings to be used during the update of a given instance.
- * It will store the settings used by the system during the setup of this object and so later can be reverted back.
- */
-class FMutableSystemSettingsOverrides
-{
-public:
-	UE_API FMutableSystemSettingsOverrides(bool bUseProgressiveMipStreaming, bool bOnlyGenerateRequestedLODs, mu::FImageOperator::FImagePixelFormatFunc InImagePixelFormatFunc);
-	
-	/** Apply the settings set in this object to the mutable system. */
-	UE_API void ApplySettingsOverrides() const;
-
-	/** Restore the mutable system settings to a state prior to the setup of this object */
-	UE_API void RestoreSettings() const;
-
-
-private:
-	
-	// Requested values
-	bool bIsProgressiveMipStreamingEnabled = false;
-	bool bIsOnlyGenerateRequestedLODsEnabled = false;
-	mu::FImageOperator::FImagePixelFormatFunc ImagePixelFormatFunc;
-	
-	// previous values. Cached at the time of creating this object
-	bool bOldIsProgressiveMipStreamingEnabled = false;
-	bool bOldIsOnlyGenerateRequestedLODsEnabled = false;
-	mu::FImageOperator::FImagePixelFormatFunc OldImagePixelFormatFunc;
-};
-
-#endif
-
-
 /** Strongly typed index for the index of a component in a UCustomizableObjectInstance. */
 class FCustomizableObjectInstanceComponentIndex
 {
@@ -325,51 +311,11 @@ public:
 };
 
 
-struct FMutableImageCacheKey
-{
-	mu::FResourceID Resource = 0;
-	int32 SkippedMips = 0;
-
-	FMutableImageCacheKey() {};
-
-	FMutableImageCacheKey(mu::FResourceID InResource, int32 InSkippedMips)
-		: Resource(InResource), SkippedMips(InSkippedMips) {}
-
-	inline bool operator==(const FMutableImageCacheKey& Other) const
-	{
-		return Resource == Other.Resource && SkippedMips == Other.SkippedMips;
-	}
-};
-
-
-inline uint32 GetTypeHash(const FMutableImageCacheKey& Key)
-{
-	return HashCombine(GetTypeHash(Key.Resource), GetTypeHash(Key.SkippedMips));
-}
-
-
-// Cache of weak references to generated resources for one single model
-struct FMutableResourceCache
-{
-	TWeakObjectPtr<const UCustomizableObject> Object;
-	TMap<mu::FResourceID, TWeakObjectPtr<USkeletalMesh> > Meshes;
-	TMap<FMutableImageCacheKey, TWeakObjectPtr<UTexture2D> > Images;
-
-	void Clear()
-	{
-		Meshes.Reset();
-		Images.Reset();
-	}
-};
-
-
 USTRUCT()
 struct FGeneratedTexture
 {
 	GENERATED_USTRUCT_BODY();
 
-	FMutableImageCacheKey Key;
-
 	UPROPERTY(Category = NoCategory, VisibleAnywhere)
 	FString Name;
 
@@ -426,7 +372,7 @@ struct FInstanceUpdateData
 	struct FImage
 	{
 		FName Name;
-		mu::FResourceID ImageID;
+		mu::FImageId ImageID;
 		
 		// LOD of the ImageId. If the texture is shared between LOD, first LOD where this image can be found. 
 		int32 BaseLOD;
@@ -480,7 +426,7 @@ struct FInstanceUpdateData
 
 	struct FLOD
 	{		
-		mu::FResourceID MeshID = TNumericLimits<mu::FResourceID>::Max();
+		mu::FMeshId MeshID;
 		TSharedPtr<const mu::FMesh> Mesh;
 
 		/** Range in the Surfaces array */
@@ -580,7 +526,7 @@ class FUpdateContextPrivate
 public:
 	FUpdateContextPrivate(UCustomizableObjectInstance& InInstance, const FCustomizableObjectInstanceDescriptor& Descriptor);
 
-	FUpdateContextPrivate(UCustomizableObjectInstance& InInstance);
+	CUSTOMIZABLEOBJECT_API FUpdateContextPrivate(UCustomizableObjectInstance& InInstance);
 	
 	bool IsContextValid() const;
 	
@@ -652,7 +598,19 @@ public:
 	bool bLowPriorityTasksBlocked = false;
 
 	/** The Context has been successfully created. */
-	bool bValid = false; 
+	bool bValid = false;
+
+	/** Do not optimize the update. */
+	bool bForce = false;
+
+	/** Update due to a bake. */
+	bool bBake = false;
+
+	/** User config. */
+	bool bProgressiveMipStreamingEnabled = false;
+	
+	/** User config. */
+	bool bIsOnlyGenerateRequestedLODsEnabled = false;
 
 	/** This option comes from the operation request. It is used to reduce the number of mipmaps that mutable must generate for images.  */
 	int32 MipsToSkip = 0;
@@ -661,6 +619,8 @@ public:
 	TSharedPtr<const mu::FInstance> MutableInstance;
 
 	TSharedPtr<mu::FModel> Model;
+	TSharedPtr<mu::FMeshIdRegistry> MeshIdRegistry;
+	TSharedPtr<mu::FImageIdRegistry> ImageIdRegistry;
 
 	// Number of possible components in the entire CO
 	// TODO: Redundant because it is in the CO
@@ -686,7 +646,7 @@ public:
 	
 	TMap<FName, uint8> FirstResidentLOD;
 
-	TMap<mu::FResourceID, FTexturePlatformData*> ImageToPlatformDataMap;
+	TMap<mu::FImageId, FTexturePlatformData*> ImageToPlatformDataMap;
 
 	EUpdateResult UpdateResult = EUpdateResult::Success;
 
@@ -695,13 +655,13 @@ public:
 private:
 
 	/** Mutable Meshes required for each component. Outermost index is the object component index, inner index is the LOD. */
-	TArray<TArray<mu::FResourceID>> MeshDescriptors;
+	TArray<TArray<mu::FMeshId>> MeshDescriptors;
 
 public:
 
 	void InitMeshDescriptors(int32 Size);
-	const TArray<TArray<mu::FResourceID>>& GetMeshDescriptors() const;
-	TArray<mu::FResourceID>* GetMeshDescriptors(FCustomizableObjectComponentIndex Index);
+	const TArray<TArray<mu::FMeshId>>& GetMeshDescriptors() const;
+	TArray<mu::FMeshId>* GetMeshDescriptors(FCustomizableObjectComponentIndex Index);
 
 	/** Used to know if the updated instances' meshes are different from the previous ones. 
 	  * The index of the array is the instance component's index.
@@ -738,16 +698,14 @@ public:
 	/** Hard references to objects. Avoids GC to collect them. */
 	TArray<TStrongObjectPtr<const UObject>> Objects;
 
-#if WITH_EDITORONLY_DATA
-	TSharedPtr<class FMutableSystemSettingsOverrides> UpdateSettingsOverride = nullptr;
-#endif
+	TArray<mu::FImageId> CachedTextures;
 };
 
 
 /** Runtime data used during a mutable instance update */
 struct FMutableReleasePlatformOperationData
 {
-	TMap<mu::FResourceID, FTexturePlatformData*> ImageToPlatformDataMap;
+	TMap<mu::FImageId, FTexturePlatformData*> ImageToPlatformDataMap;
 };
 
 
@@ -851,7 +809,6 @@ public:
 
 	/** Store the last streaming memory size in bytes, to change it when it is safe. */
 	uint64 LastWorkingMemoryBytes = 0;
-	uint32 LastGeneratedResourceCacheSize = 0;
 
 	// This object is responsible for streaming data to the MutableSystem.
 	TSharedPtr<class FUnrealMutableModelBulkReader> Streamer;
@@ -860,15 +817,7 @@ public:
 	// This object is called from the mutable thread, and it should only access data already safely submitted from
 	// the game thread and stored in FUnrealMutableImageProvider::GlobalExternalImages.
 	TSharedPtr<class FUnrealMutableResourceProvider> ResourceProvider;
-
-	// Cache of weak references to generated resources to see if they can be reused.
-	TArray<FMutableResourceCache> ModelResourcesCache;
-
-	// List of textures currently cached and valid for the current object that we are operating on.
-	// This array gets generated when the object cached resources are protected in SetResourceCacheProtected
-	// from the game thread, and it is read from the Mutable thread only while updating the instance.
-	TArray<mu::FResourceID> ProtectedObjectCachedImages;
-
+	
 	// The pending instance updates, discards or releases
 	FMutablePendingInstanceWork MutablePendingInstanceWork;
 	
@@ -897,26 +846,11 @@ public:
 
 	UCustomizableObjectSystem* GetPublic() const;
 	
-	// Remove references to cached objects that have been deleted in the unreal
-	// side, and cannot be cached anyway.
-	// This should only happen in the game thread
-	void CleanupCache();
-
-	// This should only happen in the game thread
-	FMutableResourceCache& GetObjectCache(const UCustomizableObject* Object);
-
-	void AddTextureReference(const FMutableImageCacheKey& TextureId);
-
-	// Returns true if the texture's references become zero
-	bool RemoveTextureReference(const FMutableImageCacheKey& TextureId);
-
-	bool TextureHasReferences(const FMutableImageCacheKey& TextureId) const;
-
 	EUpdateRequired IsUpdateRequired(const UCustomizableObjectInstance& Instance, bool bOnlyUpdateIfNotGenerated, bool bOnlyUpdateIfLOD, bool bIgnoreCloseDist) const;
 
 	EQueuePriorityType GetUpdatePriority(const UCustomizableObjectInstance& Instance, bool bForceHighPriority) const;
 
-	void EnqueueUpdateSkeletalMesh(const TSharedRef<FUpdateContextPrivate>& Context);
+	CUSTOMIZABLEOBJECT_API void EnqueueUpdateSkeletalMesh(const TSharedRef<FUpdateContextPrivate>& Context);
 		
 	// Init an async and safe release of the UE and Mutable resources used by the instance without actually destroying the instance, for example if it's very far away
 	void InitDiscardResourcesSkeletalMesh(UCustomizableObjectInstance* InCustomizableObjectInstance);
@@ -924,7 +858,7 @@ public:
 	// Init the async release of a Mutable Core Instance ID and all the temp resources associated with it
 	void InitInstanceIDRelease(mu::FInstance::FID IDToRelease);
 
-	void GetMipStreamingConfig(const UCustomizableObjectInstance& Instance, bool& bOutNeverStream, int32& OutMipsToSkip) const;
+	void GetMipStreamingConfig(const UCustomizableObjectInstance& Instance, bool bTextureStreamingEnabled, bool& bOutNeverStream, int32& OutMipsToSkip) const;
 	
 	bool IsReplaceDiscardedWithReferenceMeshEnabled() const;
 	void SetReplaceDiscardedWithReferenceMeshEnabled(bool bIsEnabled);
@@ -933,12 +867,9 @@ public:
 	int32 GetNumSkeletalMeshes() const;
 
 	bool bReplaceDiscardedWithReferenceMesh = false;
-	bool bReleaseTexturesImmediately = false;
-
+	
 	bool bSupport16BitBoneIndex = false;
-
-	TMap<FMutableImageCacheKey, uint32> TextureReferenceCount; // Keeps a count of texture usage to decide if they have to be blocked from GC during an update
-
+	
 	UPROPERTY(Transient)
 	TObjectPtr<UCustomizableObjectInstance> CurrentInstanceBeingUpdated = nullptr;
 
@@ -973,12 +904,7 @@ public:
 	 * to do things like check-if-up-to-date or compile.
 	 */
 	void AddPendingLoad(UCustomizableObject*);
-
 #endif
-
-	// Unprotect the resources in the instances of this object of being garbage collector
-	// while an instance is being built or updated, so that they can be reused.
-	void ClearResourceCacheProtected();
 	
 	/** Mutable TaskGraph system (Mutable Thread). */
 	FMutableTaskGraph MutableTaskGraph;
@@ -1006,11 +932,6 @@ public:
 
 	UPROPERTY(Transient)
 	TObjectPtr<UCustomizableInstanceLODManagementBase> CurrentInstanceLODManagement = nullptr;
-
-	// Array where textures are added temporarily while the mutable thread may want to
-	// reused them for some instance under construction.
-	UPROPERTY(Transient)
-	TArray<TObjectPtr<UTexture2D>> ProtectedCachedTextures;
 	
 	TSharedRef<FMutableStreamableManager> StreamableManager = MakeShared<FMutableStreamableManager>();
 
diff --git a/Engine/Plugins/Mutable/Source/CustomizableObject/Private/MuCO/CustomizableObject.cpp b/Engine/Plugins/Mutable/Source/CustomizableObject/Private/MuCO/CustomizableObject.cpp
--- a/Engine/Plugins/Mutable/Source/CustomizableObject/Private/MuCO/CustomizableObject.cpp
+++ b/Engine/Plugins/Mutable/Source/CustomizableObject/Private/MuCO/CustomizableObject.cpp
@@ -51,6 +51,7 @@
 
 #include "MuCO/CustomizableObjectCustomVersion.h"
 #include "MuCO/CustomizableObjectResourceDataTypes.h"
+#include "MuCO/CustomizableObjectSkeletalMesh.h"
 
 #include UE_INLINE_GENERATED_CPP_BY_NAME(CustomizableObject)
 
@@ -1059,6 +1060,9 @@ FString UCustomizableObject::GetStateParameterName(const FString& StateName, int
 #if WITH_EDITORONLY_DATA
 void UCustomizableObjectPrivate::PostCompile()
 {
+	MeshIdRegistry = MakeShared<mu::FMeshIdRegistry>();
+	ImageIdRegistry = MakeShared<mu::FImageIdRegistry>();
+	
 	for (TObjectIterator<UCustomizableObjectInstance> It; It; ++It)
 	{
 		if (It->GetCustomizableObject() == this->GetPublic())
@@ -1973,30 +1977,21 @@ uint8 UCustomizableObjectPrivate::GetMinLODIndex(const FName& ComponentName) con
 //-------------------------------------------------------------------------------------------------
 //-------------------------------------------------------------------------------------------------
 
-USkeletalMesh* FMeshCache::Get(const TArray<mu::FResourceID>& Key)
+UCustomizableObjectSkeletalMesh* FMeshCache::Get(const TArray<mu::FMeshId>& Key)
 {
-	const TWeakObjectPtr<USkeletalMesh>* Result = GeneratedMeshes.Find(Key);
+	const TWeakObjectPtr<UCustomizableObjectSkeletalMesh>* Result = GeneratedMeshes->Find(Key);
 	return Result ? Result->Get() : nullptr;
 }
 
 
-void FMeshCache::Add(const TArray<mu::FResourceID>& Key, USkeletalMesh* Value)
+void FMeshCache::Add(const TArray<mu::FMeshId>& Key, UCustomizableObjectSkeletalMesh* Value)
 {
 	if (!Value)
 	{
 		return;
 	}
 	
-	GeneratedMeshes.Add(Key, Value);
-
-	// Remove invalid SkeletalMeshes from the cache.
-	for (auto MeshIterator = GeneratedMeshes.CreateIterator(); MeshIterator; ++MeshIterator)
-	{
-		if (MeshIterator.Value().IsStale())
-		{
-			MeshIterator.RemoveCurrent();
-		}
-	}	
+	Value->MeshCacheId = GeneratedMeshes->Add(Key, TWeakObjectPtr(Value));
 }
 
 
@@ -2027,6 +2022,46 @@ void FSkeletonCache::Add(const TArray<uint16>& Key, USkeleton* Value)
 }
 
 
+UTexture2D* FTextureCache::Get(const FId& Key)
+{
+	const TWeakObjectPtr<UTexture2D>* Result = GeneratedTextures.Find(Key);
+	return Result ? Result->Get() : nullptr;
+}
+
+
+void FTextureCache::Add(const FId& Key, UTexture2D* Value)
+{
+	if (!Value)
+	{
+		return;
+	}
+
+	GeneratedTextures.Add(Key, Value);
+
+	// Remove invalid SkeletalMeshes from the cache.
+	for (auto SkeletonIterator = GeneratedTextures.CreateIterator(); SkeletonIterator; ++SkeletonIterator)
+	{
+		if (SkeletonIterator.Value().IsStale())
+		{
+			SkeletonIterator.RemoveCurrent();
+		}
+	}
+}
+
+
+void FTextureCache::Remove(const UTexture2D& Value)
+{
+	for (auto It = GeneratedTextures.CreateIterator(); It; ++It)
+	{
+		if (It->Value.Get() == &Value)
+		{
+			It.RemoveCurrent();
+			break;
+		}
+	}
+}
+
+
 FArchive& operator<<(FArchive& Ar, FIntegerParameterUIData& Struct)
 {
 	Ar << Struct.ParamUIMetadata;
@@ -3419,6 +3454,13 @@ FMutableParameterData::FMutableParameterData(const FMutableParamUIMetadata& InPa
 	Type = InType;
 }
 
+
+uint32 GetTypeHash(const FTextureCache::FId& Key)
+{
+	return HashCombine(GetTypeHash(Key.Resource), Key.SkippedMips);
+}
+
+
 #if WITH_EDITORONLY_DATA
 FArchive& operator<<(FArchive& Ar, FAnimBpOverridePhysicsAssetsInfo& Info)
 {
diff --git a/Engine/Plugins/Mutable/Source/CustomizableObject/Private/MuCO/CustomizableObjectInstance.cpp b/Engine/Plugins/Mutable/Source/CustomizableObject/Private/MuCO/CustomizableObjectInstance.cpp
--- a/Engine/Plugins/Mutable/Source/CustomizableObject/Private/MuCO/CustomizableObjectInstance.cpp
+++ b/Engine/Plugins/Mutable/Source/CustomizableObject/Private/MuCO/CustomizableObjectInstance.cpp
@@ -213,7 +213,7 @@ void UCustomizableInstancePrivate::InvalidateGeneratedData()
 
 	// Init Component Data
 	FCustomizableInstanceComponentData TemplateComponentData;
-	TemplateComponentData.LastMeshIdPerLOD.Init(MAX_uint64, MAX_MESH_LOD_COUNT);
+	TemplateComponentData.LastMeshIdPerLOD.Init({}, MAX_MESH_LOD_COUNT);
 	ComponentsData.Init(TemplateComponentData, ComponentsData.Num());
 
 	GeneratedMaterials.Empty();
@@ -231,7 +231,7 @@ void UCustomizableInstancePrivate::InitCustomizableObjectData(const UCustomizabl
 	
 	// Init Component Data
 	FCustomizableInstanceComponentData TemplateComponentData;
-	TemplateComponentData.LastMeshIdPerLOD.Init(MAX_uint64, MAX_MESH_LOD_COUNT);
+	TemplateComponentData.LastMeshIdPerLOD.Init({}, MAX_MESH_LOD_COUNT);
 	ComponentsData.Init(TemplateComponentData, InCustomizableObject->GetComponentCount());
 
 	ExtensionInstanceData.Empty();
@@ -438,8 +438,6 @@ void UCustomizableObjectInstance::BeginDestroy()
 		// Unbind Object delegates
 		PrivateData->BindObjectDelegates(GetCustomizableObject(), nullptr);
 #endif
-
-		PrivateData.Get()->ReleaseMutableResources(true, *this);
 	}
 	
 	Super::BeginDestroy();
@@ -462,34 +460,6 @@ void UCustomizableObjectInstance::DestroyLiveUpdateInstance()
 }
 
 
-void UCustomizableInstancePrivate::ReleaseMutableResources(bool bCalledFromBeginDestroy, const UCustomizableObjectInstance& Instance)
-{
-	GeneratedMaterials.Empty();
-
-	if (UCustomizableObjectSystem::IsCreated()) // Need to check this because the object might be destroyed after the CustomizableObjectSystem at shutdown
-	{
-		UCustomizableObjectSystemPrivate* CustomizableObjectSystem = UCustomizableObjectSystem::GetInstance()->GetPrivate();
-		// Get the cache of resources of all live instances of this object
-		FMutableResourceCache& Cache = CustomizableObjectSystem->GetObjectCache(Instance.GetCustomizableObject());
-
-		for (FGeneratedTexture& Texture : GeneratedTextures)
-		{
-			if (CustomizableObjectSystem->RemoveTextureReference(Texture.Key))
-			{
-				// Do not release textures when called from BeginDestroy, it would produce a texture artifact in the
-				// instance's remaining sk meshes and GC is being performed anyway so it will free the textures if needed
-				if (!bCalledFromBeginDestroy && CustomizableObjectSystem->bReleaseTexturesImmediately)
-				{
-					ReleaseMutableTexture(Texture.Key, Cast<UTexture2D>(Texture.Texture), Cache);
-				}
-			}
-		}
-	}
-
-	GeneratedTextures.Empty();
-}
-
-
 bool UCustomizableObjectInstance::IsReadyForFinishDestroy()
 {
 	//return ReleaseResourcesFence.IsFenceComplete();
@@ -615,24 +585,6 @@ int32 UCustomizableObjectInstance::GetTextureValueRange(const FString& ParamName
 }
 
 
-// Only safe to call if the Mutable texture ref count system returns 0 and absolutely sure nobody holds a reference to the texture
-void UCustomizableInstancePrivate::ReleaseMutableTexture(const FMutableImageCacheKey& MutableTextureKey, UTexture2D* Texture, FMutableResourceCache& Cache)
-{
-	if (ensure(Texture) && Texture->IsValidLowLevel())
-	{
-		Texture->ConditionalBeginDestroy();
-
-		for (FTexture2DMipMap& Mip : Texture->GetPlatformData()->Mips)
-		{
-			Mip.BulkData.RemoveBulkData();
-		}
-	}
-
-	// Must remove texture from cache since it has been released
-	Cache.Images.Remove(MutableTextureKey);
-}
-
-
 void UCustomizableObjectInstance::SetObject(UCustomizableObject* InObject)
 {
 #if WITH_EDITOR
@@ -2014,7 +1966,7 @@ bool UCustomizableInstancePrivate::DoComponentsNeedUpdate(UCustomizableObjectIns
 		bool bMeshChanged = true;
 
 		// Compare Mesh IDs 
-		const TArray<mu::FResourceID>* MeshIDs = OperationData->GetMeshDescriptors(ObjectComponentIndex);
+		const TArray<mu::FMeshId>* MeshIDs = OperationData->GetMeshDescriptors(ObjectComponentIndex);
 		if (ensure(MeshIDs))
 		{
 			if (const FCustomizableInstanceComponentData* ComponentData = GetComponentData(ObjectComponentIndex)) // Could be nullptr if the component has not been generated.
@@ -2228,7 +2180,7 @@ bool UCustomizableInstancePrivate::UpdateSkeletalMesh_PostBeginUpdate0(UCustomiz
 			continue;
 		}
 		
-		TArray<mu::FResourceID>* MeshDescriptors = OperationData->GetMeshDescriptors(ObjectComponentIndex);
+		TArray<mu::FMeshId>* MeshDescriptors = OperationData->GetMeshDescriptors(ObjectComponentIndex);
 		check(MeshDescriptors);
 		check(MeshDescriptors->Num() == MAX_MESH_LOD_COUNT);
 		ComponentData->LastMeshIdPerLOD = *MeshDescriptors;
@@ -2304,7 +2256,7 @@ bool UCustomizableInstancePrivate::UpdateSkeletalMesh_PostBeginUpdate0(UCustomiz
 		
 		FString SkeletalMeshName = FString::Printf (TEXT("SK_%s_%s_%lu"), *Public->GetCustomizableObject()->GetName(), *ComponentName.ToString(), MeshHash);
 		SkeletalMeshName = MakeUniqueObjectName(GetTransientPackage(), USkeletalMesh::StaticClass(), *SkeletalMeshName, EUniqueObjectNameOptions::GloballyUnique).ToString();
-		if (OperationData->bStreamMeshLODs)
+		if (!OperationData->bBake)
 		{
 			SkeletalMesh = NewObject<UCustomizableObjectSkeletalMesh>(GetTransientPackage(), FName(*SkeletalMeshName), RF_Transient);
 		}
@@ -2421,10 +2373,10 @@ bool UCustomizableInstancePrivate::UpdateSkeletalMesh_PostBeginUpdate0(UCustomiz
 
 		if (OperationData->bUseMeshCache)
 		{
-			const TArray<mu::FResourceID>* MeshId = OperationData->GetMeshDescriptors(ObjectComponentIndex);
+			const TArray<mu::FMeshId>* MeshId = OperationData->GetMeshDescriptors(ObjectComponentIndex);
 			if (MeshId)
 			{
-				CustomizableObject->GetPrivate()->MeshCache.Add(*MeshId, SkeletalMesh);
+				CustomizableObject->GetPrivate()->MeshCache.Add(*MeshId, CastChecked<UCustomizableObjectSkeletalMesh>(SkeletalMesh));
 			}
 		}
 
@@ -2830,22 +2782,6 @@ int32 UCustomizableInstancePrivate::FindProjectorParameterNameIndex(const FStrin
 }
 
 
-#if WITH_EDITOR
-void UCustomizableInstancePrivate::UpdateSkeletalMeshAsyncResult(FInstanceUpdateNativeDelegate Callback, bool bIgnoreCloseDist, bool bForceHighPriority, TSharedPtr<FMutableSystemSettingsOverrides> MutableSystemSettingsOverride)
-{
-	UCustomizableObjectSystemPrivate* SystemPrivate = UCustomizableObjectSystem::GetInstance()->GetPrivate();
-
-	const TSharedRef<FUpdateContextPrivate> Context = MakeShared<FUpdateContextPrivate>(*GetPublic());
-	Context->bIgnoreCloseDist = bIgnoreCloseDist;
-	Context->bForceHighPriority = bForceHighPriority;
-	Context->UpdateNativeCallback = Callback;
-	Context->UpdateSettingsOverride = MutableSystemSettingsOverride;
-	
-	SystemPrivate->EnqueueUpdateSkeletalMesh(Context);
-}
-
-#endif
-
 void UCustomizableObjectInstance::SetRandomValues()
 {
 	Descriptor.SetRandomValues();
@@ -3120,8 +3056,6 @@ void UCustomizableInstancePrivate::DiscardResources()
 		}
 		
 		SkeletalMeshes.Empty();
-
-		ReleaseMutableResources(false, *Instance);
 	}
 	
 	InvalidateGeneratedData();
@@ -5744,6 +5678,8 @@ void UCustomizableInstancePrivate::BuildMaterials(const TSharedRef<FUpdateContex
 		Public->GetPathName(),
 		UCustomizableObjectSystem::GetInstance()->GetPrivate()->MutableSystem,
 		OperationData->Model,
+		CustomizableObject->GetPrivate()->MeshIdRegistry,
+	    CustomizableObject->GetPrivate()->ImageIdRegistry,
 		CustomizableObject->GetPrivate()->GetModelStreamableBulkData(),
 		OperationData->Parameters,
 	    OperationData->GetCapturedDescriptor().GetState());
@@ -6043,7 +5979,7 @@ void UCustomizableInstancePrivate::BuildMaterials(const TSharedRef<FUpdateContex
 					MUTABLE_CPUPROFILER_SCOPE(BuildMaterials_ImageLoop);
 
 					// Get the cache of resources of all live instances of this object
-					FMutableResourceCache& Cache = UCustomizableObjectSystem::GetInstance()->GetPrivate()->GetObjectCache(CustomizableObject);
+					FTextureCache& TextureCache = CustomizableObject->GetPrivate()->TextureCache;
 
 					FString CurrentState = Public->GetCurrentState();
 					bool bNeverStream = OperationData->bNeverStream;
@@ -6064,16 +6000,9 @@ void UCustomizableInstancePrivate::BuildMaterials(const TSharedRef<FUpdateContex
 						FString TextureReuseCacheRef = bReuseTextures ? FString::Printf(TEXT("%d-%d-%d-%d"), Image.BaseLOD, ObjectComponentIndex.GetValue(), Surface.SurfaceId, ImageIndex) : FString();
 
 						// If the mutable image is null, it must be in the cache
-						FMutableImageCacheKey ImageCacheKey = { Image.ImageID, OperationData->MipsToSkip };
 						if (!MutableImage)
 						{
-							TWeakObjectPtr<UTexture2D>* CachedPointerPtr = Cache.Images.Find(ImageCacheKey);
-							if (CachedPointerPtr)
-							{
-								ensure(!CachedPointerPtr->IsStale());
-								MutableTexture = CachedPointerPtr->Get();
-							}
-
+							MutableTexture = TextureCache.Get(FTextureCache::FId(Image.ImageID, OperationData->MipsToSkip));
 							check(MutableTexture);
 						}
 
@@ -6124,7 +6053,7 @@ void UCustomizableInstancePrivate::BuildMaterials(const TSharedRef<FUpdateContex
 								TSharedPtr<FTexturePlatformData, ESPMode::ThreadSafe> ReusedTexturePlatformData;
 
 								// Ensure the name of the texture is unique (prevents black textures in editor previews)
-								FString MutableTextureName = FString::Printf(TEXT("T_%s_%llu"), *Props.TextureParameterName, Image.ImageID);
+								FString MutableTextureName = FString::Printf(TEXT("T_%s_%u"), *Props.TextureParameterName, GetTypeHash(*Image.ImageID.GetKey()));
 								MutableTextureName.ReplaceInline(TEXT(" "), TEXT("_"));
 								MutableTextureName = MakeUniqueObjectName(GetTransientPackage(), UTexture2D::StaticClass(), *MutableTextureName, EUniqueObjectNameOptions::GloballyUnique).ToString();
 									
@@ -6184,9 +6113,9 @@ void UCustomizableInstancePrivate::BuildMaterials(const TSharedRef<FUpdateContex
 									}
 									else
 									{
-										UE_LOG(LogMutable, Error, TEXT("Required image [%s] with ID [%lld] was not generated in the mutable thread, and it is not cached. LOD [%d]. Object Component [%d]"),
+										UE_LOG(LogMutable, Error, TEXT("Required image [%s] with ID [%s] was not generated in the mutable thread, and it is not cached. LOD [%d]. Object Component [%d]"),
 											*Props.TextureParameterName,
-											Image.ImageID,
+											*Image.ImageID.ToString(),
 											LODIndex, ObjectComponentIndex.GetValue());
 										continue;
 									}
@@ -6213,15 +6142,8 @@ void UCustomizableInstancePrivate::BuildMaterials(const TSharedRef<FUpdateContex
 										if (ReusedTexture)
 										{
 											// Must remove texture from cache since it will be reused with a different ImageID
-											for (TPair<FMutableImageCacheKey, TWeakObjectPtr<UTexture2D>>& CachedTexture : Cache.Images)
-											{
-												if (CachedTexture.Value == MutableTexture)
-												{
-													Cache.Images.Remove(CachedTexture.Key);
-													break;
-												}
-											}
-
+											TextureCache.Remove(*MutableTexture);
+											
 											check(ReusedTexturePlatformData.IsValid());
 
 											if (ReusedTexturePlatformData.IsValid())
@@ -6261,7 +6183,7 @@ void UCustomizableInstancePrivate::BuildMaterials(const TSharedRef<FUpdateContex
 										}
 									}
 
-									Cache.Images.Add(ImageCacheKey, MutableTexture);						
+									TextureCache.Add(FTextureCache::FId(Image.ImageID, OperationData->MipsToSkip), MutableTexture);						
 								}
 								else
 								{
@@ -6270,7 +6192,6 @@ void UCustomizableInstancePrivate::BuildMaterials(const TSharedRef<FUpdateContex
 							}
 
 							FGeneratedTexture TextureData;
-							TextureData.Key = ImageCacheKey;
 							TextureData.Name = Props.TextureParameterName;
 							TextureData.Texture = MutableTexture ? MutableTexture : PassThroughTexture;
 							
@@ -6574,26 +6495,7 @@ void UCustomizableInstancePrivate::BuildMaterials(const TSharedRef<FUpdateContex
 
 	{
 		MUTABLE_CPUPROFILER_SCOPE(BuildMaterials_Exchange);
-
-		UCustomizableObjectSystemPrivate* CustomizableObjectSystem = UCustomizableObjectSystem::GetInstance()->GetPrivate();
-		TexturesToRelease.Empty();
-
-		for (const FGeneratedTexture& Texture : NewGeneratedTextures)
-		{
-			CustomizableObjectSystem->AddTextureReference(Texture.Key);
-		}
-
-		for (const FGeneratedTexture& Texture : GeneratedTextures)
-		{
-			if (CustomizableObjectSystem->RemoveTextureReference(Texture.Key))
-			{
-				if (CustomizableObjectSystem->bReleaseTexturesImmediately)
-				{
-					TexturesToRelease.Add(Texture); // Texture count is zero, so prepare to release it
-				}
-			}
-		}
-
+		
 		Exchange(GeneratedTextures, NewGeneratedTextures);
 
 		// All pass-through textures and meshes have been set, no need to keep referencing them from the instance
diff --git a/Engine/Plugins/Mutable/Source/CustomizableObject/Private/MuCO/CustomizableObjectMeshUpdate.cpp b/Engine/Plugins/Mutable/Source/CustomizableObject/Private/MuCO/CustomizableObjectMeshUpdate.cpp
--- a/Engine/Plugins/Mutable/Source/CustomizableObject/Private/MuCO/CustomizableObjectMeshUpdate.cpp
+++ b/Engine/Plugins/Mutable/Source/CustomizableObject/Private/MuCO/CustomizableObjectMeshUpdate.cpp
@@ -48,6 +48,8 @@ FCustomizableObjectMeshStreamIn::FCustomizableObjectMeshStreamIn(
 
 	OperationData->Model = InMesh->Model;
 	OperationData->Parameters = InMesh->Parameters;
+	OperationData->MeshIdRegistry = InMesh->MeshIdRegistry;
+	OperationData->ImageIdRegistry = InMesh->ImageIdRegistry;
 	OperationData->State = InMesh->State;
 
 	OperationData->ModelStreamableBulkData = ModelStreamableBulkData;
@@ -196,7 +198,7 @@ namespace impl
 			return;
 		}
 
-		const mu::FResourceID MeshID = OperationData->MeshIDs[LODIndex];
+		const mu::FMeshId& MeshID = OperationData->MeshIDs[LODIndex];
 
 		constexpr mu::EMeshContentFlags MeshContentFilter = mu::EMeshContentFlags::AllFlags;
 		UE::Tasks::TTask<TSharedPtr<const mu::FMesh>> GetMeshTask = OperationData->System->GetMesh(InstanceID, MeshID, MeshContentFilter);
@@ -243,7 +245,13 @@ namespace impl
 		const uint32 LODMask = 0xFFFFFFFF;
 
 		// Main instance generation step
-		TSharedPtr<const mu::FInstance> Instance = System->BeginUpdate(InstanceID, OperationData->Parameters, OperationData->State, LODMask);
+		TSharedPtr<const mu::FInstance> Instance = System->BeginUpdate(InstanceID,
+			OperationData->Parameters,
+			OperationData->MeshIdRegistry,
+			OperationData->ImageIdRegistry,
+			OperationData->State,
+			LODMask);
+		
 		check(Instance);
 
 		Task_Mutable_UpdateMesh_Loop(OperationData, Task, InstanceID, OperationData->PendingFirstLODIdx + OperationData->AssetLODBias);
diff --git a/Engine/Plugins/Mutable/Source/CustomizableObject/Private/MuCO/CustomizableObjectMeshUpdate.h b/Engine/Plugins/Mutable/Source/CustomizableObject/Private/MuCO/CustomizableObjectMeshUpdate.h
--- a/Engine/Plugins/Mutable/Source/CustomizableObject/Private/MuCO/CustomizableObjectMeshUpdate.h
+++ b/Engine/Plugins/Mutable/Source/CustomizableObject/Private/MuCO/CustomizableObjectMeshUpdate.h
@@ -26,6 +26,8 @@ struct FMutableMeshOperationData
 {
 	TSharedPtr<mu::FSystem> System;
 	TSharedPtr<mu::FModel, ESPMode::ThreadSafe> Model;
+	TSharedPtr<mu::FMeshIdRegistry> MeshIdRegistry;
+	TSharedPtr<mu::FImageIdRegistry> ImageIdRegistry;
 	TSharedPtr<const mu::FParameters> Parameters;
 
 	TSharedPtr<FModelStreamableBulkData> ModelStreamableBulkData;
@@ -33,7 +35,7 @@ struct FMutableMeshOperationData
 	int32 State = -1;
 
 	// IDs of the meshes to generate per LOD.
-	TArray<mu::FResourceID> MeshIDs;
+	TArray<mu::FMeshId> MeshIDs;
 
 	// Meshes generated per LOD
 	TArray<TSharedPtr<const mu::FMesh>> Meshes;
diff --git a/Engine/Plugins/Mutable/Source/CustomizableObject/Private/MuCO/CustomizableObjectMipDataProvider.cpp b/Engine/Plugins/Mutable/Source/CustomizableObject/Private/MuCO/CustomizableObjectMipDataProvider.cpp
--- a/Engine/Plugins/Mutable/Source/CustomizableObject/Private/MuCO/CustomizableObjectMipDataProvider.cpp
+++ b/Engine/Plugins/Mutable/Source/CustomizableObject/Private/MuCO/CustomizableObjectMipDataProvider.cpp
@@ -40,15 +40,19 @@ static FAutoConsoleVariableRef CVarEnablePrefetchHighQualityMips(
 FMutableUpdateContext::FMutableUpdateContext(
 	const FString& InCustomizableObjectPathName,
 	const FString& InInstancePathName,
-	TSharedPtr<mu::FSystem> InSystem,
-	TSharedPtr<mu::FModel, ESPMode::ThreadSafe> InModel,
-	TSharedPtr<FModelStreamableBulkData> InModelStreamable,
-	TSharedPtr<const mu::FParameters> InParameters,
+	const TSharedPtr<mu::FSystem>& InSystem,
+	const TSharedPtr<mu::FModel>& InModel,
+	const TSharedPtr<mu::FMeshIdRegistry>& InMeshIdRegistry,
+	const TSharedPtr<mu::FImageIdRegistry>& InImageIdRegistry,
+	const TSharedPtr<FModelStreamableBulkData>& InModelStreamable,
+	const TSharedPtr<const mu::FParameters>& InParameters,
 	int32 InState) :
 	CustomizableObjectPathName(InCustomizableObjectPathName),
 	InstancePathName(InInstancePathName),
 	System(InSystem),
 	Model(InModel),
+	MeshIdRegistry(InMeshIdRegistry),
+	ImageIdRegistry(InImageIdRegistry),
 	ModelStreamableBulkData(InModelStreamable),
 	Parameters(InParameters),
 	State(InState)
@@ -108,7 +112,7 @@ FMutableTextureMipDataProvider::FMutableTextureMipDataProvider(const UTexture* T
 	: FTextureMipDataProvider(Texture, ETickState::Init, ETickThread::Async),
 	CustomizableObjectInstance(InCustomizableObjectInstance), ImageRef(InImageRef)
 {
-	check(ImageRef.ImageID > 0);
+	check(ImageRef.ImageID);
 }
 
 
@@ -207,7 +211,13 @@ namespace Impl
 			// LOD mask, set to all ones to build all LODs
 			uint32 LODMask = 0xFFFFFFFF;
 
-			Instance = System->BeginUpdate(InstanceID, OperationData->UpdateContext->GetParameters(), OperationData->UpdateContext->GetState(), LODMask);
+			Instance = System->BeginUpdate(InstanceID,
+				OperationData->UpdateContext->GetParameters(),
+				OperationData->UpdateContext->MeshIdRegistry,
+				OperationData->UpdateContext->ImageIdRegistry,
+				OperationData->UpdateContext->GetState(),
+				LODMask);
+			
 			check(Instance);
 		}
 
@@ -218,7 +228,7 @@ namespace Impl
 
 		// This ID may be different than the ID obtained the first time the image was generated, because the mutable
 		// runtime cannot remember all the resources it has built, and only remembers a fixed amount.
-		mu::FResourceID MipImageID = Instance->GetImageId(ImageRef.Component, ImageRef.LOD, SurfaceIndex, ImageRef.Image);
+		mu::FImageId MipImageID = Instance->GetImageId(ImageRef.Component, ImageRef.LOD, SurfaceIndex, ImageRef.Image);
 
 		UE::Tasks::TTask<TSharedPtr<const mu::FImage>> GetImageTask = 
 				System->GetImage(InstanceID, MipImageID, ImageRef.BaseMip + OperationData->MipsToSkip, ImageRef.LOD);
@@ -348,7 +358,7 @@ int32 FMutableTextureMipDataProvider::GetMips(const FTextureUpdateContext& Conte
 	const TIndirectArray<FTexture2DMipMap>& OwnerMips = Texture->GetPlatformMips();
 
 	const int32 NumMips = OwnerMips.Num();
-	check(ImageRef.ImageID > 0);
+	check(ImageRef.ImageID);
 
 	// Maximum value to skip, will be minimized by the first mip level requested
 	int32 MipsToSkip = 256;
diff --git a/Engine/Plugins/Mutable/Source/CustomizableObject/Private/MuCO/CustomizableObjectSkeletalMesh.cpp b/Engine/Plugins/Mutable/Source/CustomizableObject/Private/MuCO/CustomizableObjectSkeletalMesh.cpp
--- a/Engine/Plugins/Mutable/Source/CustomizableObject/Private/MuCO/CustomizableObjectSkeletalMesh.cpp
+++ b/Engine/Plugins/Mutable/Source/CustomizableObject/Private/MuCO/CustomizableObjectSkeletalMesh.cpp
@@ -27,11 +27,13 @@ void UCustomizableObjectSkeletalMesh::InitMutableStreamingData(const TSharedRef<
 
 	// Init properties
 	Model = CustomizableObject.GetPrivate()->GetModel();
+	MeshIdRegistry = CustomizableObject.GetPrivate()->MeshIdRegistry;
+	ImageIdRegistry = CustomizableObject.GetPrivate()->ImageIdRegistry;
 
 	Parameters = InOperationData->Parameters;
 	State = InOperationData->GetCapturedDescriptor().GetState();
 
-	MeshIDs.Init(MAX_uint64, MAX_MESH_LOD_COUNT);
+	MeshIDs.Init({}, MAX_MESH_LOD_COUNT);
 	SurfaceIDs.SetNum(MAX_MESH_LOD_COUNT);
 
 	const int32 FirstLOD = InOperationData->bStreamMeshLODs ?
diff --git a/Engine/Plugins/Mutable/Source/CustomizableObject/Private/MuCO/CustomizableObjectSystem.cpp b/Engine/Plugins/Mutable/Source/CustomizableObject/Private/MuCO/CustomizableObjectSystem.cpp
--- a/Engine/Plugins/Mutable/Source/CustomizableObject/Private/MuCO/CustomizableObjectSystem.cpp
+++ b/Engine/Plugins/Mutable/Source/CustomizableObject/Private/MuCO/CustomizableObjectSystem.cpp
@@ -42,6 +42,7 @@
 #endif
 
 #include "MutableStreamRequest.h"
+#include "MuCO/CustomizableObjectSkeletalMesh.h"
 #include "MuCO/LoadUtils.h"
 
 #include UE_INLINE_GENERATED_CPP_BY_NAME(CustomizableObjectSystem)
@@ -93,12 +94,6 @@ TAutoConsoleVariable<bool> CVarReuseImagesBetweenInstances(
 	TEXT("Enables or disables the reuse of images between instances."),
 	ECVF_Scalability);
 
-static TAutoConsoleVariable<int32> CVarGeneratedResourcesCacheSize(
-	TEXT("mutable.GeneratedResourcesCacheSize"),
-	512,
-	TEXT("Limit the number of resources (images and meshes) that will be tracked for reusal. Each tracked resource uses a small amout of memory for its key."),
-	ECVF_Scalability);
-
 TAutoConsoleVariable<bool> CVarPreserveUserLODsOnFirstGeneration(
 	TEXT("mutable.PreserveUserLODsOnFirstGeneration"),
 	true,
@@ -182,72 +177,6 @@ static void CVarMutableSinkFunction()
 	}
 }
 
-#if WITH_EDITOR
-
-FMutableSystemSettingsOverrides::FMutableSystemSettingsOverrides(bool bUseProgressiveMipStreaming, bool bOnlyGenerateRequestedLODs,
-	mu::FImageOperator::FImagePixelFormatFunc InImagePixelFormatFunc)
-{
-	// Cache previous settings
-	UCustomizableObjectSystem* System = UCustomizableObjectSystem::GetInstanceChecked();
-	bOldIsProgressiveMipStreamingEnabled = System->IsProgressiveMipStreamingEnabled();
-	bOldIsOnlyGenerateRequestedLODsEnabled = System->IsOnlyGenerateRequestedLODsEnabled();
-	OldImagePixelFormatFunc = System->GetPrivate()->ImageFormatOverrideFunc;
-	
-	// Cache desired settings
-	bIsProgressiveMipStreamingEnabled = bUseProgressiveMipStreaming;
-	bIsOnlyGenerateRequestedLODsEnabled = bOnlyGenerateRequestedLODs;
-	ImagePixelFormatFunc = InImagePixelFormatFunc;
-}
-
-
-void FMutableSystemSettingsOverrides::ApplySettingsOverrides() const
-{
-	UE_LOG(LogMutable, Display, TEXT("Overriding system settings for this instance update."))
-	
-	UCustomizableObjectSystem* System = UCustomizableObjectSystem::GetInstanceChecked();
-	System->SetProgressiveMipStreamingEnabled(bIsProgressiveMipStreamingEnabled);
-	System->SetOnlyGenerateRequestedLODsEnabled(bIsOnlyGenerateRequestedLODsEnabled);
-	
-	if (ensure(IsInGameThread()))
-	{
-		PrepareUnrealCompression();
-		
-		UCustomizableObjectSystemPrivate* Private = System->GetPrivate();
-		check(Private)
-		Private->ImageFormatOverrideFunc = ImagePixelFormatFunc;
-	}
-	else
-	{
-		UE_LOG(LogMutable, Error, TEXT("Failed to apply the MutableSystem Pixel Format override. Invocation of method performed outside the MainThread"))
-	}
-}
-
-
-void FMutableSystemSettingsOverrides::RestoreSettings() const
-{
-	UE_LOG(LogMutable, Display, TEXT("Restoring system settings."))
-	
- 	// apply the settings to the system using the old data
- 	UCustomizableObjectSystem* System = UCustomizableObjectSystem::GetInstanceChecked();
- 	System->SetProgressiveMipStreamingEnabled(bOldIsProgressiveMipStreamingEnabled);
- 	System->SetOnlyGenerateRequestedLODsEnabled(bOldIsOnlyGenerateRequestedLODsEnabled);
-
-	if (ensure(IsInGameThread()))
-	{
-		PrepareUnrealCompression();
-		
-		UCustomizableObjectSystemPrivate* Private = System->GetPrivate();
-		check(Private)
-		Private->ImageFormatOverrideFunc = OldImagePixelFormatFunc;
-	}
-	else
-	{
-		UE_LOG(LogMutable, Error, TEXT("Failed to restore the MutableSystem Pixel Format override. Invocation of method performed outside the MainThread"))
-	}
-}
-
-#endif
-
 
 static FAutoConsoleVariableSink CVarMutableSink(FConsoleCommandDelegate::CreateStatic(&CVarMutableSinkFunction));
 
@@ -327,9 +256,18 @@ FUpdateContextPrivate::FUpdateContextPrivate(UCustomizableObjectInstance& InInst
 		ComponentNames = ModelResources->ComponentNamesPerObjectComponent;
 	}
 
-	MutableSystem = UCustomizableObjectSystem::GetInstance()->GetPrivate()->MutableSystem;	
-	check(MutableSystem);
+	UCustomizableObjectSystem* System = UCustomizableObjectSystem::GetInstanceChecked();
+	UCustomizableObjectSystemPrivate* SystemPrivate = System->GetPrivate();
 
+	MutableSystem = SystemPrivate->MutableSystem;	
+	check(MutableSystem);
+	
+	bProgressiveMipStreamingEnabled = System->IsProgressiveMipStreamingEnabled();
+	bIsOnlyGenerateRequestedLODsEnabled = System->IsOnlyGenerateRequestedLODsEnabled();
+#if WITH_EDITOR
+	PixelFormatOverride = SystemPrivate->ImageFormatOverrideFunc;
+#endif
+	
 	bValid = true;
 }
 
@@ -652,13 +590,13 @@ void FUpdateContextPrivate::InitMeshDescriptors(int32 Size)
 }
 
 
-const TArray<TArray<mu::FResourceID>>& FUpdateContextPrivate::GetMeshDescriptors() const
+const TArray<TArray<mu::FMeshId>>& FUpdateContextPrivate::GetMeshDescriptors() const
 {
 	return MeshDescriptors;
 }
 
 
-TArray<mu::FResourceID>* FUpdateContextPrivate::GetMeshDescriptors(FCustomizableObjectComponentIndex Index)
+TArray<mu::FMeshId>* FUpdateContextPrivate::GetMeshDescriptors(FCustomizableObjectComponentIndex Index)
 {
 	if (MeshDescriptors.IsValidIndex(Index.GetValue()))
 	{
@@ -897,7 +835,6 @@ void UCustomizableObjectSystem::InitSystem()
 	Private->CurrentInstanceBeingUpdated = nullptr;
 
 	Private->LastWorkingMemoryBytes = CVarWorkingMemoryKB->GetInt() * 1024;
-	Private->LastGeneratedResourceCacheSize = CVarGeneratedResourcesCacheSize.GetValueOnGameThread();
 
 	mu::FSettings Settings;
 	Settings.SetProfile(false);
@@ -994,76 +931,6 @@ UCustomizableObjectSystem* UCustomizableObjectSystemPrivate::GetPublic() const
 }
 
 
-TAutoConsoleVariable<bool> CVarCleanupTextureCache(
-	TEXT("mutable.EnableCleanupCache"),
-	true,
-	TEXT("If enabled stale textures and meshes in mutable's cache will be removed."),
-	ECVF_Scalability);
-
-
-void UCustomizableObjectSystemPrivate::CleanupCache()
-{
-	check(IsInGameThread());
-
-	const bool bCleanupEnabled = CVarCleanupTextureCache.GetValueOnGameThread();
-
-	for (int32 ModelIndex = 0; ModelIndex < ModelResourcesCache.Num();)
-	{
-		if (!ModelResourcesCache[ModelIndex].Object.IsValid(false, true))
-		{
-			// The whole object has been destroyed. Remove everything.
-			ModelResourcesCache.RemoveAtSwap(ModelIndex);
-		}
-		else
-		{
-			if (bCleanupEnabled)
-			{
-				// Remove stale textures
-				for (auto Iterator = ModelResourcesCache[ModelIndex].Images.CreateIterator(); Iterator; ++Iterator)
-				{
-					if (Iterator->Value.IsStale())
-					{
-						Iterator.RemoveCurrent();
-					}
-				}
-
-				// Remove stale meshes
-				for (auto Iterator = ModelResourcesCache[ModelIndex].Meshes.CreateIterator(); Iterator; ++Iterator)
-				{
-					if (Iterator->Value.IsStale())
-					{
-						Iterator.RemoveCurrent();
-					}
-				}
-			}
-
-			++ModelIndex;
-		}
-	}
-}
-
-
-FMutableResourceCache& UCustomizableObjectSystemPrivate::GetObjectCache(const UCustomizableObject* Object)
-{
-	check(IsInGameThread());
-
-	// Not mandatory, but a good place for a cleanup
-	CleanupCache();
-
-	for (int ModelIndex = 0; ModelIndex < ModelResourcesCache.Num(); ++ModelIndex)
-	{
-		if (ModelResourcesCache[ModelIndex].Object==Object)
-		{
-			return ModelResourcesCache[ModelIndex];
-		}
-	}
-		
-	// Not found, create and add it.
-	ModelResourcesCache.Push(FMutableResourceCache());
-	ModelResourcesCache.Last().Object = Object;
-	return ModelResourcesCache.Last();
-}
-
 bool bForceStreamMeshLODs = false;
 
 static FAutoConsoleVariableRef CVarMutableForceStreamMeshLODs(
@@ -1153,7 +1020,6 @@ void FinishUpdateGlobal(const TSharedRef<FUpdateContextPrivate>& Context)
 	{
 		SystemPrivate->CurrentInstanceBeingUpdated = nullptr;
 		SystemPrivate->CurrentMutableOperation = nullptr;
-		SystemPrivate->ClearResourceCacheProtected();
 	}
 	
 	if (Instance)
@@ -1247,15 +1113,6 @@ void FinishUpdateGlobal(const TSharedRef<FUpdateContextPrivate>& Context)
 	FUpdateContext ContextPublic;
 	ContextPublic.UpdateResult = Context->UpdateResult;
 	ContextPublic.Instance = Instance;
-
-#if WITH_EDITORONLY_DATA
-	// Restore the overriden system settings (bake operation)
-	if (Context->UpdateSettingsOverride)
-	{
-		Context->UpdateSettingsOverride->RestoreSettings();
-		Context->UpdateSettingsOverride = nullptr;
-	}
-#endif
 	
 	Context->UpdateCallback.ExecuteIfBound(ContextPublic);
 	Context->UpdateNativeCallback.Broadcast(ContextPublic);
@@ -1506,7 +1363,7 @@ void UpdateSkeletalMesh(const TSharedRef<FUpdateContextPrivate>& Context)
 }
 
 
-void UCustomizableObjectSystemPrivate::GetMipStreamingConfig(const UCustomizableObjectInstance& Instance, bool& bOutNeverStream, int32& OutMipsToSkip) const
+void UCustomizableObjectSystemPrivate::GetMipStreamingConfig(const UCustomizableObjectInstance& Instance, bool bTextureStreamingEnabled, bool& bOutNeverStream, int32& OutMipsToSkip) const
 {
 	bOutNeverStream = false;
 
@@ -1533,7 +1390,7 @@ void UCustomizableObjectSystemPrivate::GetMipStreamingConfig(const UCustomizable
 
 	// Streaming disabled from platform settings or from platform CustomizableObjectSystem properties?
 #if PLATFORM_SUPPORTS_TEXTURE_STREAMING
-	if (!IStreamingManager::Get().IsTextureStreamingEnabled() || !EnableMutableProgressiveMipStreaming)
+	if (!IStreamingManager::Get().IsTextureStreamingEnabled() || !bTextureStreamingEnabled)
 	{
 		bOutNeverStream = true;
 	}
@@ -1566,52 +1423,6 @@ int32 UCustomizableObjectSystemPrivate::GetNumSkeletalMeshes() const
 }
 
 
-void UCustomizableObjectSystemPrivate::AddTextureReference(const FMutableImageCacheKey& TextureId)
-{
-	uint32& CountRef = TextureReferenceCount.FindOrAdd(TextureId);
-
-	CountRef++;
-}
-
-
-bool UCustomizableObjectSystemPrivate::RemoveTextureReference(const FMutableImageCacheKey& TextureId)
-{
-	uint32* CountPtr = TextureReferenceCount.Find(TextureId);
-
-	if (CountPtr && *CountPtr > 0)
-	{
-		(*CountPtr)--;
-
-		if (*CountPtr == 0)
-		{
-			TextureReferenceCount.Remove(TextureId);
-
-			return true;
-		}
-	}
-	else
-	{
-		ensure(false); // Mutable texture reference count is incorrect
-		TextureReferenceCount.Remove(TextureId);
-	}
-
-	return false;
-}
-
-
-bool UCustomizableObjectSystemPrivate::TextureHasReferences(const FMutableImageCacheKey& TextureId) const
-{
-	const uint32* CountPtr = TextureReferenceCount.Find(TextureId);
-
-	if (CountPtr && *CountPtr > 0)
-	{
-		return true;
-	}
-
-	return false;
-}
-
-
 EUpdateRequired UCustomizableObjectSystemPrivate::IsUpdateRequired(const UCustomizableObjectInstance& Instance, bool bOnlyUpdateIfNotGenerated, bool bOnlyUpdateIfLODs, bool bIgnoreCloseDist) const
 {
 	UCustomizableObjectSystem* System = UCustomizableObjectSystem::GetInstance();
@@ -1756,71 +1567,72 @@ void UCustomizableObjectSystemPrivate::EnqueueUpdateSkeletalMesh(const TSharedRe
 	}		
 	case EUpdateRequired::Update:
 	{
-		if (const FMutablePendingInstanceUpdate* QueueElem = MutablePendingInstanceWork.GetUpdate(Instance))
+		if (!Context->bForce)
 		{
-			if (Context->GetCapturedDescriptorHash().IsSubset(QueueElem->Context->GetCapturedDescriptorHash()))
+			if (const FMutablePendingInstanceUpdate* QueueElem = MutablePendingInstanceWork.GetUpdate(Instance))
 			{
-				Context->bOptimizedUpdate = true;
-				Context->UpdateResult = EUpdateResult::ErrorOptimized;
-				FinishUpdateGlobal(Context);			
-				return; // The requested update is equal to the last enqueued update.
-			}
-		}	
+				if (Context->GetCapturedDescriptorHash().IsSubset(QueueElem->Context->GetCapturedDescriptorHash()))
+				{
+					Context->bOptimizedUpdate = true;
+					Context->UpdateResult = EUpdateResult::ErrorOptimized;
+					FinishUpdateGlobal(Context);			
+					return; // The requested update is equal to the last enqueued update.
+				}
+			}	
 
-		if (CurrentMutableOperation &&
-			Instance == CurrentMutableOperation->Instance &&
-			Context->GetCapturedDescriptorHash().IsSubset(CurrentMutableOperation->GetCapturedDescriptorHash()))
-		{
-			Context->bOptimizedUpdate = true;
-			Context->UpdateResult = EUpdateResult::ErrorOptimized;
-			FinishUpdateGlobal(Context);
-			return; // The requested update is equal to the running update.
-		}
-	
-		if (Context->GetCapturedDescriptorHash().IsSubset(Instance->GetPrivate()->CommittedDescriptorHash) &&
-			!(CurrentMutableOperation &&
-			Instance == CurrentMutableOperation->Instance)) // This condition is necessary because even if the descriptor is a subset, it will be replaced by the CurrentMutableOperation
-		{
-			if (CVarEnableUpdateOptimization.GetValueOnGameThread()) // TODO Remove hotfix: UE-218957 
+			if (CurrentMutableOperation &&
+				Instance == CurrentMutableOperation->Instance &&
+				Context->GetCapturedDescriptorHash().IsSubset(CurrentMutableOperation->GetCapturedDescriptorHash()))
 			{
 				Context->bOptimizedUpdate = true;
-
-				// The user may have changed the AttachParent and we need to re-customize it.
-				// In case nothing need to be re-customized, the update will be considered ErrorOptimized.
-				UpdateSkeletalMesh(Context);
-				Context->UpdateResult = Context->AttachedParentUpdated.IsEmpty() ? EUpdateResult::ErrorOptimized : EUpdateResult::Success;
-
+				Context->UpdateResult = EUpdateResult::ErrorOptimized;
 				FinishUpdateGlobal(Context);
+				return; // The requested update is equal to the running update.
 			}
-			else 
+		
+			if (Context->GetCapturedDescriptorHash().IsSubset(Instance->GetPrivate()->CommittedDescriptorHash) &&
+				!(CurrentMutableOperation &&
+				Instance == CurrentMutableOperation->Instance)) // This condition is necessary because even if the descriptor is a subset, it will be replaced by the CurrentMutableOperation
 			{
-				Context->bOptimizedUpdate = false;
+				if (CVarEnableUpdateOptimization.GetValueOnGameThread()) // TODO Remove hotfix: UE-218957 
+				{
+					Context->bOptimizedUpdate = true;
 
-				// The user may have changed the AttachParent and we need to re-customize it.
-				// In case nothing need to be re-customized, the update will be considered ErrorOptimized.
-				UpdateSkeletalMesh(Context);
-				Context->UpdateResult = EUpdateResult::Success;
+					// The user may have changed the AttachParent and we need to re-customize it.
+					// In case nothing need to be re-customized, the update will be considered ErrorOptimized.
+					UpdateSkeletalMesh(Context);
+					Context->UpdateResult = Context->AttachedParentUpdated.IsEmpty() ? EUpdateResult::ErrorOptimized : EUpdateResult::Success;
 
-				FinishUpdateGlobal(Context);
+					FinishUpdateGlobal(Context);
+				}
+				else 
+				{
+					Context->bOptimizedUpdate = false;
+
+					// The user may have changed the AttachParent and we need to re-customize it.
+					// In case nothing need to be re-customized, the update will be considered ErrorOptimized.
+					UpdateSkeletalMesh(Context);
+					Context->UpdateResult = EUpdateResult::Success;
+
+					FinishUpdateGlobal(Context);
+				}
 			}
 		}
-		else
-		{
-			if (CVarDescriptorDebugPrint->GetBool())
-			{
-				FString String = TEXT("DESCRIPTOR DEBUG PRINT\n");
-				String += "================================\n";				
-				String += FString::Printf(TEXT("=== DESCRIPTOR HASH ===\n%s\n"), *Context->GetCapturedDescriptorHash().ToString());
-				String += FString::Printf(TEXT("=== DESCRIPTOR ===\n%s"), *Instance->GetPrivate()->GetDescriptor().ToString());
-				String += "================================";
-				
-				UE_LOG(LogMutable, Log, TEXT("%s"), *String);
-			}
 
-			const FMutablePendingInstanceUpdate InstanceUpdate(Context);
-			MutablePendingInstanceWork.AddUpdate(InstanceUpdate);
+		if (CVarDescriptorDebugPrint->GetBool())
+		{
+			FString String = TEXT("DESCRIPTOR DEBUG PRINT\n");
+			String += "================================\n";				
+			String += FString::Printf(TEXT("=== DESCRIPTOR HASH ===\n%s\n"), *Context->GetCapturedDescriptorHash().ToString());
+			String += FString::Printf(TEXT("=== DESCRIPTOR ===\n%s"), *Instance->GetPrivate()->GetDescriptor().ToString());
+			String += "================================";
+			
+			UE_LOG(LogMutable, Log, TEXT("%s"), *String);
 		}
 
+		const FMutablePendingInstanceUpdate InstanceUpdate(Context);
+		MutablePendingInstanceWork.AddUpdate(InstanceUpdate);
+
 		break;
 	}
 
@@ -1879,15 +1691,6 @@ void UCustomizableObjectSystem::SetReplaceDiscardedWithReferenceMeshEnabled(bool
 }
 
 
-void UCustomizableObjectSystemPrivate::ClearResourceCacheProtected()
-{
-	check(IsInGameThread());
-
-	ProtectedCachedTextures.Reset(0);
-	ProtectedObjectCachedImages.Reset(0);
-}
-
-
 #if WITH_EDITOR
 
 void UCustomizableObjectSystemPrivate::AddPendingLoad(UCustomizableObject* CO)
@@ -1964,11 +1767,7 @@ bool UCustomizableObjectSystem::LockObject(UCustomizableObject* InObject)
 		}
 
 		Private->MutablePendingInstanceWork.RemoveUpdatesForObject(InObject);
-
-		// Clear the cache for the instance, since we will remake it
-		FMutableResourceCache& Cache = GetPrivate()->GetObjectCache(InObject);
-		Cache.Clear();
-
+		
 		check(InObject->GetPrivate()->bLocked);
 
 		return true;
@@ -2068,14 +1867,6 @@ void UCustomizableObjectSystemPrivate::UpdateMemoryLimit()
 		check(MutableSystem);
 		MutableSystem->SetWorkingMemoryBytes(MemoryBytes);
 	}
-
-	const uint32 GeneratedResourceCacheSize = CVarGeneratedResourcesCacheSize.GetValueOnGameThread();
-	if (GeneratedResourceCacheSize != LastGeneratedResourceCacheSize)
-	{
-		LastGeneratedResourceCacheSize = GeneratedResourceCacheSize;
-		check(MutableSystem);
-		MutableSystem->SetGeneratedCacheSize(GeneratedResourceCacheSize);
-	}
 }
 
 
@@ -2089,7 +1880,7 @@ namespace impl
 	struct FGetImageData
 	{
 		int32 ImageIndex;
-		mu::FResourceID ImageID;
+		mu::FImageId ImageID;
 	};
 	
 
@@ -2104,7 +1895,7 @@ namespace impl
 	struct FGetMeshData
 	{
 		int32 InstanceUpdateLODIndex;
-		mu::FResourceID MeshID;
+		mu::FMeshId MeshID;
 		mu::EMeshContentFlags ContentFilter;
 	};
 
@@ -2122,7 +1913,7 @@ namespace impl
 	void Task_Mutable_GetImages_GetImage(
 		const TSharedRef<FUpdateContextPrivate>& OperationData,
 		double StartTime,
-		const TSharedPtr<TArray<mu::FResourceID>>& ImagesInThisInstance,
+		const TSharedPtr<TArray<mu::FImageId>>& ImagesInThisInstance,
 		int32 ImageIndex,
 		UE::Tasks::TTask<mu::FExtendedImageDesc> GetImageDescTask);
 
@@ -2131,7 +1922,7 @@ namespace impl
 	void Task_Mutable_GetImages_Loop(
 		const TSharedRef<FUpdateContextPrivate>& OperationData,
 		double StartTime,
-		const TSharedPtr<TArray<mu::FResourceID>>& ImagesInThisInstance,
+		const TSharedPtr<TArray<mu::FImageId>>& ImagesInThisInstance,
 		int32 ImageIndex);
 
 	
@@ -2921,7 +2712,7 @@ namespace impl
 	void Task_Mutable_GetImages_GetImageDesc(
 		const TSharedRef<FUpdateContextPrivate>& OperationData,
 		double StartTime,
-		const TSharedPtr<TArray<mu::FResourceID>>& ImagesInThisInstance,
+		const TSharedPtr<TArray<mu::FImageId>>& ImagesInThisInstance,
 		int32 ImageIndex)
 	{
 		MUTABLE_CPUPROFILER_SCOPE(Task_Mutable_GetImages_GetImageDesc)
@@ -2944,7 +2735,7 @@ namespace impl
 	void Task_Mutable_GetImages_GetImage_Post(
 		const TSharedRef<FUpdateContextPrivate>& OperationData,
     	double StartTime,
-    	const TSharedPtr<TArray<mu::FResourceID>>& ImagesInThisInstance,
+    	const TSharedPtr<TArray<mu::FImageId>>& ImagesInThisInstance,
     	int32 ImageIndex,
     	UE::Tasks::TTask<TSharedPtr<const mu::FImage>> GetImageTask,
     	int32 MipSizeX,
@@ -2965,7 +2756,7 @@ namespace impl
 		if (bSizeMissmatch)
 		{
 			// Generate a correctly-sized but empty image instead, to avoid crashes.
-			UE_LOG(LogMutable, Warning, TEXT("Mutable generated a wrongly-sized image %llu."), Image.ImageID);
+			UE_LOG(LogMutable, Warning, TEXT("Mutable generated a wrongly-sized image %s."), *Image.ImageID.ToString());
 			Image.Image = MakeShared<mu::FImage>(MipSizeX, MipSizeY, FullLODContent - MipsToSkip, Image.Image->GetFormat(), mu::EInitializationType::Black);
 		}
 
@@ -2981,7 +2772,7 @@ namespace impl
 		{
 			MUTABLE_CPUPROFILER_SCOPE(GetImage_MipFix);
 
-			UE_LOG(LogMutable, Warning, TEXT("Mutable generated an incomplete mip chain for image %llu."), Image.ImageID);
+			UE_LOG(LogMutable, Warning, TEXT("Mutable generated an incomplete mip chain for image %s."), *Image.ImageID.ToString());
 
 			// Force the right number of mips. The missing data will be black.
 			const TSharedPtr<mu::FImage> NewImage = MakeShared<mu::FImage>(Image.Image->GetSizeX(), Image.Image->GetSizeY(), FullMipCount, Image.Image->GetFormat(), mu::EInitializationType::Black);
@@ -3009,7 +2800,7 @@ namespace impl
 	void Task_Mutable_GetImages_GetImage(
 		const TSharedRef<FUpdateContextPrivate>& OperationData,
 		double StartTime,
-		const TSharedPtr<TArray<mu::FResourceID>>& ImagesInThisInstance,
+		const TSharedPtr<TArray<mu::FImageId>>& ImagesInThisInstance,
 		int32 ImageIndex,
 		UE::Tasks::TTask<mu::FExtendedImageDesc> GetImageDescTask)
 	{
@@ -3045,11 +2836,11 @@ namespace impl
 		}
 
 		const bool bCached = ImagesInThisInstance->Contains(Image.ImageID) || // See if it is cached from this same instance (can happen with LODs)
-			(UCustomizableObjectSystem::ShouldReuseTexturesBetweenInstances() && CustomizableObjectSystemPrivate->ProtectedObjectCachedImages.Contains(Image.ImageID)); // See if it is cached from another instance
+			(UCustomizableObjectSystem::ShouldReuseTexturesBetweenInstances() && OperationData->CachedTextures.Contains(Image.ImageID)); // See if it is cached from another instance
 
 		if (bCached)
 		{
-			UE_LOG(LogMutable, VeryVerbose, TEXT("Texture resource with id [%llu] is cached."), Image.ImageID);
+			UE_LOG(LogMutable, VeryVerbose, TEXT("Texture resource with id [%s] is cached."), *Image.ImageID.ToString());
 
 			Task_Mutable_GetImages_Loop(OperationData, StartTime, ImagesInThisInstance, ++ImageIndex);
 			return;
@@ -3093,7 +2884,7 @@ namespace impl
 	void Task_Mutable_GetImages_Loop(
 		const TSharedRef<FUpdateContextPrivate>& OperationData,
 		double StartTime,
-		const TSharedPtr<TArray<mu::FResourceID>>& ImagesInThisInstance,
+		const TSharedPtr<TArray<mu::FImageId>>& ImagesInThisInstance,
 		int32 ImageIndex)
 	{
 		MUTABLE_CPUPROFILER_SCOPE(Task_Mutable_GetImages_Loop)
@@ -3124,7 +2915,7 @@ namespace impl
 
 			const double StartTime = FPlatformTime::Seconds();		
 
-			const TSharedPtr<TArray<mu::FResourceID>> ImagesInThisInstance = MakeShared<TArray<mu::FResourceID>>();
+			const TSharedPtr<TArray<mu::FImageId>> ImagesInThisInstance = MakeShared<TArray<mu::FImageId>>();
 			Task_Mutable_GetImages_Loop(OperationData, StartTime, ImagesInThisInstance, 0);
 		}
 		
@@ -3189,8 +2980,8 @@ namespace impl
 
 		check(OperationData);
 		
-		TMap<mu::FResourceID, FTexturePlatformData*>& ImageToPlatformDataMap = OperationData->ImageToPlatformDataMap;
-		for (const TPair<mu::FResourceID, FTexturePlatformData*>& Pair : ImageToPlatformDataMap)
+		TMap<mu::FImageId, FTexturePlatformData*>& ImageToPlatformDataMap = OperationData->ImageToPlatformDataMap;
+		for (const TPair<mu::FImageId, FTexturePlatformData*>& Pair : ImageToPlatformDataMap)
 		{
 			delete Pair.Value; // If this is not null then it must mean it hasn't been used, otherwise they would have taken ownership and nulled it
 		}
@@ -3237,20 +3028,6 @@ namespace impl
 		// TODO MTBL-391: Review This hotfix
 		UpdateSkeletalMesh(OperationData);
 
-		// All work is done, release unused textures.
-		if (CustomizableObjectSystemPrivate->bReleaseTexturesImmediately)
-		{
-			FMutableResourceCache& Cache = CustomizableObjectSystemPrivate->GetObjectCache(CustomizableObjectInstance->GetCustomizableObject());
-
-			UCustomizableInstancePrivate* CustomizableObjectInstancePrivateData = CustomizableObjectInstance->GetPrivate();
-			for (FGeneratedTexture& GeneratedTexture : CustomizableObjectInstancePrivateData->TexturesToRelease)
-			{
-				UCustomizableInstancePrivate::ReleaseMutableTexture(GeneratedTexture.Key, Cast<UTexture2D>(GeneratedTexture.Texture), Cache);
-			}
-
-			CustomizableObjectInstancePrivateData->TexturesToRelease.Empty();
-		}
-
 		// End Update
 		FinishUpdateGlobal(OperationData);
 	}
@@ -3361,16 +3138,6 @@ namespace impl
 				}
 			);
 
-
-			// Unlock step
-			//-------------------------------------------------------------
-			if (CustomizableObjectInstance->GetCustomizableObject())
-			{
-				// Unlock the resource cache for the object used by this instance to avoid
-				// the destruction of resources that we may want to reuse.
-				System->GetPrivate()->ClearResourceCacheProtected();
-			}
-
 			// Next Task: Callbacks
 			//-------------------------------------------------------------
 			CustomizableObjectSystemPrivate->MutableTaskGraph.AddGameThreadTask(
@@ -3407,7 +3174,7 @@ namespace impl
 		
 		UCustomizableObjectSystemPrivate* SystemPrivate = System->GetPrivate();
 
-		const UCustomizableObject* CustomizableObject = OperationData->Object.Get(); 
+		UCustomizableObject* CustomizableObject = OperationData->Object.Get(); 
 		if (!CustomizableObject)
 		{
 			OperationData->UpdateResult = EUpdateResult::Error;
@@ -3443,10 +3210,6 @@ namespace impl
 			ObjectInstancePrivateData->RelevantParameters = OperationData->RelevantParametersInProgress;
 		}
 		
-		// Selectively lock the resource cache for the object used by this instance to avoid the destruction of resources that we may want to reuse.
-		// When protecting textures there mustn't be any left from a previous update
-		check(SystemPrivate->ProtectedCachedTextures.Num() == 0);
-		
 		// TODO: If this is the first code that runs after the CO program has finished AND if it's
 		// guaranteed that the next CO program hasn't started yet, we need to call ClearActiveObject
 		// and CancelPendingLoads on SystemPrivate->ExtensionDataStreamer.
@@ -3456,20 +3219,14 @@ namespace impl
 		// If the program aborted, AreAnyLoadsPending may return true, as the program doesn't cancel
 		// its own loads on exit (maybe it should?)
 
-		FMutableResourceCache& Cache = SystemPrivate->GetObjectCache(CustomizableObject);
-
-		SystemPrivate->ProtectedCachedTextures.Reset(Cache.Images.Num());
-		SystemPrivate->ProtectedObjectCachedImages.Reset(Cache.Images.Num());
-
+		FTextureCache& TextureCache = CustomizableObject->GetPrivate()->TextureCache;
+		
 		for (const FInstanceUpdateData::FImage& Image : OperationData->InstanceUpdateData.Images)
 		{
-			FMutableImageCacheKey Key(Image.ImageID, OperationData->MipsToSkip);
-			const TWeakObjectPtr<UTexture2D>* TexturePtr = Cache.Images.Find(Key);
-
-			if (TexturePtr && TexturePtr->Get() && SystemPrivate->TextureHasReferences(Key))
+			if (const UTexture2D* Texture = TextureCache.Get(FTextureCache::FId(Image.ImageID, OperationData->MipsToSkip)))
 			{
-				SystemPrivate->ProtectedCachedTextures.Add(TexturePtr->Get());
-				SystemPrivate->ProtectedObjectCachedImages.Add(Image.ImageID);
+				OperationData->Objects.Add(TStrongObjectPtr(Texture));
+				OperationData->CachedTextures.Add(Image.ImageID);
 			}
 		}
 
@@ -3655,7 +3412,7 @@ namespace impl
 		Operation->bNeverStream = false;
 		Operation->MipsToSkip = 0;
 
-		SystemPrivate->GetMipStreamingConfig(*Instance, Operation->bNeverStream, Operation->MipsToSkip);
+		SystemPrivate->GetMipStreamingConfig(*Instance, Operation->bProgressiveMipStreamingEnabled, Operation->bNeverStream, Operation->MipsToSkip);
 		
 		if (Operation->bLiveUpdateMode && (!Operation->bNeverStream || Operation->MipsToSkip > 0))
 		{
@@ -3686,6 +3443,8 @@ namespace impl
 		}
 		
 		Operation->Model = ObjectPrivate->GetModel().ToSharedRef();
+		Operation->MeshIdRegistry = ObjectPrivate->MeshIdRegistry;
+		Operation->ImageIdRegistry = ObjectPrivate->ImageIdRegistry;
 
 		// Register the current CO into the external resource provider for runtime meshes and images.
 		// It is needed to get access to the bone map when converting meshes at runtime, for example.
@@ -3737,7 +3496,7 @@ namespace impl
 		// Task: Mutable Update and GetMesh
 		//-------------------------------------------------------------
 		Operation->InstanceID = Operation->bLiveUpdateMode ? InstancePrivate->LiveUpdateModeInstanceID : 0;
-		Operation->bUseMeshCache = Object->bEnableMeshCache && !Operation->bLiveUpdateMode && UCustomizableObjectSystem::IsMeshCacheEnabled(true);
+		Operation->bUseMeshCache = Object->bEnableMeshCache && !Operation->bBake && !Operation->bLiveUpdateMode && UCustomizableObjectSystem::IsMeshCacheEnabled(true);
 
 		Operation->bStreamMeshLODs = IsStreamingEnabled(*Object);
 
@@ -3750,7 +3509,7 @@ namespace impl
 			ObjectPrivate->GetLowPriorityTextureNames(Operation->LowPriorityTextures);
 		}
 
-		bool bRequestAllLODs = !System->IsOnlyGenerateRequestedLODsEnabled() ||
+		bool bRequestAllLODs = !Operation->bIsOnlyGenerateRequestedLODsEnabled ||
 			!SystemPrivate->CurrentInstanceLODManagement->IsOnlyGenerateRequestedLODLevelsEnabled();
 
 #if WITH_EDITOR
@@ -3841,7 +3600,13 @@ namespace impl
 				UE_LOG(LogMutable, Verbose, TEXT("Creating Mutable instance with id [%d] "), Operation->InstanceID);
 			}
 
-			Operation->MutableInstance = MutableSystem->BeginUpdate(Operation->InstanceID, Operation->Parameters, Operation->GetCapturedDescriptor().GetState(), mu::FSystem::AllLODs);
+			Operation->MutableInstance = MutableSystem->BeginUpdate(Operation->InstanceID,
+				Operation->Parameters,
+				Operation->MeshIdRegistry,
+				Operation->ImageIdRegistry,
+				Operation->GetCapturedDescriptor().GetState(),
+				mu::FSystem::AllLODs);
+			
 			Operation->NumInstanceComponents = Operation->MutableInstance->GetComponentCount();
 		}
 		
@@ -3854,7 +3619,7 @@ namespace impl
 		for (int32 InstanceComponentIndex = 0; InstanceComponentIndex < Operation->NumInstanceComponents; ++InstanceComponentIndex)
 		{
 			FCustomizableObjectComponentIndex ObjectComponentIndex(Operation->MutableInstance->GetComponentId(InstanceComponentIndex));
-			TArray<mu::FResourceID>* MeshId = Operation->GetMeshDescriptors(ObjectComponentIndex);
+			TArray<mu::FMeshId>* MeshId = Operation->GetMeshDescriptors(ObjectComponentIndex);
 			check(MeshId);
 
 			if (!MeshId)
@@ -3862,7 +3627,7 @@ namespace impl
 				continue;
 			}
 
-			MeshId->Init(MAX_uint64, MAX_MESH_LOD_COUNT);
+			MeshId->Init({}, MAX_MESH_LOD_COUNT);
 
 			const FName& ComponentName = Operation->ComponentNames[ObjectComponentIndex.GetValue()];
 
@@ -3878,9 +3643,9 @@ namespace impl
 
 		if (Operation->bUseMeshCache)
 		{
-			for (const TArray<mu::FResourceID>& MeshId : Operation->GetMeshDescriptors())
+			for (const TArray<mu::FMeshId>& MeshId : Operation->GetMeshDescriptors())
 			{
-				if (USkeletalMesh* CachedMesh = ObjectPrivate->MeshCache.Get(MeshId))
+				if (UCustomizableObjectSkeletalMesh* CachedMesh = ObjectPrivate->MeshCache.Get(MeshId))
 				{
 					Operation->Objects.Emplace(CachedMesh);
 				}
@@ -4193,14 +3958,6 @@ int32 UCustomizableObjectSystem::TickInternal(const bool bBlocking)
 
 				UCustomizableObjectInstance* PendingInstance = PendingInstanceUpdateFound->Context->Instance.Get();
 				check(PendingInstance);
-
-#if WITH_EDITORONLY_DATA
-				// Used for the baking of instances so we do not want it to be possible to run in a package
-				if (PendingInstanceUpdateFound->Context->UpdateSettingsOverride)
-				{
-					PendingInstanceUpdateFound->Context->UpdateSettingsOverride->ApplySettingsOverrides();
-				}
-#endif
 				
 				// Maybe there's a LODUpdate that has the same instance, merge both updates as an optimization
 				FMutableUpdateCandidate* LODUpdateWithSameInstance = RequestedLODUpdates.Find(PendingInstance);
@@ -4538,12 +4295,6 @@ void UCustomizableObjectSystem::AddUncompiledCOWarning(const UCustomizableObject
 }
 
 
-void UCustomizableObjectSystem::SetReleaseMutableTexturesImmediately(bool bReleaseTextures)
-{
-	GetPrivate()->bReleaseTexturesImmediately = bReleaseTextures;
-}
-
-
 void UCustomizableObjectSystem::EnableBenchmark()
 {
 	// Start reporting benchmarking data (log and .csv file)
diff --git a/Engine/Plugins/Mutable/Source/CustomizableObject/Public/MuCO/CustomizableObjectSkeletalMesh.h b/Engine/Plugins/Mutable/Source/CustomizableObject/Public/MuCO/CustomizableObjectSkeletalMesh.h
--- a/Engine/Plugins/Mutable/Source/CustomizableObject/Public/MuCO/CustomizableObjectSkeletalMesh.h
+++ b/Engine/Plugins/Mutable/Source/CustomizableObject/Public/MuCO/CustomizableObjectSkeletalMesh.h
@@ -7,16 +7,13 @@
 #include "MuR/Parameters.h"
 #include "Engine/SkeletalMesh.h"
 #include "CoreMinimal.h"
+#include "MuCO/CustomizableObjectPrivate.h"
+#include "MuR/ResourceId.h"
 
 #include "CustomizableObjectSkeletalMesh.generated.h"
 
 #define UE_API CUSTOMIZABLEOBJECT_API
 
-namespace mu
-{
-	typedef uint64 FResourceID;
-}
-
 class UCustomizableObjectInstance;
 class UCustomizableObject;
 class UCustomizableInstancePrivate;
@@ -36,6 +33,7 @@ class UCustomizableObjectSkeletalMesh : public USkeletalMesh
 
 	friend FCustomizableObjectMeshStreamIn;
 	friend UCustomizableInstancePrivate;
+	friend FMeshCache;
 	
 	UE_API void InitMutableStreamingData(const TSharedRef<FUpdateContextPrivate>& InOperationData, const FName& ComponentName, const int32 FirstLOD, const int32 LODCount);
 
@@ -48,15 +46,19 @@ class UCustomizableObjectSkeletalMesh : public USkeletalMesh
 	TSharedPtr<FModelStreamableBulkData> ModelStreamableBulkData;
 
 	TSharedPtr<mu::FModel, ESPMode::ThreadSafe> Model;
-	
+	TSharedPtr<mu::FMeshIdRegistry> MeshIdRegistry;
+	TSharedPtr<mu::FImageIdRegistry> ImageIdRegistry;
+
 	/** Instance parameters at the time of creation. */
 	TSharedPtr<mu::FParameters> Parameters;
 	int32 State = -1;
 
 	/** Mesh IDs per LOD */
-	TArray<mu::FResourceID> MeshIDs;
+	TArray<mu::FMeshId> MeshIDs;
 	TArray<TArray<uint32>> SurfaceIDs;
 	TArray<TPair<uint32, FName>> SkinWeightProfileIDs;
+
+	FMeshCache::FId MeshCacheId;
 };
 
 #undef UE_API
diff --git a/Engine/Plugins/Mutable/Source/CustomizableObject/Public/MuCO/CustomizableObjectSystem.h b/Engine/Plugins/Mutable/Source/CustomizableObject/Public/MuCO/CustomizableObjectSystem.h
--- a/Engine/Plugins/Mutable/Source/CustomizableObject/Public/MuCO/CustomizableObjectSystem.h
+++ b/Engine/Plugins/Mutable/Source/CustomizableObject/Public/MuCO/CustomizableObjectSystem.h
@@ -9,7 +9,6 @@
 #include "AssetRegistry/AssetData.h"
 #include "MuCO/CustomizableObjectInstance.h"
 #include "MuR/Parameters.h"
-#include "MuR/Image.h"
 #include "MuR/Types.h"
 
 #include "CustomizableObjectSystem.generated.h"
@@ -106,30 +105,6 @@ struct FPendingReleaseMaterialsInfo
 
 
 
-/** Key to identify an image inside a generated mutable runtime instance. */
-struct FMutableImageReference
-{
-	/** Original image ID. Once generated it will be unique. However, future updates of the image may return a different ID for the
-	* same image, if many other resources have been built in the middle. For this reason the rest of the data in the struct is what
-	* must be used to request the additional mips.	*/
-	uint32 ImageID = 0;
-
-	uint32 SurfaceId = 0;
-
-	uint8 LOD = 0;
-	
-	/** Object component index. */
-	uint8 Component = 0;
-
-	uint8 Image = 0;
-
-	uint8 BaseMip = 0;
-
-	TArray<int32> ConstantImagesNeededToGenerate;
-};
-
-
-
 // Opaque representation of a possible registered value for texture parameters.
 struct FCustomizableObjectExternalTexture
 {
@@ -255,12 +230,6 @@ public:
 	UFUNCTION(BlueprintCallable, Category = Status)
 	UE_API int32 GetAverageBuildTime() const;
 
-	// If set to true, Mutable will release Mutable-generated textures immediately when they are no longer used without waiting for GC
-	// IMPORTANT!!! Do NOT keep references to any Mutable generated textures or skeletal meshes if this is enabled,
-	// they are owned by Mutable and will be destroyed without notice
-	UFUNCTION(BlueprintCallable, Category = Status)
-	UE_API void SetReleaseMutableTexturesImmediately(bool bReleaseTextures);
-
 	UE_API bool IsMutableAnimInfoDebuggingEnabled() const;
 
 #if WITH_EDITOR
diff --git a/Engine/Plugins/Mutable/Source/CustomizableObjectEditor/Private/MuCOE/CustomizableObjectInstanceBakingUtils.cpp b/Engine/Plugins/Mutable/Source/CustomizableObjectEditor/Private/MuCOE/CustomizableObjectInstanceBakingUtils.cpp
--- a/Engine/Plugins/Mutable/Source/CustomizableObjectEditor/Private/MuCOE/CustomizableObjectInstanceBakingUtils.cpp
+++ b/Engine/Plugins/Mutable/Source/CustomizableObjectEditor/Private/MuCOE/CustomizableObjectInstanceBakingUtils.cpp
@@ -286,16 +286,19 @@ void ScheduleInstanceUpdateForBaking(UCustomizableObjectInstance& InInstance, FI
 
 	InInstanceUpdateDelegate.AddStatic(&OnInstanceUpdateFinish);
 	
-	// Desired System settings for the update part of a baking operation
-	TSharedPtr<FMutableSystemSettingsOverrides> SystemSettingsOverride =
-		MakeShared<FMutableSystemSettingsOverrides>(
-			false,
-			false,
-			UnrealPixelFormatFunc);
-	
 	// Schedule the update
-	check (InInstance.GetPrivate());
-	InInstance.GetPrivate()->UpdateSkeletalMeshAsyncResult(InInstanceUpdateDelegate,true,true, SystemSettingsOverride);
+	const TSharedRef<FUpdateContextPrivate> Context = MakeShared<FUpdateContextPrivate>(InInstance);
+	Context->bForce = true;
+	Context->bBake = true;
+	Context->bIgnoreCloseDist = true;
+	Context->bForceHighPriority = true;
+	Context->UpdateNativeCallback = InInstanceUpdateDelegate;
+	Context->bProgressiveMipStreamingEnabled = false;
+	Context->bIsOnlyGenerateRequestedLODsEnabled = false;
+	Context->PixelFormatOverride = UnrealPixelFormatFunc;
+
+	UCustomizableObjectSystemPrivate* SystemPrivate = UCustomizableObjectSystem::GetInstance()->GetPrivate();
+	SystemPrivate->EnqueueUpdateSkeletalMesh(Context);
 	
 	// Tell the baking system that an instance for baking will be processed
 	bIsBakeOperationAlreadyScheduled = true;
diff --git a/Engine/Plugins/Mutable/Source/CustomizableObjectEditor/Private/MuCOE/SMutableCodeViewer.cpp b/Engine/Plugins/Mutable/Source/CustomizableObjectEditor/Private/MuCOE/SMutableCodeViewer.cpp
--- a/Engine/Plugins/Mutable/Source/CustomizableObjectEditor/Private/MuCOE/SMutableCodeViewer.cpp
+++ b/Engine/Plugins/Mutable/Source/CustomizableObjectEditor/Private/MuCOE/SMutableCodeViewer.cpp
@@ -293,11 +293,21 @@ void SMutableCodeViewer::ClearSelectedTreeRow() const
 	TreeView->ClearSelection();
 }
 
-void SMutableCodeViewer::SetCurrentModel(const TSharedPtr<mu::FModel, ESPMode::ThreadSafe>& InMutableModel, 
+
+void SMutableCodeViewer::Construct(const FArguments& InArgs,
+	const TSharedPtr<mu::FModel>& InMutableModel,
+	const TSharedPtr<mu::FMeshIdRegistry>& InMeshIdRegistry,
+	const TSharedPtr<mu::FImageIdRegistry>& InImageIdRegistry,
 	const TArray<TSoftObjectPtr<const UTexture>>& InReferencedTextures,
 	const TArray<TSoftObjectPtr<const UStreamableRenderAsset>>& InReferencedMeshes )
 {
+	// Min width allowed for the column. Needed to avoid having issues with the constants space being to small
+	// and then getting too tall on the y axis crashing the UI drawer.
+	constexpr float MinParametersCollWidth = 200;
+
 	MutableModel = InMutableModel;
+	MeshIdRegistry = InMeshIdRegistry;
+	ImageIdRegistry = InImageIdRegistry;
 	ReferencedTextures = InReferencedTextures;
 	ReferencedMeshes = InReferencedMeshes;
 	PreviewParameters = mu::FModel::NewParameters(MutableModel);
@@ -341,18 +351,6 @@ void SMutableCodeViewer::SetCurrentModel(const TSharedPtr<mu::FModel, ESPMode::T
 		check(FoundModelOperationTypeElements.Num())
 		CurrentlySelectedOperationTypeElement = NoneOperationEntry;
 	}
-}
-
-
-void SMutableCodeViewer::Construct(const FArguments& InArgs, const TSharedPtr<mu::FModel, ESPMode::ThreadSafe>& InMutableModel,
-	const TArray<TSoftObjectPtr<const UTexture>>& InReferencedTextures,
-	const TArray<TSoftObjectPtr<const UStreamableRenderAsset>>& InReferencedMeshes )
-{
-	// Min width allowed for the column. Needed to avoid having issues with the constants space being to small
-	// and then getting too tall on the y axis crashing the UI drawer.
-	constexpr float MinParametersCollWidth = 200;
-
-	SetCurrentModel(InMutableModel, InReferencedTextures, InReferencedMeshes);
 	
 	FToolBarBuilder ToolbarBuilder(TSharedPtr<const FUICommandList>(), FMultiBoxCustomization::None, TSharedPtr<FExtender>(), true);
 	ToolbarBuilder.SetLabelVisibility(EVisibility::Visible);
@@ -2975,7 +2973,7 @@ void SMutableCodeViewer::Tick(const FGeometry& AllottedGeometry, const double In
 	ExternalResourceProvider->ReferencedMeshes = ReferencedMeshes;
 	System->SetExternalResourceProvider(ExternalResourceProvider);
 
-	System->GetPrivate()->BeginBuild(MutableModel);
+	System->GetPrivate()->BeginBuild(MutableModel, MeshIdRegistry, ImageIdRegistry);
 
 	switch (OperationDataType)
 	{
diff --git a/Engine/Plugins/Mutable/Source/CustomizableObjectEditor/Private/MuCOE/SMutableCodeViewer.h b/Engine/Plugins/Mutable/Source/CustomizableObjectEditor/Private/MuCOE/SMutableCodeViewer.h
--- a/Engine/Plugins/Mutable/Source/CustomizableObjectEditor/Private/MuCOE/SMutableCodeViewer.h
+++ b/Engine/Plugins/Mutable/Source/CustomizableObjectEditor/Private/MuCOE/SMutableCodeViewer.h
@@ -59,6 +59,8 @@ public:
 	SLATE_END_ARGS()
 
 	void Construct(const FArguments& InArgs, const TSharedPtr<mu::FModel, ESPMode::ThreadSafe>& InMutableModel,
+		const TSharedPtr<mu::FMeshIdRegistry>& InMeshIdRegistry,
+		const TSharedPtr<mu::FImageIdRegistry>& InImageIdRegistry,
 		const TArray<TSoftObjectPtr<const UTexture>>& ReferencedTextures,
 		const TArray<TSoftObjectPtr<const UStreamableRenderAsset>>& ReferencedMeshes);
 
@@ -75,7 +77,9 @@ public:
 private:
 
 	/** The Mutable Model that we are showing. */
-	TSharedPtr<mu::FModel, ESPMode::ThreadSafe> MutableModel;
+	TSharedPtr<mu::FModel> MutableModel;
+	TSharedPtr<mu::FMeshIdRegistry> MeshIdRegistry;
+	TSharedPtr<mu::FImageIdRegistry> ImageIdRegistry;
 	
 	/** Array of external referenced textures in MutableModel, indexed by id. */
 	TArray<TSoftObjectPtr<const UTexture>> ReferencedTextures;
@@ -186,9 +190,6 @@ private:
 	/** Array with all the elements that have been manually expanded by the user */
 	TMap< mu::OP::ADDRESS, TSharedPtr<FMutableCodeTreeElement>> ExpandedElements;
 
-	/** Prepare the widget for the given model. */
-	void SetCurrentModel(const TSharedPtr<mu::FModel, ESPMode::ThreadSafe>&, const TArray<TSoftObjectPtr<const UTexture>>& ReferencedTextures, const TArray<TSoftObjectPtr<const UStreamableRenderAsset>>& ReferencedMeshes);
-
 	/** Before any UI operation generate all the elements that may be navigable over the tree. No children of duplicated 
 	 * addresses will be generated.
 	 */
diff --git a/Engine/Plugins/Mutable/Source/CustomizableObjectEditor/Private/MuCOE/SMutableInstanceViewer.cpp b/Engine/Plugins/Mutable/Source/CustomizableObjectEditor/Private/MuCOE/SMutableInstanceViewer.cpp
--- a/Engine/Plugins/Mutable/Source/CustomizableObjectEditor/Private/MuCOE/SMutableInstanceViewer.cpp
+++ b/Engine/Plugins/Mutable/Source/CustomizableObjectEditor/Private/MuCOE/SMutableInstanceViewer.cpp
@@ -306,12 +306,12 @@ void SMutableInstanceViewer::SetInstance(const TSharedPtr<const mu::FInstance>&
 					LODArray->Add(NewLOD);
 
 					NewLOD->LODIndex = FText::AsNumber(LODIndex);
-					mu::FResourceID MeshId = MutableInstance->GetMeshId(ComponentIndex, LODIndex);
-					NewLOD->MeshId = FText::AsNumber(MeshId);
+					mu::FMeshId MeshId = MutableInstance->GetMeshId(ComponentIndex, LODIndex);
+					NewLOD->MeshId = FText::FromString(MeshId.ToString());
 
 					// Generate the mesh to get some data
 					{
-						mu::OP::ADDRESS MeshRootAddress = mu::GetResourceIDRoot(MeshId);
+						mu::OP::ADDRESS MeshRootAddress = MeshId.GetKey()->Address;
 						TSharedPtr<const mu::FMesh> MutableMesh = System.GetPrivate()->BuildMesh(
 								Model, 
 								Parameters.Get(), 
@@ -321,7 +321,7 @@ void SMutableInstanceViewer::SetInstance(const TSharedPtr<const mu::FInstance>&
 						if (MutableMesh && MutableMesh->IsReference())
 						{
 							NewLOD->MeshId = FText::Format(
-								LOCTEXT("MeshReference", "Id [{0}]  Ref [{1}] "), MeshId, MutableMesh->GetReferencedMesh());
+								LOCTEXT("MeshReference", "Id [{0}]  Ref [{1}] "), FText::FromString(MeshId.ToString()), MutableMesh->GetReferencedMesh());
 						}
 					}
 
diff --git a/Engine/Plugins/Mutable/Source/CustomizableObjectEditor/Private/MuCOE/SMutableObjectViewer.cpp b/Engine/Plugins/Mutable/Source/CustomizableObjectEditor/Private/MuCOE/SMutableObjectViewer.cpp
--- a/Engine/Plugins/Mutable/Source/CustomizableObjectEditor/Private/MuCOE/SMutableObjectViewer.cpp
+++ b/Engine/Plugins/Mutable/Source/CustomizableObjectEditor/Private/MuCOE/SMutableObjectViewer.cpp
@@ -194,7 +194,7 @@ void SMutableObjectViewer::CompileMutableCodePressed()
 
 	FString DataTag = FString::Printf( TEXT("Mutable Code : %s opt %d"), *CompileOptions.TargetPlatform->PlatformName(), CompileOptions.OptimizationLevel );
 
-	TSharedRef<SMutableCodeViewer> CodeViewer = SNew(SMutableCodeViewer, CompileTask->Model, RuntimeTextures, RuntimeMeshes)
+	TSharedRef<SMutableCodeViewer> CodeViewer = SNew(SMutableCodeViewer, CompileTask->Model,  CustomizableObject->GetPrivate()->MeshIdRegistry, CustomizableObject->GetPrivate()->ImageIdRegistry, RuntimeTextures, RuntimeMeshes)
 		.DataTag(DataTag);
 	
 	DebuggerContentsBox->ClearContent();
diff --git a/Engine/Plugins/Mutable/Source/MutableRuntime/Internal/MuR/Image.h b/Engine/Plugins/Mutable/Source/MutableRuntime/Internal/MuR/Image.h
new file mode 100644
--- /dev/null
+++ b/Engine/Plugins/Mutable/Source/MutableRuntime/Internal/MuR/Image.h
@@ -0,0 +1,422 @@
+// Copyright Epic Games, Inc. All Rights Reserved.
+
+#pragma once
+
+#include "MuR/ImageTypes.h"
+
+#include "MuR/MutableMath.h"
+#include "MuR/MutableTrace.h"
+#include "MuR/Ptr.h"
+#include "MuR/RefCounted.h"
+#include "MuR/Serialisation.h"
+#include "Containers/Array.h"
+#include "HAL/LowLevelMemTracker.h"
+#include "HAL/Platform.h"
+#include "Misc/AssertionMacros.h"
+
+#include "MuR/ImageDataStorage.h"
+
+#define UE_API MUTABLERUNTIME_API
+
+namespace mu
+{
+    /** 2D image resource with mipmaps.
+	*/
+    class FImage : public FResource
+	{
+	public:
+
+		//-----------------------------------------------------------------------------------------
+		// Life cycle
+		//-----------------------------------------------------------------------------------------
+
+		UE_API FImage();
+
+		//! Constructor from the size and format.
+		//! It will allocate the memory for the image, but its content will be undefined.
+		//! \param sizeX Width of the image.
+        //! \param sizeY Height of the image.
+        //! \param lods Number of levels of detail (mipmaps) to include in the image, inclduing the
+        //!         base level. It must be a number between 1 and the maximum possible levels, which
+        //!         depends on the image size.
+        //! \param format Pixel format.
+		UE_API FImage(uint32 SizeX, uint32 SizeY, uint32 Lods, EImageFormat Format, EInitializationType InitType);
+
+		/** Create a new empty image that repreents an external resource image. */
+		static UE_API TSharedPtr<FImage> CreateAsReference(uint32 ID, const FImageDesc& Desc, bool bForceLoad);
+
+		//! Serialisation
+		static UE_API void Serialise( const FImage*, FOutputArchive& );
+		static UE_API TSharedPtr<FImage> StaticUnserialise( FInputArchive& );
+
+		// Resource interface
+		UE_API int32 GetDataSize() const override;
+
+		//-----------------------------------------------------------------------------------------
+		// Own interface
+		//-----------------------------------------------------------------------------------------
+
+		/** */
+		UE_API void Init(uint32 SizeX, uint32 SizeY, uint32 Lods, EImageFormat Format, EInitializationType InitType);
+
+		/** Clear the image to black colour. */
+		UE_API void InitToBlack();
+
+		/** Return true if this is a reference to an engine image. */
+		UE_API bool IsReference() const;
+
+		/** If true, this is a reference that must be resolved at compile time. */
+		UE_API bool IsForceLoad() const;
+
+		/** Return the id of the engine referenced texture. Only valid if IsReference. */
+		UE_API uint32 GetReferencedTexture() const;
+
+		//! Return the width of the image.
+        UE_API uint16 GetSizeX() const;
+
+		//! Return the height of the image.
+        UE_API uint16 GetSizeY() const;
+
+		UE_API const FImageSize& GetSize() const;
+
+		//! Return the pixel format of the image.
+		UE_API EImageFormat GetFormat() const;
+
+		//! Return the number of levels of detail (mipmaps) in the texture. The base lavel is also
+		//! counted, so the minimum is 1.
+		UE_API int32 GetLODCount() const;
+
+		//! Return a pointer to a instance-owned buffer where the image pixels are.
+        UE_API const uint8* GetLODData(int32 LODIndex) const;
+        UE_API uint8* GetLODData(int32 LODIndex);
+
+        //! Return the size in bytes of a specific LOD of the image.
+        UE_API int32 GetLODDataSize(int32 LODIndex) const;
+		
+	public:
+
+		// This used to be the data in the private implementation of the image interface
+		//-----------------------------------------------------------------------------------------
+		/** These set of flags are used to cache information for images at runtime. */
+		typedef enum
+		{
+			// Set if the next flag has been calculated and its value is valid.
+			IF_IS_PLAIN_COLOUR_VALID = 1 << 0,
+
+			// If the previous flag is set and this one too, the image is single colour.
+			IF_IS_PLAIN_COLOUR = 1 << 1,
+
+			// If this is set, the image shouldn't be scaled: it's contents is resoultion-dependent.
+			IF_CANNOT_BE_SCALED = 1 << 2,
+
+			// If this is set, the image has an updated relevancy map. This flag is not persisent.
+			IF_HAS_RELEVANCY_MAP = 1 << 3,
+
+			/** If this is set, this is a reference to an external image, and the ReferenceID is valid. */
+			IF_IS_REFERENCE = 1 << 4,
+
+			/** For reference images, this indicates that they should be loaded into full images as soon as they are generated. */
+			IF_IS_FORCELOAD = 1 << 5
+		} EImageFlags;
+
+		/** Persistent flags with some image properties. The meaning will depend of every context. */
+		mutable uint8 Flags = 0;
+
+		/** Non-persistent relevancy map. */
+		uint16 RelevancyMinY = 0;
+		uint16 RelevancyMaxY = 0;
+
+		/** Only valid if the right flags are set, this identifies a referenced image. */
+		uint32 ReferenceID = 0;
+
+		/** Pixel data for all lods. */
+		FImageDataStorage DataStorage;
+
+		// This used to be the methods in the private implementation of the image interface
+		//-----------------------------------------------------------------------------------------
+
+		//! Deep clone.
+		TSharedPtr<FImage> Clone() const
+		{
+			LLM_SCOPE_BYNAME(TEXT("MutableRuntime"));
+			MUTABLE_CPUPROFILER_SCOPE(ImageClone)
+
+			TSharedPtr<FImage> Result = MakeShared<FImage>();
+
+			Result->Flags = Flags;
+			Result->RelevancyMinY = RelevancyMinY;
+			Result->RelevancyMaxY = RelevancyMaxY;
+			Result->DataStorage = DataStorage;
+			Result->ReferenceID = ReferenceID;
+
+			return Result;
+		}
+
+		//! Copy another image.
+		void Copy(const FImage* Other)
+		{
+			LLM_SCOPE_BYNAME(TEXT("MutableRuntime"));
+			MUTABLE_CPUPROFILER_SCOPE(Copy);
+
+			if (Other == this)
+			{
+				return;
+			}
+
+			Flags = Other->Flags;
+			RelevancyMinY = Other->RelevancyMinY;
+			RelevancyMaxY = Other->RelevancyMaxY;
+			DataStorage = Other->DataStorage;
+			ReferenceID = Other->ReferenceID;
+		}
+
+		void CopyMove(FImage* Other)
+		{
+			LLM_SCOPE_BYNAME(TEXT("MutableRuntime"));
+			MUTABLE_CPUPROFILER_SCOPE(CopyMove);
+
+			if (Other == this)
+			{
+				return;
+			}
+
+			Flags = Other->Flags;
+			RelevancyMinY = Other->RelevancyMinY;
+			RelevancyMaxY = Other->RelevancyMaxY;
+			ReferenceID = Other->ReferenceID;
+			
+			DataStorage = MoveTemp(Other->DataStorage);
+		}
+
+
+		//!
+		UE_API void Serialise(FOutputArchive&) const;
+
+		//!
+		UE_API void Unserialise(FInputArchive&);
+
+		//!
+		inline bool operator==(const FImage& Other) const
+		{
+			return 
+				(DataStorage == Other.DataStorage) &&
+				(ReferenceID == Other.ReferenceID);
+		}
+
+		//-----------------------------------------------------------------------------------------
+
+		//! Sample the image and return an RGB colour.
+		UE_API FVector4f Sample(FVector2f Coords) const;
+
+		//! Calculate the size of the image data in bytes, regardless of what is allocated in
+		//! m_data, only using the image descriptions. For non-block-compressed images, it returns
+		//! 0.
+		static UE_API int32 CalculateDataSize(int32 SizeX, int32 SizeY, int32 LodCount, EImageFormat Format);
+
+		//! Calculate the size in pixels of a particular mipmap of this image. The size doesn't
+		//! include pixels necessary for completing blocks in block-compressed formats.
+		UE_API FIntVector2 CalculateMipSize(int32 LOD) const;
+
+		//! Return a pointer to the beginning of the data for a particular mip.
+		UE_API uint8* GetMipData(int32 Mip);
+		UE_API const uint8* GetMipData(int32 Mip) const;
+		
+		//! Return the size of the resident mips. for block-compressed images the result is the same as in CalculateDataSize()
+		//! for non-block-compressed images, the sizes encoded in the image data is used to compute the final size.
+		UE_API int32 GetMipsDataSize() const;
+
+		//! See if all the pixels in the image are equal and return the colour.
+		UE_API bool IsPlainColour(FVector4f& Colour) const;
+
+		//! See if all the pixels in the alpha channel are the max value of the pixel format
+		//! (white).
+		UE_API bool IsFullAlpha() const;
+
+		//! Reduce the LODs of the image to the given amount. Don't do anything if there are already 
+		//! less LODs than specified.
+		UE_API void ReduceLODsTo(int32 NewLODCount);
+		UE_API void ReduceLODs(int32 LODsToSkip);
+
+		//! Calculate the number of mipmaps for a particular image size.
+		static UE_API int32 GetMipmapCount(int32 SizeX, int32 SizeY);
+
+		//! Get the rect inside the image bounding the non-black content of the image.
+		UE_API void GetNonBlackRect(FImageRect& OutRect) const;
+		UE_API void GetNonBlackRect_Reference(FImageRect& OutRect) const;
+	};
+
+	/** This struct contains image operations that may need to allocate and free temporary images to work.
+	* It's purpose is to override the creation, release and clone functions depending on the context.
+	*/
+	struct FImageOperator
+	{
+		/** Common callback used for functions that can create temporary images. */
+		typedef TFunction<TSharedPtr<FImage>(int32, int32, int32, EImageFormat, EInitializationType)> FImageCreateFunc;
+		typedef TFunction<void(TSharedPtr<FImage>&)> FImageReleaseFunc;
+		typedef TFunction<TSharedPtr<FImage>(const FImage*)> FImageCloneFunc;
+
+		FImageCreateFunc CreateImage;
+		FImageReleaseFunc ReleaseImage;
+		FImageCloneFunc CloneImage;
+
+		/** Interface to override the internal mutable image pixel format conversion functions.
+		* Arguments match the FImageOperator::ImagePixelFormat function.
+		*/
+		typedef TFunction<void(bool&, int32, FImage*, const FImage*, int32)> FImagePixelFormatFunc;
+		FImagePixelFormatFunc FormatImageOverride;
+
+		FImageOperator(FImageCreateFunc Create, FImageReleaseFunc Release, FImageCloneFunc Clone, FImagePixelFormatFunc FormatOverride) : CreateImage(Create), ReleaseImage(Release), CloneImage(Clone), FormatImageOverride(FormatOverride) {};
+
+		/** Create an default version for untracked resources. */
+		static FImageOperator GetDefault( const FImagePixelFormatFunc& InFormatOverride )
+		{
+			return FImageOperator(
+				[](int32 x, int32 y, int32 m, EImageFormat f, EInitializationType i) { return MakeShared<FImage>(x, y, m, f, i); },
+				[](TSharedPtr<FImage>& In) {In = nullptr; },
+				[](const FImage* In) { return In->Clone(); },
+				InFormatOverride
+			);
+		}
+
+		/** Convert an image to another pixel format.
+		* Allocates the destination image.
+		* \warning Not all format conversions are implemented.
+		* \param onlyLOD If different than -1, only the specified lod level will be converted in the returned image.
+		* \return nullptr if the conversion failed, usually because not enough memory was allocated in the result. This is only checked for RLE compression.
+		*/
+		MUTABLERUNTIME_API TSharedPtr<FImage> ImagePixelFormat(int32 Quality, const FImage* Base, EImageFormat TargetFormat, int32 OnlyLOD = -1);
+
+		/** Convert an image to another pixel format.
+		* \warning Not all format conversions are implemented.
+		* \param onlyLOD If different than -1, only the specified lod level will be converted in the returned image.
+		* \return false if the conversion failed, usually because not enough memory was allocated in the result. This is only checked for RLE compression.
+		*/
+		MUTABLERUNTIME_API void ImagePixelFormat(bool& bOutSuccess, int32 Quality, FImage* Result, const FImage* Base, int32 OnlyLOD = -1);
+		MUTABLERUNTIME_API void ImagePixelFormat(bool& bOutSuccess, int32 Quality, FImage* Result, const FImage* Base, int32 BeginResultLOD, int32 BeginBaseLOD, int32 NumLODs);
+
+		MUTABLERUNTIME_API TSharedPtr<FImage> ImageSwizzle(EImageFormat Format, const TSharedPtr<const FImage> Sources[], const uint8 Channels[]);
+
+		/** */
+		MUTABLERUNTIME_API TSharedPtr<FImage> ExtractMip(const FImage* From, int32 Mip);
+
+		/** Bilinear filter image resize. */
+		MUTABLERUNTIME_API void ImageResizeLinear(FImage* Dest, int32 ImageCompressionQuality, const FImage* Base);
+
+		/** Fill the image with a plain colour. */
+		void FillColor(FImage* Image, FVector4f Color);
+
+		/** Support struct to keep preallocated data required for some mipmap operations. */
+		struct FScratchImageMipmap
+		{
+			TSharedPtr<FImage> Uncompressed;
+			TSharedPtr<FImage> UncompressedMips;
+			TSharedPtr<FImage> CompressedMips;
+		};
+
+		/** Generate the mipmaps for images.
+		* if bGenerateOnlyTail is true, generates the mips missing from Base to LevelCount and sets
+		* them in Dest (the full chain is spit in two images). Otherwise generate the mips missing
+		* from Base up to LevelCount and append them in Dest to the already generated Base's mips.
+		*/
+		void ImageMipmap(int32 CompressionQuality, FImage* Dest, const FImage* Base,
+			int32 StartLevel, int32 LevelCount,
+			const FMipmapGenerationSettings&, bool bGenerateOnlyTail = false);
+
+		/** Mipmap separating the worst case treatment in 3 steps to manage allocations of temp data. */
+		void ImageMipmap_PrepareScratch(const FImage* DestImage, int32 StartLevel, int32 LevelCount, FScratchImageMipmap&);
+		void ImageMipmap(FImageOperator::FScratchImageMipmap&, int32 CompressionQuality, FImage* Dest, const FImage* Base,
+			int32 LevelStart, int32 LevelCount,
+			const FMipmapGenerationSettings&, bool bGenerateOnlyTail = false);
+		void ImageMipmap_ReleaseScratch(FScratchImageMipmap&);
+
+		/** */
+		MUTABLERUNTIME_API bool ImageCrop(TSharedPtr<FImage>& OutCropped, int32 CompressionQuality, const FImage* Base, const box<FIntVector2>& Rect);
+
+		/** */
+		void ImageCompose(FImage* Base, const FImage* Block, const box<FIntVector2>& Rect);
+
+	};
+
+	/** Update all the mipmaps in the image from the data in the base one.
+	* Only the mipmaps already existing in the image are updated.
+	*/
+	MUTABLERUNTIME_API void ImageMipmapInPlace(int32 CompressionQuality, FImage* Base, const FMipmapGenerationSettings&);
+
+	/** */
+	MUTABLERUNTIME_API void ImageSwizzle(FImage* Result, const TSharedPtr<const FImage> Sources[], const uint8 Channels[]);
+
+	MUTABLERUNTIME_API inline EImageFormat GetUncompressedFormat(EImageFormat Format)
+	{
+		check(Format < EImageFormat::Count);
+
+		EImageFormat Result = Format;
+
+		switch (Result)
+		{
+		case EImageFormat::L_UBitRLE: Result = EImageFormat::L_UByte; break;
+		case EImageFormat::L_UByteRLE: Result = EImageFormat::L_UByte; break;
+		case EImageFormat::RGB_UByteRLE: Result = EImageFormat::RGB_UByte; break;
+        case EImageFormat::RGBA_UByteRLE: Result = EImageFormat::RGBA_UByte; break;
+        case EImageFormat::BC1: Result = EImageFormat::RGBA_UByte; break;
+        case EImageFormat::BC2: Result = EImageFormat::RGBA_UByte; break;
+        case EImageFormat::BC3: Result = EImageFormat::RGBA_UByte; break;
+        case EImageFormat::BC4: Result = EImageFormat::L_UByte; break;
+        case EImageFormat::BC5: Result = EImageFormat::RGB_UByte; break;
+		case EImageFormat::ASTC_4x4_RGB_LDR: Result = EImageFormat::RGB_UByte; break;
+		case EImageFormat::ASTC_4x4_RGBA_LDR: Result = EImageFormat::RGBA_UByte; break;
+		case EImageFormat::ASTC_4x4_RG_LDR: Result = EImageFormat::RGB_UByte; break;
+		case EImageFormat::ASTC_6x6_RGB_LDR: Result = EImageFormat::RGB_UByte; break;
+		case EImageFormat::ASTC_6x6_RGBA_LDR: Result = EImageFormat::RGBA_UByte; break;
+		case EImageFormat::ASTC_6x6_RG_LDR: Result = EImageFormat::RGB_UByte; break;
+		case EImageFormat::ASTC_8x8_RGB_LDR: Result = EImageFormat::RGB_UByte; break;
+		case EImageFormat::ASTC_8x8_RGBA_LDR: Result = EImageFormat::RGBA_UByte; break;
+		case EImageFormat::ASTC_8x8_RG_LDR: Result = EImageFormat::RGB_UByte; break;
+		case EImageFormat::ASTC_10x10_RGB_LDR: Result = EImageFormat::RGB_UByte; break;
+		case EImageFormat::ASTC_10x10_RGBA_LDR: Result = EImageFormat::RGBA_UByte; break;
+		case EImageFormat::ASTC_10x10_RG_LDR: Result = EImageFormat::RGB_UByte; break;
+		case EImageFormat::ASTC_12x12_RGB_LDR: Result = EImageFormat::RGB_UByte; break;
+		case EImageFormat::ASTC_12x12_RGBA_LDR: Result = EImageFormat::RGBA_UByte; break;
+		case EImageFormat::ASTC_12x12_RG_LDR: Result = EImageFormat::RGB_UByte; break;
+		default: break;
+		}
+
+		return Result;
+	}
+
+	/** */
+	template<class T>
+	TSharedPtr<T> CloneOrTakeOver(const TSharedPtr<T>& Source)
+	{
+		TSharedPtr<T> Result;
+		if (Source.IsUnique())
+		{
+			Result = ConstCastSharedPtr<T>(Source);
+		}
+		else
+		{
+			Result = Source->Clone();
+		}
+
+		return Result;
+	}
+
+	/** */
+	template<class T>
+	TSharedPtr<T> CloneOrTakeOver(const TSharedPtr<const T>& Source)
+	{
+		TSharedPtr<T> Result;
+		if (Source.IsUnique())
+		{
+			Result = ConstCastSharedPtr<T>(Source);
+		}
+		else
+		{
+			Result = Source->Clone();
+		}
+
+		return Result;
+	}
+
+}
+
+#undef UE_API
diff --git a/Engine/Plugins/Mutable/Source/MutableRuntime/Internal/MuR/ImageDataStorage.h b/Engine/Plugins/Mutable/Source/MutableRuntime/Internal/MuR/ImageDataStorage.h
new file mode 100644
--- /dev/null
+++ b/Engine/Plugins/Mutable/Source/MutableRuntime/Internal/MuR/ImageDataStorage.h
@@ -0,0 +1,211 @@
+// Copyright Epic Games, Inc. All Rights Reserved.
+
+#pragma once
+
+#include "MuR/ImageTypes.h"
+#include "MuR/MemoryTrackingAllocationPolicy.h"
+
+#include "Containers/Array.h"
+#include "Containers/ArrayView.h"
+#include "Containers/StaticArray.h"
+
+#include <atomic>
+
+#define UE_API MUTABLERUNTIME_API
+
+namespace mu::MemoryCounters
+{
+	struct FImageMemoryCounter
+	{
+		static UE_API std::atomic<SSIZE_T>& Get();
+	};
+}
+
+
+namespace mu
+{
+	using FImageArray = TArray<uint8, FDefaultMemoryTrackingAllocator<MemoryCounters::FImageMemoryCounter>>;
+}
+
+namespace mu::FImageDataStorageInternal
+{
+	inline int32 ComputeNumLODsForSize(FImageSize Size)
+	{
+		return static_cast<int32>(FMath::CeilLogTwo(FMath::Max(Size.X, Size.Y))) + 1;
+	}
+}
+
+namespace mu
+{
+class FImageDataStorage
+{
+public:
+	TArray<FImageArray, TInlineAllocator<1>> Buffers;
+	
+	FImageSize ImageSize = FImageSize(0, 0);
+	EImageFormat ImageFormat = EImageFormat::None;
+	uint8 NumLODs = 0;
+
+	static constexpr int32 NumLODsInCompactedTail = 7;
+	// This is needed for images its size cannot be known from dimensions and format, e.g., RLE compressed formats.
+	// It stores the offset to the end of the LOD so that CompactedTailOffsets[LOD] - CompactedTailOffsets[LOD - 1]
+	// is the size of LOD.
+	TStaticArray<uint16, NumLODsInCompactedTail> CompactedTailOffsets = MakeUniformStaticArray<uint16, NumLODsInCompactedTail>(0);
+
+public:
+
+	MUTABLERUNTIME_API FImageDataStorage();
+	MUTABLERUNTIME_API FImageDataStorage(const FImageDesc& Desc);
+
+	MUTABLERUNTIME_API FImageDataStorage(const FImageDataStorage& Other);
+	MUTABLERUNTIME_API FImageDataStorage& operator=(const FImageDataStorage& Other);
+	
+	FImageDataStorage(FImageDataStorage&& Other) = default;
+	FImageDataStorage& operator=(FImageDataStorage&& Other) = default;
+
+	MUTABLERUNTIME_API bool operator==(const FImageDataStorage& Other) const;
+
+	inline FImageDesc MakeImageDesc() const
+	{
+		return FImageDesc(ImageSize, ImageFormat, NumLODs);
+	}
+
+	inline void InitVoid(const FImageDesc& Desc) 
+	{
+		Buffers.Empty();	
+		CompactedTailOffsets[0] = TNumericLimits<uint16>::Max();
+
+		ImageSize = Desc.m_size;
+		ImageFormat = Desc.m_format;
+		NumLODs = Desc.m_lods;
+	}
+
+	inline bool IsVoid() const
+	{
+		return Buffers.IsEmpty() && CompactedTailOffsets[0] == TNumericLimits<uint16>::Max();
+	}
+
+	MUTABLERUNTIME_API void InitInternalArray(int32 Index, int32 Size, EInitializationType InitType);
+	MUTABLERUNTIME_API void Init(const FImageDesc& ImageDesc, EInitializationType InitType);
+
+	MUTABLERUNTIME_API void CopyFrom(const FImageDataStorage& Other);
+
+	inline int32 GetNumLODs() const
+	{
+		return NumLODs;
+	}
+
+	/**
+	* Get a const view to the data containing the LODIndex.
+	*/
+	MUTABLERUNTIME_API TArrayView<const uint8> GetLOD(int32 LODIndex) const;
+
+	/**
+	* Get a view to the data containing the LODIndex.
+	*/
+	MUTABLERUNTIME_API TArrayView<uint8> GetLOD(int32 LODIndex);
+
+	/**
+	 * Change the number of lods in the image. If InNumLODs is smaller than the current lod count, 
+	 * LODs are dropped from the tail up.
+	 */
+	MUTABLERUNTIME_API void SetNumLODs(int32 InNumLODs, EInitializationType InitType = EInitializationType::NotInitialized);
+
+	/**
+	 * Remove NumLODsToDrop starting form LOD 0. The resulting image will be smaller. 
+	 */
+	MUTABLERUNTIME_API void DropLODs(int32 NumLODsToDrop);
+
+	/**
+	 * Resizes LODIndex Data to NewSizeBytes
+	 */
+	MUTABLERUNTIME_API void ResizeLOD(int32 LODIndex, int32 NewSizeBytes, EAllowShrinking AllowShrinking = EAllowShrinking::Default);
+
+	inline const FImageArray& GetInternalArray(int32 LODIndex) const
+	{	
+		check(!IsVoid());
+		check(LODIndex < NumLODs);
+		return Buffers[FMath::Min(LODIndex, ComputeFirstCompactedTailLOD())];
+	}
+
+	inline FImageArray& GetInternalArray(int32 LODIndex)
+	{	
+		check(!IsVoid());
+		check(LODIndex < NumLODs);
+		return Buffers[FMath::Min(LODIndex, ComputeFirstCompactedTailLOD())];
+	}
+
+	inline int32 ComputeFirstCompactedTailLOD() const
+	{
+		using namespace FImageDataStorageInternal;
+		return FMath::Max(0, ComputeNumLODsForSize(ImageSize) - NumLODsInCompactedTail); 
+	}
+
+	/**
+	 * Stateless iteration in batches.
+	 */
+	MUTABLERUNTIME_API int32 GetNumBatches(int32 BatchSizeInElems, int32 BatchElemSizeInBytes) const;
+	MUTABLERUNTIME_API TArrayView<const uint8> GetBatch(int32 BatchId, int32 BatchSizeInElems, int32 BatchElemSizeInBytes) const;
+	MUTABLERUNTIME_API TArrayView<uint8> GetBatch(int32 BatchId, int32 BatchSizeInElems, int32 BatchElemSizeInBytes);
+
+	/**
+	 * Returns the number of batches GetBatchFirtsLODOffet will admit and that will cover
+	 * the first LOD buffer with OffsetInBytes removed from the front. 
+	 * 
+	 * A batch cannot be larger than BatchElemsSize. The last batch of a buffer will be smaller if 
+	 * BatchSizeInElems*BatchElemSizeInBytes is not multiple of the buffer size.
+	 *
+	 **/
+	MUTABLERUNTIME_API int32 GetNumBatchesFirstLODOffset(int32 BatchSizeInElems, int32 BatchElemSizeInBytes, int32 OffsetInBytes) const;
+
+	/**
+	 * Returns a non modifiable view to the portion of the first LOD buffer with OffsetInBytes for the BatchId 
+	 *
+	 * A batch cannot be larger than BatchElemsSize. The last batch of a buffer will be smaller if 
+	 * BatchSizeInElems*BatchElemSizeInBytes is not multiple of the buffer size.
+	 *
+	 */
+	MUTABLERUNTIME_API TArrayView<const uint8> GetBatchFirstLODOffset(int32 BatchId, int32 BatchSizeInElems, int32 BatchElemSizeInBytes, int32 OffsetInBytes) const;
+
+	/**
+	 * Non const version of GetBatchFirstLODOffset(). 
+	 */
+	MUTABLERUNTIME_API TArrayView<uint8> GetBatchFirstLODOffset(int32 BatchId, int32 BatchSizeInElems, int32 BatchElemSizeInBytes, int32 OffsetInBytes = 0);
+
+	/**
+	 * Returns the number of batches GetBatchLODRange will admit and that will cover
+	 * the [LODBegin, LODEnd) range.
+	 *
+	 **/
+	MUTABLERUNTIME_API int32 GetNumBatchesLODRange(int32 BatchSizeInElems, int32 BatchElemSizeInBytes, int32 LODBegin, int32 LODEnd) const;
+
+	/**
+	 * Returns a non modifiable view to the portion of the [LODBegin, LODEnd) for the BatchId. 
+	 *
+	 * A batch cannot be larger than BatchElemsSize. The last batch of a buffer will be smaller if 
+	 * BatchSizeInElems*BatchElemSizeInBytes is not multiple of the buffer size.
+	 *
+	 */
+	MUTABLERUNTIME_API TArrayView<const uint8> GetBatchLODRange(int32 BatchId, int32 BatchSizeInElems, int32 BatchElemSizeInBytes, int32 LODBegin, int32 LODEnd) const;
+	
+	/**
+	 * Non const version of GetBatchLODRange().
+	 */
+	MUTABLERUNTIME_API TArrayView<uint8> GetBatchLODRange(int32 BatchId, int32 BatchSizeInElems, int32 BatchElemSizeInBytes, int32 LODBegin, int32 LODEnd);
+
+	MUTABLERUNTIME_API int32 GetAllocatedSize() const;
+	MUTABLERUNTIME_API int32 GetDataSize() const;
+
+	/**
+	 * Returns true if all buffers are empty. This can be true after initialization for image formats that report 
+	 * BytesPerBlock 0, e.g., RLE formats. If not initialized will also return true.
+	 */
+	MUTABLERUNTIME_API bool IsEmpty() const;
+
+	MUTABLERUNTIME_API void Serialise(FOutputArchive& Arch) const;
+	MUTABLERUNTIME_API void Unserialise(FInputArchive& Arch);
+};
+
+}
+
+#undef UE_API
diff --git a/Engine/Plugins/Mutable/Source/MutableRuntime/Internal/MuR/MemoryCounters.h b/Engine/Plugins/Mutable/Source/MutableRuntime/Internal/MuR/MemoryCounters.h
new file mode 100644
--- /dev/null
+++ b/Engine/Plugins/Mutable/Source/MutableRuntime/Internal/MuR/MemoryCounters.h
@@ -0,0 +1,23 @@
+﻿// Copyright Epic Games, Inc. All Rights Reserved.
+
+#pragma once
+
+#include "MuR/MemoryTrackingAllocationPolicy.h"
+
+
+namespace mu::MemoryCounters
+{
+	struct FInternalMemoryCounter
+	{
+		static MUTABLERUNTIME_API std::atomic<SSIZE_T>& Get();
+	};
+}
+
+
+namespace mu
+{
+	using AllocType = FDefaultMemoryTrackingAllocator<MemoryCounters::FInternalMemoryCounter>;
+
+	template<class Type, class Alloc = AllocType >
+	using TMemoryTrackedArray = TArray<Type, Alloc>;
+}
\ No newline at end of file
diff --git a/Engine/Plugins/Mutable/Source/MutableRuntime/Internal/MuR/MemoryTrackingAllocationPolicy.h b/Engine/Plugins/Mutable/Source/MutableRuntime/Internal/MuR/MemoryTrackingAllocationPolicy.h
new file mode 100644
--- /dev/null
+++ b/Engine/Plugins/Mutable/Source/MutableRuntime/Internal/MuR/MemoryTrackingAllocationPolicy.h
@@ -0,0 +1,187 @@
+// Copyright Epic Games, Inc. All Rights Reserved.
+
+#pragma once
+
+#include "Containers/ContainerAllocationPolicies.h"
+#include "Containers/ContainersFwd.h"
+#include "MuR/MemoryTrackingUtils.h"
+
+#include <atomic>
+#include <type_traits>
+
+namespace mu
+{
+	
+	
+	// CounterType is expected to be of the following form. 
+	//struct FCounterTypeName
+	//{
+	//	static std::atomic<SSIZE_T>& Get()
+	//	{
+	//		static std::atomic<SSIZE_T> Counter{0};
+	//		return Counter;
+	//	}
+	//};
+
+
+	template<typename BaseAlloc, typename CounterType>
+	class TMemoryTrackingAllocatorWrapper
+	{
+		static_assert(std::is_same_v<decltype(CounterType::Get().load()), SSIZE_T>, "CounterType must be signed.");
+
+	public:
+		using SizeType = typename BaseAlloc::SizeType;
+
+		enum { NeedsElementType = BaseAlloc::NeedsElementType };
+		enum { RequireRangeCheck = BaseAlloc::RequireRangeCheck };
+
+		/** 
+		 * ForAnyElementType is privately inherited from the wrapped allocator ForAnyElementType so 
+		 * the base class members must be explicitly defined to compile. This way if any new method 
+		 * is added to the allocator interface, it forces its addition here. This is useful in case 
+		 * the new method needs to do some memory tracking, otherwise a simple using declaration may 
+		 * suffice.
+		 */
+		class ForAnyElementType : private BaseAlloc::ForAnyElementType 
+		{
+		public:
+			ForAnyElementType() 
+			{
+			}
+
+			/** Destructor. */
+			inline ~ForAnyElementType()
+			{
+				CounterType::Get().fetch_sub(AllocSize, std::memory_order_relaxed);
+
+#if UE_MUTABLE_TRACK_ALLOCATOR_MEMORY_PEAK
+				FGlobalMemoryCounter::Update(-AllocSize);
+#endif
+
+				AllocSize = 0;
+			}
+
+			ForAnyElementType(const ForAnyElementType&) = delete;
+			ForAnyElementType& operator=(const ForAnyElementType&) = delete;
+
+			inline void MoveToEmpty(ForAnyElementType& Other)
+			{
+				BaseAlloc::ForAnyElementType::MoveToEmpty(Other);
+
+				CounterType::Get().fetch_sub(AllocSize, std::memory_order_relaxed);
+
+#if UE_MUTABLE_TRACK_ALLOCATOR_MEMORY_PEAK
+				FGlobalMemoryCounter::Update(-AllocSize);
+#endif
+
+				AllocSize = Other.AllocSize;
+				Other.AllocSize = 0;
+			}
+
+			inline void ResizeAllocation(
+				SizeType CurrentNum,
+				SizeType NewMax,
+				SIZE_T NumBytesPerElement
+			)
+			{
+				BaseAlloc::ForAnyElementType::ResizeAllocation(CurrentNum, NewMax, NumBytesPerElement);
+
+				const SSIZE_T AllocatedSize = (SSIZE_T)BaseAlloc::ForAnyElementType::GetAllocatedSize(NewMax, NumBytesPerElement); 
+				const SSIZE_T Differential = AllocatedSize - AllocSize;
+				const SSIZE_T PrevCounterValue = CounterType::Get().fetch_add(Differential, std::memory_order_relaxed);
+				check(PrevCounterValue >= AllocSize);
+
+#if UE_MUTABLE_TRACK_ALLOCATOR_MEMORY_PEAK
+				FGlobalMemoryCounter::Update(Differential);
+#endif
+
+				AllocSize = AllocatedSize;
+			}
+
+			template <
+				class T = BaseAlloc
+				UE_REQUIRES(TAllocatorTraits<T>::SupportsElementAlignment)
+			>
+			inline void ResizeAllocation(
+				SizeType CurrentNum,
+				SizeType NewMax,
+				SIZE_T NumBytesPerElement,
+				uint32 AlignmentOfElement
+			)
+			{
+				BaseAlloc::ForAnyElementType::ResizeAllocation(CurrentNum, NewMax, NumBytesPerElement, AlignmentOfElement);
+
+				const SSIZE_T AllocatedSize = (SSIZE_T)BaseAlloc::ForAnyElementType::GetAllocatedSize(NewMax, NumBytesPerElement); 
+				const SSIZE_T Differential = AllocatedSize - AllocSize;
+				const SSIZE_T PrevCounterValue = CounterType::Get().fetch_add(Differential, std::memory_order_relaxed);
+				check(PrevCounterValue >= AllocSize);
+
+#if UE_MUTABLE_TRACK_ALLOCATOR_MEMORY_PEAK
+				FGlobalMemoryCounter::Update(Differential);
+#endif
+
+				AllocSize = AllocatedSize;
+			}
+
+			// Explicitly incorporate pass-through base allocator member functions.
+			using BaseAlloc::ForAnyElementType::GetAllocation;
+			using BaseAlloc::ForAnyElementType::CalculateSlackReserve;
+			using BaseAlloc::ForAnyElementType::CalculateSlackShrink;
+			using BaseAlloc::ForAnyElementType::CalculateSlackGrow;
+			using BaseAlloc::ForAnyElementType::GetAllocatedSize;
+			using BaseAlloc::ForAnyElementType::HasAllocation;
+			using BaseAlloc::ForAnyElementType::GetInitialCapacity;
+
+#if UE_ENABLE_ARRAY_SLACK_TRACKING
+			using BaseAlloc::ForAnyElementType::SlackTrackerLogNum;
+#endif
+
+		private:
+			SSIZE_T AllocSize = 0;
+		};
+
+		template<typename ElementType>
+		class ForElementType : public ForAnyElementType
+		{
+		public:
+			ForElementType()
+			{
+			}
+
+			inline ElementType* GetAllocation() const
+			{
+				return (ElementType*)ForAnyElementType::GetAllocation();
+			}
+		};
+
+	};
+
+
+	/** Default memory tracking allocators needed for TArray and TMap. */
+
+	template<typename CounterType>
+	using FDefaultMemoryTrackingAllocator = TMemoryTrackingAllocatorWrapper<FDefaultAllocator, CounterType>;
+
+	template<typename CounterType>
+	using FDefaultMemoryTrackingBitArrayAllocator = TInlineAllocator<4, FDefaultMemoryTrackingAllocator<CounterType>>;
+
+	template<typename CounterType>
+	using FDefaultMemoryTrackingSparceArrayAllocator = TSparseArrayAllocator<
+		FDefaultMemoryTrackingAllocator<CounterType>,
+		FDefaultMemoryTrackingBitArrayAllocator<CounterType>>;
+		
+	template<typename CounterType>
+	using FDefaultMemoryTrackingSetAllocator = TSetAllocator<
+		FDefaultMemoryTrackingSparceArrayAllocator<CounterType>,
+		TInlineAllocator<1, FDefaultMemoryTrackingAllocator<CounterType>>,
+		DEFAULT_NUMBER_OF_ELEMENTS_PER_HASH_BUCKET,
+		DEFAULT_BASE_NUMBER_OF_HASH_BUCKETS,
+		DEFAULT_MIN_NUMBER_OF_HASHED_ELEMENTS>;
+}
+
+template<typename BaseAlloc, typename Counter>
+struct TAllocatorTraits<mu::TMemoryTrackingAllocatorWrapper<BaseAlloc, Counter>> : public TAllocatorTraits<BaseAlloc>
+{
+	enum { SupportsMoveFromOtherAllocator = false };
+};
+
diff --git a/Engine/Plugins/Mutable/Source/MutableRuntime/Internal/MuR/MemoryTrackingUtils.h b/Engine/Plugins/Mutable/Source/MutableRuntime/Internal/MuR/MemoryTrackingUtils.h
new file mode 100644
--- /dev/null
+++ b/Engine/Plugins/Mutable/Source/MutableRuntime/Internal/MuR/MemoryTrackingUtils.h
@@ -0,0 +1,55 @@
+// Copyright Epic Games, Inc. All Rights Reserved.
+
+#pragma once
+
+#include "Misc/Build.h"
+#include "HAL/Platform.h"
+#include "HAL/CriticalSection.h"
+
+#define UE_API MUTABLERUNTIME_API
+
+#ifndef UE_MUTABLE_TRACK_ALLOCATOR_MEMORY_PEAK
+	#if UE_BUILD_DEBUG || UE_BUILD_DEVELOPMENT || UE_BUILD_TEST
+		#define UE_MUTABLE_TRACK_ALLOCATOR_MEMORY_PEAK 1
+	#else
+		#define UE_MUTABLE_TRACK_ALLOCATOR_MEMORY_PEAK 0
+	#endif
+#endif
+
+namespace mu
+{
+#if UE_MUTABLE_TRACK_ALLOCATOR_MEMORY_PEAK 
+	struct FGlobalMemoryCounter
+	{
+	private:
+		 static inline SSIZE_T AbsoluteCounter   { 0 };
+		 static inline SSIZE_T AbsolutePeakValue { 0 };
+		 static inline SSIZE_T Counter           { 0 };
+		 static inline SSIZE_T PeakValue         { 0 };
+		
+		 static inline FCriticalSection Mutex {};
+
+	public:
+		 static UE_API void Update(SSIZE_T Differential);
+		 static UE_API void Zero();
+		 static UE_API void Restore();
+		 static UE_API SSIZE_T GetPeak();
+		 static UE_API SSIZE_T GetCounter();
+		 static UE_API SSIZE_T GetAbsolutePeak();
+		 static UE_API SSIZE_T GetAbsoluteCounter();
+	};
+#else
+	struct FGlobalMemoryCounter
+	{
+		 //static void Update(SSIZE_T Differential);
+		 static UE_API void Zero();
+		 static UE_API void Restore();
+		 static UE_API SSIZE_T GetPeak();
+		 static UE_API SSIZE_T GetCounter();
+		 static UE_API SSIZE_T GetAbsolutePeak();
+		 static UE_API SSIZE_T GetAbsoluteCounter();
+	};
+#endif
+}
+
+#undef UE_API
diff --git a/Engine/Plugins/Mutable/Source/MutableRuntime/Internal/MuR/Registry.h b/Engine/Plugins/Mutable/Source/MutableRuntime/Internal/MuR/Registry.h
new file mode 100644
--- /dev/null
+++ b/Engine/Plugins/Mutable/Source/MutableRuntime/Internal/MuR/Registry.h
@@ -0,0 +1,232 @@
+﻿// Copyright Epic Games, Inc. All Rights Reserved.
+
+#pragma once
+
+#include "Chaos/Character/CharacterGroundConstraintContainer.h"
+#include "Chaos/Deformable/MuscleActivationConstraints.h"
+#include "Templates/SharedPointer.h"
+#include "Misc/ScopeLock.h"
+
+
+template <typename KeyType, typename ValueType>
+class TRegistry : public TSharedFromThis<TRegistry<KeyType, ValueType>>
+{
+public:
+	class FHandle
+	{
+		friend TRegistry;
+	
+	public:
+		FHandle() = default;
+
+	private:
+		FHandle(uint32 InId, const TWeakPtr<TRegistry>& InWeakRegistry)
+		{
+			Id = InId;
+			WeakRegistry = InWeakRegistry;
+		}
+
+	public:
+		FHandle(const FHandle& Other)
+		{
+			*this = Other;
+		}
+
+		FHandle(FHandle&& Other)
+		{
+			*this = MoveTemp(Other);
+		}
+
+		operator bool() const
+		{
+			return Id != ID_NONE && WeakRegistry.IsValid();
+		}
+		
+		FHandle& operator=(const FHandle& Other)
+		{
+			Id = Other.Id;
+			WeakRegistry = Other.WeakRegistry;
+
+			if (TSharedPtr<TRegistry> Registry = WeakRegistry.Pin())
+			{
+				Registry->Increment(Id);
+			}
+			
+			return *this;
+		}
+
+		FHandle& operator=(FHandle&& Other)
+		{
+			Id = Other.Id;
+			Other.Id = ID_NONE;
+
+			WeakRegistry = Other.WeakRegistry;
+			Other.WeakRegistry = nullptr;
+			
+			return *this;
+		}
+		
+		bool operator==(const FHandle& Other) const
+		{
+			return Id == Other.Id && WeakRegistry == Other.WeakRegistry;
+		}
+		
+		~FHandle()
+		{
+			if (TSharedPtr<TRegistry> Registry = WeakRegistry.Pin())
+			{
+				Registry->Decrement(Id);
+			}
+		}
+
+		const KeyType* GetKey() const
+		{
+			if (TSharedPtr<TRegistry> Registry = WeakRegistry.Pin())
+			{
+				return &Registry->GetKey(Id);
+			}
+			else
+			{
+				return nullptr;
+			}
+		}
+
+		const ValueType* GetValue() const
+		{
+			if (TSharedPtr<TRegistry> Registry = WeakRegistry.Pin())
+			{
+				return &Registry->GetValue(Id);
+			}
+			else
+			{
+				return nullptr;
+			}
+		}
+
+		
+		FString ToString() const
+		{
+			return FString::FromInt(Id);
+		}
+
+		friend uint32 GetTypeHash(const FHandle& Handle)
+		{
+			return HashCombine(Handle.Id, GetTypeHash(Handle.WeakRegistry));
+		}
+
+	private:
+		uint32 Id = ID_NONE;
+
+		TWeakPtr<TRegistry> WeakRegistry;
+	};
+
+private:
+	struct FKeyContainer
+	{
+		TSharedRef<KeyType> Key;
+		
+		bool operator==(const FKeyContainer& Other) const
+		{
+			return Key.Get() == Other.Key.Get();
+		}
+
+		bool operator==(const KeyType& Other) const
+		{
+			return Key.Get() == Other;
+		}
+		
+		friend uint32 GetTypeHash(const FKeyContainer& KeyContainer)
+		{
+			return GetTypeHash(KeyContainer.Key.Get());
+		}
+	};
+	
+	struct FEntry
+	{
+		uint32 RefCount = 0;
+
+		FKeyContainer KeyContainer;
+
+		ValueType Value;
+	};
+	
+public:
+	FHandle Add(const KeyType& Key, const ValueType& Value)
+	{
+		FScopeLock Lock(&CriticalSection);
+
+		if (uint32* Result = KeyIndex.FindByHash(GetTypeHash(Key), Key))
+		{
+			const uint32 Id = *Result;
+			
+			++Entries[Id].RefCount;
+			return FHandle(Id, TSharedFromThis<TRegistry>::AsShared());
+		}
+
+		uint32 Id = GlobalId++;
+
+		FKeyContainer KeyContainer { MakeShared<KeyType>(Key) };
+		
+		Entries.Emplace(Id, FEntry(1, KeyContainer, Value));
+		KeyIndex.Emplace(KeyContainer, Id);
+		
+		return FHandle(Id, TSharedFromThis<TRegistry>::AsShared());
+	}
+
+	ValueType* Find(const KeyType& Key)
+	{
+		if (uint32* Result = KeyIndex.FindByHash(GetTypeHash(Key), Key))
+		{
+			const uint32 Id = *Result;
+			
+			return &Entries[Id].Value;
+		}
+		else
+		{
+			return nullptr;
+		}
+	}
+	
+private:
+	const KeyType& GetKey(int32 Id) const
+	{
+		return Entries[Id].KeyContainer.Key.Get();
+	}
+
+	const ValueType& GetValue(int32 Id) const
+	{
+		return Entries[Id].Value;
+	}
+	
+	void Increment(uint32 Id)
+	{
+		FScopeLock Lock(&CriticalSection);
+
+		FEntry& Entry = Entries[Id];
+		++Entry.RefCount;
+	}
+
+	void Decrement(uint32 Id)
+	{
+		FScopeLock Lock(&CriticalSection);
+
+		FEntry& Entry = Entries[Id];
+		--Entry.RefCount;
+
+		if (Entry.RefCount == 0)
+		{
+			KeyIndex.Remove(Entry.KeyContainer);
+			Entries.Remove(Id);
+		}
+	}
+	
+	static constexpr uint32 ID_NONE = TNumericLimits<uint32>::Max();
+
+	/** Global incremental Id counter. */
+	uint32 GlobalId = 0;
+	
+	TMap<uint32, FEntry> Entries;
+	TMap<FKeyContainer, uint32> KeyIndex;
+
+	FCriticalSection CriticalSection;
+};
\ No newline at end of file
diff --git a/Engine/Plugins/Mutable/Source/MutableRuntime/Internal/MuR/ResourceID.h b/Engine/Plugins/Mutable/Source/MutableRuntime/Internal/MuR/ResourceID.h
new file mode 100644
--- /dev/null
+++ b/Engine/Plugins/Mutable/Source/MutableRuntime/Internal/MuR/ResourceID.h
@@ -0,0 +1,51 @@
+﻿// Copyright Epic Games, Inc. All Rights Reserved.
+
+#pragma once
+
+#include "MemoryCounters.h"
+#include "MuR/Registry.h"
+#include "MuR/Operations.h"
+
+
+namespace mu
+{
+	struct FGeneratedMeshKey
+	{
+		OP::ADDRESS Address = 0;
+		
+		TMemoryTrackedArray<uint8> ParameterValuesBlob;
+		
+		bool operator==(const FGeneratedMeshKey&) const = default;
+	};
+
+	struct FGeneratedMeshData
+	{
+	};
+
+	
+	MUTABLERUNTIME_API uint32 GetTypeHash(const FGeneratedMeshKey& Key);
+
+	typedef TRegistry<FGeneratedMeshKey, FGeneratedMeshData>::FHandle FMeshId;
+
+	typedef TRegistry<FGeneratedMeshKey, FGeneratedMeshData> FMeshIdRegistry;
+
+	struct FGeneratedImageKey
+	{
+		OP::ADDRESS Address = 0;
+
+		TMemoryTrackedArray<uint8> ParameterValuesBlob;
+
+		bool operator==(const FGeneratedImageKey&) const = default;
+	};
+
+	
+	struct FGeneratedImageData
+	{
+	};
+	
+	MUTABLERUNTIME_API uint32 GetTypeHash(const FGeneratedImageKey& Key);
+
+	typedef TRegistry<FGeneratedImageKey, FGeneratedImageData>::FHandle FImageId;
+
+	typedef TRegistry<FGeneratedImageKey, FGeneratedImageData> FImageIdRegistry;
+}
diff --git a/Engine/Plugins/Mutable/Source/MutableRuntime/Internal/MuR/SystemPrivate.h b/Engine/Plugins/Mutable/Source/MutableRuntime/Internal/MuR/SystemPrivate.h
--- a/Engine/Plugins/Mutable/Source/MutableRuntime/Internal/MuR/SystemPrivate.h
+++ b/Engine/Plugins/Mutable/Source/MutableRuntime/Internal/MuR/SystemPrivate.h
@@ -2,6 +2,7 @@
 
 #pragma once
 
+#include "MuR/MemoryCounters.h"
 #include "MuR/System.h"
 
 #include "MuR/Settings.h"
@@ -26,21 +27,9 @@
 	#define MUTABLE_DEBUG_CODERUNNER_TASK_SCHEDULE_CALLSTACK 0	
 #endif
 
-namespace mu::MemoryCounters
-{
-	struct FInternalMemoryCounter
-	{
-		static MUTABLERUNTIME_API std::atomic<SSIZE_T>& Get();
-	};
-}
 
 namespace mu
 {
-	inline uint32 GetResourceIDRoot(FResourceID Id)
-	{
-		return uint32(Id >> 32);
-	}
-
 	class FExtensionDataStreamer;
 
 	// Call the tick of the LLM system (we do this to simulate a frame since the LLM system is not entirelly designed to run over a program)
@@ -452,12 +441,6 @@ namespace mu
 	class FProgramCache
 	{
 	public:
-
-		using AllocType = FDefaultMemoryTrackingAllocator<MemoryCounters::FInternalMemoryCounter>;
-
-		template<class Type, class Alloc = AllocType >
-		using TMemoryTrackedArray = TArray<Type, Alloc>;
-		
 		TMemoryTrackedArray<ExecutionIndex, TInlineAllocator<4, AllocType>> m_usedRangeIndices;
 
 		/** Runtime data for each program op. */
@@ -1216,51 +1199,18 @@ namespace mu
 			TMemoryTrackedArray<uint16> PendingOpsPerRom;
         };
 
-		//! Management of generated resources
-		//! @{
-
-		//! This is used to uniquely identify a generated resource like meshes or images.
-		struct FGeneratedResourceData
-		{
-			/** Model for this resource. */
-			TWeakPtr<const FModel> Model;
-
-			//! The id assigned to the generated resource.
-			FResourceID Id;
-
-			//! The last request operation for this resource
-			uint32 LastRequestId;
-
-			//! An opaque blob with the values of the relevant parameters
-			TMemoryTrackedArray<uint8> ParameterValuesBlob;
-		};
-
-		//! The last id generated for a resource
-		uint32 LastResourceKeyId = 0;
-
-		//! The last id generated for a resource request. This is used to check the
-		//! relevancy of the resources when flushing the cache
-		uint32 LastResourceResquestId = 0;
-
-		//! Cached ids for returned assets
-		//! This is non-persistent runtime data
-		TMemoryTrackedArray<FGeneratedResourceData> GeneratedResources;
-
-		/** */
-		FResourceID GetResourceKey(const TSharedPtr<const FModel>&, const FParameters*, uint32 ParamListIndex, OP::ADDRESS RootAt);
-
-		//! @}
+		FLiveInstance CurrentLiveInstance;
 
+		FMeshId GetMeshId(const TSharedPtr<const FModel>& Model, const FParameters* Params, uint32 ParamListIndex, OP::ADDRESS RootAt);
+    	
+		FImageId GetImageId(const TSharedPtr<const FModel>& Model, const FParameters* Params, uint32 ParamListIndex, OP::ADDRESS RootAt);
+    	
 		/** Maximum working memory that mutable should be using. */
 		int64 BudgetBytes = 0;
 
 		/** Maximum excess memory reached suring the current operation. */
 		int64 BudgetExcessBytes = 0;
-
-		/** Maximum number of resource keys that will be stored for resource reusal. */
-		int32 MaxGeneratedResourceCacheSize = 1024;
-
-
+    	
 		/** This value is used to track the order of loading of roms. */
         uint64 RomTick = 0;
 
@@ -1273,6 +1223,10 @@ namespace mu
 		/** Temporary reference to the memory of the current instance being updated. Only valid during a mutable "atomic" operation, like a BeginUpdate or a GetImage. */
 		TSharedPtr<FProgramCache> CurrentInstanceCache;
 
+    	TSharedPtr<FMeshIdRegistry> CurrentMeshIdRegistry;
+    	
+		TSharedPtr<FImageIdRegistry> CurrentImageIdRegistry;
+    	
 		/** Resources that have been used in the past, but haven't been deallocated because they still fitted the memory budget and they could be reused. */
 		TMemoryTrackedArray<TSharedPtr<FImage>> PooledImages;
 
@@ -1975,7 +1929,10 @@ namespace mu
         //-----------------------------------------------------------------------------------------
 
 		/** This method can be used to internally prepare for code execution. */
-		MUTABLERUNTIME_API void BeginBuild(const TSharedPtr<const FModel>&);
+		MUTABLERUNTIME_API void BeginBuild(const TSharedPtr<const FModel>&,
+			const TSharedPtr<FMeshIdRegistry>&,
+			const TSharedPtr<FImageIdRegistry>&);
+    	
 		MUTABLERUNTIME_API void EndBuild();
 
 		MUTABLERUNTIME_API bool BuildBool(const TSharedPtr<const FModel>&, const FParameters*, OP::ADDRESS) ;
diff --git a/Engine/Plugins/Mutable/Source/MutableRuntime/Private/MuR/CodeRunner.cpp b/Engine/Plugins/Mutable/Source/MutableRuntime/Private/MuR/CodeRunner.cpp
--- a/Engine/Plugins/Mutable/Source/MutableRuntime/Private/MuR/CodeRunner.cpp
+++ b/Engine/Plugins/Mutable/Source/MutableRuntime/Private/MuR/CodeRunner.cpp
@@ -1023,7 +1023,7 @@ namespace mu
 
                 if ( Args.value )
                 {
-					FResourceID MeshId = System->WorkingMemoryManager.GetResourceKey(InModel,InParams,Args.relevantParametersListIndex, Args.value);
+					FMeshId MeshId = System->WorkingMemoryManager.GetMeshId(InModel, InParams, Args.relevantParametersListIndex, Args.value);
 					OP::ADDRESS NameAd = Args.name;
 					check(NameAd < (uint32)Program.ConstantStrings.Num());
 					const FString& Name = Program.ConstantStrings[NameAd];
@@ -1064,7 +1064,7 @@ namespace mu
 
                 if ( Args.value )
                 {
-					FResourceID ImageId = System->WorkingMemoryManager.GetResourceKey(InModel, InParams, Args.relevantParametersListIndex, Args.value);
+					FImageId ImageId = System->WorkingMemoryManager.GetImageId(InModel, InParams, Args.relevantParametersListIndex, Args.value);
 					OP::ADDRESS NameAd = Args.name;
 					check(NameAd < (uint32)Program.ConstantStrings.Num());
 					const FString& Name = Program.ConstantStrings[NameAd];
diff --git a/Engine/Plugins/Mutable/Source/MutableRuntime/Private/MuR/Instance.cpp b/Engine/Plugins/Mutable/Source/MutableRuntime/Private/MuR/Instance.cpp
--- a/Engine/Plugins/Mutable/Source/MutableRuntime/Private/MuR/Instance.cpp
+++ b/Engine/Plugins/Mutable/Source/MutableRuntime/Private/MuR/Instance.cpp
@@ -212,7 +212,7 @@ namespace mu
 	}
 	
     
-	FResourceID FInstance::GetMeshId(int32 ComponentIndex, int32 LODIndex ) const
+	FMeshId FInstance::GetMeshId(int32 ComponentIndex, int32 LODIndex) const
     {
         check(Components.IsValidIndex(ComponentIndex));
         check(Components[ComponentIndex].LODs.IsValidIndex(LODIndex));
@@ -221,7 +221,7 @@ namespace mu
     }
 
 	
-	FResourceID FInstance::GetImageId(int32 ComponentIndex, int32 LODIndex, int32 SurfaceIndex, int32 ImageIndex) const
+	FImageId FInstance::GetImageId(int32 ComponentIndex, int32 LODIndex, int32 SurfaceIndex, int32 ImageIndex) const
 	{
 		check(Components.IsValidIndex(ComponentIndex));
 		check(Components[ComponentIndex].LODs.IsValidIndex(LODIndex));
@@ -405,8 +405,9 @@ namespace mu
         surface.Name = Name;
     }
 
+
 	
-	void FInstance::SetMesh(int32 ComponentIndex, int32 LODIndex, FResourceID meshId, FName Name)
+	void FInstance::SetMesh(int32 ComponentIndex, int32 LODIndex, FMeshId meshId, FName Name)
 	{
 		// Automatically create the necessary lods and components
 		while (ComponentIndex >= Components.Num())
@@ -422,9 +423,9 @@ namespace mu
 		LOD.MeshId = meshId;
 		LOD.MeshName = Name;
 	}
-	
-	
-    int32 FInstance::AddImage( int32 ComponentIndex, int32 LODIndex, int32 SurfaceIndex, FResourceID imageId, FName Name)
+
+
+    int32 FInstance::AddImage( int32 ComponentIndex, int32 LODIndex, int32 SurfaceIndex, FImageId imageId, FName Name)
 	{
 		// Automatically create the necessary lods and components
 		while (ComponentIndex >= Components.Num())
diff --git a/Engine/Plugins/Mutable/Source/MutableRuntime/Private/MuR/ResourceID.cpp b/Engine/Plugins/Mutable/Source/MutableRuntime/Private/MuR/ResourceID.cpp
new file mode 100644
--- /dev/null
+++ b/Engine/Plugins/Mutable/Source/MutableRuntime/Private/MuR/ResourceID.cpp
@@ -0,0 +1,25 @@
+// Copyright Epic Games, Inc. All Rights Reserved.
+
+#include "MuR/ResourceID.h"
+
+#include "Hash/xxhash.h"
+
+
+uint32 mu::GetTypeHash(const FGeneratedMeshKey& Key)
+{
+	uint32 Hash = Key.Address;
+
+	Hash = HashCombineFast(Hash, static_cast<uint32>(FXxHash64::HashBuffer(Key.ParameterValuesBlob.GetData(), Key.ParameterValuesBlob.Num()).Hash));
+
+	return Hash;
+}
+
+
+uint32 mu::GetTypeHash(const FGeneratedImageKey& Key)
+{
+	uint32 Hash = Key.Address;
+
+	Hash = HashCombineFast(Hash, static_cast<uint32>(FXxHash64::HashBuffer(Key.ParameterValuesBlob.GetData(), Key.ParameterValuesBlob.Num()).Hash));
+
+	return Hash;
+}
diff --git a/Engine/Plugins/Mutable/Source/MutableRuntime/Private/MuR/System.cpp b/Engine/Plugins/Mutable/Source/MutableRuntime/Private/MuR/System.cpp
--- a/Engine/Plugins/Mutable/Source/MutableRuntime/Private/MuR/System.cpp
+++ b/Engine/Plugins/Mutable/Source/MutableRuntime/Private/MuR/System.cpp
@@ -119,7 +119,6 @@ namespace mu
         Settings = InSettings;
 	
 		WorkingMemoryManager.BudgetBytes = Settings.WorkingMemoryBytes;
-		WorkingMemoryManager.GeneratedResources.Reserve(WorkingMemoryManager.MaxGeneratedResourceCacheSize);
 
 		UpdateStats();
 	}
@@ -158,22 +157,6 @@ namespace mu
 		m_pD->UpdateStats();
 	}
 
-
-    //---------------------------------------------------------------------------------------------
-    void FSystem::SetGeneratedCacheSize( uint32 InCount )
-    {
-		LLM_SCOPE_BYNAME(TEXT("MutableRuntime"));
-		MUTABLE_CPUPROFILER_SCOPE(SetGeneratedCacheSize);
-
-		m_pD->WorkingMemoryManager.MaxGeneratedResourceCacheSize = InCount;
-		m_pD->WorkingMemoryManager.GeneratedResources.Reserve(m_pD->WorkingMemoryManager.MaxGeneratedResourceCacheSize);
-		if (m_pD->WorkingMemoryManager.GeneratedResources.Num()>int32(InCount))
-		{
-			// Discard some random resource keys.
-			m_pD->WorkingMemoryManager.GeneratedResources.SetNum(InCount);
-		}
-	}
-
 	
 	//---------------------------------------------------------------------------------------------
 	void FSystem::ClearWorkingMemory()
@@ -244,6 +227,8 @@ namespace mu
     //---------------------------------------------------------------------------------------------
     TSharedPtr<const FInstance> FSystem::BeginUpdate( FInstance::FID InInstanceID,
                                      const TSharedPtr<const FParameters>& InParams,
+                                     const TSharedPtr<FMeshIdRegistry>& InMeshIdRegistry,
+									 const TSharedPtr<FImageIdRegistry>& InImageIdRegistry,
                                      int32 InStateIndex,
                                      uint32 InLodMask )
     {
@@ -266,6 +251,9 @@ namespace mu
 
 		m_pD->WorkingMemoryManager.CurrentInstanceCache = LiveInstance->Cache;
 
+    	m_pD->WorkingMemoryManager.CurrentMeshIdRegistry = InMeshIdRegistry;
+    	m_pD->WorkingMemoryManager.CurrentImageIdRegistry = InImageIdRegistry;
+
 		FProgram& Program = LiveInstance->Model->GetPrivate()->Program;
 
 		bool bValidState = InStateIndex >= 0 && InStateIndex < (int)Program.States.Num();
@@ -329,11 +317,22 @@ namespace mu
 	}
 
 
-	TSharedPtr<const FImage> FSystem::GetImageInline(FInstance::FID instanceID, FResourceID ImageId, int32 MipsToSkip, int32 InImageLOD)
+    TSharedRef<FImage> BlackImage()
+    {
+    	return MakeShared<mu::FImage>(16, 16, 1, EImageFormat::RGBA_UByte, EInitializationType::Black);
+    }
+	
+
+	TSharedPtr<const FImage> FSystem::GetImageInline(FInstance::FID instanceID, const FImageId& ImageId, int32 MipsToSkip, int32 InImageLOD)
 	{
 		LLM_SCOPE_BYNAME(TEXT("MutableRuntime"));
 		MUTABLE_CPUPROFILER_SCOPE(SystemGetImage);
 
+    	if (!ImageId)
+    	{
+			return BlackImage();
+    	}
+    	
 		TSharedPtr<const FImage> pResult;
 
 		// Find the live instance
@@ -341,13 +340,13 @@ namespace mu
 		check(LiveInstance);
 		m_pD->WorkingMemoryManager.CurrentInstanceCache = LiveInstance->Cache;
 
-		OP::ADDRESS RootAddress = GetResourceIDRoot(ImageId);
+		OP::ADDRESS RootAddress = ImageId.GetKey()->Address;
 		pResult = m_pD->BuildImage(LiveInstance->Model, LiveInstance->OldParameters.Get(), RootAddress, MipsToSkip, InImageLOD);
 
 		// We always need to return something valid.
 		if (!pResult)
 		{
-			pResult = MakeShared<mu::FImage>(16, 16, 1, EImageFormat::RGBA_UByte, EInitializationType::Black);
+			pResult = BlackImage();
 		}
 
 		m_pD->WorkingMemoryManager.CurrentInstanceCache = nullptr;
@@ -355,11 +354,16 @@ namespace mu
 		return pResult;
 	}
 
-	UE::Tasks::TTask<TSharedPtr<const FImage>> FSystem::GetImage(FInstance::FID instanceID, FResourceID ImageId, int32 MipsToSkip, int32 InImageLOD)
+	UE::Tasks::TTask<TSharedPtr<const FImage>> FSystem::GetImage(FInstance::FID instanceID, const FImageId& ImageId, int32 MipsToSkip, int32 InImageLOD)
 	{
 		LLM_SCOPE_BYNAME(TEXT("MutableRuntime"));
 		MUTABLE_CPUPROFILER_SCOPE(SystemGetImage);
 
+    	if (!ImageId)
+    	{
+			return UE::Tasks::MakeCompletedTask<TSharedPtr<const FImage>>(BlackImage());
+    	}
+    	
 		TSharedPtr<const FImage> pResult;
 
 		// Find the live instance
@@ -367,7 +371,7 @@ namespace mu
 		check(LiveInstance);
 		m_pD->WorkingMemoryManager.CurrentInstanceCache = LiveInstance->Cache;
 
-		OP::ADDRESS RootAddress = GetResourceIDRoot(ImageId);
+		OP::ADDRESS RootAddress = ImageId.GetKey()->Address;
 		
 		m_pD->WorkingMemoryManager.BeginRunnerThread();
 		
@@ -377,8 +381,7 @@ namespace mu
 			m_pD->WorkingMemoryManager.EndRunnerThread();
 			m_pD->WorkingMemoryManager.CurrentInstanceCache = nullptr;
 
-			return UE::Tasks::MakeCompletedTask<TSharedPtr<const FImage>>(
-					MakeShared<mu::FImage>(16, 16, 1, EImageFormat::RGBA_UByte, EInitializationType::Black));
+			return UE::Tasks::MakeCompletedTask<TSharedPtr<const FImage>>(BlackImage());
 		}
 		
 		TSharedRef<CodeRunner> Runner = CodeRunner::Create(
@@ -423,11 +426,16 @@ namespace mu
 		TEXT("If different than 0, clear the image desc cache at every image."),
 		ECVF_Scalability);
 
-	FExtendedImageDesc FSystem::GetImageDescInline(FInstance::FID instanceID, FResourceID ImageId)
+	FExtendedImageDesc FSystem::GetImageDescInline(FInstance::FID instanceID, const FImageId& ImageId)
 	{
 		LLM_SCOPE_BYNAME(TEXT("MutableRuntime"));
 		MUTABLE_CPUPROFILER_SCOPE(SystemGetImageDesc);
 
+		if (!ImageId)
+		{
+			return {};
+		}
+		
 		FExtendedImageDesc Result;
 
 		// Find the live instance
@@ -435,7 +443,7 @@ namespace mu
 		check(LiveInstance);
 		m_pD->WorkingMemoryManager.CurrentInstanceCache = LiveInstance->Cache;
 
-		OP::ADDRESS RootAddress = GetResourceIDRoot(ImageId);
+		OP::ADDRESS RootAddress = ImageId.GetKey()->Address;
 
 		const mu::FModel* Model = LiveInstance->Model.Get();
 		const mu::FProgram& Program = Model->GetPrivate()->Program;
@@ -472,17 +480,22 @@ namespace mu
 	}
 
 
-	UE::Tasks::TTask<FExtendedImageDesc> FSystem::GetImageDesc(FInstance::FID instanceID, FResourceID ImageId)
+	UE::Tasks::TTask<FExtendedImageDesc> FSystem::GetImageDesc(FInstance::FID instanceID, const FImageId& ImageId)
 	{
 		LLM_SCOPE_BYNAME(TEXT("MutableRuntime"));
 		MUTABLE_CPUPROFILER_SCOPE(SystemGetImageDesc);
 
+		if (!ImageId)
+		{
+			return UE::Tasks::MakeCompletedTask<FExtendedImageDesc>();
+		}
+		
 		// Find the live instance
 		FLiveInstance* LiveInstance = m_pD->FindLiveInstance(instanceID);
 		check(LiveInstance);
 		m_pD->WorkingMemoryManager.CurrentInstanceCache = LiveInstance->Cache;
 
-		OP::ADDRESS RootAddress = GetResourceIDRoot(ImageId);
+		OP::ADDRESS RootAddress = ImageId.GetKey()->Address;
 
 		const mu::FModel* Model = LiveInstance->Model.Get();
 		const mu::FProgram& Program = Model->GetPrivate()->Program;
@@ -528,11 +541,16 @@ namespace mu
 
 	}	
 
-    TSharedPtr<const FMesh> FSystem::GetMeshInline(FInstance::FID instanceID, FResourceID MeshId, EMeshContentFlags MeshContentFilter)
+    TSharedPtr<const FMesh> FSystem::GetMeshInline(FInstance::FID instanceID, const FImageId& MeshId, EMeshContentFlags MeshContentFilter)
     {
 		LLM_SCOPE_BYNAME(TEXT("MutableRuntime"));
 		MUTABLE_CPUPROFILER_SCOPE(SystemGetMesh);
 
+		if (!MeshId)
+		{
+			return MakeShared<FMesh>();
+		}
+		
 		TSharedPtr<const FMesh> Result;
 
 		// Find the live instance
@@ -540,7 +558,7 @@ namespace mu
 		check(LiveInstance);
 		m_pD->WorkingMemoryManager.CurrentInstanceCache = LiveInstance->Cache;
 
-		OP::ADDRESS RootAddress = GetResourceIDRoot(MeshId);
+		OP::ADDRESS RootAddress = MeshId.GetKey()->Address;
 		Result = m_pD->BuildMesh(
 				LiveInstance->Model, 
 				LiveInstance->OldParameters.Get(), 
@@ -557,10 +575,15 @@ namespace mu
 		return Result;
 	}
 
-	UE::Tasks::TTask<TSharedPtr<const FMesh>> FSystem::GetMesh(FInstance::FID instanceID, FResourceID MeshId, EMeshContentFlags MeshContentFilter)
+	UE::Tasks::TTask<TSharedPtr<const FMesh>> FSystem::GetMesh(FInstance::FID instanceID, const FMeshId& MeshId, EMeshContentFlags MeshContentFilter)
     {
 		LLM_SCOPE_BYNAME(TEXT("MutableRuntime"));
 		MUTABLE_CPUPROFILER_SCOPE(SystemGetMesh);
+		
+		if (!MeshId)
+		{
+			return UE::Tasks::MakeCompletedTask<TSharedPtr<const FMesh>>(new FMesh());
+		}
 
 		TSharedPtr<const FMesh> ResultMesh;
 
@@ -570,7 +593,7 @@ namespace mu
 		m_pD->WorkingMemoryManager.CurrentInstanceCache = LiveInstance->Cache;
 		m_pD->WorkingMemoryManager.BeginRunnerThread();
 		
-		OP::ADDRESS RootAddress = GetResourceIDRoot(MeshId);
+		OP::ADDRESS RootAddress = MeshId.GetKey()->Address;
 			
 		mu::EOpType OpType = LiveInstance->Model->GetPrivate()->Program.GetOpType(RootAddress);
 		if (GetOpDataType(OpType) != EDataType::Mesh)
@@ -841,13 +864,18 @@ namespace mu
 
 
 	//---------------------------------------------------------------------------------------------
-	void FSystem::Private::BeginBuild(const TSharedPtr<const FModel>& InModel)
+	void FSystem::Private::BeginBuild(const TSharedPtr<const FModel>& InModel,
+		const TSharedPtr<FMeshIdRegistry>& InMeshIdRegistry,
+		const TSharedPtr<FImageIdRegistry>& InImageIdRegistry)
 	{
 		// We don't have a FLiveInstance, let's create the memory
 		// \TODO: There is no clear moment to remove this... EndBuild?
 		WorkingMemoryManager.CurrentInstanceCache = MakeShared<FProgramCache>();
 		WorkingMemoryManager.CurrentInstanceCache->Init(InModel->GetPrivate()->Program.OpAddress.Num());
 
+    	WorkingMemoryManager.CurrentMeshIdRegistry = InMeshIdRegistry;
+    	WorkingMemoryManager.CurrentImageIdRegistry = InImageIdRegistry;
+    	
 		// Ensure the model cache has been created
 		WorkingMemoryManager.FindOrAddModelCache(InModel);
 
@@ -1714,15 +1742,13 @@ namespace mu
 
 
 	//---------------------------------------------------------------------------------------------
-	FResourceID FWorkingMemoryManager::GetResourceKey(const TSharedPtr<const FModel>& Model, const FParameters* Params, uint32 ParamListIndex, OP::ADDRESS RootAt)
+	void GetResourceKey(const TSharedPtr<const FModel>& Model, const FParameters* Params, uint32 ParamListIndex, OP::ADDRESS RootAt, TMemoryTrackedArray<uint8>& ParameterValuesBlob)
 	{
 		MUTABLE_CPUPROFILER_SCOPE(GetResourceKey);
-
-		constexpr uint32 ErrorId = 0xffff;
-
+    	
 		if (!Model)
 		{
-			return ErrorId;
+			return;
 		}
 
 		const FProgram& Program = Model->GetPrivate()->Program;
@@ -1736,7 +1762,7 @@ namespace mu
 		check(RelevantParams);
 		if (!RelevantParams)
 		{
-			return ErrorId;
+			return;
 		}
 
 		// Generate the relevant parameters blob
@@ -1965,61 +1991,32 @@ namespace mu
 				check(false);
 			}
 		}
+    	
+		int32 BlobBytes = Blob.GetNumBytes();
+		ParameterValuesBlob.SetNum(BlobBytes);
+		FMemory::Memcpy(ParameterValuesBlob.GetData(), Blob.GetData(), BlobBytes);
+	}
 
-		// Increase the request id
-		++LastResourceResquestId;
 
-		FGeneratedResourceData NewKey;
-		int32 BlobBytes = Blob.GetNumBytes();
-		NewKey.ParameterValuesBlob.SetNum(BlobBytes);
-		FMemory::Memcpy(NewKey.ParameterValuesBlob.GetData(), Blob.GetData(), BlobBytes);
+	FMeshId FWorkingMemoryManager::GetMeshId(const TSharedPtr<const FModel>& Model, const FParameters* Params, uint32 ParamListIndex, OP::ADDRESS RootAt)
+    {
+    	FGeneratedMeshKey NewKey;
+    	NewKey.Address = RootAt;
 
-		// See if we already have this id
-		int32 OldestCachePosition = 0;
-		uint32 OldestRequestId = 0;
-		for (int32 CacheIndex = 0; CacheIndex < GeneratedResources.Num(); ++CacheIndex)
-		{
-			FGeneratedResourceData& Data = GeneratedResources[CacheIndex];
-			bool bSameModel = Data.Model == Model;
-			bool bSameRoot = GetResourceIDRoot(Data.Id) == RootAt;
-			if (bSameModel
-				&&
-				bSameRoot
-				)
-			{
-				bool bSameBlob = Data.ParameterValuesBlob == NewKey.ParameterValuesBlob;
-				if (bSameBlob)
-				{
-					Data.LastRequestId = LastResourceResquestId;
-					return Data.Id;
-				}
-			}
-			
-			if (!Data.Model.Pin().Get() 
-				||
-				OldestRequestId > Data.LastRequestId)
-			{
-				OldestCachePosition = CacheIndex;
-				OldestRequestId = Data.Model.Pin() ? Data.LastRequestId : 0;
-			}
-		}
+		GetResourceKey(Model, Params, ParamListIndex, RootAt, NewKey.ParameterValuesBlob);
 
-		// Generate a new id
-		uint32 NewBlobId = ++LastResourceKeyId;
-		NewKey.Id = MakeResourceID( RootAt, NewBlobId);
-		NewKey.LastRequestId = LastResourceResquestId;
-		NewKey.Model = Model;
+    	return CurrentMeshIdRegistry->Add(NewKey, {});
+    }
+	
+    	
+	FImageId FWorkingMemoryManager::GetImageId(const TSharedPtr<const FModel>& Model, const FParameters* Params, uint32 ParamListIndex, OP::ADDRESS RootAt)
+    {
+    	FGeneratedImageKey NewKey;
+    	NewKey.Address = RootAt;
 
-		if (GeneratedResources.Num() >= MaxGeneratedResourceCacheSize)
-		{
-			GeneratedResources[OldestCachePosition] = MoveTemp(NewKey);
-		}
-		else
-		{
-			GeneratedResources.Add(MoveTemp(NewKey));
-		}
+    	GetResourceKey(Model, Params, ParamListIndex, RootAt, NewKey.ParameterValuesBlob);
 
-		return NewKey.Id;
-	}
+    	return CurrentImageIdRegistry->Add(NewKey, { });
+    }
 
 }
diff --git a/Engine/Plugins/Mutable/Source/MutableRuntime/Public/MuR/Image.h b/Engine/Plugins/Mutable/Source/MutableRuntime/Public/MuR/Image.h
deleted file mode 100644
diff --git a/Engine/Plugins/Mutable/Source/MutableRuntime/Public/MuR/ImageDataStorage.h b/Engine/Plugins/Mutable/Source/MutableRuntime/Public/MuR/ImageDataStorage.h
deleted file mode 100644
diff --git a/Engine/Plugins/Mutable/Source/MutableRuntime/Public/MuR/Instance.h b/Engine/Plugins/Mutable/Source/MutableRuntime/Public/MuR/Instance.h
--- a/Engine/Plugins/Mutable/Source/MutableRuntime/Public/MuR/Instance.h
+++ b/Engine/Plugins/Mutable/Source/MutableRuntime/Public/MuR/Instance.h
@@ -2,6 +2,7 @@
 
 #pragma once
 
+#include "MuR/ResourceId.h"
 #include "HAL/PlatformMath.h"
 #include "MuR/ExtensionData.h"
 #include "MuR/Ptr.h"
@@ -14,14 +15,6 @@
 
 namespace mu
 {
-
-	/** Helper functions to make and read FResourceIDs */
-	inline FResourceID MakeResourceID(uint32 RootAddress, uint32 ParameterBlobIndex)
-	{
-		return (uint64(RootAddress) << 32) | uint64(ParameterBlobIndex);
-	}
-
-	
 	struct FInstanceSurface
 	{
 		FName Name;
@@ -33,7 +26,7 @@ namespace mu
 
 		struct FInstanceImage
 		{
-			FResourceID Id;
+			FImageId Id;
 			FName Name;
 		};
 
@@ -67,7 +60,7 @@ namespace mu
 
 	struct FInstanceLOD
 	{
-		FResourceID MeshId;
+		FMeshId MeshId;
 		FName MeshName;
 
 		// The order must match the meshes surfaces
@@ -171,7 +164,7 @@ namespace mu
         //! Get the mesh resource id from a component
         //! \param ComponentIndex Index of the component, from 0 to GetComponentCount()-1
 		//! \param LODIndex Index of the level of detail, from 0 to GetLODCount()-1
-  		UE_API FResourceID GetMeshId( int32 ComponentIndex, int32 LODIndex ) const;
+  		UE_API FMeshId GetMeshId( int32 ComponentIndex, int32 LODIndex ) const;
 
 		//! Get the number of images in a component
 		//! \param ComponentIndex Index of the component, from 0 to GetComponentCount()-1
@@ -184,7 +177,7 @@ namespace mu
 		//! \param LODIndex Index of the level of detail, from 0 to GetLODCount()-1
 		//! \param SurfaceIndex Index of the surface, from 0 to GetSurfaceCount()-1
         //! \param ImageIndex Index of the image, from 0 to GetImageCount()-1
- 		UE_API FResourceID GetImageId( int32 ComponentIndex, int32 LODIndex, int32 SurfaceIndex, int32 ImageIndex) const;
+ 		UE_API FImageId GetImageId( int32 ComponentIndex, int32 LODIndex, int32 SurfaceIndex, int32 ImageIndex) const;
 
 		//! Get the name of an image in a component
 		//! \param ComponentIndex Index of the component, from 0 to GetComponentCount()-1
@@ -272,13 +265,13 @@ namespace mu
 		
 		int32 AddLOD(int32 ComponentIndex);
 
-		void SetMesh(int32 ComponentIndex, int32 LODIndex, FResourceID, FName Name);
+		void SetMesh(int32 ComponentIndex, int32 LODIndex, FMeshId, FName Name);
 
 		int32 AddSurface(int32 ComponentIndex, int32 LODIndex);
 
 		void SetSurfaceName(int32 ComponentIndex, int32 LODIndex, int32 SurfaceIndex, FName Name);
 
-		int32 AddImage(int32 ComponentIndex, int32 LODIndex, int32 SurfaceIndex, FResourceID, FName Name);
+		int32 AddImage(int32 ComponentIndex, int32 LODIndex, int32 SurfaceIndex, FImageId, FName Name);
 
 		int32 AddVector(int32 ComponentIndex, int32 LODIndex, int32 SurfaceIndex, const FVector4f&, FName Name);
 
diff --git a/Engine/Plugins/Mutable/Source/MutableRuntime/Public/MuR/MemoryTrackingAllocationPolicy.h b/Engine/Plugins/Mutable/Source/MutableRuntime/Public/MuR/MemoryTrackingAllocationPolicy.h
deleted file mode 100644
diff --git a/Engine/Plugins/Mutable/Source/MutableRuntime/Public/MuR/MemoryTrackingUtils.h b/Engine/Plugins/Mutable/Source/MutableRuntime/Public/MuR/MemoryTrackingUtils.h
deleted file mode 100644
diff --git a/Engine/Plugins/Mutable/Source/MutableRuntime/Public/MuR/System.h b/Engine/Plugins/Mutable/Source/MutableRuntime/Public/MuR/System.h
--- a/Engine/Plugins/Mutable/Source/MutableRuntime/Public/MuR/System.h
+++ b/Engine/Plugins/Mutable/Source/MutableRuntime/Public/MuR/System.h
@@ -120,9 +120,6 @@ namespace mu
 		*/
 		UE_API void ClearWorkingMemory();
 		
-		/** Set the amount of generated resources keys that will be stored for resource reusal. */
-		UE_API void SetGeneratedCacheSize(uint32 InCount);
-
         /** Set a new provider for external image data. This is only necessary if image parameters are used in the models. */
         UE_API void SetExternalResourceProvider(const TSharedPtr<FExternalResourceProvider>&);
 
@@ -153,33 +150,38 @@ namespace mu
         //! \return the instance data with all the LOD, components, and ids to generate the meshes 
         //! and images. The returned Instance is only valid until the next call to EndUpdate with 
         //! the same instanceID parameter.
-		UE_API TSharedPtr<const FInstance> BeginUpdate(FInstance::FID InstanceID, const TSharedPtr<const FParameters>& Params, int32 StateIndex, uint32 LodMask);
+		UE_API TSharedPtr<const FInstance> BeginUpdate(FInstance::FID InstanceID,
+			const TSharedPtr<const FParameters>& Params,
+			const TSharedPtr<FMeshIdRegistry>& InMeshIdRegistry,
+			const TSharedPtr<FImageIdRegistry>& InImageIdRegistry,
+			int32 StateIndex,
+			uint32 LodMask);
 
 		//! Only valid between BeginUpdate and EndUpdate
 		//! Calculate the description of an image, without generating it.
-		UE_API UE::Tasks::TTask<FExtendedImageDesc> GetImageDesc(FInstance::FID InstanceID, FResourceID ImageId);
+		UE_API UE::Tasks::TTask<FExtendedImageDesc> GetImageDesc(FInstance::FID InstanceID, const FImageId& ImageId);
 
 		//! Only valid between BeginUpdate and EndUpdate
 		//! \param MipsToSkip Number of mips to skip compared from the full image.
 		//! If 0, all mip levels will be generated. If more levels than possible to discard are specified, 
 		//! the image will still contain a minimum number of mips specified at model compile time.
-		UE_API UE::Tasks::TTask<TSharedPtr<const FImage>> GetImage(FInstance::FID InstanceID, FResourceID ImageId, int32 MipsToSkip = 0, int32 LOD = 0);
+		UE_API UE::Tasks::TTask<TSharedPtr<const FImage>> GetImage(FInstance::FID InstanceID, const FImageId& ImageId, int32 MipsToSkip = 0, int32 LOD = 0);
 
         //! Only valid between BeginUpdate and EndUpdate
-		UE_API UE::Tasks::TTask<TSharedPtr<const FMesh>> GetMesh(FInstance::FID InstanceID, FResourceID MeshId, EMeshContentFlags MeshContentFilter = EMeshContentFlags::AllFlags);
+		UE_API UE::Tasks::TTask<TSharedPtr<const FMesh>> GetMesh(FInstance::FID InstanceID, const FMeshId& MeshId, EMeshContentFlags MeshContentFilter = EMeshContentFlags::AllFlags);
 
 		//! Only valid between BeginUpdate and EndUpdate
 		//! Calculate the description of an image, without generating it.
-		UE_API FExtendedImageDesc GetImageDescInline(FInstance::FID InstanceID, FResourceID ImageId);
+		UE_API FExtendedImageDesc GetImageDescInline(FInstance::FID InstanceID, const FImageId& ImageId);
 
 		//! Only valid between BeginUpdate and EndUpdate
 		//! \param MipsToSkip Number of mips to skip compared from the full image.
 		//! If 0, all mip levels will be generated. If more levels than possible to discard are specified, 
 		//! the image will still contain a minimum number of mips specified at model compile time.
-		UE_API TSharedPtr<const FImage> GetImageInline(FInstance::FID InstanceID, FResourceID ImageId, int32 MipsToSkip = 0, int32 LOD = 0);
+		UE_API TSharedPtr<const FImage> GetImageInline(FInstance::FID InstanceID, const FImageId& ImageId, int32 MipsToSkip = 0, int32 LOD = 0);
 
         //! Only valid between BeginUpdate and EndUpdate
-		UE_API TSharedPtr<const FMesh> GetMeshInline(FInstance::FID InstanceID, FResourceID MeshId, EMeshContentFlags MeshContentFilter = EMeshContentFlags::AllFlags);
+		UE_API TSharedPtr<const FMesh> GetMeshInline(FInstance::FID InstanceID, const FImageId& MeshId, EMeshContentFlags MeshContentFilter = EMeshContentFlags::AllFlags);
 
         //! Invalidate and free the last Instance data returned by a call to BeginUpdate with
         //! the same instance index. After a call to this method, that Instance cannot be used any
diff --git a/Engine/Plugins/Mutable/Source/MutableRuntime/Public/MuR/Types.h b/Engine/Plugins/Mutable/Source/MutableRuntime/Public/MuR/Types.h
--- a/Engine/Plugins/Mutable/Source/MutableRuntime/Public/MuR/Types.h
+++ b/Engine/Plugins/Mutable/Source/MutableRuntime/Public/MuR/Types.h
@@ -35,16 +35,12 @@ namespace mu
 {
 	class FOutputArchive;
 	class FInputArchive;
-
-	/** Handle used to identify separate resources generated by mutable: images and meshes included in instances. */
-    typedef uint64 FResourceID;
-
+	
 	enum class EMemoryInitPolicy
 	{
 		Uninitialized,
 		Zeroed,
 	};
-
 }
 
 
diff --git a/Engine/Plugins/Mutable/Source/MutableTools/Private/MuT/CodeOptimiser.cpp b/Engine/Plugins/Mutable/Source/MutableTools/Private/MuT/CodeOptimiser.cpp
--- a/Engine/Plugins/Mutable/Source/MutableTools/Private/MuT/CodeOptimiser.cpp
+++ b/Engine/Plugins/Mutable/Source/MutableTools/Private/MuT/CodeOptimiser.cpp
@@ -406,15 +406,15 @@ namespace mu
 			LinkerOptions.MinTextureResidentMipCount = 255;
 			LinkerOptions.bSeparateImageMips = false;
 
-			TSharedPtr<const FModel> model = MakeShared<FModel>();
-			OP::ADDRESS at = ASTOp::FullLink(SourceCloned, model->GetPrivate()->Program, &LinkerOptions);
+			TSharedPtr<const FModel> Model = MakeShared<FModel>();
+			OP::ADDRESS at = ASTOp::FullLink(SourceCloned, Model->GetPrivate()->Program, &LinkerOptions);
 
 			FProgram::FState state;
 			state.Root = at;
-			model->GetPrivate()->Program.States.Add(state);
+			Model->GetPrivate()->Program.States.Add(state);
 
-			TSharedPtr<FParameters> LocalParams = FModel::NewParameters(model);
-			System->GetPrivate()->BeginBuild( model );
+			TSharedPtr<FParameters> LocalParams = FModel::NewParameters(Model);
+			System->GetPrivate()->BeginBuild(Model, nullptr, nullptr);
 
 			// Calculate the value and replace this op by a constant
 			switch( DataType )
@@ -423,7 +423,7 @@ namespace mu
 			{
 				MUTABLE_CPUPROFILER_SCOPE(ConstantMesh);
 
-				TSharedPtr<const FMesh> MeshBuild = System->GetPrivate()->BuildMesh(model, LocalParams.Get(), at, EMeshContentFlags::AllFlags);
+				TSharedPtr<const FMesh> MeshBuild = System->GetPrivate()->BuildMesh(Model, LocalParams.Get(), at, EMeshContentFlags::AllFlags);
 
 				if (MeshBuild)
 				{
@@ -440,7 +440,7 @@ namespace mu
 			{
 				MUTABLE_CPUPROFILER_SCOPE(ConstantImage);
 
-				TSharedPtr<const FImage> pImage = System->GetPrivate()->BuildImage( model, LocalParams.Get(), at, 0, 0 );
+				TSharedPtr<const FImage> pImage = System->GetPrivate()->BuildImage( Model, LocalParams.Get(), at, 0, 0 );
 
 				if (pImage)
 				{
@@ -457,7 +457,7 @@ namespace mu
 			{
 				MUTABLE_CPUPROFILER_SCOPE(ConstantLayout);
 
-				TSharedPtr<const FLayout> pLayout = System->GetPrivate()->BuildLayout( model, LocalParams.Get(), at );
+				TSharedPtr<const FLayout> pLayout = System->GetPrivate()->BuildLayout( Model, LocalParams.Get(), at );
 
 				if (pLayout)
 				{
@@ -473,7 +473,7 @@ namespace mu
 			{
 				MUTABLE_CPUPROFILER_SCOPE(ConstantBool);
 
-				bool value = System->GetPrivate()->BuildBool( model, LocalParams.Get(), at );
+				bool value = System->GetPrivate()->BuildBool( Model, LocalParams.Get(), at );
 				Result = new ASTOpConstantBool(value);
 				break;
 			}
@@ -483,7 +483,7 @@ namespace mu
 				MUTABLE_CPUPROFILER_SCOPE(ConstantBool);
 
 				FVector4f ResultColor(0, 0, 0, 0);
-				ResultColor = System->GetPrivate()->BuildColour( model, LocalParams.Get(), at );
+				ResultColor = System->GetPrivate()->BuildColour( Model, LocalParams.Get(), at );
 
 				{
 					mu::Ptr<ASTOpConstantColor> ConstantOp = new ASTOpConstantColor();
