[Mutable] Fix Skeletal Mesh Parameter morph check.

diff --git a/Engine/Plugins/Mutable/Source/CustomizableObject/Private/MuCO/UnrealConversionUtils.cpp b/Engine/Plugins/Mutable/Source/CustomizableObject/Private/MuCO/UnrealConversionUtils.cpp
--- a/Engine/Plugins/Mutable/Source/CustomizableObject/Private/MuCO/UnrealConversionUtils.cpp
+++ b/Engine/Plugins/Mutable/Source/CustomizableObject/Private/MuCO/UnrealConversionUtils.cpp
@@ -1331,11 +1331,6 @@ namespace UnrealConversionUtils
 				const TArray<TObjectPtr<UMorphTarget>>& MorphTargets = SkeletalMesh->GetMorphTargets();
 				for (UMorphTarget* MorphTarget : MorphTargets)
 				{
-					TArray<FMorphTargetLODModel>& MorphLODModels = MorphTarget->GetMorphLODModels();
-
-					const FMorphTargetLODModel* MorphTargetLOD = MorphLODModels.IsValidIndex(LODIndex) 
-							? &MorphLODModels[LODIndex] 
-							: nullptr;
 					Result->Morph.Names.Add(MorphTarget->GetFName());
 				}
 
@@ -1384,7 +1379,6 @@ namespace UnrealConversionUtils
 					TArray<FDeltaBatchHeader> BatchHeaders;
 					TArray<uint32> BitstreamData;
 				
-					const int32 NumMorphTargets = MorphTargetLODs.Num();
 					for (const FMorphTargetLODModel* MorphModel : MorphTargetLODs)
 					{
 						uint32 BatchStartOffset = BatchHeaders.Num();
@@ -1438,7 +1432,6 @@ namespace UnrealConversionUtils
 					RuntimeCompressedMorphData.Append(BitstreamData);
 				}
 
-				check(!RuntimeCompressedMorphData.IsEmpty() || MorphTargetBuffers.GetMorphDataSizeInBytes() == 0);
 				const TConstArrayView<uint32> SourceMorphDataView = !RuntimeCompressedMorphData.IsEmpty() 
 						? TConstArrayView<uint32>(RuntimeCompressedMorphData)
 						: TConstArrayView<uint32>(MorphTargetBuffers.GetData(), MorphTargetBuffers.GetMorphDataSizeInBytes()/sizeof(uint32));
