[mutable] Fixed Mesh Clip Deform node ensure due to an invalid manipulation of the Positions array
- When the PositionBuffer is compatible the shape positions array was left untouched witch caused the issue.

#rnx

diff --git a/Engine/Plugins/Mutable/Source/MutableRuntime/Private/MuR/OpMeshClipDeform.h b/Engine/Plugins/Mutable/Source/MutableRuntime/Private/MuR/OpMeshClipDeform.h
--- a/Engine/Plugins/Mutable/Source/MutableRuntime/Private/MuR/OpMeshClipDeform.h
+++ b/Engine/Plugins/Mutable/Source/MutableRuntime/Private/MuR/OpMeshClipDeform.h
@@ -172,9 +172,6 @@ namespace mu
 		{
 			MUTABLE_CPUPROFILER_SCOPE(GenerateVertexQueryData);
 
-			ShapeNormals.SetNum(ShapeVertexCount);
-			ShapeWeights.SetNum(ShapeVertexCount);
-
 			UntypedMeshBufferIteratorConst ItPositionBegin(ShapeMesh->GetVertexBuffers(), EMeshBufferSemantic::Position);
 			UntypedMeshBufferIteratorConst ItNormalBegin(ShapeMesh->GetVertexBuffers(), EMeshBufferSemantic::Normal);
 			UntypedMeshBufferIteratorConst ItUvsBegin(ShapeMesh->GetVertexBuffers(), EMeshBufferSemantic::TexCoords);
@@ -222,6 +219,7 @@ namespace mu
 			}
 	
 			// Don't try to use the buffer directly for weights since we only need a component.
+			ShapeWeights.SetNum(ShapeVertexCount);
 			UntypedMeshBufferIteratorConst ItUvs = ItUvsBegin;
 			for (int32 ShapeVertexIndex = 0; ShapeVertexIndex < ShapeVertexCount; ++ShapeVertexIndex)
 			{
@@ -237,8 +235,6 @@ namespace mu
 		{
 			MUTABLE_CPUPROFILER_SCOPE(GenerateTriangleQueryData);
 
-			ShapeTriangles.SetNum(ShapeTriangleCount);
-
 			const UntypedMeshBufferIteratorConst ItIndicesBegin(ShapeMesh->GetIndexBuffers(), EMeshBufferSemantic::VertexIndex);
 
 			check(ShapeMesh->GetIndexCount() % 3 == 0);
@@ -254,6 +250,7 @@ namespace mu
 			}
 			else
 			{
+				ShapeTriangles.SetNum(ShapeTriangleCount);
 				UntypedMeshBufferIteratorConst ItIndices = ItIndicesBegin;
 				for (int32 TriangleIndex = 0; TriangleIndex < ShapeTriangleCount; ++TriangleIndex)
 				{
@@ -271,11 +268,7 @@ namespace mu
 				ShapeMeshDescriptor.Triangles = TArrayView<const UE::Geometry::FIndex3i>(ShapeTriangles.GetData(), ShapeTriangles.Num());
 			}
 		}
-
-		ShapeMeshDescriptor.Positions = TArrayView<const FVector3f>(ShapePositions.GetData(), ShapePositions.Num());
-		ShapeMeshDescriptor.Normals = TArrayView<const FVector3f>(ShapeNormals.GetData(), ShapeNormals.Num());
-		ShapeMeshDescriptor.Weights = TArrayView<const float>(ShapeWeights.GetData(), ShapeWeights.Num());
-
+		
 		// Update the result mesh positions
 		int32 MeshVertexCount = BaseMesh->GetVertexCount();
 
