[mutable] Updated severity category for the log raised when a parameter is not found.
- Requested by an UDN client.

#rnx

diff --git a/Engine/Plugins/Mutable/Source/CustomizableObject/Private/MuCO/CustomizableObjectInstanceDescriptor.cpp b/Engine/Plugins/Mutable/Source/CustomizableObject/Private/MuCO/CustomizableObjectInstanceDescriptor.cpp
--- a/Engine/Plugins/Mutable/Source/CustomizableObject/Private/MuCO/CustomizableObjectInstanceDescriptor.cpp
+++ b/Engine/Plugins/Mutable/Source/CustomizableObject/Private/MuCO/CustomizableObjectInstanceDescriptor.cpp
@@ -1480,7 +1480,7 @@ bool FCustomizableObjectInstanceDescriptor::HasAnyParameters() const
  
 void LogParameterNotFoundWarning(const FString& ParameterName, const int32 ObjectParameterIndex, const int32 InstanceParameterIndex, const UCustomizableObject* CustomizableObject, char const* CallingFunction)
 {
-	UE_LOG(LogMutable, Error,
+	UE_LOG(LogMutable, Warning,
 		TEXT("%hs: Failed to find parameter (%s) on CO (%s). CO parameter index: (%d). COI parameter index: (%d)s"),
 		CallingFunction, *ParameterName, *GetNameSafe(CustomizableObject), ObjectParameterIndex, InstanceParameterIndex
 	);
