[Mutable] Added SharedPtrs to mu::FMaterials.

#rnx

diff --git a/Engine/Plugins/Mutable/Source/MutableRuntime/Internal/MuR/SystemPrivate.h b/Engine/Plugins/Mutable/Source/MutableRuntime/Internal/MuR/SystemPrivate.h
--- a/Engine/Plugins/Mutable/Source/MutableRuntime/Internal/MuR/SystemPrivate.h
+++ b/Engine/Plugins/Mutable/Source/MutableRuntime/Internal/MuR/SystemPrivate.h
@@ -497,7 +497,7 @@ namespace mu
 		TMemoryTrackedArray<TSharedPtr<const String>> StringResults;
 		TMemoryTrackedArray<TSharedPtr<const FExtensionData>> ExtensionDataResults;
 		TMemoryTrackedArray<FMatrix44f> MatrixResults;
-		TMemoryTrackedArray<FMaterial> MaterialResults;
+		TMemoryTrackedArray<TSharedPtr<const FMaterial>> MaterialResults;
 
 		/** */
 		inline const ExecutionIndex& GetRangeIndex(uint32_t i)
@@ -556,7 +556,7 @@ namespace mu
 				StringResults.Add(nullptr);
 				MatrixResults.Add(FMatrix44f());
 				ExtensionDataResults.Add(nullptr);
-				MaterialResults.Add(FMaterial());
+				MaterialResults.Add(nullptr);
 			}
 		}
 
@@ -722,7 +722,7 @@ namespace mu
 			return MatrixResults[Data->DataTypeIndex];
 		}
 
-		FMaterial GetMaterial(FCacheAddress at)
+		TSharedPtr<const FMaterial> GetMaterial(FCacheAddress at)
 		{
 			if (!at.At) return {};
 			const FOpExecutionData* Data = OpExecutionData.get_ptr(at);
@@ -926,7 +926,7 @@ namespace mu
 			check(Data->DataTypeIndex != 0);
 		}
 
-		void SetMaterial(FCacheAddress at, const FMaterial& v)
+		void SetMaterial(FCacheAddress at, TSharedPtr<const FMaterial> Value)
 		{
 			check(at.At < OpExecutionData.size_code());
 			FOpExecutionData* Data = OpExecutionData.get_ptr(at);
@@ -937,11 +937,11 @@ namespace mu
 			if (!Data->DataTypeIndex)
 			{
 				Data->DataTypeIndex = MaterialResults.Num();
-				MaterialResults.Add(v);
+				MaterialResults.Add(Value);
 			}
 			else
 			{
-				MaterialResults[Data->DataTypeIndex] = v;
+				MaterialResults[Data->DataTypeIndex] = Value;
 			}
 			check(Data->DataTypeIndex != 0);
 		}
diff --git a/Engine/Plugins/Mutable/Source/MutableRuntime/Private/MuR/CodeRunner.cpp b/Engine/Plugins/Mutable/Source/MutableRuntime/Private/MuR/CodeRunner.cpp
--- a/Engine/Plugins/Mutable/Source/MutableRuntime/Private/MuR/CodeRunner.cpp
+++ b/Engine/Plugins/Mutable/Source/MutableRuntime/Private/MuR/CodeRunner.cpp
@@ -925,7 +925,9 @@ namespace mu
 					Result = mu::CloneOrTakeOver<FInstance>(Base);
 				}
 
-				Result->AddOverlayMaterial(0, LoadMaterial(FCacheAddress(Args.Material, Item)));
+				// OverlayMaterial Material
+				TSharedPtr<FMaterial> LoadedMaterial = mu::CloneOrTakeOver<FMaterial>(LoadMaterial(FCacheAddress(Args.Material, Item)));
+				Result->AddOverlayMaterial(0, *LoadedMaterial.Get());
 
 				StoreInstance(Item, Result);
 				break;
@@ -964,7 +966,8 @@ namespace mu
 				}
 
 				// Material
-				Result->AddMaterial(0, 0, 0, LoadMaterial(FCacheAddress(Args.Material, Item)));
+				TSharedPtr<FMaterial> LoadedMaterial = mu::CloneOrTakeOver<FMaterial>(LoadMaterial(FCacheAddress(Args.Material, Item)));
+				Result->AddMaterial(0, 0, 0, *LoadedMaterial.Get());
 
 				StoreInstance(Item, Result);
 				break;
@@ -1159,8 +1162,8 @@ namespace mu
 		{
 			OP::ResourceConstantArgs Args = Program.GetOpArgs<OP::ResourceConstantArgs>(Item.At);
 
-			FMaterial Source;
-			Source.MaterialId = Args.value;
+			TSharedPtr<FMaterial> Source = MakeShared<FMaterial>();
+			Source->MaterialId = Args.value;
 
 			StoreMaterial(Item, Source);
 			break;
@@ -6170,8 +6173,8 @@ namespace mu
 			OP::ParameterArgs Args = Program.GetOpArgs<OP::ParameterArgs>(Item.At);
 			TSharedPtr<FRangeIndex> Index = BuildCurrentOpRangeIndex(Item, InParams, InModel, Args.variable);
 
-			FMaterial Material;
-			Material.Material = TStrongObjectPtr(InParams->GetMaterialValue(Args.variable, Index.Get()));
+			TSharedPtr<FMaterial> Material = MakeShared<FMaterial>();
+			Material->Material = TStrongObjectPtr(InParams->GetMaterialValue(Args.variable, Index.Get()));
 			
 			StoreMaterial(Item, Material);
 			break;
diff --git a/Engine/Plugins/Mutable/Source/MutableRuntime/Private/MuR/CodeRunner.h b/Engine/Plugins/Mutable/Source/MutableRuntime/Private/MuR/CodeRunner.h
--- a/Engine/Plugins/Mutable/Source/MutableRuntime/Private/MuR/CodeRunner.h
+++ b/Engine/Plugins/Mutable/Source/MutableRuntime/Private/MuR/CodeRunner.h
@@ -458,7 +458,7 @@ namespace  mu
 			return System->WorkingMemoryManager.CurrentInstanceCache->GetMatrix(From);
 		}
 
-		inline FMaterial LoadMaterial(const FCacheAddress& From)
+		inline TSharedPtr<const FMaterial> LoadMaterial(const FCacheAddress& From)
 		{
 			return System->WorkingMemoryManager.CurrentInstanceCache->GetMaterial(From);
 		}
@@ -533,7 +533,7 @@ namespace  mu
 			System->WorkingMemoryManager.CurrentInstanceCache->SetMatrix(To, Value);
 		}
 
-		inline void StoreMaterial(const FCacheAddress& To, const FMaterial& Value)
+		inline void StoreMaterial(const FCacheAddress& To, TSharedPtr<const FMaterial> Value)
 		{
 			System->WorkingMemoryManager.CurrentInstanceCache->SetMaterial(To, Value);
 		}
