[Mutable] Re-enable DDC.

#rnx

diff --git a/Engine/Plugins/Mutable/Source/CustomizableObject/Internal/MuCO/CustomizableObjectPrivate.h b/Engine/Plugins/Mutable/Source/CustomizableObject/Internal/MuCO/CustomizableObjectPrivate.h
--- a/Engine/Plugins/Mutable/Source/CustomizableObject/Internal/MuCO/CustomizableObjectPrivate.h
+++ b/Engine/Plugins/Mutable/Source/CustomizableObject/Internal/MuCO/CustomizableObjectPrivate.h
@@ -1950,7 +1950,7 @@ public:
 
 #endif
 	
-	static constexpr int32 DerivedDataVersion = 0x2c537a2f;
+	static constexpr int32 DerivedDataVersion = 0x290ba346;
 };
 
 #if WITH_EDITOR
diff --git a/Engine/Plugins/Mutable/Source/CustomizableObjectEditor/Private/MuCOE/CustomizableObjectEditorModule.cpp b/Engine/Plugins/Mutable/Source/CustomizableObjectEditor/Private/MuCOE/CustomizableObjectEditorModule.cpp
--- a/Engine/Plugins/Mutable/Source/CustomizableObjectEditor/Private/MuCOE/CustomizableObjectEditorModule.cpp
+++ b/Engine/Plugins/Mutable/Source/CustomizableObjectEditor/Private/MuCOE/CustomizableObjectEditorModule.cpp
@@ -131,7 +131,7 @@ constexpr int32 MaxAccessTimespan = 30;
 
 TAutoConsoleVariable<bool> CVarMutableDerivedDataCacheUsage(
 	TEXT("mutable.DerivedDataCacheUsage"),
-	false,
+	true,
 	TEXT("Derived data cache access for compiled data.")
 	TEXT("true - Allows DDC access according to the DDC policies specified in the plugin settings."),
 	ECVF_Default);
