[Mutable] Added missing COI Contains API functions.

diff --git a/Engine/Plugins/Mutable/Source/CustomizableObject/Private/MuCO/CustomizableObjectInstance.cpp b/Engine/Plugins/Mutable/Source/CustomizableObject/Private/MuCO/CustomizableObjectInstance.cpp
--- a/Engine/Plugins/Mutable/Source/CustomizableObject/Private/MuCO/CustomizableObjectInstance.cpp
+++ b/Engine/Plugins/Mutable/Source/CustomizableObject/Private/MuCO/CustomizableObjectInstance.cpp
@@ -2728,6 +2728,24 @@ bool UCustomizableObjectInstance::ContainsFloatParameter(const FString& Paramete
 }
 
 
+bool UCustomizableObjectInstance::ContainsTextureParameter(const FString& ParameterName) const
+{
+	return Descriptor.FindTypedParameterIndex(ParameterName, EMutableParameterType::Texture) != INDEX_NONE;
+}
+
+
+bool UCustomizableObjectInstance::ContainsSkeletalMeshParameter(const FString& ParameterName) const
+{
+	return Descriptor.FindTypedParameterIndex(ParameterName, EMutableParameterType::SkeletalMesh) != INDEX_NONE;
+}
+
+
+bool UCustomizableObjectInstance::ContainsMaterialParameter(const FString& ParameterName) const
+{
+	return Descriptor.FindTypedParameterIndex(ParameterName, EMutableParameterType::Material) != INDEX_NONE;
+}
+
+
 bool UCustomizableObjectInstance::ContainsBoolParameter(const FString& ParameterName) const
 {
 	return Descriptor.FindTypedParameterIndex(ParameterName, EMutableParameterType::Bool) != INDEX_NONE;
diff --git a/Engine/Plugins/Mutable/Source/CustomizableObject/Public/MuCO/CustomizableObjectInstance.h b/Engine/Plugins/Mutable/Source/CustomizableObject/Public/MuCO/CustomizableObjectInstance.h
--- a/Engine/Plugins/Mutable/Source/CustomizableObject/Public/MuCO/CustomizableObjectInstance.h
+++ b/Engine/Plugins/Mutable/Source/CustomizableObject/Public/MuCO/CustomizableObjectInstance.h
@@ -573,6 +573,18 @@ public:
 	UFUNCTION(BlueprintCallable, Category = CustomizableObjectInstance)
 	CUSTOMIZABLEOBJECT_API bool ContainsFloatParameter(const FString& ParameterName) const;
 
+	/** Return true if the Texture Parameter exists. */
+	UFUNCTION(BlueprintCallable, Category = CustomizableObjectInstance)
+	CUSTOMIZABLEOBJECT_API bool ContainsTextureParameter(const FString& ParameterName) const;
+	
+	/** Return true if the Skeletal Mesh Parameter exists. */
+	UFUNCTION(BlueprintCallable, Category = CustomizableObjectInstance)
+	CUSTOMIZABLEOBJECT_API bool ContainsSkeletalMeshParameter(const FString& ParameterName) const;
+	
+	/** Return true if the Material Parameter exists. */
+	UFUNCTION(BlueprintCallable, Category = CustomizableObjectInstance)
+	CUSTOMIZABLEOBJECT_API bool ContainsMaterialParameter(const FString& ParameterName) const;
+	
 	/** Return true if the Bool Parameter exists. */
 	UFUNCTION(BlueprintCallable, Category = CustomizableObjectInstance)
 	CUSTOMIZABLEOBJECT_API bool ContainsBoolParameter(const FString& ParameterName) const;
@@ -589,6 +601,7 @@ public:
 	UFUNCTION(BlueprintCallable, Category = CustomizableObjectInstance)
 	CUSTOMIZABLEOBJECT_API bool ContainsTransformParameter(const FString& ParameterName) const;
 	
+
 	/** Increases the range of values of the integer with ParamName and returns the index of the new integer value, -1 otherwise.
 	  * The added value is initialized with the first integer option and is the last one of the range. */
 	UFUNCTION(BlueprintCallable, Category = CustomizableObjectInstance)
