[mutable] Subsequent bakes will now override old bakes and made multiple fixes (Update and Bake related).
- Made the baking resources generated during the update able to override old baked data. The name of the resources does still represent the hash of the mutable resource/material parameters.
- Fixed baking bug that was causing the unnecessary generation of materials during the bake.
- Fixed an issue with the generation of unique names for the mutable generated resources during the update.
- Fixed issue that was causing the names of SkeletalMeshes, Skeletons and Physics materials not to be generated properly (bake)
- Created key based on a struct for the TextureReuseCache InstancePrivate map (marked as TODO)
- Added constexpr bBake variable for cases where no editor is present (as the bake will then not be possible) to optimize some part of the Instance Update runtime.
- Fixed indeterminism issue for the hash of the generated mutable meshes.
- Fixed indeterminism issue for the hash of the mutable generated FMutableMaterialPlaceholder object.
- Improved generation of hashes (added backwards compatible hashCombine path)

#rnx

diff --git a/Engine/Plugins/Mutable/Source/CustomizableObject/Internal/MuCO/CustomizableObjectInstancePrivate.h b/Engine/Plugins/Mutable/Source/CustomizableObject/Internal/MuCO/CustomizableObjectInstancePrivate.h
--- a/Engine/Plugins/Mutable/Source/CustomizableObject/Internal/MuCO/CustomizableObjectInstancePrivate.h
+++ b/Engine/Plugins/Mutable/Source/CustomizableObject/Internal/MuCO/CustomizableObjectInstancePrivate.h
@@ -205,6 +205,38 @@ enum class ESkeletalMeshStatus : uint8
 };
 
 
+USTRUCT()
+struct FTextureReuseCacheKey
+{
+	GENERATED_BODY()
+	
+	FTextureReuseCacheKey() = default;
+	FTextureReuseCacheKey(int32 InLODIndex, int32 InComponentIndex, uint32 InSurfaceIndex, int32 InImageIndex)
+	{
+		LODIndex = InLODIndex;
+		ComponentIndex = InComponentIndex;
+		SurfaceIndex = InSurfaceIndex;
+		ImageIndex = InImageIndex;
+	}
+
+	friend uint32 GetTypeHash(FTextureReuseCacheKey In)
+	{
+		return HashCombine(GetTypeHash(In.LODIndex), GetTypeHash(In.ComponentIndex), GetTypeHash(In.SurfaceIndex), GetTypeHash(In.ImageIndex));
+	}
+	
+	bool operator==(const FTextureReuseCacheKey& Other) const
+	{
+		return	LODIndex == Other.LODIndex && ComponentIndex == Other.ComponentIndex &&
+				SurfaceIndex == Other.SurfaceIndex && ImageIndex == Other.ImageIndex;
+	}
+
+	uint32 SurfaceIndex = 0;
+	int32 LODIndex = -1;
+	int32 ComponentIndex = -1;
+	int32 ImageIndex = -1;
+};
+
+
 UCLASS(MinimalAPI)
 class UCustomizableInstancePrivate : public UObject
 {
@@ -226,9 +258,8 @@ public:
 	TArray<int> RelevantParameters;
 
 	// If Texture reuse is enabled, stores which texture is being used in a particular <LODIndex, ComponentIndex, MeshSurfaceIndex, image>
-	// \TODO: Create a key based on a struct instead of generating strings dynamically.
 	UPROPERTY(Transient)
-	TMap<FString, TWeakObjectPtr<UTexture2D>> TextureReuseCache;
+	TMap<FTextureReuseCacheKey, TWeakObjectPtr<UTexture2D>> TextureReuseCache;
 	
 	// Only used in LiveUpdateMode to reuse core instances between updates and their temp data to speed up updates, but spend way more memory
 	mu::FInstance::FID LiveUpdateModeInstanceID = 0;
diff --git a/Engine/Plugins/Mutable/Source/CustomizableObject/Private/MuCO/CustomizableObjectInstance.cpp b/Engine/Plugins/Mutable/Source/CustomizableObject/Private/MuCO/CustomizableObjectInstance.cpp
--- a/Engine/Plugins/Mutable/Source/CustomizableObject/Private/MuCO/CustomizableObjectInstance.cpp
+++ b/Engine/Plugins/Mutable/Source/CustomizableObject/Private/MuCO/CustomizableObjectInstance.cpp
@@ -143,6 +143,8 @@ public:
 	bool operator==(const FMutableMaterialPlaceholder& Other) const;
 
 	friend uint32 GetTypeHash(const FMutableMaterialPlaceholder& PlaceHolder);
+	
+	friend uint32 GetHash(const FMutableMaterialPlaceholder& PlaceHolder, const bool bGenerateDeterministicHash);
 };
 
 
@@ -153,9 +155,47 @@ bool FMutableMaterialPlaceholder::operator==(const FMutableMaterialPlaceholder&
 }
 
 
-// Return a hash of the material and its parameters
+/**
+ * Makes the provided name unique and takes care of not losing the hashed value that could be present at the end of the string
+ * @warn This should only be used for mutable resources as they have the hash of their contents at the end of the string. 
+ * @param ObjectName The name to be given to a UObject that was generated with data from a mu object
+ * @param ResourceHash The hash of the resource used to create the UAsset you want to make its name unique.
+ * @param ObjectClass The class of the asset you want to make its name unique
+ */
+void MakeMutableGeneratedObjectNameUnique(FString& ObjectName, const uint32 ResourceHash ,const UClass* ObjectClass)
+{
+#if WITH_EDITOR
+	const FString HashAsString = FString::Printf(TEXT("%u"), ResourceHash);
+	ensureMsgf(ObjectName.Contains(HashAsString), TEXT("Unexpected mutable object name \"%s\" : The name was expected to end with the resource hash value of \"%u\""), *ObjectName, ResourceHash);
+#endif
+	
+	ObjectName.Append(TEXT("_"));
+	const int32 IndexOfAppendedChar = ObjectName.Len() - 1;
+	ObjectName = MakeUniqueObjectName(GetTransientPackage(), ObjectClass, *ObjectName).ToString();
+	ObjectName.RemoveAt(IndexOfAppendedChar);
+}
+
+
 uint32 GetTypeHash(const FMutableMaterialPlaceholder& PlaceHolder)
 {
+	return GetHash(PlaceHolder, false);
+}
+
+
+uint32 GetHash(const FMutableMaterialPlaceholder& PlaceHolder, bool bGenerateDeterministicHash)
+{
+	static auto FastApproach = [](const uint32 HashA, const uint32 HashB) -> uint32
+	{
+		return HashCombineFast(HashA, HashB);
+	};
+
+	static auto BackwardsCompatibleApproach = [](const uint32 HashA, const uint32 HashB) -> uint32
+	{
+		return HashCombine(HashA,HashB);
+	};
+
+	auto HashingMethod = bGenerateDeterministicHash ? BackwardsCompatibleApproach : FastApproach;
+	
 	uint32 Hash = GetTypeHash(PlaceHolder.ParentMaterialID);
 
 	// Sort parameters before building the hash.
@@ -163,26 +203,26 @@ uint32 GetTypeHash(const FMutableMaterialPlaceholder& PlaceHolder)
 
 	for (const FMutableMaterialPlaceholder::FMutableMaterialPlaceHolderParam& Param : PlaceHolder.Params)
 	{
-		uint32 ParamHash = GetTypeHash(Param.ParamName);
-		ParamHash = HashCombineFast(ParamHash, (uint32)Param.LayerIndex);
-		ParamHash = HashCombineFast(ParamHash, (uint32)Param.Type);
+		uint32 ParamHash = GetTypeHash(Param.ParamName.ToString().ToLower());
+		ParamHash = HashingMethod(ParamHash, (uint32)Param.LayerIndex);
+		ParamHash = HashingMethod(ParamHash, (uint32)Param.Type);
 
 		switch (Param.Type)
 		{
 		case FMutableMaterialPlaceholder::EPlaceHolderParamType::Vector:
-			ParamHash = HashCombineFast(ParamHash, GetTypeHash(Param.Vector));
+			ParamHash = HashingMethod(ParamHash, GetTypeHash(Param.Vector));
 			break;
 
 		case FMutableMaterialPlaceholder::EPlaceHolderParamType::Scalar:
-			ParamHash = HashCombineFast(ParamHash, GetTypeHash(Param.Scalar));
+			ParamHash = HashingMethod(ParamHash, GetTypeHash(Param.Scalar));
 			break;
 
 		case FMutableMaterialPlaceholder::EPlaceHolderParamType::Texture:
-			ParamHash = HashCombineFast(ParamHash, Param.Texture.Texture->GetUniqueID());
+			ParamHash = HashingMethod(ParamHash, GetTypeHash(Param.Texture.Texture.GetName().ToLower()));
 			break;
 		}
 
-		Hash = HashCombineFast(Hash, ParamHash);
+		Hash = HashingMethod(Hash, ParamHash);
 	}
 
 	return Hash;
@@ -2188,7 +2228,17 @@ bool UCustomizableInstancePrivate::UpdateSkeletalMesh_PostBeginUpdate0(UCustomiz
 		uint32 MeshHash = INDEX_NONE;
 		for (int32 LODIndex = 0; LODIndex < Component->LODCount && LODIndex < MAX_MESH_LOD_COUNT; ++LODIndex)
 		{
-			MeshHash = HashCombine(MeshHash, GetTypeHash((*MeshDescriptors)[LODIndex]));
+			const mu::FGeneratedMeshKey* MeshKey = (*MeshDescriptors)[LODIndex].GetKey();
+			const uint32 MeshIDHash = GetTypeHash(*MeshKey, OperationData->bBake);
+
+			if (!OperationData->bBake)
+			{
+				MeshHash = HashCombineFast(MeshHash, MeshIDHash);
+			}
+			else
+			{
+				MeshHash = HashCombine(MeshHash, MeshIDHash);
+			}
 		}
 
 		if (OperationData->bUseMeshCache)
@@ -2253,21 +2303,34 @@ bool UCustomizableInstancePrivate::UpdateSkeletalMesh_PostBeginUpdate0(UCustomiz
 		MUTABLE_CPUPROFILER_SCOPE(ConstructMesh);
 
 		USkeletalMesh* SkeletalMesh = nullptr;
-		
-		FString SkeletalMeshName = FString::Printf (TEXT("SK_%s_%s_%lu"), *Public->GetCustomizableObject()->GetName(), *ComponentName.ToString(), MeshHash);
-		SkeletalMeshName = MakeUniqueObjectName(GetTransientPackage(), USkeletalMesh::StaticClass(), *SkeletalMeshName, EUniqueObjectNameOptions::GloballyUnique).ToString();
-		if (!OperationData->bBake)
 		{
-			SkeletalMesh = NewObject<UCustomizableObjectSkeletalMesh>(GetTransientPackage(), FName(*SkeletalMeshName), RF_Transient);
-		}
-		else
-		{
-			SkeletalMesh = NewObject<USkeletalMesh>(GetTransientPackage(), FName(*SkeletalMeshName), RF_Transient);
+			FString SkeletalMeshName = FString::Printf (TEXT("SK_%s_%s_%lu"), *Public->GetCustomizableObject()->GetName(), *ComponentName.ToString(), MeshHash);
+
+#if WITH_EDITOR
+			if (OperationData->bBake)
+			{
+				SkeletalMeshName = FBakingConfiguration::BakedResourcePrefix + SkeletalMeshName;
+			}
+#endif
+
+			// Make name unique to avoid collisions with other objects.
+			MakeMutableGeneratedObjectNameUnique(SkeletalMeshName, MeshHash, USkeletalMesh::StaticClass());
+			
+			if (!OperationData->bBake)
+			{
+				SkeletalMesh = NewObject<UCustomizableObjectSkeletalMesh>(GetTransientPackage(), FName(*SkeletalMeshName), RF_Transient);
+			}
+#if WITH_EDITOR
+			else
+			{
+				SkeletalMesh = NewObject<USkeletalMesh>(GetTransientPackage(), FName(*SkeletalMeshName), RF_Transient);
+			}
+#endif
+				
+			check(SkeletalMesh);
+			SkeletalMeshes.Add(ComponentName, SkeletalMesh);
 		}
 		
-		check(SkeletalMesh);
-		SkeletalMeshes.Add(ComponentName, SkeletalMesh);
-		
 		const FMutableRefSkeletalMeshData& RefSkeletalMeshData = ModelResources.ReferenceSkeletalMeshesData[ObjectComponentIndex.GetValue()];
 
 		// Set up the default information any mesh from this component will have (LODArrayInfos, RenderData, Mesh settings, etc). 
@@ -2302,10 +2365,10 @@ bool UCustomizableInstancePrivate::UpdateSkeletalMesh_PostBeginUpdate0(UCustomiz
 		const int32 NumAdditionalPhysicsNum = ComponentMesh->AdditionalPhysicsBodies.Num();
 		for (int32 I = 0; I < NumAdditionalPhysicsNum; ++I)
 		{
-			TSharedPtr<const mu::FPhysicsBody> AdditionalPhysiscBody = ComponentMesh->AdditionalPhysicsBodies[I];
+			TSharedPtr<const mu::FPhysicsBody> AdditionalPhysicsBody = ComponentMesh->AdditionalPhysicsBodies[I];
 			
-			check(AdditionalPhysiscBody);
-			if (!AdditionalPhysiscBody->bBodiesModified)
+			check(AdditionalPhysicsBody);
+			if (!AdditionalPhysicsBody->bBodiesModified)
 			{
 				continue;
 			}
@@ -2333,7 +2396,7 @@ bool UCustomizableInstancePrivate::UpdateSkeletalMesh_PostBeginUpdate0(UCustomiz
 
 			Entry.PropertyIndex = Info.PropertyIndex;
 			Entry.PhysicsAsset = MakePhysicsAssetFromTemplateAndMutableBody(
-				OperationData,PhysicsAssetTemplate, AdditionalPhysiscBody.Get(), InstanceComponentIndex);
+				OperationData,PhysicsAssetTemplate, AdditionalPhysicsBody.Get(), InstanceComponentIndex);
 		}
 
 		// Add sockets from the SkeletalMesh of reference and from the MutableMesh
@@ -2370,7 +2433,7 @@ bool UCustomizableInstancePrivate::UpdateSkeletalMesh_PostBeginUpdate0(UCustomiz
 				UnrealConversionUtils::UpdateSkeletalMeshLODRenderDataBuffersSize(LODResource);
 			}
 		}
-
+		
 		if (OperationData->bUseMeshCache)
 		{
 			const TArray<mu::FMeshId>* MeshId = OperationData->GetMeshDescriptors(ObjectComponentIndex);
@@ -5688,8 +5751,6 @@ void UCustomizableInstancePrivate::BuildMaterials(const TSharedRef<FUpdateContex
 	TArray<FGeneratedMaterial> OldGeneratedMaterials;
 	Exchange(OldGeneratedMaterials, GeneratedMaterials);
 	
-	GeneratedMaterials.Reset();
-
 	// Prepare the data to store in order to regenerate resources for this instance (usually texture mips).
 	TSharedPtr<FMutableUpdateContext> UpdateContext = MakeShared<FMutableUpdateContext>(
 		CustomizableObject->GetPathName(),
@@ -5711,9 +5772,9 @@ void UCustomizableInstancePrivate::BuildMaterials(const TSharedRef<FUpdateContex
 			UpdateContext->bLevelBegunPlay = GWorld->GetBegunPlay();
 		}
 	}
-	
-	const bool bReuseTextures = OperationData->bReuseInstanceTextures;
 
+	const bool bReuseTextures = OperationData->bBake ? false : OperationData->bReuseInstanceTextures;
+	
 	TArray<bool> RecreateRenderStateOnInstanceComponent;
 	RecreateRenderStateOnInstanceComponent.Init(false, OperationData->NumInstanceComponents);
 
@@ -5731,18 +5792,17 @@ void UCustomizableInstancePrivate::BuildMaterials(const TSharedRef<FUpdateContex
 			continue;
 		}
 		const FName& ComponentName = ModelResources.ComponentNamesPerObjectComponent[ObjectComponentIndex.GetValue()];
-
-		TObjectPtr<USkeletalMesh>* Result = SkeletalMeshes.Find(ComponentName);
-		TObjectPtr<USkeletalMesh> SkeletalMesh = Result ? *Result : nullptr;
+		
+		TObjectPtr<USkeletalMesh> SkeletalMesh = SkeletalMeshes.Contains(ComponentName) ? SkeletalMeshes[ComponentName] : nullptr;
 		if (!SkeletalMesh)
 		{
 			continue;
 		}
 
-		const bool bReuseMaterials = !OperationData->MeshChangedPerInstanceComponent[InstanceComponentIndex];
+		const bool bReuseMaterials = !OperationData->bBake && !OperationData->MeshChangedPerInstanceComponent[InstanceComponentIndex];
 
 		// If the mesh is not transient, it means it's pass-through so it should use material overrides and not be modified in any way
-		const bool bIsTransientMesh = static_cast<bool>(SkeletalMesh->HasAllFlags(EObjectFlags::RF_Transient));
+		const bool bIsTransientMesh = SkeletalMesh->HasAllFlags(EObjectFlags::RF_Transient);
 
 		// It is not safe to replace the materials of a SkeletalMesh whose resources are initialized. Use overrides instead.
 		const bool bUseOverrideMaterialsOnly = !bIsTransientMesh || (OperationData->bUseMeshCache && SkeletalMesh->GetResourceForRendering()->IsInitialized());
@@ -6014,9 +6074,8 @@ void UCustomizableInstancePrivate::BuildMaterials(const TSharedRef<FUpdateContex
 						UTexture2D* MutableTexture = nullptr; // Texture generated by mutable
 						UTexture* PassThroughTexture = nullptr; // Texture not generated by mutable
 
-						// \TODO: Change this key to a struct.
-						FString TextureReuseCacheRef = bReuseTextures ? FString::Printf(TEXT("%d-%d-%d-%d"), Image.BaseLOD, ObjectComponentIndex.GetValue(), Surface.SurfaceId, ImageIndex) : FString();
-
+						const FTextureReuseCacheKey TextureReuseCacheRef = bReuseTextures ? FTextureReuseCacheKey{Image.BaseLOD, ObjectComponentIndex.GetValue(), Surface.SurfaceId, ImageIndex} : FTextureReuseCacheKey{};
+						
 						// If the mutable image is null, it must be in the cache
 						if (!MutableImage)
 						{
@@ -6070,11 +6129,26 @@ void UCustomizableInstancePrivate::BuildMaterials(const TSharedRef<FUpdateContex
 								// and delete it automatically
 								TSharedPtr<FTexturePlatformData, ESPMode::ThreadSafe> ReusedTexturePlatformData;
 
-								// Ensure the name of the texture is unique (prevents black textures in editor previews)
-								FString MutableTextureName = FString::Printf(TEXT("T_%s_%u"), *Props.TextureParameterName, GetTypeHash(*Image.ImageID.GetKey()));
-								MutableTextureName.ReplaceInline(TEXT(" "), TEXT("_"));
-								MutableTextureName = MakeUniqueObjectName(GetTransientPackage(), UTexture2D::StaticClass(), *MutableTextureName, EUniqueObjectNameOptions::GloballyUnique).ToString();
-									
+								FString MutableTextureName;
+								{
+									const uint32 MutableImageHash = GetTypeHash(*Image.ImageID.GetKey(), OperationData->bBake);
+									MutableTextureName = FString::Printf(TEXT("T_%s_%lu"), *Props.TextureParameterName, MutableImageHash);
+									MutableTextureName.ReplaceInline(TEXT(" "), TEXT("_"));
+
+									// If not for bake work with this resource as normal and make it unique to avoid collisions
+									if (!OperationData->bBake)
+									{
+										MakeMutableGeneratedObjectNameUnique(MutableTextureName, MutableImageHash, UMaterialInstanceDynamic::StaticClass());
+									}
+#if WITH_EDITOR
+									else
+									{
+										// If this is a bake operation just add the baked prefix to make the name unique for all baked instances of this resource
+										MutableTextureName = FBakingConfiguration::BakedResourcePrefix + MutableTextureName;
+									}
+#endif
+								}
+								
 								if (ReusedTexture && (*ReusedTexture).IsValid() && !(*ReusedTexture)->HasAnyFlags(RF_BeginDestroyed))
 								{
 									// Only uncompressed textures can be reused. This also fixes an issue in the editor where textures supposedly 
@@ -6174,26 +6248,24 @@ void UCustomizableInstancePrivate::BuildMaterials(const TSharedRef<FUpdateContex
 										{	
 											//if (!bNeverStream) // No need to check bNeverStream. In that case, the texture won't use 
 											// the MutableMipDataProviderFactory anyway and it's needed for detecting Mutable textures elsewhere
+											UMutableTextureMipDataProviderFactory* MutableMipDataProviderFactory = Cast<UMutableTextureMipDataProviderFactory>(MutableTexture->GetAssetUserDataOfClass(UMutableTextureMipDataProviderFactory::StaticClass()));
+											if (!MutableMipDataProviderFactory)
 											{
-												UMutableTextureMipDataProviderFactory* MutableMipDataProviderFactory = Cast<UMutableTextureMipDataProviderFactory>(MutableTexture->GetAssetUserDataOfClass(UMutableTextureMipDataProviderFactory::StaticClass()));
-												if (!MutableMipDataProviderFactory)
+												MutableMipDataProviderFactory = NewObject<UMutableTextureMipDataProviderFactory>();
+
+												if (MutableMipDataProviderFactory)
 												{
-													MutableMipDataProviderFactory = NewObject<UMutableTextureMipDataProviderFactory>();
-
-													if (MutableMipDataProviderFactory)
-													{
-														MutableMipDataProviderFactory->CustomizableObjectInstance = Public;
-														check(LODIndex < 256 && InstanceComponentIndex < 256 && ImageIndex < 256);
-														MutableMipDataProviderFactory->ImageRef.ImageID = Image.ImageID;
-														MutableMipDataProviderFactory->ImageRef.SurfaceId = Surface.SurfaceId;
-														MutableMipDataProviderFactory->ImageRef.LOD = uint8(Image.BaseLOD);
-														MutableMipDataProviderFactory->ImageRef.Component = uint8(InstanceComponentIndex);
-														MutableMipDataProviderFactory->ImageRef.Image = uint8(ImageIndex);
-														MutableMipDataProviderFactory->ImageRef.BaseMip = uint8(Image.BaseMip);
-														MutableMipDataProviderFactory->ImageRef.ConstantImagesNeededToGenerate = Image.ConstantImagesNeededToGenerate;
-														MutableMipDataProviderFactory->UpdateContext = UpdateContext;
-														MutableTexture->AddAssetUserData(MutableMipDataProviderFactory);
-													}
+													MutableMipDataProviderFactory->CustomizableObjectInstance = Public;
+													check(LODIndex < 256 && InstanceComponentIndex < 256 && ImageIndex < 256);
+													MutableMipDataProviderFactory->ImageRef.ImageID = Image.ImageID;
+													MutableMipDataProviderFactory->ImageRef.SurfaceId = Surface.SurfaceId;
+													MutableMipDataProviderFactory->ImageRef.LOD = uint8(Image.BaseLOD);
+													MutableMipDataProviderFactory->ImageRef.Component = uint8(InstanceComponentIndex);
+													MutableMipDataProviderFactory->ImageRef.Image = uint8(ImageIndex);
+													MutableMipDataProviderFactory->ImageRef.BaseMip = uint8(Image.BaseMip);
+													MutableMipDataProviderFactory->ImageRef.ConstantImagesNeededToGenerate = Image.ConstantImagesNeededToGenerate;
+													MutableMipDataProviderFactory->UpdateContext = UpdateContext;
+													MutableTexture->AddAssetUserData(MutableMipDataProviderFactory);
 												}
 											}
 
@@ -6201,7 +6273,11 @@ void UCustomizableInstancePrivate::BuildMaterials(const TSharedRef<FUpdateContex
 										}
 									}
 
-									TextureCache.Add(FTextureCache::FId(Image.ImageID, OperationData->MipsToSkip), MutableTexture);						
+									// During the bake do not add the texture to the cache as we do not want it to be used for other instance updates
+									if (!OperationData->bBake)
+									{
+										TextureCache.Add(FTextureCache::FId(Image.ImageID, OperationData->MipsToSkip), MutableTexture);						
+									}
 								}
 								else
 								{
@@ -6274,7 +6350,7 @@ void UCustomizableInstancePrivate::BuildMaterials(const TSharedRef<FUpdateContex
 					MUTABLE_CPUPROFILER_SCOPE(BuildMaterials_CreateMaterial);
 
 					ReuseMaterialCache.Add(MutableMaterialPlaceholder);
-
+					
 					FGeneratedMaterial& Material = GeneratedMaterials.AddDefaulted_GetRef();
 					Material.SurfaceId = Surface.SurfaceId;
 					Material.Material = Surface.Material;
@@ -6316,9 +6392,17 @@ void UCustomizableInstancePrivate::BuildMaterials(const TSharedRef<FUpdateContex
 								MIDName = MIDName.RightChop(MaterialInstanceConstantPrefix.Len());
 							}
 
-							const uint32 MaterialPlaceHolderHash = GetTypeHash(MutableMaterialPlaceholder);
+							const uint32 MaterialPlaceHolderHash = GetHash(MutableMaterialPlaceholder, OperationData->bBake);
 							MIDName = FString::Printf(TEXT("%s%s_%u"),TEXT("MID_"),*MIDName, MaterialPlaceHolderHash);
-							MIDName = MakeUniqueObjectName(GetTransientPackage(), UMaterialInstanceDynamic::StaticClass(), *MIDName, EUniqueObjectNameOptions::GloballyUnique).ToString();
+
+							// Add a "tag" to better identify this asset during the bake
+							if (OperationData->bBake)
+							{
+								MIDName = FBakingConfiguration::BakedResourcePrefix + MIDName;
+							}
+
+							MakeMutableGeneratedObjectNameUnique(MIDName, MaterialPlaceHolderHash, UMaterialInstanceDynamic::StaticClass());
+							check(MIDName.Contains(FString::Printf(TEXT("%u"), MaterialPlaceHolderHash)))
 						}
 						
 						MaterialInstance = UMaterialInstanceDynamic::Create(MaterialTemplate, GetTransientPackage(), FName(MIDName));
@@ -6513,7 +6597,7 @@ void UCustomizableInstancePrivate::BuildMaterials(const TSharedRef<FUpdateContex
 
 	{
 		MUTABLE_CPUPROFILER_SCOPE(BuildMaterials_Exchange);
-		
+
 		Exchange(GeneratedTextures, NewGeneratedTextures);
 
 		// All pass-through textures and meshes have been set, no need to keep referencing them from the instance
diff --git a/Engine/Plugins/Mutable/Source/CustomizableObject/Private/MuCO/CustomizableObjectSystem.cpp b/Engine/Plugins/Mutable/Source/CustomizableObject/Private/MuCO/CustomizableObjectSystem.cpp
--- a/Engine/Plugins/Mutable/Source/CustomizableObject/Private/MuCO/CustomizableObjectSystem.cpp
+++ b/Engine/Plugins/Mutable/Source/CustomizableObject/Private/MuCO/CustomizableObjectSystem.cpp
@@ -3219,14 +3219,18 @@ namespace impl
 		// If the program aborted, AreAnyLoadsPending may return true, as the program doesn't cancel
 		// its own loads on exit (maybe it should?)
 
-		FTextureCache& TextureCache = CustomizableObject->GetPrivate()->TextureCache;
-		
-		for (const FInstanceUpdateData::FImage& Image : OperationData->InstanceUpdateData.Images)
+		// Do not update the CachedTextures during a bake as we do not want the instance update to use cached textures
+		if (!OperationData->bBake)
 		{
-			if (const UTexture2D* Texture = TextureCache.Get(FTextureCache::FId(Image.ImageID, OperationData->MipsToSkip)))
+			FTextureCache& TextureCache = CustomizableObject->GetPrivate()->TextureCache;
+		
+			for (const FInstanceUpdateData::FImage& Image : OperationData->InstanceUpdateData.Images)
 			{
-				OperationData->Objects.Add(TStrongObjectPtr(Texture));
-				OperationData->CachedTextures.Add(Image.ImageID);
+				if (const UTexture2D* Texture = TextureCache.Get(FTextureCache::FId(Image.ImageID, OperationData->MipsToSkip)))
+				{
+					OperationData->Objects.Add(TStrongObjectPtr(Texture));
+					OperationData->CachedTextures.Add(Image.ImageID);
+				}
 			}
 		}
 
@@ -3421,8 +3425,8 @@ namespace impl
 		}
 
 		Operation->bReuseInstanceTextures = false;
-
-		if (SystemPrivate->EnableReuseInstanceTextures)
+		
+		if (!Operation->bBake && SystemPrivate->EnableReuseInstanceTextures)
 		{
 			Operation->bReuseInstanceTextures = StateData ? StateData->bReuseInstanceTextures : false;
 			Operation->bReuseInstanceTextures |= InstancePrivate->HasCOInstanceFlags(ReuseTextures);
@@ -3497,7 +3501,7 @@ namespace impl
 		//-------------------------------------------------------------
 		Operation->InstanceID = Operation->bLiveUpdateMode ? InstancePrivate->LiveUpdateModeInstanceID : 0;
 		Operation->bUseMeshCache = Object->bEnableMeshCache && !Operation->bBake && !Operation->bLiveUpdateMode && UCustomizableObjectSystem::IsMeshCacheEnabled(true);
-
+		
 		Operation->bStreamMeshLODs = IsStreamingEnabled(*Object);
 
 #if WITH_EDITOR
diff --git a/Engine/Plugins/Mutable/Source/CustomizableObject/Public/MuCO/CustomizableObjectInstance.h b/Engine/Plugins/Mutable/Source/CustomizableObject/Public/MuCO/CustomizableObjectInstance.h
--- a/Engine/Plugins/Mutable/Source/CustomizableObject/Public/MuCO/CustomizableObjectInstance.h
+++ b/Engine/Plugins/Mutable/Source/CustomizableObject/Public/MuCO/CustomizableObjectInstance.h
@@ -172,6 +172,11 @@ struct FBakingConfiguration
 	 */
 	UPROPERTY(BlueprintReadWrite, Blueprintable, Category = CustomizableObjectInstanceBaker)
 	FBakeOperationCompletedDelegate OnBakeOperationCompletedCallback;
+
+	/**
+	 * The contents of the prefix used to mark resources generated for a bake operation
+	 */
+	const static inline FString BakedResourcePrefix = TEXT("BAKE_");
 };
 
 
diff --git a/Engine/Plugins/Mutable/Source/CustomizableObjectEditor/Private/MuCOE/CustomizableObjectInstanceBakingUtils.cpp b/Engine/Plugins/Mutable/Source/CustomizableObjectEditor/Private/MuCOE/CustomizableObjectInstanceBakingUtils.cpp
--- a/Engine/Plugins/Mutable/Source/CustomizableObjectEditor/Private/MuCOE/CustomizableObjectInstanceBakingUtils.cpp
+++ b/Engine/Plugins/Mutable/Source/CustomizableObjectEditor/Private/MuCOE/CustomizableObjectInstanceBakingUtils.cpp
@@ -26,21 +26,50 @@
 #define LOCTEXT_NAMESPACE "CustomizableObjectEditor" 
 
 
+/**
+ * Removes the "_n" string found at the end of the material and skeletal mesh name.
+ * @warning It does not check if that part is there or not.
+ * @param MaterialName The name of the material that you want not to be unique.
+ */
+void MakeAssetNameNotUnique(FString& MaterialName)
+{
+	// This should be a material generated by mutable so try to remove the value that makes it unique
+	int32 UniquenessValueSeparatorIndex = INDEX_NONE;
+	MaterialName.FindLastChar(TEXT('_'),UniquenessValueSeparatorIndex);
+	if (UniquenessValueSeparatorIndex != INDEX_NONE)
+	{
+		MaterialName = MaterialName.Left(UniquenessValueSeparatorIndex);
+	}
+						
+	// Ensure that the new name contains numerical values representing the hash (speculative)
+	int32 HashValueSeparatorIndex = INDEX_NONE;
+	MaterialName.FindLastChar(TEXT('_'),HashValueSeparatorIndex);
+	if (ensure(HashValueSeparatorIndex != INDEX_NONE))
+	{
+		const FString HashValue = MaterialName.RightChop(HashValueSeparatorIndex+1);
+		// If this fails it means that the name of the material did not contain or the hash or the value that makes it unique.
+		check(HashValue.IsNumeric())
+	}
+}
+
+
 /**
  * Remove the prefix from the provided string.
  * @param InOutString The string to update
  * @param ToCheckPrefixes The list of possible prefixes that, if present, will get removed from the target string
  */
-void RemovePrefix(FString& InOutString, const TArray<FString>& ToCheckPrefixes)
+bool RemovePrefix(FString& InOutString, const TArray<FString>& ToCheckPrefixes)
 {
 	for (const FString& PrefixToRemove : ToCheckPrefixes)
 	{
 		if (InOutString.Find(PrefixToRemove, ESearchCase::CaseSensitive) == 0)
 		{
 			InOutString = InOutString.RightChop(PrefixToRemove.Len());
-			return;
+			return true;
 		}
 	}
+
+	return false;
 }
 
 
@@ -339,7 +368,7 @@ bool IsAMutableTexture (const UTexture2D* InTexture)
  */
 void AddUniqueToSavePackage(TPair<UPackage*, EPackageSaveResolutionType> NewObjectPair, TMap<UPackage*,EPackageSaveResolutionType>& OutSavedPackages)
 {
-	if (!OutSavedPackages.Find(NewObjectPair.Key))
+	if (!OutSavedPackages.Contains(NewObjectPair.Key))
 	{
 		OutSavedPackages.Add(NewObjectPair);
 	}
@@ -517,6 +546,7 @@ bool BakeCustomizableObjectInstance(
 							
 							// Update the name adding to it the name of the component alongside other data
 							FString TextureName = SourceTexture->GetName();
+							RemovePrefix(TextureName, {FBakingConfiguration::BakedResourcePrefix});
 							RemovePrefix(TextureName, {TextureAssetPrefix});
 							ComposeResourceName(TextureAssetPrefix, ObjectBaseName, TextureName);
 
@@ -606,7 +636,7 @@ bool BakeCustomizableObjectInstance(
 				
 				if (Material)
 				{
-					// If the material has already been processed (so it is already part of the ReplacementMap) just skip it as we only need
+					// If the interface has already been processed (so it is already part of the ReplacementMap) just skip it as we only need
 					// to update the Material Interface once (Material Interface <- Material)
 					if (HandledSourceObjects.Contains(Interface))
 					{
@@ -616,6 +646,12 @@ bool BakeCustomizableObjectInstance(
 					}
 					
 					FString MaterialName = Interface->GetName();
+					// If this is a material marked with the BAKE prefix then it should include a mark to make it unique. Remove it for the baking
+					// operation
+					if (RemovePrefix(MaterialName, {FBakingConfiguration::BakedResourcePrefix}))
+					{
+						MakeAssetNameNotUnique(MaterialName);
+					}
 					RemovePrefix(MaterialName, MaterialResourcePrefixes);
 					ComposeResourceName(MaterialAssetPrefix, ObjectBaseName, MaterialName);
 					
@@ -630,7 +666,7 @@ bool BakeCustomizableObjectInstance(
 					UObject* DuplicatedObject = FUnrealBakeHelpers::BakeHelper_DuplicateAsset(Material, MaterialName, PackageName, 
 						 nullptr, Configuration.bGenerateConstantMaterialInstancesOnBake, SaveResolution);
 
-					HandledSourceObjects.Add(Material);
+					HandledSourceObjects.Add(Interface);
 					DuplicatedObjects.Add(DuplicatedObject);
 
 					// Trying to add the same interface twice would mean something is not being done correctly in this method
@@ -671,7 +707,11 @@ bool BakeCustomizableObjectInstance(
 
 				FString MaterialName;
 				{
-					MaterialName = Interface->GetName();		// Use the name of the interface as it is unique
+					MaterialName = Interface->GetName();		// Use the name of the interface as it is unique because of the hash it contains
+					if (RemovePrefix(MaterialName, {FBakingConfiguration::BakedResourcePrefix}))
+					{
+						MakeAssetNameNotUnique(MaterialName);
+					}
 					RemovePrefix(MaterialName, MaterialResourcePrefixes);
 					
 					if (Configuration.bGenerateConstantMaterialInstancesOnBake && Interface->IsA<UMaterialInstance>())
@@ -759,6 +799,7 @@ bool BakeCustomizableObjectInstance(
 										
 										// Generate a new name for the duplicated texture so it does not collide with the original one
 										FString TextureName = SourceTexture->GetName();
+										RemovePrefix(TextureName, {FBakingConfiguration::BakedResourcePrefix});
 										RemovePrefix(TextureName, {TextureAssetPrefix});
 										ComposeResourceName(TextureAssetPrefix, ObjectBaseName, TextureName);
 										
@@ -810,13 +851,18 @@ bool BakeCustomizableObjectInstance(
 
 		// TODO: As it happens with the textures we may be able to re-use the other resources generated for one component into the baked data of the other
 		//		for now we will keep them separated for the assets being processed bellow
+
+		const FString CustomizableObjectName = InInstance.GetCustomizableObject()->GetName() + TEXT("_");
 		
 		// Get the clean name of the Skeletal mesh (without the embedded CO name)
 		FString SkeletalMeshName = Mesh->GetName();
 		if (Mesh->HasAnyFlags(RF_Transient))
-		{
-			const FString CustomizableObjectName = InInstance.GetCustomizableObject()->GetName() + TEXT("_");
+		{		
 			SkeletalMeshName = SkeletalMeshName.Replace(*CustomizableObjectName,TEXT(""));
+			if (RemovePrefix(SkeletalMeshName, {FBakingConfiguration::BakedResourcePrefix}))
+			{
+				MakeAssetNameNotUnique(SkeletalMeshName);
+			}
 		}
 		RemovePrefix(SkeletalMeshName, {SkeletalMeshAssetPrefix});
 		ComposeResourceName(SkeletalMeshAssetPrefix, ObjectBaseName, SkeletalMeshName);
@@ -831,6 +877,14 @@ bool BakeCustomizableObjectInstance(
 			if (bTransient || Configuration.bExportAllResourcesOnBake)
 			{
 				FString SkeletonName = Mesh->GetName();
+				if (Mesh->HasAnyFlags(RF_Transient))
+				{
+					SkeletonName = SkeletonName.Replace(*CustomizableObjectName,TEXT(""));
+					if (RemovePrefix(SkeletonName, {FBakingConfiguration::BakedResourcePrefix}))
+					{
+						MakeAssetNameNotUnique(SkeletonName);
+					}
+				}
 				RemovePrefix(SkeletonName, {SkeletalMeshAssetPrefix});
 				ComposeResourceName(SkeletonAssetPrefix, ObjectBaseName, SkeletonName);
 				MakeResourceNameUnique(SkeletonName, DuplicatedObjects);
@@ -859,6 +913,14 @@ bool BakeCustomizableObjectInstance(
 			if (bTransient || Configuration.bExportAllResourcesOnBake)
 			{
 				FString PhysicsAssetName = Mesh->GetName();
+				if (Mesh->HasAnyFlags(RF_Transient))
+				{
+					PhysicsAssetName = PhysicsAssetName.Replace(*CustomizableObjectName,TEXT(""));
+					if (RemovePrefix(PhysicsAssetName, {FBakingConfiguration::BakedResourcePrefix}))
+					{
+						MakeAssetNameNotUnique(PhysicsAssetName);
+					}
+				}
 				RemovePrefix(PhysicsAssetName, {SkeletalMeshAssetPrefix});
 				ComposeResourceName(PhysicsAssetPrefix, ObjectBaseName, PhysicsAssetName);
 				MakeResourceNameUnique(PhysicsAssetName, DuplicatedObjects);
diff --git a/Engine/Plugins/Mutable/Source/MutableRuntime/Internal/MuR/ResourceID.h b/Engine/Plugins/Mutable/Source/MutableRuntime/Internal/MuR/ResourceID.h
--- a/Engine/Plugins/Mutable/Source/MutableRuntime/Internal/MuR/ResourceID.h
+++ b/Engine/Plugins/Mutable/Source/MutableRuntime/Internal/MuR/ResourceID.h
@@ -22,8 +22,13 @@ namespace mu
 	{
 	};
 
-	
-	MUTABLERUNTIME_API uint32 GetTypeHash(const FGeneratedMeshKey& Key);
+	/**
+	 * Generates a hash based on the provided mesh key
+	 * @param Key The key you want to hash
+	 * @param bGenerateDeterministicHash Determines what hashing approach to follow.
+	 * @return The mesh key hash
+	 */
+	MUTABLERUNTIME_API uint32 GetTypeHash(const FGeneratedMeshKey& Key, const bool bGenerateDeterministicHash = false);
 
 	typedef TRegistry<FGeneratedMeshKey, FGeneratedMeshData>::FHandle FMeshId;
 
@@ -43,7 +48,13 @@ namespace mu
 	{
 	};
 	
-	MUTABLERUNTIME_API uint32 GetTypeHash(const FGeneratedImageKey& Key);
+	/**
+     * Generates a hash based on the provided image key
+     * @param Key The key you want to hash
+     * @param bGenerateDeterministicHash Determines what hashing approach to follow.
+     * @return The image key hash
+     */
+	MUTABLERUNTIME_API uint32 GetTypeHash(const FGeneratedImageKey& Key, const bool bGenerateDeterministicHash = false);
 
 	typedef TRegistry<FGeneratedImageKey, FGeneratedImageData>::FHandle FImageId;
 
diff --git a/Engine/Plugins/Mutable/Source/MutableRuntime/Private/MuR/ResourceID.cpp b/Engine/Plugins/Mutable/Source/MutableRuntime/Private/MuR/ResourceID.cpp
--- a/Engine/Plugins/Mutable/Source/MutableRuntime/Private/MuR/ResourceID.cpp
+++ b/Engine/Plugins/Mutable/Source/MutableRuntime/Private/MuR/ResourceID.cpp
@@ -5,21 +5,35 @@
 #include "Hash/xxhash.h"
 
 
-uint32 mu::GetTypeHash(const FGeneratedMeshKey& Key)
+uint32 mu::GetTypeHash(const FGeneratedMeshKey& Key, const bool bGenerateDeterministicHash)
 {
 	uint32 Hash = Key.Address;
 
-	Hash = HashCombineFast(Hash, static_cast<uint32>(FXxHash64::HashBuffer(Key.ParameterValuesBlob.GetData(), Key.ParameterValuesBlob.Num()).Hash));
+	if (!bGenerateDeterministicHash)
+	{
+		Hash = HashCombineFast(Hash, static_cast<uint32>(FXxHash64::HashBuffer(Key.ParameterValuesBlob.GetData(), Key.ParameterValuesBlob.Num()).Hash));
+	}
+	else
+	{
+		Hash = HashCombine(Hash, static_cast<uint32>(FXxHash64::HashBuffer(Key.ParameterValuesBlob.GetData(), Key.ParameterValuesBlob.Num()).Hash));
+	}
 
 	return Hash;
 }
 
 
-uint32 mu::GetTypeHash(const FGeneratedImageKey& Key)
+uint32 mu::GetTypeHash(const FGeneratedImageKey& Key, const bool bGenerateDeterministicHash)
 {
 	uint32 Hash = Key.Address;
 
-	Hash = HashCombineFast(Hash, static_cast<uint32>(FXxHash64::HashBuffer(Key.ParameterValuesBlob.GetData(), Key.ParameterValuesBlob.Num()).Hash));
+	if (!bGenerateDeterministicHash)
+	{
+		Hash = HashCombineFast(Hash, static_cast<uint32>(FXxHash64::HashBuffer(Key.ParameterValuesBlob.GetData(), Key.ParameterValuesBlob.Num()).Hash));
+	}
+	else
+	{
+		Hash = HashCombine(Hash, static_cast<uint32>(FXxHash64::HashBuffer(Key.ParameterValuesBlob.GetData(), Key.ParameterValuesBlob.Num()).Hash));
+	}
 
 	return Hash;
 }
