[mutable] Fixed localization warning in GenerateMutableSourceMaterial.cpp.
#tests Editor
#rnx

diff --git a/Engine/Plugins/Mutable/Source/CustomizableObjectEditor/Private/MuCOE/GenerateMutableSource/GenerateMutableSourceMaterial.cpp b/Engine/Plugins/Mutable/Source/CustomizableObjectEditor/Private/MuCOE/GenerateMutableSource/GenerateMutableSourceMaterial.cpp
--- a/Engine/Plugins/Mutable/Source/CustomizableObjectEditor/Private/MuCOE/GenerateMutableSource/GenerateMutableSourceMaterial.cpp
+++ b/Engine/Plugins/Mutable/Source/CustomizableObjectEditor/Private/MuCOE/GenerateMutableSource/GenerateMutableSourceMaterial.cpp
@@ -120,7 +120,7 @@ mu::Ptr<mu::NodeMaterial> GenerateMutableSourceMaterial(const UEdGraphPin* Pin,
 					}
 					else
 					{
-						const FText Message = LOCTEXT("MissingTexture", "Unable to generate material switch node. Required connection not found.");
+						const FText Message = LOCTEXT("MissingMaterial", "Unable to generate material switch node. Required connection not found.");
 						GenerationContext.Log(Message, Node);
 						return Result;
 					}
