[Mutable] Fix public API Enum Parameter functions not working with the first parameter.

- GetEnumParameterNumValues
- GetEnumParameterValue

diff --git a/Engine/Plugins/Mutable/Source/CustomizableObject/Private/MuCO/CustomizableObject.cpp b/Engine/Plugins/Mutable/Source/CustomizableObject/Private/MuCO/CustomizableObject.cpp
--- a/Engine/Plugins/Mutable/Source/CustomizableObject/Private/MuCO/CustomizableObject.cpp
+++ b/Engine/Plugins/Mutable/Source/CustomizableObject/Private/MuCO/CustomizableObject.cpp
@@ -1478,7 +1478,8 @@ void UCustomizableObjectPrivate::UpdateParameterPropertiesFromModel(const TShare
 
 int32 UCustomizableObject::GetEnumParameterNumValues(const FString& ParamName) const
 {
-	if (int32 ParamIndex = GetPrivate()->FindParameter(ParamName))
+	const int32 ParamIndex = GetPrivate()->FindParameter(ParamName);
+	if (ParamIndex != INDEX_NONE)
 	{
 		return GetPrivate()->GetEnumParameterNumValues(ParamIndex);
 	}
@@ -1491,7 +1492,8 @@ int32 UCustomizableObject::GetEnumParameterNumValues(const FString& ParamName) c
 
 const FString& UCustomizableObject::GetEnumParameterValue(const FString& ParamName, int32 ValueIndex) const
 {
-	if (int32 ParamIndex = GetPrivate()->FindParameter(ParamName))
+	const int32 ParamIndex = GetPrivate()->FindParameter(ParamName);
+	if (ParamIndex != INDEX_NONE)
 	{
 		return GetPrivate()->GetIntParameterAvailableOption(ParamIndex, ValueIndex);
 	}
@@ -1504,7 +1506,7 @@ const FString& UCustomizableObject::GetEnumParameterValue(const FString& ParamNa
 
 bool UCustomizableObject::ContainsEnumParameterValue(const FString& ParameterName, const FString Value) const
 {
-	int32 ParamIndex = GetPrivate()->FindParameter(ParameterName);
+	const int32 ParamIndex = GetPrivate()->FindParameter(ParameterName);
 	if (ParamIndex == INDEX_NONE)
 	{
 		return false;
