[Mutable] Fix assert when merge mesh with explicit vertexids generated at compile time.

- The data packer removed the vertexids buffers. Now ME_MERGE Ops will force the data packer to keep the vertexids buffer.

#rnx

diff --git a/Engine/Plugins/Mutable/Source/MutableTools/Private/MuT/DataPacker.cpp b/Engine/Plugins/Mutable/Source/MutableTools/Private/MuT/DataPacker.cpp
--- a/Engine/Plugins/Mutable/Source/MutableTools/Private/MuT/DataPacker.cpp
+++ b/Engine/Plugins/Mutable/Source/MutableTools/Private/MuT/DataPacker.cpp
@@ -28,6 +28,7 @@
 #include "MuT/ASTOpMeshRemoveMask.h"
 #include "MuT/ASTOpMeshDifference.h"
 #include "MuT/ASTOpMeshMorph.h"
+#include "MuT/ASTOpMeshMerge.h"
 #include "MuT/ASTOpMeshProject.h"
 #include "MuT/ASTOpMeshApplyLayout.h"
 #include "MuT/ASTOpMeshPrepareLayout.h"
@@ -396,6 +397,21 @@ namespace mu
                 break;
             }
 
+			case EOpType::ME_MERGE:
+			{
+				bRecurse = false;
+
+				const ASTOpMeshMerge* op = static_cast<const ASTOpMeshMerge*>(node.get());
+
+				// ME_MERGE will need VertexIds if the mesh has converted them from implicit to explicit. 
+				uint64 NewState = CurrentSemantics;
+				NewState |= (UINT64_C(1) << uint32(EMeshBufferSemantic::VertexIndex));
+
+				RecurseWithState(op->Base.child(), NewState);
+				RecurseWithState(op->Added.child(), NewState);
+				break;
+			}
+
 			case EOpType::LA_FROMMESH:
 			{
 				bRecurse = false;
