[Mutable] Added path that handles the comparison of the RequiredBones by checking the actual bone names
- Added CVar to disable this new code path.
- Instead of comparing the arrays (amount and content) directly and determining if the pose should be reinitialized or not the new approach compares the actual names of the bones being pointed at by the index set in the arrays.

#rnx

diff --git a/Engine/Plugins/Mutable/Source/CustomizableObject/Private/MuCO/CustomizableObjectSystem.cpp b/Engine/Plugins/Mutable/Source/CustomizableObject/Private/MuCO/CustomizableObjectSystem.cpp
--- a/Engine/Plugins/Mutable/Source/CustomizableObject/Private/MuCO/CustomizableObjectSystem.cpp
+++ b/Engine/Plugins/Mutable/Source/CustomizableObject/Private/MuCO/CustomizableObjectSystem.cpp
@@ -127,6 +127,12 @@ TAutoConsoleVariable<bool> CVarForceGeometryOnFirstGeneration(
 	false,
 	TEXT("If set to true, forces geometry generation on first generation even if the LOD will be streamed."));
 
+
+TAutoConsoleVariable<bool> CVarRequiresReinitCompareBoneNames(
+	TEXT("mutable.RequiresReinitCompareBoneNames"),
+	true,
+	TEXT("If set to true, instead of comparing the indices of the RequiredBones the names will get compared when checking if the bone Pose should re reinitialized or not."));
+
 #if WITH_EDITOR
 bool bEnableLODManagmentInEditor = false;
 
@@ -1184,36 +1190,82 @@ void FinishUpdateGlobal(const TSharedRef<FUpdateContextPrivate>& Context)
 }
 
 
-bool RequiresReinitPose(USkeletalMesh* CurrentSkeletalMesh, USkeletalMesh* SkeletalMesh)
+bool RequiresReinitPose(USkeletalMesh* CurrentSkeletalMesh, USkeletalMesh* NewSkeletalMesh)
 {
-	if (CurrentSkeletalMesh == SkeletalMesh)
+	MUTABLE_CPUPROFILER_SCOPE(RequiresReinitPose);
+	
+	if (CurrentSkeletalMesh == NewSkeletalMesh)
 	{
 		return false;
 	}
 
-	if (!CurrentSkeletalMesh || !SkeletalMesh)
+	if (!CurrentSkeletalMesh || !NewSkeletalMesh)
 	{
-		return SkeletalMesh != nullptr;
+		return NewSkeletalMesh != nullptr;
 	}
 
-	if (CurrentSkeletalMesh->GetLODNum() != SkeletalMesh->GetLODNum())
+	if (CurrentSkeletalMesh->GetLODNum() != NewSkeletalMesh->GetLODNum())
 	{
 		return true;
 	}
 
 	const FSkeletalMeshRenderData* CurrentRenderData = CurrentSkeletalMesh->GetResourceForRendering();
-	const FSkeletalMeshRenderData* NewRenderData = SkeletalMesh->GetResourceForRendering();
+	const FSkeletalMeshRenderData* NewRenderData = NewSkeletalMesh->GetResourceForRendering();
 	if (!CurrentRenderData || !NewRenderData)
 	{
 		return false;
 	}
 
-	const int32 NumLODs = SkeletalMesh->GetLODNum();
-	for (int32 LODIndex = 0; LODIndex < NumLODs; ++LODIndex)
+	// Instead of comparing the bone indices compare the name of the bones
+	if (CVarRequiresReinitCompareBoneNames.GetValueOnAnyThread())
 	{
-		if (CurrentRenderData->LODRenderData[LODIndex].RequiredBones != NewRenderData->LODRenderData[LODIndex].RequiredBones)
+		const FReferenceSkeleton& NewSkeletalMeshRefSkeleton = NewSkeletalMesh->GetRefSkeleton();
+		const FReferenceSkeleton& CurrentSkeletalMeshRefSkeleton = CurrentSkeletalMesh->GetRefSkeleton();
+	
+		const int32 NumLODs = NewSkeletalMesh->GetLODNum();
+		for (int32 LODIndex = 0; LODIndex < NumLODs; ++LODIndex)
 		{
-			return true;
+			const TArrayView<const FBoneIndexType> CurrentRequiredBones (CurrentRenderData->LODRenderData[LODIndex].RequiredBones);
+			const TArrayView<const FBoneIndexType> NewRequiredBones (NewRenderData->LODRenderData[LODIndex].RequiredBones);
+
+			if (CurrentRequiredBones.Num() != NewRequiredBones.Num())
+			{
+				return true;
+			}
+
+			// Iterate over the ref skeletons comparing the names of each of the bones found there by using the value for each one of the LODs
+			// If the names match we should assume that it is safe to not re-init the pose.
+		
+			const TArray<FMeshBoneInfo>& CurrentBoneInfo = CurrentSkeletalMeshRefSkeleton.GetRefBoneInfo();
+			const TArray<FMeshBoneInfo>& NewBoneInfo = NewSkeletalMeshRefSkeleton.GetRefBoneInfo();
+
+			const int32 RequiredBoneCount = CurrentRequiredBones.Num();
+			for (int32 Index = 0; Index < RequiredBoneCount; ++Index)
+			{
+				const FBoneIndexType& NewBoneIndex = NewRequiredBones[Index];
+				const FBoneIndexType& CurrentBoneIndex = CurrentRequiredBones[Index];
+
+				if (NewBoneIndex == CurrentBoneIndex)
+				{
+					continue;
+				}
+			
+				if (NewBoneInfo[NewBoneIndex].Name != CurrentBoneInfo[CurrentBoneIndex].Name)
+				{
+					return true;
+				}
+			}
+		}
+	}
+	else
+	{
+		const int32 NumLODs = NewSkeletalMesh->GetLODNum();
+		for (int32 LODIndex = 0; LODIndex < NumLODs; ++LODIndex)
+		{
+			if (CurrentRenderData->LODRenderData[LODIndex].RequiredBones != NewRenderData->LODRenderData[LODIndex].RequiredBones)
+			{
+				return true;
+			}
 		}
 	}
 
@@ -3781,7 +3833,7 @@ int32 UCustomizableObjectSystem::TickInternal(const bool bBlocking)
 		}
 
 		UCustomizableObjectInstanceUsagePrivate* UsagePrivate = CustomizableObjectInstanceUsage->GetPrivate();
-
+		
 		if (!UsagePrivate->bPendingSetSkeletalMesh)
 		{
 			continue;
