[Mutable] Add node to convert from LinearColor to sRGB.

#rnx

diff --git a/Engine/Plugins/Mutable/Source/CustomizableObjectEditor/Private/MuCOE/EdGraphSchema_CustomizableObject.cpp b/Engine/Plugins/Mutable/Source/CustomizableObjectEditor/Private/MuCOE/EdGraphSchema_CustomizableObject.cpp
--- a/Engine/Plugins/Mutable/Source/CustomizableObjectEditor/Private/MuCOE/EdGraphSchema_CustomizableObject.cpp
+++ b/Engine/Plugins/Mutable/Source/CustomizableObjectEditor/Private/MuCOE/EdGraphSchema_CustomizableObject.cpp
@@ -26,6 +26,7 @@
 #include "MuCOE/Nodes/CustomizableObjectNodeColorParameter.h"
 #include "MuCOE/Nodes/CustomizableObjectNodeColorSwitch.h"
 #include "MuCOE/Nodes/CustomizableObjectNodeColorVariation.h"
+#include "MuCOE/Nodes/CustomizableObjectNodeColorToSRGB.h"
 #include "MuCOE/Nodes/CustomizableObjectNodeCopyMaterial.h"
 #include "MuCOE/Nodes/CustomizableObjectNodeCurve.h"
 #include "MuCOE/Nodes/CustomizableObjectNodeEnumParameter.h"
@@ -540,6 +541,7 @@ void UEdGraphSchema_CustomizableObject::GetGraphContextActions(FGraphContextMenu
 			ContextMenuBuilder.CreateTemplateNode<UCustomizableObjectNodeTextureSample>(),
 			ContextMenuBuilder.CreateTemplateNode<UCustomizableObjectNodeColorArithmeticOp>(),
 			ContextMenuBuilder.CreateTemplateNode<UCustomizableObjectNodeColorFromFloats>(),
+			ContextMenuBuilder.CreateTemplateNode<UCustomizableObjectNodeColorToSRGB>(),
 		};
 
 		AddNewNodeCategoryActionsFiltered(ColorTemplateNodes, ContextMenuBuilder, TEXT("Color"), GeneralGrouping, Filter);
diff --git a/Engine/Plugins/Mutable/Source/CustomizableObjectEditor/Private/MuCOE/GenerateMutableSource/GenerateMutableSourceColor.cpp b/Engine/Plugins/Mutable/Source/CustomizableObjectEditor/Private/MuCOE/GenerateMutableSource/GenerateMutableSourceColor.cpp
--- a/Engine/Plugins/Mutable/Source/CustomizableObjectEditor/Private/MuCOE/GenerateMutableSource/GenerateMutableSourceColor.cpp
+++ b/Engine/Plugins/Mutable/Source/CustomizableObjectEditor/Private/MuCOE/GenerateMutableSource/GenerateMutableSourceColor.cpp
@@ -17,6 +17,7 @@
 #include "MuCOE/Nodes/CustomizableObjectNodeColorParameter.h"
 #include "MuCOE/Nodes/CustomizableObjectNodeColorSwitch.h"
 #include "MuCOE/Nodes/CustomizableObjectNodeColorVariation.h"
+#include "MuCOE/Nodes/CustomizableObjectNodeColorToSRGB.h"
 #include "MuCOE/Nodes/CustomizableObjectNodeMacroInstance.h"
 #include "MuCOE/Nodes/CustomizableObjectNodeTable.h"
 #include "MuCOE/Nodes/CustomizableObjectNodeTextureSample.h"
@@ -29,6 +30,7 @@
 #include "MuT/NodeColourSwitch.h"
 #include "MuT/NodeColourTable.h"
 #include "MuT/NodeColourVariation.h"
+#include "MuT/NodeColorToSRGB.h"
 #include "MuT/NodeImageBinarise.h"
 
 
@@ -274,6 +276,17 @@ mu::Ptr<mu::NodeColour> GenerateMutableSourceColor(const UEdGraphPin* Pin, FMuta
 		}
 	}
 
+	else if (const UCustomizableObjectNodeColorToSRGB* TypedNodeColorToSRGB = Cast<UCustomizableObjectNodeColorToSRGB>(Node))
+	{
+		mu::Ptr<mu::NodeColorToSRGB> ColorNode = new mu::NodeColorToSRGB();
+		Result = ColorNode;
+
+		if (const UEdGraphPin* ConnectedPin = FollowInputPin(*TypedNodeColorToSRGB->GetInputPin()))
+		{
+			ColorNode->Color = GenerateMutableSourceColor(ConnectedPin, GenerationContext);
+		}
+	}
+
 	else if (const UCustomizableObjectNodeTable* TypedNodeTable = Cast<UCustomizableObjectNodeTable>(Node))
 	{
 		//This node will add a white color in case of error
diff --git a/Engine/Plugins/Mutable/Source/CustomizableObjectEditor/Private/MuCOE/Nodes/CustomizableObjectNodeColorToSRGB.cpp b/Engine/Plugins/Mutable/Source/CustomizableObjectEditor/Private/MuCOE/Nodes/CustomizableObjectNodeColorToSRGB.cpp
new file mode 100644
--- /dev/null
+++ b/Engine/Plugins/Mutable/Source/CustomizableObjectEditor/Private/MuCOE/Nodes/CustomizableObjectNodeColorToSRGB.cpp
@@ -0,0 +1,66 @@
+// Copyright Epic Games, Inc. All Rights Reserved.
+
+#include "MuCOE/Nodes/CustomizableObjectNodeColorToSRGB.h"
+
+#include "MuCOE/EdGraphSchema_CustomizableObject.h"
+
+#include UE_INLINE_GENERATED_CPP_BY_NAME(CustomizableObjectNodeColorToSRGB)
+
+class UCustomizableObjectNodeRemapPins;
+
+#define LOCTEXT_NAMESPACE "CustomizableObjectEditor"
+
+
+UCustomizableObjectNodeColorToSRGB::UCustomizableObjectNodeColorToSRGB()
+	: Super()
+{
+}
+
+
+void UCustomizableObjectNodeColorToSRGB::AllocateDefaultPins(UCustomizableObjectNodeRemapPins* RemapPins)
+{
+	const UEdGraphSchema_CustomizableObject* Schema = GetDefault<UEdGraphSchema_CustomizableObject>();
+	
+	UEdGraphPin* InPin = CustomCreatePin(EGPD_Input, Schema->PC_Color, FName("Linear"));
+	InPin->bDefaultValueIsIgnored = true;
+	InputPin = InPin;
+
+	UEdGraphPin* OutPin = CustomCreatePin(EGPD_Output, Schema->PC_Color, FName("sRGB"));
+	OutPin->bDefaultValueIsIgnored = true;
+	OutputPin = OutPin;
+}
+
+
+FText UCustomizableObjectNodeColorToSRGB::GetNodeTitle(ENodeTitleType::Type TitleType) const
+{
+	return LOCTEXT("LinearColor_To_sRGB", "Linear To sRGB");
+}
+
+
+FLinearColor UCustomizableObjectNodeColorToSRGB::GetNodeTitleColor() const
+{
+	const UEdGraphSchema_CustomizableObject* Schema = GetDefault<UEdGraphSchema_CustomizableObject>();
+	return Schema->GetPinTypeColor(UEdGraphSchema_CustomizableObject::PC_Color);
+}
+
+
+FText UCustomizableObjectNodeColorToSRGB::GetTooltipText() const
+{
+	return LOCTEXT("Color_TO_SRGB_Tooltip", "Converts a linear color to sRGB.");
+}
+
+
+UEdGraphPin* UCustomizableObjectNodeColorToSRGB::GetInputPin() const
+{
+	return InputPin.Get();
+}
+
+
+UEdGraphPin* UCustomizableObjectNodeColorToSRGB::GetOutputPin() const
+{
+	return OutputPin.Get();
+}
+
+
+#undef LOCTEXT_NAMESPACE
+
diff --git a/Engine/Plugins/Mutable/Source/CustomizableObjectEditor/Private/MuCOE/Nodes/CustomizableObjectNodeColorToSRGB.h b/Engine/Plugins/Mutable/Source/CustomizableObjectEditor/Private/MuCOE/Nodes/CustomizableObjectNodeColorToSRGB.h
new file mode 100644
--- /dev/null
+++ b/Engine/Plugins/Mutable/Source/CustomizableObjectEditor/Private/MuCOE/Nodes/CustomizableObjectNodeColorToSRGB.h
@@ -0,0 +1,48 @@
+// Copyright Epic Games, Inc. All Rights Reserved.
+
+#pragma once
+
+#include "MuCOE/Nodes/CustomizableObjectNode.h"
+
+#include "CustomizableObjectNodeColorToSRGB.generated.h"
+
+#define UE_API CUSTOMIZABLEOBJECTEDITOR_API
+
+namespace ENodeTitleType { enum Type : int; }
+
+class UCustomizableObjectNodeRemapPins;
+class UEdGraphPin;
+
+UCLASS(MinimalAPI)
+class UCustomizableObjectNodeColorToSRGB : public UCustomizableObjectNode
+{
+public:
+	GENERATED_BODY()
+
+	UE_API UCustomizableObjectNodeColorToSRGB();
+
+	// Begin EdGraphNode interface
+	UE_API FText GetNodeTitle(ENodeTitleType::Type TitleType) const override;
+	UE_API FLinearColor GetNodeTitleColor() const override;
+	UE_API FText GetTooltipText() const override;
+	bool ShouldOverridePinNames() const override { return true; }
+
+	// UCustomizableObjectNode interface
+	UE_API void AllocateDefaultPins(UCustomizableObjectNodeRemapPins* RemapPins) override;
+
+	virtual bool IsAffectedByLOD() const { return false; }
+
+	// Own interface
+	UEdGraphPin* GetInputPin() const;
+	UEdGraphPin* GetOutputPin() const;
+
+private:
+
+	UPROPERTY()
+	FEdGraphPinReference InputPin;
+
+	UPROPERTY()
+	FEdGraphPinReference OutputPin;
+};
+
+#undef UE_API
diff --git a/Engine/Plugins/Mutable/Source/CustomizableObjectEditor/Private/MuCOE/SMutableGraphViewer.cpp b/Engine/Plugins/Mutable/Source/CustomizableObjectEditor/Private/MuCOE/SMutableGraphViewer.cpp
--- a/Engine/Plugins/Mutable/Source/CustomizableObjectEditor/Private/MuCOE/SMutableGraphViewer.cpp
+++ b/Engine/Plugins/Mutable/Source/CustomizableObjectEditor/Private/MuCOE/SMutableGraphViewer.cpp
@@ -144,6 +144,7 @@ static const char* MutableNodeNames[] =
 	"ColorArithmeticOperation",
 	"ColorSwitch",
 	"ColorVariation",
+	"ColorToSRGB",
 
 	"Scalar",
 	"ScalarConstant",
diff --git a/Engine/Plugins/Mutable/Source/MutableRuntime/Internal/MuR/Operations.h b/Engine/Plugins/Mutable/Source/MutableRuntime/Internal/MuR/Operations.h
--- a/Engine/Plugins/Mutable/Source/MutableRuntime/Internal/MuR/Operations.h
+++ b/Engine/Plugins/Mutable/Source/MutableRuntime/Internal/MuR/Operations.h
@@ -121,6 +121,9 @@ namespace mu
         //! Apply component-wise arithmetic operations to two colours
         CO_ARITHMETIC,
 
+		//! Apply a Linear to sRGB color transformation on a given color vector.
+		CO_LINEARTOSRGB,
+
         //-----------------------------------------------------------------------------------------
         // Image operations
         //-----------------------------------------------------------------------------------------
@@ -509,6 +512,11 @@ namespace mu
 			ADDRESS A, B;
 		};
 
+		struct ColorArgs
+		{
+			ADDRESS Color;
+		};
+
 
         //-------------------------------------------------------------------------------------
         struct ImageLayerArgs
diff --git a/Engine/Plugins/Mutable/Source/MutableRuntime/Private/MuR/CodeRunner.cpp b/Engine/Plugins/Mutable/Source/MutableRuntime/Private/MuR/CodeRunner.cpp
--- a/Engine/Plugins/Mutable/Source/MutableRuntime/Private/MuR/CodeRunner.cpp
+++ b/Engine/Plugins/Mutable/Source/MutableRuntime/Private/MuR/CodeRunner.cpp
@@ -6078,6 +6078,33 @@ namespace mu
             break;
         }
 
+		case EOpType::CO_LINEARTOSRGB:
+		{
+			OP::ColorArgs Args = Program.GetOpArgs<OP::ColorArgs>(Item.At);
+			switch (Item.Stage)
+			{
+			case 0:
+				AddOp(FScheduledOp(Item.At, Item, 1),
+					FScheduledOp(Args.Color, Item));
+				break;
+
+			case 1:
+			{
+				FVector4f Color = Args.Color ? LoadColor(FCacheAddress(Args.Color, Item))
+					: FVector4f(0, 0, 0, 0);
+
+				ConvertLinearColorToSRGB(Color);
+
+				StoreColor(Item, Color);
+				break;
+			}
+
+			default: unimplemented();
+			}
+
+			break;
+		}
+
         default:
             check( false );
             break;
diff --git a/Engine/Plugins/Mutable/Source/MutableRuntime/Private/MuR/CodeRunnerTasked.cpp b/Engine/Plugins/Mutable/Source/MutableRuntime/Private/MuR/CodeRunnerTasked.cpp
--- a/Engine/Plugins/Mutable/Source/MutableRuntime/Private/MuR/CodeRunnerTasked.cpp
+++ b/Engine/Plugins/Mutable/Source/MutableRuntime/Private/MuR/CodeRunnerTasked.cpp
@@ -1770,8 +1770,12 @@ namespace mu
 
 		TSharedPtr<const FImage> Source = Runner->LoadImage({ Args.Base, Op.ExecutionIndex, Op.ExecutionOptions });
 
-		// Create destination data
-		Result = Runner->CloneOrTakeOver(Source);
+		if (Source)
+		{
+			// Create destination data
+			Result = Runner->CloneOrTakeOver(Source);
+		}
+
 		return true;
 	}
 
diff --git a/Engine/Plugins/Mutable/Source/MutableRuntime/Private/MuR/Operations.cpp b/Engine/Plugins/Mutable/Source/MutableRuntime/Private/MuR/Operations.cpp
--- a/Engine/Plugins/Mutable/Source/MutableRuntime/Private/MuR/Operations.cpp
+++ b/Engine/Plugins/Mutable/Source/MutableRuntime/Private/MuR/Operations.cpp
@@ -76,6 +76,7 @@ namespace mu
 		{ EDataType::Color,			},	// CO_SWIZZLE
 		{ EDataType::Color,			},	// CO_FROMSCALARS
 		{ EDataType::Color,			},	// CO_ARITHMETIC
+		{ EDataType::Color,			},	// CO_LINEARTOSRGB
 
 		{ EDataType::Image,			},	// IM_LAYER
 		{ EDataType::Image,			},	// IM_LAYERCOLOUR
@@ -345,6 +346,13 @@ namespace mu
             f(args.B);
             break;
         }
+		
+		case EOpType::CO_LINEARTOSRGB:
+        {
+			OP::ColorArgs Args = program.GetOpArgs<OP::ColorArgs>(at);
+            f(Args.Color);
+            break;
+        }
 
         //-------------------------------------------------------------------------------------
         case EOpType::IM_LAYER:
diff --git a/Engine/Plugins/Mutable/Source/MutableRuntime/Private/MuR/System.cpp b/Engine/Plugins/Mutable/Source/MutableRuntime/Private/MuR/System.cpp
--- a/Engine/Plugins/Mutable/Source/MutableRuntime/Private/MuR/System.cpp
+++ b/Engine/Plugins/Mutable/Source/MutableRuntime/Private/MuR/System.cpp
@@ -760,7 +760,6 @@ namespace mu
             case EOpType::CO_PARAMETER:
             case EOpType::PR_PARAMETER:
 			case EOpType::IM_PARAMETER:
-			case EOpType::ME_PARAMETER:
 			case EOpType::MA_PARAMETER:
 			case EOpType::MI_PARAMETER:
             {
@@ -769,7 +768,13 @@ namespace mu
                 Flags[ParamIndex] = true;
                 break;
             }
-
+			case EOpType::ME_PARAMETER:
+			{
+				OP::MeshParameterArgs args = Program.GetOpArgs<OP::MeshParameterArgs>(at);
+				OP::ADDRESS ParamIndex = args.variable;
+				Flags[ParamIndex] = true;
+				break;
+			}
             default:
                 break;
 
diff --git a/Engine/Plugins/Mutable/Source/MutableRuntime/Public/MuR/MutableMath.h b/Engine/Plugins/Mutable/Source/MutableRuntime/Public/MuR/MutableMath.h
--- a/Engine/Plugins/Mutable/Source/MutableRuntime/Public/MuR/MutableMath.h
+++ b/Engine/Plugins/Mutable/Source/MutableRuntime/Public/MuR/MutableMath.h
@@ -3,6 +3,7 @@
 #pragma once
 
 #include "MuR/Types.h"
+#include "MathUtil.h"
 #include "Math/Matrix.h"
 #include "Math/IntVector.h"
 
@@ -76,6 +77,19 @@ namespace mu
 		return  (v0[0] <= v1[0]) && (v0[1] <= v1[1]);
 	}
 
+	/** Apply a Linear to sRGB color transformation on a given color vector. */
+	inline void ConvertLinearColorToSRGB(FVector4f& Color)
+	{
+		// Sourced from MPCDIUtils.ush
+		auto LinearToSRGBFloat = [](const float Color) -> float
+			{
+				return (Color <= 0.0031308f) ? Color * 12.92f : 1.055f * FMathf::Pow(Color, 1.0f / 2.4f) - 0.055f;
+			};
+
+		Color.X = LinearToSRGBFloat(Color.X);
+		Color.Y = LinearToSRGBFloat(Color.Y);
+		Color.Z = LinearToSRGBFloat(Color.Z);
+	}
 
 	//!
     template <class VECTOR>
diff --git a/Engine/Plugins/Mutable/Source/MutableTools/Private/MuT/ASTOpColorToSRGB.cpp b/Engine/Plugins/Mutable/Source/MutableTools/Private/MuT/ASTOpColorToSRGB.cpp
new file mode 100644
--- /dev/null
+++ b/Engine/Plugins/Mutable/Source/MutableTools/Private/MuT/ASTOpColorToSRGB.cpp
@@ -0,0 +1,78 @@
+// Copyright Epic Games, Inc. All Rights Reserved.
+
+#include "MuT/ASTOpColorToSRGB.h"
+
+#include "HAL/PlatformMath.h"
+#include "MuR/ModelPrivate.h"
+#include "MuR/RefCounted.h"
+#include "MuR/Types.h"
+
+namespace mu
+{
+
+	ASTOpColorToSRGB::ASTOpColorToSRGB()
+		: Color(this)
+	{
+	}
+
+
+	ASTOpColorToSRGB::~ASTOpColorToSRGB()
+	{
+		// Explicit call needed to avoid recursive destruction
+		ASTOp::RemoveChildren();
+	}
+
+
+	bool ASTOpColorToSRGB::IsEqual(const ASTOp& OtherUntyped) const
+	{
+		if (OtherUntyped.GetOpType() == GetOpType())
+		{
+			const ASTOpColorToSRGB* Other = static_cast<const ASTOpColorToSRGB*>(&OtherUntyped);
+			return Color == Other->Color;
+		}
+		return false;
+	}
+
+
+	uint64 ASTOpColorToSRGB::Hash() const
+	{
+		uint64 Res = std::hash<void*>()(Color.child().get());
+		return Res;
+	}
+
+
+	mu::Ptr<ASTOp> ASTOpColorToSRGB::Clone(MapChildFuncRef MapChild) const
+	{
+		mu::Ptr<ASTOpColorToSRGB> New = new ASTOpColorToSRGB();
+		New->Color = MapChild(Color.child());
+		return New;
+	}
+
+
+	void ASTOpColorToSRGB::ForEachChild(const TFunctionRef<void(ASTChild&)> Func)
+	{
+		Func(Color);
+	}
+
+
+	void ASTOpColorToSRGB::Link(FProgram& Program, FLinkerOptions*)
+	{
+		// Already linked?
+		if (!linkedAddress)
+		{
+			OP::ColorArgs Args;
+			FMemory::Memzero(Args);
+
+			if (Color)
+			{
+				Args.Color = Color->linkedAddress;
+			}
+
+			linkedAddress = (OP::ADDRESS)Program.OpAddress.Num();
+			Program.OpAddress.Add(Program.ByteCode.Num());
+			AppendCode(Program.ByteCode, GetOpType());
+			AppendCode(Program.ByteCode, Args);
+		}
+	}
+
+}
diff --git a/Engine/Plugins/Mutable/Source/MutableTools/Private/MuT/ASTOpColorToSRGB.h b/Engine/Plugins/Mutable/Source/MutableTools/Private/MuT/ASTOpColorToSRGB.h
new file mode 100644
--- /dev/null
+++ b/Engine/Plugins/Mutable/Source/MutableTools/Private/MuT/ASTOpColorToSRGB.h
@@ -0,0 +1,35 @@
+// Copyright Epic Games, Inc. All Rights Reserved.
+
+#pragma once
+
+#include "MuR/Operations.h"
+#include "MuR/Ptr.h"
+#include "MuT/AST.h"
+
+
+namespace mu
+{
+	struct FProgram;
+
+	class ASTOpColorToSRGB final : public ASTOp
+	{
+	public:
+
+		ASTChild Color;
+
+	public:
+
+		ASTOpColorToSRGB();
+		ASTOpColorToSRGB(const ASTOpColorToSRGB&) = delete;
+		~ASTOpColorToSRGB();
+
+		virtual EOpType GetOpType() const override { return EOpType::CO_LINEARTOSRGB; }
+		virtual uint64 Hash() const override;
+		virtual bool IsEqual(const ASTOp& OtherUntyped) const override;
+		virtual Ptr<ASTOp> Clone(MapChildFuncRef MapChild) const override;
+		virtual void ForEachChild(const TFunctionRef<void(ASTChild&)>) override;
+		virtual void Link(FProgram& Program, FLinkerOptions* Options) override;
+	};
+
+}
+
diff --git a/Engine/Plugins/Mutable/Source/MutableTools/Private/MuT/CodeGenerator.h b/Engine/Plugins/Mutable/Source/MutableTools/Private/MuT/CodeGenerator.h
--- a/Engine/Plugins/Mutable/Source/MutableTools/Private/MuT/CodeGenerator.h
+++ b/Engine/Plugins/Mutable/Source/MutableTools/Private/MuT/CodeGenerator.h
@@ -52,6 +52,7 @@ namespace mu
 	class NodeColourSwitch;
 	class NodeColourTable;
 	class NodeColourVariation;
+	class NodeColorToSRGB;
 	class NodeImageBinarise;
 	class NodeImageColourMap;
 	class NodeImageConditional;
@@ -874,6 +875,7 @@ namespace mu
 		void GenerateColor_FromScalars(FColorGenerationResult&, const FGenericGenerationOptions&, const Ptr<const NodeColourFromScalars>&);
 		void GenerateColor_Arithmetic(FColorGenerationResult&, const FGenericGenerationOptions&, const Ptr<const NodeColourArithmeticOperation>&);
 		void GenerateColor_Variation(FColorGenerationResult&, const FGenericGenerationOptions&, const Ptr<const NodeColourVariation>&);
+		void GenerateColor_ToSRGB(FColorGenerationResult&, const FGenericGenerationOptions&, const Ptr<const NodeColorToSRGB>&);
 		void GenerateColor_Table(FColorGenerationResult&, const FGenericGenerationOptions&, const Ptr<const NodeColourTable>&);
 		Ptr<ASTOp> GenerateMissingColourCode(const TCHAR* strWhere, const void* errorContext);
 
diff --git a/Engine/Plugins/Mutable/Source/MutableTools/Private/MuT/CodeGenerator_Color.cpp b/Engine/Plugins/Mutable/Source/MutableTools/Private/MuT/CodeGenerator_Color.cpp
--- a/Engine/Plugins/Mutable/Source/MutableTools/Private/MuT/CodeGenerator_Color.cpp
+++ b/Engine/Plugins/Mutable/Source/MutableTools/Private/MuT/CodeGenerator_Color.cpp
@@ -20,6 +20,7 @@
 #include "MuT/ASTOpColorArithmetic.h"
 #include "MuT/ASTOpColorSampleImage.h"
 #include "MuT/ASTOpColorFromScalars.h"
+#include "MuT/ASTOpColorToSRGB.h"
 #include "MuT/CodeGenerator.h"
 #include "MuT/CodeGenerator_FirstPass.h"
 #include "MuT/ErrorLog.h"
@@ -33,6 +34,7 @@
 #include "MuT/NodeColourSwitch.h"
 #include "MuT/NodeColourTable.h"
 #include "MuT/NodeColourVariation.h"
+#include "MuT/NodeColorToSRGB.h"
 #include "MuT/NodeImage.h"
 #include "MuT/NodeImageProject.h"
 #include "MuT/NodeMeshTable.h"
@@ -106,6 +108,11 @@ namespace mu
 			const NodeColourVariation* Variation = static_cast<const NodeColourVariation*>(Untyped.get());
 			GenerateColor_Variation(Result, Options, Variation);
 		}
+		else if (Untyped->GetType() == NodeColorToSRGB::GetStaticType())
+		{
+			const NodeColorToSRGB* ToSRGB = static_cast<const NodeColorToSRGB*>(Untyped.get());
+			GenerateColor_ToSRGB(Result, Options, ToSRGB);
+		}
 		else if (Untyped->GetType() == NodeColourTable::GetStaticType())
 		{
 			const NodeColourTable* Table = static_cast<const NodeColourTable*>(Untyped.get());
@@ -493,6 +500,21 @@ namespace mu
 	}
 
 
+	//-------------------------------------------------------------------------------------------------
+	void CodeGenerator::GenerateColor_ToSRGB(FColorGenerationResult& result, const FGenericGenerationOptions& Options, const Ptr<const NodeColorToSRGB>& Typed)
+	{
+		Ptr<ASTOpColorToSRGB> Op = new ASTOpColorToSRGB;
+		if (Typed->Color)
+		{
+			FColorGenerationResult ChildResult;
+			GenerateColor(ChildResult, Options, Typed->Color);
+			Op->Color = ChildResult.op;
+		}
+
+		result.op = Op;
+	}
+
+
 	//-------------------------------------------------------------------------------------------------
 	void CodeGenerator::GenerateColor_Table(FColorGenerationResult& result, const FGenericGenerationOptions& Options, const Ptr<const NodeColourTable>& Typed)
 	{
diff --git a/Engine/Plugins/Mutable/Source/MutableTools/Private/MuT/Compiler.cpp b/Engine/Plugins/Mutable/Source/MutableTools/Private/MuT/Compiler.cpp
--- a/Engine/Plugins/Mutable/Source/MutableTools/Private/MuT/Compiler.cpp
+++ b/Engine/Plugins/Mutable/Source/MutableTools/Private/MuT/Compiler.cpp
@@ -115,6 +115,7 @@ namespace mu
 		{ false,	{ 0, 0, 0, 0, 0, 0, 0, 0, 0, 0, 0, 0, 0, 0, 0, 0, 0, 0, 0, 0, 0, 0, 0, 0, 0, 0, 0, 0, 0 }		},	// CO_SWIZZLE
 		{ false,	{ 0, 0, 0, 0, 0, 0, 0, 0, 0, 0, 0, 0, 0, 0, 0, 0, 0, 0, 0, 0, 0, 0, 0, 0, 0, 0, 0, 0, 0 }		},	// CO_FROMSCALARS
 		{ false,	{ 0, 0, 0, 0, 0, 0, 0, 0, 0, 0, 0, 0, 0, 0, 0, 0, 0, 0, 0, 0, 0, 0, 0, 0, 0, 0, 0, 0, 0 }		},	// CO_ARITHMETIC
+		{ false,	{ 0, 0, 0, 0, 0, 0, 0, 0, 0, 0, 0, 0, 0, 0, 0, 0, 0, 0, 0, 0, 0, 0, 0, 0, 0, 0, 0, 0, 0 }		},	// CO_LINEARTOSRGB
 
 		{ true,		{ 0, 1, 1, 1, 0, 0, 0, 0, 0, 0, 0, 0, 0, 0, 0, 0, 0, 0, 0, 0, 0, 0, 0, 0, 0, 0, 0, 0, 0 }		},	// IM_LAYER
 		{ true,		{ 0, 1, 1, 1, 0, 0, 0, 0, 0, 0, 0, 0, 0, 0, 0, 0, 0, 0, 0, 0, 0, 0, 0, 0, 0, 0, 0, 0, 0 }		},	// IM_LAYERCOLOUR
diff --git a/Engine/Plugins/Mutable/Source/MutableTools/Private/MuT/ErrorLog.cpp b/Engine/Plugins/Mutable/Source/MutableTools/Private/MuT/ErrorLog.cpp
--- a/Engine/Plugins/Mutable/Source/MutableTools/Private/MuT/ErrorLog.cpp
+++ b/Engine/Plugins/Mutable/Source/MutableTools/Private/MuT/ErrorLog.cpp
@@ -260,6 +260,7 @@ namespace mu
 		TEXT("CO_SWIZZLE       "),
 		TEXT("CO_FROMSCALARS   "),
 		TEXT("CO_ARITHMETIC    "),
+		TEXT("CO_LINEARTOSRGB  "),
 
 		TEXT("IM_LAYER         "),
 		TEXT("IM_LAYERCOLOUR   "),
diff --git a/Engine/Plugins/Mutable/Source/MutableTools/Private/MuT/Node.cpp b/Engine/Plugins/Mutable/Source/MutableTools/Private/MuT/Node.cpp
--- a/Engine/Plugins/Mutable/Source/MutableTools/Private/MuT/Node.cpp
+++ b/Engine/Plugins/Mutable/Source/MutableTools/Private/MuT/Node.cpp
@@ -52,6 +52,7 @@
 #include "MuT/NodeColourArithmeticOperation.h"
 #include "MuT/NodeColourFromScalars.h"
 #include "MuT/NodeColourTable.h"
+#include "MuT/NodeColorToSRGB.h"
 #include "MuT/NodeScalarSwitch.h"
 #include "MuT/NodeScalarConstant.h"
 #include "MuT/NodeScalarVariation.h"
@@ -175,6 +176,7 @@ namespace mu
 	FNodeType NodeColourArithmeticOperation::StaticType = FNodeType(Node::EType::ColorArithmeticOperation, NodeColour::GetStaticType());
 	FNodeType NodeColourSampleImage::StaticType = FNodeType(Node::EType::ColorSampleImage, NodeColour::GetStaticType());
 	FNodeType NodeColourFromScalars::StaticType = FNodeType(Node::EType::ColorFromScalars, NodeColour::GetStaticType());
+	FNodeType NodeColorToSRGB::StaticType = FNodeType(Node::EType::ColorLinearToSRGB, NodeColour::GetStaticType());
 
 	FNodeType NodeMesh::StaticType = FNodeType(Node::EType::Mesh, Node::GetStaticType());
 	FNodeType NodeMeshConstant::StaticType = FNodeType(Node::EType::MeshConstant, NodeMesh::GetStaticType());
diff --git a/Engine/Plugins/Mutable/Source/MutableTools/Public/MuT/Node.h b/Engine/Plugins/Mutable/Source/MutableTools/Public/MuT/Node.h
--- a/Engine/Plugins/Mutable/Source/MutableTools/Public/MuT/Node.h
+++ b/Engine/Plugins/Mutable/Source/MutableTools/Public/MuT/Node.h
@@ -92,6 +92,7 @@ namespace mu
 			ColorArithmeticOperation,
 			ColorSwitch,
 			ColorVariation,
+			ColorLinearToSRGB,
 
 			Scalar,
 			ScalarConstant,
diff --git a/Engine/Plugins/Mutable/Source/MutableTools/Public/MuT/NodeColorToSRGB.h b/Engine/Plugins/Mutable/Source/MutableTools/Public/MuT/NodeColorToSRGB.h
new file mode 100644
--- /dev/null
+++ b/Engine/Plugins/Mutable/Source/MutableTools/Public/MuT/NodeColorToSRGB.h
@@ -0,0 +1,44 @@
+// Copyright Epic Games, Inc. All Rights Reserved.
+
+#pragma once
+
+#include "MuR/Ptr.h"
+#include "MuR/RefCounted.h"
+#include "MuT/Node.h"
+#include "MuT/NodeColour.h"
+
+#define UE_API MUTABLETOOLS_API
+
+
+namespace mu
+{
+
+	/** Convert a linear color to sRGB
+	*/
+	class NodeColorToSRGB : public NodeColour
+	{
+	public:
+
+		Ptr<NodeColour> Color;
+
+	public:
+
+		// Node interface
+		virtual const FNodeType* GetType() const { return &StaticType; }
+		static const FNodeType* GetStaticType() { return &StaticType; }
+
+	protected:
+
+		/** Forbidden. Manage with the Ptr<> template. */
+		~NodeColorToSRGB() {}
+
+	private:
+
+		static UE_API FNodeType StaticType;
+
+	};
+
+
+}
+
+#undef UE_API
