[Mutable] Add missing version bump for Linear to sRGB op.

#rnx

diff --git a/Engine/Plugins/Mutable/Source/CustomizableObject/Internal/MuCO/CustomizableObjectPrivate.h b/Engine/Plugins/Mutable/Source/CustomizableObject/Internal/MuCO/CustomizableObjectPrivate.h
--- a/Engine/Plugins/Mutable/Source/CustomizableObject/Internal/MuCO/CustomizableObjectPrivate.h
+++ b/Engine/Plugins/Mutable/Source/CustomizableObject/Internal/MuCO/CustomizableObjectPrivate.h
@@ -1678,6 +1678,8 @@ enum class ECustomizableObjectVersions : int32
 
 	MaterialParameter,
 
+	AddLinearTosRGB,
+
 	// -----<new versions can be added above this line>--------
 	LastCustomizableObjectVersion
 };
