[Mutable] Added new developer node Schema.

- Split Project and Editor settings.

diff --git a/Engine/Plugins/Mutable/Source/CustomizableObjectEditor/Private/MuCOE/CustomizableObjectEditorModule.cpp b/Engine/Plugins/Mutable/Source/CustomizableObjectEditor/Private/MuCOE/CustomizableObjectEditorModule.cpp
--- a/Engine/Plugins/Mutable/Source/CustomizableObjectEditor/Private/MuCOE/CustomizableObjectEditorModule.cpp
+++ b/Engine/Plugins/Mutable/Source/CustomizableObjectEditor/Private/MuCOE/CustomizableObjectEditorModule.cpp
@@ -4,6 +4,7 @@
 
 #include "Animation/AnimBlueprint.h"
 #include "AssetToolsModule.h"
+#include "CustomizableObjectProjectSettings.h"
 #include "GameFramework/Pawn.h"
 #include "ISettingsModule.h"
 #include "ISettingsSection.h"
@@ -387,33 +388,38 @@ void FCustomizableObjectEditorModule::RegisterSettings()
 
 	if (SettingsModule != nullptr)
     {
-        ISettingsSectionPtr SettingsSectionPtr = SettingsModule->RegisterSettings("Project", "Plugins", "CustomizableObjectSettings",
-            LOCTEXT("MutableSettings_Setting", "Mutable"),
-            LOCTEXT("MutableSettings_Setting_Desc", "Mutable Settings"),
-            GetMutableDefault<UCustomizableObjectEditorSettings>()
-        );
+		UCustomizableObjectEditorSettings* EditorSettings = GetMutableDefault<UCustomizableObjectEditorSettings>();
 
-        if (SettingsSectionPtr.IsValid())
-        {
-            SettingsSectionPtr->OnModified().BindRaw(this, &FCustomizableObjectEditorModule::HandleSettingsSaved);
-        }
+        ISettingsSectionPtr SettingsSectionPtr = SettingsModule->RegisterSettings("Editor", "Plugins", "CustomizableObjectSettings",
+            LOCTEXT("MutableEditorSettings_Setting", "Mutable"),
+            LOCTEXT("MutableEditorSettings_Setting_Desc", "Mutable Settings"),
+            EditorSettings);
 
+		if (SettingsSectionPtr.IsValid())
+		{
+			SettingsSectionPtr->OnModified().BindRaw(this, &FCustomizableObjectEditorModule::HandleSettingsSaved);
+		}
+		
 		if (UCustomizableObjectSystem::GetInstance() != nullptr)
 		{
-			UCustomizableObjectEditorSettings* CustomizableObjectSettings = GetMutableDefault<UCustomizableObjectEditorSettings>();
-			if (CustomizableObjectSettings != nullptr)
+			if (EditorSettings != nullptr)
 			{
 				FEditorCompileSettings CompileSettings;
-				CompileSettings.bIsMutableEnabled = !CustomizableObjectSettings->bDisableMutableCompileInEditor;
-				CompileSettings.bEnableAutomaticCompilation = CustomizableObjectSettings->bEnableAutomaticCompilation;
-				CompileSettings.bCompileObjectsSynchronously = CustomizableObjectSettings->bCompileObjectsSynchronously;
-				CompileSettings.bCompileRootObjectsOnStartPIE = CustomizableObjectSettings->bCompileRootObjectsOnStartPIE;
-				CompileSettings.EditorDerivedDataCachePolicy = CustomizableObjectSettings->EditorDerivedDataCachePolicy;
-				CompileSettings.CookDerivedDataCachePolicy = CustomizableObjectSettings->CookDerivedDataCachePolicy;
+				CompileSettings.bIsMutableEnabled = !EditorSettings->bDisableMutableCompileInEditor;
+				CompileSettings.bEnableAutomaticCompilation = EditorSettings->bEnableAutomaticCompilation;
+				CompileSettings.bCompileObjectsSynchronously = EditorSettings->bCompileObjectsSynchronously;
+				CompileSettings.bCompileRootObjectsOnStartPIE = EditorSettings->bCompileRootObjectsOnStartPIE;
+				CompileSettings.EditorDerivedDataCachePolicy = EditorSettings->EditorDerivedDataCachePolicy;
+				CompileSettings.CookDerivedDataCachePolicy = EditorSettings->CookDerivedDataCachePolicy;
 
 				UCustomizableObjectSystem::GetInstance()->EditorSettingsChanged(CompileSettings);
 			}
 		}
+		
+		SettingsModule->RegisterSettings("Project", "Plugins", "CustomizableObjectSettings",
+			LOCTEXT("MutableProjectSettings_Setting", "Mutable"),
+			LOCTEXT("MutableProjectSettings_Setting_Desc", "Mutable Settings"),
+			GetMutableDefault<UCustomizableObjectProjectSettings>());
     }
 }
 
diff --git a/Engine/Plugins/Mutable/Source/CustomizableObjectEditor/Private/MuCOE/CustomizableObjectEditorSettings.cpp b/Engine/Plugins/Mutable/Source/CustomizableObjectEditor/Private/MuCOE/CustomizableObjectEditorSettings.cpp
deleted file mode 100644
diff --git a/Engine/Plugins/Mutable/Source/CustomizableObjectEditor/Private/MuCOE/CustomizableObjectEditorSettings.h b/Engine/Plugins/Mutable/Source/CustomizableObjectEditor/Private/MuCOE/CustomizableObjectEditorSettings.h
--- a/Engine/Plugins/Mutable/Source/CustomizableObjectEditor/Private/MuCOE/CustomizableObjectEditorSettings.h
+++ b/Engine/Plugins/Mutable/Source/CustomizableObjectEditor/Private/MuCOE/CustomizableObjectEditorSettings.h
@@ -2,20 +2,20 @@
 
 #pragma once
 
+#include "MuCO/CustomizableObjectSystemPrivate.h"
 
 #include "CustomizableObjectEditorSettings.generated.h"
 
 enum class ECustomizableObjectDDCPolicy : uint8;
 
-UCLASS(config = Engine)
+
+/** COE Module, Editor Settings. */
+UCLASS(config = EditorPerProjectUserSettings)
 class UCustomizableObjectEditorSettings : public UObject
 {
 	GENERATED_BODY()
 
 public:
-
-	UCustomizableObjectEditorSettings();
-
 	/** If true, Mutable won't compile any COs in the Editor. */
 	UPROPERTY(config, EditAnywhere, Category = Compilation)
 	bool bDisableMutableCompileInEditor;
@@ -32,9 +32,12 @@ public:
 	UPROPERTY(config, EditAnywhere, Category = AutomaticCompilation)
 	bool bCompileRootObjectsOnStartPIE = false;
 
+	UPROPERTY(config, EditAnywhere, Category = Developer)
+	bool bEnableDeveloperOptions = false;
+	
 	UPROPERTY(config, EditAnywhere, Category = DerivedDataCache, DisplayName = "DDC policy for editor compilations")
-	ECustomizableObjectDDCPolicy EditorDerivedDataCachePolicy;
+	ECustomizableObjectDDCPolicy EditorDerivedDataCachePolicy = ECustomizableObjectDDCPolicy::Default;
 	
 	UPROPERTY(config, EditAnywhere, Category = DerivedDataCache, DisplayName = "DDC policy for cook compilations")
-	ECustomizableObjectDDCPolicy CookDerivedDataCachePolicy;
+	ECustomizableObjectDDCPolicy CookDerivedDataCachePolicy = ECustomizableObjectDDCPolicy::Default;
 };
diff --git a/Engine/Plugins/Mutable/Source/CustomizableObjectEditor/Private/MuCOE/CustomizableObjectProjectSettings.h b/Engine/Plugins/Mutable/Source/CustomizableObjectEditor/Private/MuCOE/CustomizableObjectProjectSettings.h
new file mode 100644
--- /dev/null
+++ b/Engine/Plugins/Mutable/Source/CustomizableObjectEditor/Private/MuCOE/CustomizableObjectProjectSettings.h
@@ -0,0 +1,15 @@
+﻿// Copyright Epic Games, Inc. All Rights Reserved.
+
+#pragma once
+
+#include "CustomizableObjectProjectSettings.generated.h"
+
+
+/** COE Module, Project Settings. */
+UCLASS(config = Mutable, defaultconfig)
+class UCustomizableObjectProjectSettings : public UObject
+{
+	GENERATED_BODY()
+
+public:
+};
diff --git a/Engine/Plugins/Mutable/Source/CustomizableObjectEditor/Private/MuCOE/EdGraphSchema_CustomizableObject.cpp b/Engine/Plugins/Mutable/Source/CustomizableObjectEditor/Private/MuCOE/EdGraphSchema_CustomizableObject.cpp
--- a/Engine/Plugins/Mutable/Source/CustomizableObjectEditor/Private/MuCOE/EdGraphSchema_CustomizableObject.cpp
+++ b/Engine/Plugins/Mutable/Source/CustomizableObjectEditor/Private/MuCOE/EdGraphSchema_CustomizableObject.cpp
@@ -4,6 +4,7 @@
 #include "MuCOE/CustomizableObjectSchemaActions.h"
 
 #include "CustomizableObjectConnectionDrawingPolicy.h"
+#include "CustomizableObjectEditorSettings.h"
 #include "EdGraphNode_Comment.h"
 #include "Engine/SkeletalMesh.h"
 #include "Engine/StaticMesh.h"
@@ -100,9 +101,11 @@
 #include "MuCOE/Nodes/CustomizableObjectNodeComponentMeshAddTo.h"
 #include "MuCOE/Nodes/CustomizableObjectNodeTransformConstant.h"
 #include "MuCOE/Nodes/CustomizableObjectNodeTransformParameter.h"
+#include "Nodes/CONodeSchema.h"
 
 #include UE_INLINE_GENERATED_CPP_BY_NAME(EdGraphSchema_CustomizableObject)
 
+class UCustomizableObjectEditorSettings;
 class IToolkit;
 
 
@@ -691,6 +694,14 @@ void UEdGraphSchema_CustomizableObject::GetGraphContextActions(FGraphContextMenu
 		Action->NodeTemplate = Node;
 	}
 
+	const UCustomizableObjectEditorSettings* EditorSettings = GetDefault<UCustomizableObjectEditorSettings>();
+	if (EditorSettings->bEnableDeveloperOptions)
+	{
+		UCONodeSchema* Node = NewObject<UCONodeSchema>();
+		TSharedPtr<FCustomizableObjectSchemaAction_NewNode> Action = AddNewNodeAction(ContextMenuBuilder, TEXT("Developer"), Node->GetNodeTitle(ENodeTitleType::ListView), FText(), 1);
+		Action->NodeTemplate = Node;
+	}
+	
 	{
 		UCustomizableObjectNodeReroute* Node = NewObject<UCustomizableObjectNodeReroute>();
 		TSharedPtr<FCustomizableObjectSchemaAction_NewNode> Action = AddNewNodeAction(ContextMenuBuilder, FString(), Node->GetNodeTitle(ENodeTitleType::ListView), FText(), 1);
diff --git a/Engine/Plugins/Mutable/Source/CustomizableObjectEditor/Private/MuCOE/Nodes/CONodeSchema.cpp b/Engine/Plugins/Mutable/Source/CustomizableObjectEditor/Private/MuCOE/Nodes/CONodeSchema.cpp
new file mode 100644
--- /dev/null
+++ b/Engine/Plugins/Mutable/Source/CustomizableObjectEditor/Private/MuCOE/Nodes/CONodeSchema.cpp
@@ -0,0 +1,60 @@
+﻿// Copyright Epic Games, Inc. All Rights Reserved.
+
+#include "MuCOE/Nodes/CONodeSchema.h"
+
+#include "MuCOE/RemapPins/CustomizableObjectNodeRemapPinsByPosition.h"
+
+
+FText UCONodeSchema::GetNodeTitle(ENodeTitleType::Type TitleType) const
+{
+	return Title;
+}
+
+
+FLinearColor UCONodeSchema::GetNodeTitleColor() const
+{
+	return TitleColor;
+}
+
+
+void UCONodeSchema::AllocateDefaultPins(UCustomizableObjectNodeRemapPins* RemapPins)
+{
+	Super::AllocateDefaultPins(RemapPins);
+	
+	for (const FPinSchema& PinSchema : PinSchemas)
+	{
+		const EEdGraphPinDirection Direction = PinSchema.bIsInput ? EGPD_Input : EGPD_Output;
+
+		UEdGraphPin* Pin = CreatePin(Direction, UEdGraphSchema_CustomizableObject::PC_Wildcard, FName(PinSchema.FriendlyName.ToString()));
+		Pin->PinType.ContainerType = PinSchema.bIsArray ? EPinContainerType::Array : EPinContainerType::None;
+		Pin->PinFriendlyName = PinSchema.FriendlyName;
+
+		UCONodeSchemaPinData* PinData = NewObject<UCONodeSchemaPinData>(this);
+		PinData->Color = PinSchema.Color;
+		
+		AddPinData(*Pin, *PinData);
+	}
+}
+
+
+FLinearColor UCONodeSchema::GetPinColor(const UEdGraphPin& Pin) const
+{
+	return GetPinData<UCONodeSchemaPinData>(Pin).Color;
+}
+
+
+UCustomizableObjectNodeRemapPins* UCONodeSchema::CreateRemapPinsDefault() const
+{
+	return NewObject<UCustomizableObjectNodeRemapPinsByPosition>();
+}
+
+
+void UCONodeSchema::PostEditChangeProperty(FPropertyChangedEvent& PropertyChangedEvent)
+{
+	Super::PostEditChangeProperty(PropertyChangedEvent);
+
+	if (PropertyChangedEvent.Property)
+	{
+		ReconstructNode();
+	}
+}
diff --git a/Engine/Plugins/Mutable/Source/CustomizableObjectEditor/Private/MuCOE/Nodes/CONodeSchema.h b/Engine/Plugins/Mutable/Source/CustomizableObjectEditor/Private/MuCOE/Nodes/CONodeSchema.h
new file mode 100644
--- /dev/null
+++ b/Engine/Plugins/Mutable/Source/CustomizableObjectEditor/Private/MuCOE/Nodes/CONodeSchema.h
@@ -0,0 +1,79 @@
+﻿// Copyright Epic Games, Inc. All Rights Reserved.
+
+#pragma once
+
+#include "MuCOE/EdGraphSchema_CustomizableObject.h"
+#include "MuCOE/Nodes/CustomizableObjectNode.h"
+
+#include "CONodeSchema.generated.h"
+
+#define LOCTEXT_NAMESPACE "CustomizableObjectEditor"
+
+
+UCLASS()
+class UCONodeSchemaPinData : public UCustomizableObjectNodePinData
+{
+	GENERATED_BODY()
+
+public:	
+	UPROPERTY()
+	FLinearColor Color;
+};
+
+
+USTRUCT()
+struct FPinSchema
+{
+	GENERATED_BODY()
+	
+	UPROPERTY(EditAnywhere, Category = Schema)
+	FText FriendlyName = LOCTEXT("Pin", "Pin");
+	
+	UPROPERTY(EditAnywhere, Category = Schema)
+	bool bIsInput = false;
+
+	UPROPERTY(EditAnywhere, Category = Schema)
+	bool bIsArray = false;
+
+	UPROPERTY(EditAnywhere, Category = Schema)
+	bool bEdit = false;
+	
+	UPROPERTY(EditAnywhere, Category = Schema)
+	FLinearColor Color = FLinearColor::White;
+};
+
+
+UCLASS(MinimalAPI)
+class UCONodeSchema : public UCustomizableObjectNode
+{
+public:
+	GENERATED_BODY()
+
+	// UObject interface
+	virtual void PostEditChangeProperty(FPropertyChangedEvent& PropertyChangedEvent) override;
+
+	// UEdGraphNode interface
+	virtual FText GetNodeTitle(ENodeTitleType::Type TitleType) const override;
+	virtual FLinearColor GetNodeTitleColor() const override;
+
+	// UCustomizableObjectNode interface
+	virtual void AllocateDefaultPins(UCustomizableObjectNodeRemapPins* RemapPins) override;
+	virtual FLinearColor GetPinColor(const UEdGraphPin &Pin) const override;
+
+protected:
+	virtual UCustomizableObjectNodeRemapPins* CreateRemapPinsDefault() const override;
+
+public:
+	UPROPERTY(EditAnywhere, Category = Schema)
+	FText Title = LOCTEXT("Schema", "Schema");
+
+	UPROPERTY(EditAnywhere, Category = Schema)
+	FLinearColor TitleColor;
+	
+	UPROPERTY(EditAnywhere, Category = Schema)
+	TArray<FPinSchema> PinSchemas;
+};
+
+
+#undef LOCTEXT_NAMESPACE
+
diff --git a/Engine/Plugins/Mutable/Source/CustomizableObjectEditor/Private/MuCOE/Nodes/CustomizableObjectNode.cpp b/Engine/Plugins/Mutable/Source/CustomizableObjectEditor/Private/MuCOE/Nodes/CustomizableObjectNode.cpp
--- a/Engine/Plugins/Mutable/Source/CustomizableObjectEditor/Private/MuCOE/Nodes/CustomizableObjectNode.cpp
+++ b/Engine/Plugins/Mutable/Source/CustomizableObjectEditor/Private/MuCOE/Nodes/CustomizableObjectNode.cpp
@@ -754,6 +754,13 @@ bool UCustomizableObjectNode::CanPinBeHidden(const UEdGraphPin& Pin) const
 }
 
 
+FLinearColor UCustomizableObjectNode::GetPinColor(const UEdGraphPin& Pin) const
+{
+	const UEdGraphSchema_CustomizableObject* Schema = GetDefault<UEdGraphSchema_CustomizableObject>();
+	return Schema->GetPinTypeColor(Pin.PinType);
+}
+
+
 bool UCustomizableObjectNode::HasPinViewer() const
 {
 	return false;
diff --git a/Engine/Plugins/Mutable/Source/CustomizableObjectEditor/Private/MuCOE/Nodes/SCustomizableObjectNodePin.cpp b/Engine/Plugins/Mutable/Source/CustomizableObjectEditor/Private/MuCOE/Nodes/SCustomizableObjectNodePin.cpp
--- a/Engine/Plugins/Mutable/Source/CustomizableObjectEditor/Private/MuCOE/Nodes/SCustomizableObjectNodePin.cpp
+++ b/Engine/Plugins/Mutable/Source/CustomizableObjectEditor/Private/MuCOE/Nodes/SCustomizableObjectNodePin.cpp
@@ -4,6 +4,7 @@
 
 #include "MuCOE/EdGraphSchema_CustomizableObject.h"
 #include "MuCOE/UnrealEditorPortabilityHelpers.h"
+#include "MuCOE/Nodes/CustomizableObjectNode.h"
 
 
 void SCustomizableObjectNodePin::Construct(const FArguments& InArgs, UEdGraphPin* InGraphPinObj)
@@ -29,3 +30,33 @@ const FSlateBrush* SCustomizableObjectNodePin::GetPinIcon() const
 		return SGraphPin::GetPinIcon();
 	}
 }
+
+
+FSlateColor SCustomizableObjectNodePin::GetPinColor() const
+{
+	UEdGraphPin* GraphPin = GetPinObj();
+	if (GraphPin && !GraphPin->IsPendingKill())
+	{
+		if (bIsDiffHighlighted)
+		{
+			return FSlateColor(FLinearColor(0.9f, 0.2f, 0.15f));
+		}
+		if (GraphPin->bOrphanedPin)
+		{
+			return FSlateColor(FLinearColor::Red);
+		}
+		
+		UCustomizableObjectNode* Node = CastChecked<UCustomizableObjectNode>(GetPinObj()->GetOwningNode());
+
+		const FLinearColor Color = Node->GetPinColor(*GraphPin);
+			
+		if (!Node->IsNodeEnabled() || Node->IsDisplayAsDisabledForced() || !IsEditingEnabled() || Node->IsNodeUnrelated())
+		{
+			return Color * FLinearColor(1.0f, 1.0f, 1.0f, 0.5f);
+		}
+
+		return Color * PinColorModifier;
+	}
+
+	return FLinearColor::White;
+}
diff --git a/Engine/Plugins/Mutable/Source/CustomizableObjectEditor/Private/MuCOE/Nodes/SCustomizableObjectNodePin.h b/Engine/Plugins/Mutable/Source/CustomizableObjectEditor/Private/MuCOE/Nodes/SCustomizableObjectNodePin.h
--- a/Engine/Plugins/Mutable/Source/CustomizableObjectEditor/Private/MuCOE/Nodes/SCustomizableObjectNodePin.h
+++ b/Engine/Plugins/Mutable/Source/CustomizableObjectEditor/Private/MuCOE/Nodes/SCustomizableObjectNodePin.h
@@ -21,6 +21,7 @@ public:
 	// SGraphPin interface
 protected:
 	virtual const FSlateBrush* GetPinIcon() const override;
+	virtual FSlateColor GetPinColor() const override;
 
 private:
 	const FSlateBrush* PassThroughImageConnected = nullptr;
diff --git a/Engine/Plugins/Mutable/Source/CustomizableObjectEditor/Public/MuCOE/Nodes/CustomizableObjectNode.h b/Engine/Plugins/Mutable/Source/CustomizableObjectEditor/Public/MuCOE/Nodes/CustomizableObjectNode.h
--- a/Engine/Plugins/Mutable/Source/CustomizableObjectEditor/Public/MuCOE/Nodes/CustomizableObjectNode.h
+++ b/Engine/Plugins/Mutable/Source/CustomizableObjectEditor/Public/MuCOE/Nodes/CustomizableObjectNode.h
@@ -269,6 +269,8 @@ public:
 	/** Return true if a pin can be hidden. Override it to have a custom behaviour. */
 	UE_API virtual bool CanPinBeHidden(const UEdGraphPin& Pin) const;
 
+	UE_API virtual FLinearColor GetPinColor(const UEdGraphPin &Pin) const;
+	
 	/** @return true if the details should spawn a Pin Viewer. */
 	UE_API virtual bool HasPinViewer() const;
 
