[Mutable] Allow data tables to have rows with missing data.

#rnx

diff --git a/Engine/Plugins/Mutable/Source/CustomizableObject/Internal/MuCO/CustomizableObjectPrivate.h b/Engine/Plugins/Mutable/Source/CustomizableObject/Internal/MuCO/CustomizableObjectPrivate.h
--- a/Engine/Plugins/Mutable/Source/CustomizableObject/Internal/MuCO/CustomizableObjectPrivate.h
+++ b/Engine/Plugins/Mutable/Source/CustomizableObject/Internal/MuCO/CustomizableObjectPrivate.h
@@ -1950,6 +1950,9 @@ public:
 	UPROPERTY()
 	int32 ImageTiling = 0;
 
+	UPROPERTY()
+	bool bDisableTableMissingDataWarning = false;
+
 #endif
 	
 	static constexpr int32 DerivedDataVersion = 0x290ba346;
diff --git a/Engine/Plugins/Mutable/Source/CustomizableObjectEditor/Private/MuCOE/CustomizableObjectDetails.cpp b/Engine/Plugins/Mutable/Source/CustomizableObjectEditor/Private/MuCOE/CustomizableObjectDetails.cpp
--- a/Engine/Plugins/Mutable/Source/CustomizableObjectEditor/Private/MuCOE/CustomizableObjectDetails.cpp
+++ b/Engine/Plugins/Mutable/Source/CustomizableObjectEditor/Private/MuCOE/CustomizableObjectDetails.cpp
@@ -92,12 +92,25 @@ void FCustomizableObjectDetails::CustomizeDetails( IDetailLayoutBuilder& DetailB
 	IDetailCategoryBuilder& CustomizableObjectCategory = DetailBuilder.EditCategory("CustomizableObject");
 	IDetailCategoryBuilder& StatesCategory = DetailBuilder.EditCategory( "States" );
 	IDetailCategoryBuilder& VersioningCategory = DetailBuilder.EditCategory( "Versioning" );
+	IDetailCategoryBuilder& CompileOptions = DetailBuilder.EditCategory("CompileOptions");
 	//StatesCategory.CategoryIcon( "ActorClassIcon.CustomizableObject" );
 
 	DetailBuilder.HideProperty("States");
 	DetailBuilder.HideProperty("CustomizableObjectClassTags");
 	DetailBuilder.HideProperty("PopulationClassTags");
 
+	if (CustomizableObject && CustomizableObject->GetPrivate())
+	{
+		TArray<TSharedRef<IPropertyHandle>> CompileOptionsProperties;
+		CompileOptions.GetDefaultProperties(CompileOptionsProperties);
+		for (TSharedRef<IPropertyHandle>& Property : CompileOptionsProperties)
+		{
+			CompileOptions.AddProperty(Property);
+		}
+
+		IDetailPropertyRow* Row = CompileOptions.AddExternalObjectProperty({ CustomizableObject->GetPrivate() }, GET_MEMBER_NAME_CHECKED(UCustomizableObjectPrivate, bDisableTableMissingDataWarning));
+	}
+
 	// Make the tree get automatically updated each time we compile the host CO
 	if (CustomizableObject && !CustomizableObject->GetPrivate()->PostCompileDelegate.IsBoundToObject(this))
 	{
diff --git a/Engine/Plugins/Mutable/Source/CustomizableObjectEditor/Private/MuCOE/GenerateMutableSource/GenerateMutableSourceMaterial.cpp b/Engine/Plugins/Mutable/Source/CustomizableObjectEditor/Private/MuCOE/GenerateMutableSource/GenerateMutableSourceMaterial.cpp
--- a/Engine/Plugins/Mutable/Source/CustomizableObjectEditor/Private/MuCOE/GenerateMutableSource/GenerateMutableSourceMaterial.cpp
+++ b/Engine/Plugins/Mutable/Source/CustomizableObjectEditor/Private/MuCOE/GenerateMutableSource/GenerateMutableSourceMaterial.cpp
@@ -228,8 +228,12 @@ mu::Ptr<mu::NodeMaterial> GenerateMutableSourceMaterial(const UEdGraphPin* Pin,
 					check(GenerationContext.ReferencedMaterials.IsValidIndex(GenerationContext.CurrentReferencedMaterialIndex));
 					UMaterialInterface* BaseMaterial = GenerationContext.ReferencedMaterials[GenerationContext.CurrentReferencedMaterialIndex];
 
+					const bool bTableMaterialCheckDisabled = GenerationContext.CompilationContext.Object->GetPrivate()->IsTableMaterialsParentCheckDisabled();
+					const bool bMaterialParentMismatch = !bTableMaterialCheckDisabled && BaseMaterial
+						&& TableMaterial->GetMaterial() != BaseMaterial->GetMaterial();
+
 					// Checking if the reference material of the Table Node has the same parent as the material of the Material Node 
-					if (!BaseMaterial || TableMaterial->GetMaterial() != BaseMaterial->GetMaterial())
+					if (bMaterialParentMismatch)
 					{
 						GenerationContext.Log(LOCTEXT("DifferentParentMaterial", "The Default Material of the Data Table and the Mesh Section must have the same Parent Material."), Node);
 						bSuccess = false;
diff --git a/Engine/Plugins/Mutable/Source/CustomizableObjectEditor/Private/MuCOE/GenerateMutableSource/GenerateMutableSourceMesh.cpp b/Engine/Plugins/Mutable/Source/CustomizableObjectEditor/Private/MuCOE/GenerateMutableSource/GenerateMutableSourceMesh.cpp
--- a/Engine/Plugins/Mutable/Source/CustomizableObjectEditor/Private/MuCOE/GenerateMutableSource/GenerateMutableSourceMesh.cpp
+++ b/Engine/Plugins/Mutable/Source/CustomizableObjectEditor/Private/MuCOE/GenerateMutableSource/GenerateMutableSourceMesh.cpp
@@ -880,6 +880,9 @@ TSharedPtr<mu::FMesh> ConvertSkeletalMeshToMutable(
 		return nullptr;
 	}
 
+	// Data tables may allow missing data. Don't throw warnings of missing LODs or sections.
+	const bool bIgnoreMissingData = Source.TableReferenceSkeletalMesh || Context.Object->GetPrivate()->bDisableTableMissingDataWarning;
+
 	if (!ImportedModel->LODModels.IsValidIndex(Source.LODIndex))
 	{
 		if (!Source.bMeshMustExist)
@@ -888,13 +891,15 @@ TSharedPtr<mu::FMesh> ConvertSkeletalMeshToMutable(
 		}
 		else
 		{
-			const FString Msg = FString::Printf(
-				TEXT("The SkeletalMesh [%s] doesn't have the expected number of LODs [need %d, has %d]. Changed after reimporting?"),
-				*MeshName,
-				Source.LODIndex + 1,
-				ImportedModel->LODModels.Num());
-			Context.Log(FText::FromString(Msg), Source.MessageContext);
-
+			if (!bIgnoreMissingData)
+			{
+				const FString Msg = FString::Printf(
+					TEXT("The SkeletalMesh [%s] doesn't have the expected number of LODs [need %d, has %d]. Changed after reimporting?"),
+					*MeshName,
+					Source.LODIndex + 1,
+					ImportedModel->LODModels.Num());
+				Context.Log(FText::FromString(Msg), Source.MessageContext);
+			}
 			return nullptr;
 		}
 	}
@@ -908,14 +913,17 @@ TSharedPtr<mu::FMesh> ConvertSkeletalMeshToMutable(
 		}
 		else
 		{
-			const FString Msg = FString::Printf(
-				TEXT("The SkeletalMesh [%s] doesn't have the expected structure. Maybe the number of LODs [need %d, has %d] or Materials [need %d, has %d] has changed after reimporting?"),
-				*SkeletalMesh->GetName(),
-				Source.LODIndex + 1,
-				ImportedModel->LODModels.Num(),
-				Source.SectionIndex + 1,
-				LODModel.Sections.Num());
-			Context.Log(FText::FromString(Msg), Source.MessageContext);
+			if (!bIgnoreMissingData)
+			{
+				const FString Msg = FString::Printf(
+					TEXT("The SkeletalMesh [%s] doesn't have the expected structure. Maybe the number of LODs [need %d, has %d] or Materials [need %d, has %d] has changed after reimporting?"),
+					*SkeletalMesh->GetName(),
+					Source.LODIndex + 1,
+					ImportedModel->LODModels.Num(),
+					Source.SectionIndex + 1,
+					LODModel.Sections.Num());
+				Context.Log(FText::FromString(Msg), Source.MessageContext);
+			}
 			return nullptr;
 		}
 	}
diff --git a/Engine/Plugins/Mutable/Source/CustomizableObjectEditor/Private/MuCOE/GenerateMutableSource/GenerateMutableSourceTable.cpp b/Engine/Plugins/Mutable/Source/CustomizableObjectEditor/Private/MuCOE/GenerateMutableSource/GenerateMutableSourceTable.cpp
--- a/Engine/Plugins/Mutable/Source/CustomizableObjectEditor/Private/MuCOE/GenerateMutableSource/GenerateMutableSourceTable.cpp
+++ b/Engine/Plugins/Mutable/Source/CustomizableObjectEditor/Private/MuCOE/GenerateMutableSource/GenerateMutableSourceTable.cpp
@@ -36,6 +36,8 @@ bool FillTableColumn(const UCustomizableObjectNodeTable* TableNode, mu::Ptr<mu::
 	UDataTable* DataTablePtr = GetDataTable(TableNode, GenerationContext);
 	FString ColumnPropertyName = ColumnProperty->GetAuthoredName();
 
+	const bool bIgnoreMissingData = GenerationContext.CompilationContext.Object->GetPrivate()->bDisableTableMissingDataWarning;
+
 	// Getting property type
 	if (const FSoftObjectProperty* SoftObjectProperty = CastField<FSoftObjectProperty>(ColumnProperty))
 	{
@@ -128,7 +130,7 @@ bool FillTableColumn(const UCustomizableObjectNodeTable* TableNode, mu::Ptr<mu::
 									}
 								}
 							}
-							else
+							else if (!bIgnoreMissingData)
 							{
 								FString msg = FString::Printf(TEXT("Could not find the Slot column of the animation blueprint column [%s] for the mesh column [%s] row [%s]."), *AnimBP, *ColumnName, *RowName);
 								LogRowGenerationMessage(TableNode, DataTablePtr, GenerationContext, msg, RowName);
@@ -136,7 +138,7 @@ bool FillTableColumn(const UCustomizableObjectNodeTable* TableNode, mu::Ptr<mu::
 						}
 					}
 				}
-				else
+				else if(!bIgnoreMissingData)
 				{
 					FString msg = FString::Printf(TEXT("Could not found the Slot column of the animation blueprint column [%s] for the mesh column [%s]."), *AnimBP, *ColumnName);
 					GenerationContext.Log(FText::FromString(msg), TableNode);
@@ -208,7 +210,7 @@ bool FillTableColumn(const UCustomizableObjectNodeTable* TableNode, mu::Ptr<mu::
 			
 			if (ImportedModel->LODModels.IsValidIndex(LODIndex)) // Ignore error since this Section is empty due to Automatic LODs From Mesh
 			{
-				if (!ImportedModel->LODModels[LODIndex].Sections.IsValidIndex(SectionIndex))
+				if (!ImportedModel->LODModels[LODIndex].Sections.IsValidIndex(SectionIndex) && !bIgnoreMissingData)
 				{
 					FString msg = FString::Printf(TEXT("Mesh from column [%s] row [%s] does not have section %d at LOD %d"), *ColumnName, *RowName, SectionIndexConnected, LODIndex);
 					LogRowGenerationMessage(TableNode, DataTablePtr, GenerationContext, msg, RowName);
@@ -396,9 +398,12 @@ bool FillTableColumn(const UCustomizableObjectNodeTable* TableNode, mu::Ptr<mu::
 			{
 				Texture = TableNode->GetColumnDefaultAssetByType<UTexture>(ColumnPropertyName);
 
-				FString Message = Cast<UObject>(Object) ? "not a suported Texture" : "null";
-				FString WarningMessage = FString::Printf(TEXT("Texture from column [%s] row [%s] is %s. The default texture will be used instead."), *ColumnName, *RowName, *Message);
-				LogRowGenerationMessage(TableNode, DataTablePtr, GenerationContext, WarningMessage, RowName);
+				if (!bIgnoreMissingData)
+				{
+					FString Message = Cast<UObject>(Object) ? "not a suported Texture" : "null";
+					FString WarningMessage = FString::Printf(TEXT("Texture from column [%s] row [%s] is %s. The default texture will be used instead."), *ColumnName, *RowName, *Message);
+					LogRowGenerationMessage(TableNode, DataTablePtr, GenerationContext, WarningMessage, RowName);
+				}
 			}
 
 			// There will be always one of the two options
@@ -440,12 +445,21 @@ bool FillTableColumn(const UCustomizableObjectNodeTable* TableNode, mu::Ptr<mu::
 			const bool bMaterialParentMismatch = !bTableMaterialCheckDisabled && MaterialInstance
 												 && ReferenceMaterial->GetMaterial() != MaterialInstance->GetMaterial();
 
-			if (!MaterialInstance || bMaterialParentMismatch)
+			if (bMaterialParentMismatch)
 			{
-				FText Warning;
+				FText Warning = FText::Format(LOCTEXT("MatInstanceFromDifferentParent", "Material Instance from column [{0}] row [{1}] has a different Material Parent than the Default Material Instance. The Default Material Instance will be used instead."),
+					FText::FromString(MaterialColumnDisplayName), FText::FromString(RowName));
+				LogRowGenerationMessage(TableNode, DataTablePtr, GenerationContext, Warning.ToString(), RowName);
+				
+				MaterialInstance = ReferenceMaterial;
+			}
 
-				if (!MaterialInstance)
+			if (!MaterialInstance)
+			{
+				if (!bIgnoreMissingData)
 				{
+					FText Warning;
+
 					if (UMaterial* Material = Cast<UMaterial>(Object))
 					{
 						Warning = FText::Format(LOCTEXT("IsAMaterial", "Asset from column [{0}] row [{1}] is a Material and not a MaterialInstance. The default Material Instance will be used instead."),
@@ -456,18 +470,13 @@ bool FillTableColumn(const UCustomizableObjectNodeTable* TableNode, mu::Ptr<mu::
 						Warning = FText::Format(LOCTEXT("NullMaterialInstance", "Material Instance from column [{0}] row [{1}] is null. The default Material Instance will be used instead."),
 							FText::FromString(MaterialColumnDisplayName), FText::FromString(RowName));
 					}
-				}
-				else
-				{
-					Warning = FText::Format(LOCTEXT("MatInstanceFromDifferentParent","Material Instance from column [{0}] row [{1}] has a different Material Parent than the Default Material Instance. The Default Material Instance will be used instead."),
-						FText::FromString(MaterialColumnDisplayName), FText::FromString(RowName));
+
+					LogRowGenerationMessage(TableNode, DataTablePtr, GenerationContext, Warning.ToString(), RowName);
 				}
 
 				MaterialInstance = ReferenceMaterial;
-
-				LogRowGenerationMessage(TableNode, DataTablePtr, GenerationContext, Warning.ToString(), RowName);
 			}
-			
+
 			if (GenerationContext.CurrentTableColumnType == mu::ETableColumnType::Material)
 			{
 				CurrentColumn = MutableTable.get()->FindColumn(ColumnName);
diff --git a/Engine/Plugins/Mutable/Source/MutableTools/Private/MuT/CodeGenerator.cpp b/Engine/Plugins/Mutable/Source/MutableTools/Private/MuT/CodeGenerator.cpp
--- a/Engine/Plugins/Mutable/Source/MutableTools/Private/MuT/CodeGenerator.cpp
+++ b/Engine/Plugins/Mutable/Source/MutableTools/Private/MuT/CodeGenerator.cpp
@@ -3104,7 +3104,7 @@ namespace mu
 
 				if (!BaseMeshResults.ExtraMeshLayouts.IsValidIndex(ThisEditIndex))
 				{
-					ErrorLog->Add(TEXT("Trying to extend a layout that doesn't exist."), ELMT_WARNING, Edit->GetMessageContext(), ModifiedNodeErrorContext);
+					ErrorLog->Add(TEXT("Trying to extend a layout that doesn't exist."), ELMT_INFO, Edit->GetMessageContext(), ModifiedNodeErrorContext);
 					continue;
 				}
 
@@ -3112,7 +3112,7 @@ namespace mu
 
 				if (LayoutIndex >= ExtraLayouts.Num() || !ExtraLayouts[LayoutIndex].Layout)
 				{
-					ErrorLog->Add(TEXT("Trying to extend a layout that doesn't exist."), ELMT_WARNING, Edit->GetMessageContext(), ModifiedNodeErrorContext);
+					ErrorLog->Add(TEXT("Trying to extend a layout that doesn't exist."), ELMT_INFO, Edit->GetMessageContext(), ModifiedNodeErrorContext);
 					continue;
 				}
 
diff --git a/Engine/Plugins/Mutable/Source/MutableTools/Private/MuT/CodeGenerator.h b/Engine/Plugins/Mutable/Source/MutableTools/Private/MuT/CodeGenerator.h
--- a/Engine/Plugins/Mutable/Source/MutableTools/Private/MuT/CodeGenerator.h
+++ b/Engine/Plugins/Mutable/Source/MutableTools/Private/MuT/CodeGenerator.h
@@ -1036,7 +1036,7 @@ namespace mu
 		}
         else if (ColIndex < 0)
         {
-            ErrorLog->Add("Table column not found.", ELMT_ERROR, node.GetMessageContext());
+            ErrorLog->Add("Table column not found.", ELMT_INFO, node.GetMessageContext());
             return nullptr;
         }
 
@@ -1110,7 +1110,7 @@ namespace mu
 		}
         else if (ColIndex < 0)
         {
-            ErrorLog->Add("Table column not found.", ELMT_ERROR, node.GetMessageContext());
+            ErrorLog->Add("Table column not found.", ELMT_INFO, node.GetMessageContext());
             return;
         }
 
