[Mutable] BeginUpdate Game Thread

BeginUpdate is now split in two:
- BeginUpdate_GameThread
- BeginUpdate_MutableThread

BeginUpdate_GameThread runs different CodeRunner that has been optimized to run as fast a possible, reducing its time on the Game Thread.

FMaterial now has its own FMaterialId. Similar to Meshes and Images, the Material Id is generated on the BeginUpdate and the Material is evaluated on demand using the id.

diff --git a/Engine/Plugins/Mutable/Source/CustomizableObject/Internal/MuCO/CustomizableObjectMipDataProvider.h b/Engine/Plugins/Mutable/Source/CustomizableObject/Internal/MuCO/CustomizableObjectMipDataProvider.h
--- a/Engine/Plugins/Mutable/Source/CustomizableObject/Internal/MuCO/CustomizableObjectMipDataProvider.h
+++ b/Engine/Plugins/Mutable/Source/CustomizableObject/Internal/MuCO/CustomizableObjectMipDataProvider.h
@@ -39,6 +39,7 @@ public:
 		const TSharedPtr<mu::FModel>& InModel,
 		const TSharedPtr<mu::FMeshIdRegistry>& InMeshIdRegistry,
 		const TSharedPtr<mu::FImageIdRegistry>& InImageIdRegistry,
+		const TSharedPtr<mu::FMaterialIdRegistry>& InMaterialIdRegistry,
 		const TSharedPtr<FModelStreamableBulkData>& InModelStreamable,
 		const TSharedPtr<const mu::FParameters>& InParameters,
 		int32 InState);
@@ -75,6 +76,7 @@ private:
 public:
 	TSharedPtr<mu::FMeshIdRegistry> MeshIdRegistry;
 	TSharedPtr<mu::FImageIdRegistry> ImageIdRegistry;
+	TSharedPtr<mu::FMaterialIdRegistry> MaterialIdRegistry;
 
 private:
 	TSharedPtr<FModelStreamableBulkData> ModelStreamableBulkData;
diff --git a/Engine/Plugins/Mutable/Source/CustomizableObject/Internal/MuCO/CustomizableObjectPrivate.h b/Engine/Plugins/Mutable/Source/CustomizableObject/Internal/MuCO/CustomizableObjectPrivate.h
--- a/Engine/Plugins/Mutable/Source/CustomizableObject/Internal/MuCO/CustomizableObjectPrivate.h
+++ b/Engine/Plugins/Mutable/Source/CustomizableObject/Internal/MuCO/CustomizableObjectPrivate.h
@@ -1680,6 +1680,8 @@ enum class ECustomizableObjectVersions : int32
 
 	AddLinearTosRGB,
 
+	AddMaterialArgs,
+
 	// -----<new versions can be added above this line>--------
 	LastCustomizableObjectVersion
 };
@@ -1890,6 +1892,8 @@ public:
 	TSharedRef<mu::FMeshIdRegistry> MeshIdRegistry = MakeShared<mu::FMeshIdRegistry>();
 	
 	TSharedRef<mu::FImageIdRegistry> ImageIdRegistry = MakeShared<mu::FImageIdRegistry>();
+
+	TSharedRef<mu::FMaterialIdRegistry> MaterialIdRegistry = MakeShared<mu::FMaterialIdRegistry>();
 	
 	// See UCustomizableObjectSystem::LockObject. Must only be modified from the game thread
 	bool bLocked = false;
diff --git a/Engine/Plugins/Mutable/Source/CustomizableObject/Internal/MuCO/CustomizableObjectSystemPrivate.h b/Engine/Plugins/Mutable/Source/CustomizableObject/Internal/MuCO/CustomizableObjectSystemPrivate.h
--- a/Engine/Plugins/Mutable/Source/CustomizableObject/Internal/MuCO/CustomizableObjectSystemPrivate.h
+++ b/Engine/Plugins/Mutable/Source/CustomizableObject/Internal/MuCO/CustomizableObjectSystemPrivate.h
@@ -342,7 +342,7 @@ struct FGeneratedMaterial
 	uint32 SurfaceId = 0;
 
 	// Core Material. Strong object pointer to a material or a material index
-	mu::FMaterial Material;
+	TSharedPtr<const mu::FMaterial> Material;
 
 #if WITH_EDITORONLY_DATA
 	// Name of the component that contains this material
@@ -421,7 +421,8 @@ struct FInstanceUpdateData
 		/** Id of the metadata associated with this surface. */ 
 		uint32 SurfaceMetadataId = 0;
 
-		mu::FMaterial Material;
+		mu::FMaterialId MaterialId;
+		TSharedPtr<const mu::FMaterial> Material;
 	};
 
 	struct FLOD
@@ -451,7 +452,8 @@ struct FInstanceUpdateData
 		uint16 LODCount = 0;
 
 		/** Overlay material. Only one per mesh, if any. */
-		mu::FMaterial OverlayMaterial;
+		mu::FMaterialId OverlayMaterialId;
+		TSharedPtr<const mu::FMaterial> OverlayMaterial;
 	};
 
 	TArray<FComponent> Components;
@@ -621,6 +623,7 @@ public:
 	TSharedPtr<mu::FModel> Model;
 	TSharedPtr<mu::FMeshIdRegistry> MeshIdRegistry;
 	TSharedPtr<mu::FImageIdRegistry> ImageIdRegistry;
+	TSharedPtr<mu::FMaterialIdRegistry> MaterialIdRegistry;
 
 	// Number of possible components in the entire CO
 	// TODO: Redundant because it is in the CO
diff --git a/Engine/Plugins/Mutable/Source/CustomizableObject/Private/MuCO/CustomizableObjectInstance.cpp b/Engine/Plugins/Mutable/Source/CustomizableObject/Private/MuCO/CustomizableObjectInstance.cpp
--- a/Engine/Plugins/Mutable/Source/CustomizableObject/Private/MuCO/CustomizableObjectInstance.cpp
+++ b/Engine/Plugins/Mutable/Source/CustomizableObject/Private/MuCO/CustomizableObjectInstance.cpp
@@ -4922,8 +4922,13 @@ UE::Tasks::FTask UCustomizableInstancePrivate::LoadAdditionalAssetsAndData(const
 	
 	for (const FInstanceUpdateData::FSurface& Surface : OperationData->InstanceUpdateData.Surfaces)
 	{
-		const int32 MaterialIndex = Surface.Material.MaterialId;
-		if (MaterialIndex == INDEX_NONE || ObjectToInstanceIndexMap.Contains(MaterialIndex))
+		if (!Surface.Material)
+		{
+			continue;
+		}
+		
+		const int32 MaterialIndex = Surface.Material->MaterialId;
+		if (ObjectToInstanceIndexMap.Contains(MaterialIndex))
 		{
 			continue;
 		}
@@ -4942,16 +4947,22 @@ UE::Tasks::FTask UCustomizableInstancePrivate::LoadAdditionalAssetsAndData(const
 
 	for (const FInstanceUpdateData::FComponent& Component : OperationData->InstanceUpdateData.Components)
 	{
-		if (Component.OverlayMaterial.MaterialId == INDEX_NONE || ObjectToInstanceIndexMap.Contains(Component.OverlayMaterial.MaterialId))
+		if (!Component.OverlayMaterial)
+		{
+			continue;
+		}
+		
+		const int32 MaterialIndex = Component.OverlayMaterial->MaterialId;
+		if (ObjectToInstanceIndexMap.Contains(Component.OverlayMaterial->MaterialId))
 		{
 			continue;
 		}
 
-		TSoftObjectPtr<UMaterialInterface> AssetPtr = ModelResources.Materials.IsValidIndex(Component.OverlayMaterial.MaterialId) ? ModelResources.Materials[Component.OverlayMaterial.MaterialId] : nullptr;
+		TSoftObjectPtr<UMaterialInterface> AssetPtr = ModelResources.Materials.IsValidIndex(MaterialIndex) ? ModelResources.Materials[Component.OverlayMaterial->MaterialId] : nullptr;
 		UMaterialInterface* LoadedMaterial = AssetPtr.Get();
 
 		const int32 ReferencedMaterialsIndex = ReferencedMaterials.Add(LoadedMaterial);
-		ObjectToInstanceIndexMap.Add(Component.OverlayMaterial.MaterialId, ReferencedMaterialsIndex);
+		ObjectToInstanceIndexMap.Add(MaterialIndex, ReferencedMaterialsIndex);
 
 		if (!LoadedMaterial && !AssetPtr.IsNull())
 		{
@@ -5760,6 +5771,7 @@ void UCustomizableInstancePrivate::BuildMaterials(const TSharedRef<FUpdateContex
 		OperationData->Model,
 		CustomizableObject->GetPrivate()->MeshIdRegistry,
 	    CustomizableObject->GetPrivate()->ImageIdRegistry,
+	    CustomizableObject->GetPrivate()->MaterialIdRegistry,
 		CustomizableObject->GetPrivate()->GetModelStreamableBulkData(),
 		OperationData->Parameters,
 	    OperationData->GetCapturedDescriptor().GetState());
@@ -5816,21 +5828,24 @@ void UCustomizableInstancePrivate::BuildMaterials(const TSharedRef<FUpdateContex
 			ComponentData->OverrideMaterials.Reset();
 			ComponentData->OverlayMaterial = nullptr;
 
-			if (Component.OverlayMaterial.MaterialId != INDEX_NONE)
+			if (Component.OverlayMaterial)
 			{
-				if (const uint32* ReferencedMaterialIndex = ObjectToInstanceIndexMap.Find(Component.OverlayMaterial.MaterialId))
+				if (Component.OverlayMaterial->MaterialId != INDEX_NONE)
 				{
-					if (ReferencedMaterials.IsValidIndex(*ReferencedMaterialIndex))
+					if (const uint32* ReferencedMaterialIndex = ObjectToInstanceIndexMap.Find(Component.OverlayMaterial->MaterialId))
 					{
-						ComponentData->OverlayMaterial = ReferencedMaterials[*ReferencedMaterialIndex];
-						OverlayMaterial = ComponentData->OverlayMaterial;
+						if (ReferencedMaterials.IsValidIndex(*ReferencedMaterialIndex))
+						{
+							ComponentData->OverlayMaterial = ReferencedMaterials[*ReferencedMaterialIndex];
+							OverlayMaterial = ComponentData->OverlayMaterial;
+						}
 					}
 				}
-			}
-			else
-			{
-				ComponentData->OverlayMaterial = Component.OverlayMaterial.Material.Get();
-				OverlayMaterial = ComponentData->OverlayMaterial;
+				else
+				{
+					ComponentData->OverlayMaterial = Component.OverlayMaterial->Material.Get();
+					OverlayMaterial = ComponentData->OverlayMaterial;
+				}
 			}
 		}
 
@@ -5874,17 +5889,6 @@ void UCustomizableInstancePrivate::BuildMaterials(const TSharedRef<FUpdateContex
 			{
 				const FInstanceUpdateData::FSurface& Surface = OperationData->InstanceUpdateData.Surfaces[LOD.FirstSurface + SurfaceIndex];
 
-				// Is this a material in a passthrough mesh that we don't modify?
-				if (Surface.Material.MaterialId == INDEX_NONE && !Surface.Material.Material)
-				{
-					Materials.Emplace();
-#if WITH_EDITOR
-					// Without this, a change of a referenced material and recompilation doesn't show up in the preview.
-					RecreateRenderStateOnInstanceComponent[InstanceComponentIndex] = true;
-#endif
-					continue;
-				}
-
 				// Reuse MaterialSlot from the previous LOD.
 				if (const int32 MaterialIndex = SurfaceIdToMaterialIndex.Find(Surface.SharedSurfaceId); MaterialIndex != INDEX_NONE)
 				{
@@ -5897,16 +5901,26 @@ void UCustomizableInstancePrivate::BuildMaterials(const TSharedRef<FUpdateContex
 					continue;
 				}
 
+				// Is this a material in a passthrough mesh that we don't modify?
+				if (!Surface.Material)
+				{
+					Materials.Emplace();
+#if WITH_EDITOR
+					// Without this, a change of a referenced material and recompilation doesn't show up in the preview.
+					RecreateRenderStateOnInstanceComponent[InstanceComponentIndex] = true;
+#endif
+					continue;
+				}
+
 				UMaterialInterface* MaterialTemplate = nullptr;
 
-				if (Surface.Material.MaterialId == INDEX_NONE)
+				if (Surface.Material->MaterialId == INDEX_NONE)
 				{
-					MaterialTemplate = Surface.Material.Material.Get();
+					MaterialTemplate = Surface.Material->Material.Get();
 				}
-				else
+				else if (const uint32* ReferencedMaterialIndex = ObjectToInstanceIndexMap.Find(Surface.Material->MaterialId))
 				{
-					const uint32 ReferencedMaterialIndex = ObjectToInstanceIndexMap[Surface.Material.MaterialId];
-					MaterialTemplate = ReferencedMaterials[ReferencedMaterialIndex];
+					MaterialTemplate = ReferencedMaterials[*ReferencedMaterialIndex];
 				}
 				
 				if (!MaterialTemplate)
diff --git a/Engine/Plugins/Mutable/Source/CustomizableObject/Private/MuCO/CustomizableObjectMeshUpdate.cpp b/Engine/Plugins/Mutable/Source/CustomizableObject/Private/MuCO/CustomizableObjectMeshUpdate.cpp
--- a/Engine/Plugins/Mutable/Source/CustomizableObject/Private/MuCO/CustomizableObjectMeshUpdate.cpp
+++ b/Engine/Plugins/Mutable/Source/CustomizableObject/Private/MuCO/CustomizableObjectMeshUpdate.cpp
@@ -50,6 +50,7 @@ FCustomizableObjectMeshStreamIn::FCustomizableObjectMeshStreamIn(
 	OperationData->Parameters = InMesh->Parameters;
 	OperationData->MeshIdRegistry = InMesh->MeshIdRegistry;
 	OperationData->ImageIdRegistry = InMesh->ImageIdRegistry;
+	OperationData->MaterialIdRegistry = InMesh->MaterialIdRegistry;
 	OperationData->State = InMesh->State;
 
 	OperationData->ModelStreamableBulkData = ModelStreamableBulkData;
@@ -245,10 +246,11 @@ namespace impl
 		const uint32 LODMask = 0xFFFFFFFF;
 
 		// Main instance generation step
-		TSharedPtr<const mu::FInstance> Instance = System->BeginUpdate(InstanceID,
+		TSharedPtr<const mu::FInstance> Instance = System->BeginUpdate_MutableThread(InstanceID,
 			OperationData->Parameters,
 			OperationData->MeshIdRegistry,
 			OperationData->ImageIdRegistry,
+			OperationData->MaterialIdRegistry,
 			OperationData->State,
 			LODMask);
 		
diff --git a/Engine/Plugins/Mutable/Source/CustomizableObject/Private/MuCO/CustomizableObjectMeshUpdate.h b/Engine/Plugins/Mutable/Source/CustomizableObject/Private/MuCO/CustomizableObjectMeshUpdate.h
--- a/Engine/Plugins/Mutable/Source/CustomizableObject/Private/MuCO/CustomizableObjectMeshUpdate.h
+++ b/Engine/Plugins/Mutable/Source/CustomizableObject/Private/MuCO/CustomizableObjectMeshUpdate.h
@@ -28,6 +28,7 @@ struct FMutableMeshOperationData
 	TSharedPtr<mu::FModel, ESPMode::ThreadSafe> Model;
 	TSharedPtr<mu::FMeshIdRegistry> MeshIdRegistry;
 	TSharedPtr<mu::FImageIdRegistry> ImageIdRegistry;
+	TSharedPtr<mu::FMaterialIdRegistry> MaterialIdRegistry;
 	TSharedPtr<const mu::FParameters> Parameters;
 
 	TSharedPtr<FModelStreamableBulkData> ModelStreamableBulkData;
diff --git a/Engine/Plugins/Mutable/Source/CustomizableObject/Private/MuCO/CustomizableObjectMipDataProvider.cpp b/Engine/Plugins/Mutable/Source/CustomizableObject/Private/MuCO/CustomizableObjectMipDataProvider.cpp
--- a/Engine/Plugins/Mutable/Source/CustomizableObject/Private/MuCO/CustomizableObjectMipDataProvider.cpp
+++ b/Engine/Plugins/Mutable/Source/CustomizableObject/Private/MuCO/CustomizableObjectMipDataProvider.cpp
@@ -44,6 +44,7 @@ FMutableUpdateContext::FMutableUpdateContext(
 	const TSharedPtr<mu::FModel>& InModel,
 	const TSharedPtr<mu::FMeshIdRegistry>& InMeshIdRegistry,
 	const TSharedPtr<mu::FImageIdRegistry>& InImageIdRegistry,
+	const TSharedPtr<mu::FMaterialIdRegistry>& InMaterialIdRegistry,
 	const TSharedPtr<FModelStreamableBulkData>& InModelStreamable,
 	const TSharedPtr<const mu::FParameters>& InParameters,
 	int32 InState) :
@@ -53,6 +54,7 @@ FMutableUpdateContext::FMutableUpdateContext(
 	Model(InModel),
 	MeshIdRegistry(InMeshIdRegistry),
 	ImageIdRegistry(InImageIdRegistry),
+    MaterialIdRegistry(InMaterialIdRegistry),
 	ModelStreamableBulkData(InModelStreamable),
 	Parameters(InParameters),
 	State(InState)
@@ -211,10 +213,11 @@ namespace Impl
 			// LOD mask, set to all ones to build all LODs
 			uint32 LODMask = 0xFFFFFFFF;
 
-			Instance = System->BeginUpdate(InstanceID,
+			Instance = System->BeginUpdate_MutableThread(InstanceID,
 				OperationData->UpdateContext->GetParameters(),
 				OperationData->UpdateContext->MeshIdRegistry,
 				OperationData->UpdateContext->ImageIdRegistry,
+				OperationData->UpdateContext->MaterialIdRegistry,
 				OperationData->UpdateContext->GetState(),
 				LODMask);
 			
diff --git a/Engine/Plugins/Mutable/Source/CustomizableObject/Private/MuCO/CustomizableObjectSystem.cpp b/Engine/Plugins/Mutable/Source/CustomizableObject/Private/MuCO/CustomizableObjectSystem.cpp
--- a/Engine/Plugins/Mutable/Source/CustomizableObject/Private/MuCO/CustomizableObjectSystem.cpp
+++ b/Engine/Plugins/Mutable/Source/CustomizableObject/Private/MuCO/CustomizableObjectSystem.cpp
@@ -2408,6 +2408,8 @@ namespace impl
 
 		TArray<uint32> SurfacesSharedId;
 
+		TArray<int32> GetMaterialSurfaces;
+		
 		const TSharedRef<TArray<FGetImageData>> GetImagesData = MakeShared<TArray<FGetImageData>>();
 		
 		for (int32 InstanceComponentIndex = 0; InstanceComponentIndex < OperationData->NumInstanceComponents; ++InstanceComponentIndex)
@@ -2429,13 +2431,13 @@ namespace impl
 
 				// This lambda does all the work to fill up the surface data
 				auto AddSurface = 
-					[&LOD, &SurfacesSharedId, &ModelResources, &GetImagesData, OperationData, MutableInstance, CustomizableObject, InstanceComponentIndex, MutableLODIndex]
+					[&LOD, &SurfacesSharedId, &ModelResources, &GetImagesData, &GetMaterialSurfaces, OperationData, MutableInstance, CustomizableObject, InstanceComponentIndex, MutableLODIndex]
 					(uint32 SurfaceId, uint32 SurfaceMetadataId, int32 InstanceSurfaceIndex)
 					{
 						int32 BaseSurfaceIndex = InstanceSurfaceIndex;
 						int32 BaseLODIndex = MutableLODIndex;
 
-						OperationData->InstanceUpdateData.Surfaces.Push({});
+						int32 SurfaceIndex = OperationData->InstanceUpdateData.Surfaces.AddDefaulted();
 						FInstanceUpdateData::FSurface& Surface = OperationData->InstanceUpdateData.Surfaces.Last();
 						++LOD.SurfaceCount;
 						
@@ -2459,7 +2461,11 @@ namespace impl
 
 						Surface.SurfaceId = MutableInstance->GetSurfaceId(InstanceComponentIndex, BaseLODIndex, BaseSurfaceIndex);
 						Surface.SurfaceMetadataId = MutableInstance->GetSurfaceCustomId(InstanceComponentIndex, BaseLODIndex, BaseSurfaceIndex);
-						Surface.Material = MutableInstance->GetMaterial(InstanceComponentIndex, BaseLODIndex, BaseSurfaceIndex);
+						Surface.MaterialId = MutableInstance->GetMaterialId(InstanceComponentIndex, BaseLODIndex, BaseSurfaceIndex);
+						if (Surface.MaterialId)
+						{
+							GetMaterialSurfaces.Add(SurfaceIndex);
+						}
 
 						// Vectors
 						Surface.FirstVector = OperationData->InstanceUpdateData.Vectors.Num();
@@ -2563,8 +2569,33 @@ namespace impl
 
 			}	
 		}
+		
+		UE::Tasks::FTask LastTask = UE::Tasks::MakeCompletedTask<void>();
 
-		Task_Mutable_GetMeshes_GetImage_Loop(OperationData, StartTime, GetImagesData, 0);
+		for (uint32 SurfaceIndex : GetMaterialSurfaces)
+		{
+			LastTask = UE::Tasks::Launch(TEXT("Task_MutableGetMeshes_GetMesh_GetMaterial"), [SurfaceIndex, OperationData]()
+			{
+				mu::FMaterialId MaterialId = OperationData->InstanceUpdateData.Surfaces[SurfaceIndex].MaterialId;
+				UE::Tasks::TTask<TSharedPtr<const mu::FMaterial>> GetMaterialTask = OperationData->MutableSystem->GetMaterial(OperationData->InstanceID, MaterialId);
+				
+				UE::Tasks::AddNested(UE::Tasks::Launch(TEXT("Task_MutableGetMeshes_GetMesh_GetMaterial_Post"), [=]() mutable
+				{
+					OperationData->InstanceUpdateData.Surfaces[SurfaceIndex].Material = GetMaterialTask.GetResult();
+				},
+				GetMaterialTask,
+				LowLevelTasks::ETaskPriority::Inherit));
+			},
+			LastTask,
+			LowLevelTasks::ETaskPriority::Inherit);
+		}
+			
+		UE::Tasks::AddNested(UE::Tasks::Launch(TEXT("Task_Mutable_GetMeshes_GetImage_Loop"), [=]()
+		{
+			Task_Mutable_GetMeshes_GetImage_Loop(OperationData, StartTime, GetImagesData, 0);
+		},
+		LastTask,
+		LowLevelTasks::ETaskPriority::Inherit));
 	}
 
 
@@ -2663,17 +2694,20 @@ namespace impl
 				OperationData->MutableSystem->SetImagePixelConversionOverride( OperationData->PixelFormatOverride );
 			}
 			
-			// Main instance generation step
-			const mu::FInstance* Instance = OperationData->MutableInstance.Get(); // TODO GMTFuture remove
-			if (!Instance)
-			{
-				UE_LOG(LogMutable, Warning, TEXT("An Instace update has failed."));
-				Task_Mutable_GetMeshes_End(OperationData, StartTime);
-				return;
-			}
+			TSharedPtr<const mu::FInstance> Instance = OperationData->MutableSystem->BeginUpdate_MutableThread(OperationData->InstanceID,
+				OperationData->Parameters,
+				OperationData->MeshIdRegistry,
+				OperationData->ImageIdRegistry,
+				OperationData->MaterialIdRegistry,
+				OperationData->GetCapturedDescriptor().GetState(),
+				mu::FSystem::AllLODs);
 
+			OperationData->MutableInstance = Instance;
+			
 			const bool bForceGeometryGeneration = CVarForceGeometryOnFirstGeneration.GetValueOnAnyThread();
-		
+
+			TArray<int32> GetMaterialOverrideComponent;
+			
 			const TSharedRef<TArray<FGetMeshData>> GetMeshesData = MakeShared<TArray<FGetMeshData>>();
 			
 			OperationData->InstanceUpdateData.Components.SetNum(OperationData->NumInstanceComponents);
@@ -2701,8 +2735,12 @@ namespace impl
 					continue;
 				}
 
-				Component.OverlayMaterial = Instance->GetOverlayMaterial(InstanceComponentIndex);
-
+				Component.OverlayMaterialId = Instance->GetOverlayMaterialId(InstanceComponentIndex);
+				if (Component.OverlayMaterialId)
+				{
+					GetMaterialOverrideComponent.Add(InstanceComponentIndex);
+				}
+				
 				const uint8 FirstResidentLOD = OperationData->FirstResidentLOD[ComponentName];
 
 				// If the LOD is not generated we still add an empty one to keep indexes aligned.
@@ -2734,7 +2772,32 @@ namespace impl
 				}
 			}
 
-			Task_Mutable_GetMeshes_GetMesh_Loop(OperationData, StartTime, GetMeshesData, 0);
+			UE::Tasks::FTask LastTask = UE::Tasks::MakeCompletedTask<void>();
+
+			for (const int32 InstanceComponentIndex : GetMaterialOverrideComponent)
+			{
+				LastTask = UE::Tasks::Launch(TEXT("Task_MutableGetMeshes_GetMesh_GetMaterialOverride"), [InstanceComponentIndex, OperationData]()
+				{
+					mu::FMaterialId MaterialId = OperationData->InstanceUpdateData.Components[InstanceComponentIndex].OverlayMaterialId;
+					UE::Tasks::TTask<TSharedPtr<const mu::FMaterial>> GetMaterialOverrideTask = OperationData->MutableSystem->GetMaterial(OperationData->InstanceID, MaterialId);
+
+					UE::Tasks::AddNested(UE::Tasks::Launch(TEXT("Task_MutableGetMeshes_GetMesh_GetMaterialOverride_Post"), [=]() mutable
+					{
+						OperationData->InstanceUpdateData.Components[InstanceComponentIndex].OverlayMaterial = GetMaterialOverrideTask.GetResult();
+					},
+					GetMaterialOverrideTask,
+					LowLevelTasks::ETaskPriority::Inherit));
+				},
+				LastTask,
+				LowLevelTasks::ETaskPriority::Inherit);
+			}
+			
+			UE::Tasks::AddNested(UE::Tasks::Launch(TEXT("Task_MutableGetMeshes_GetMesh_Post"), [=]()
+			{
+				Task_Mutable_GetMeshes_GetMesh_Loop(OperationData, StartTime, GetMeshesData, 0);
+			},
+			LastTask,
+			LowLevelTasks::ETaskPriority::Inherit));
 		}
 	}
 
@@ -3501,6 +3564,7 @@ namespace impl
 		Operation->Model = ObjectPrivate->GetModel().ToSharedRef();
 		Operation->MeshIdRegistry = ObjectPrivate->MeshIdRegistry;
 		Operation->ImageIdRegistry = ObjectPrivate->ImageIdRegistry;
+		Operation->MaterialIdRegistry = ObjectPrivate->MaterialIdRegistry;
 
 		// Register the current CO into the external resource provider for runtime meshes and images.
 		// It is needed to get access to the bone map when converting meshes at runtime, for example.
@@ -3656,10 +3720,11 @@ namespace impl
 				UE_LOG(LogMutable, Verbose, TEXT("Creating Mutable instance with id [%d] "), Operation->InstanceID);
 			}
 
-			Operation->MutableInstance = MutableSystem->BeginUpdate(Operation->InstanceID,
+			Operation->MutableInstance = MutableSystem->BeginUpdate_GameThread(Operation->InstanceID,
 				Operation->Parameters,
 				Operation->MeshIdRegistry,
 				Operation->ImageIdRegistry,
+				Operation->MaterialIdRegistry,
 				Operation->GetCapturedDescriptor().GetState(),
 				mu::FSystem::AllLODs);
 			
diff --git a/Engine/Plugins/Mutable/Source/CustomizableObject/Public/MuCO/CustomizableObjectSkeletalMesh.h b/Engine/Plugins/Mutable/Source/CustomizableObject/Public/MuCO/CustomizableObjectSkeletalMesh.h
--- a/Engine/Plugins/Mutable/Source/CustomizableObject/Public/MuCO/CustomizableObjectSkeletalMesh.h
+++ b/Engine/Plugins/Mutable/Source/CustomizableObject/Public/MuCO/CustomizableObjectSkeletalMesh.h
@@ -45,9 +45,10 @@ class UCustomizableObjectSkeletalMesh : public USkeletalMesh
 
 	TSharedPtr<FModelStreamableBulkData> ModelStreamableBulkData;
 
-	TSharedPtr<mu::FModel, ESPMode::ThreadSafe> Model;
+	TSharedPtr<mu::FModel> Model;
 	TSharedPtr<mu::FMeshIdRegistry> MeshIdRegistry;
 	TSharedPtr<mu::FImageIdRegistry> ImageIdRegistry;
+	TSharedPtr<mu::FMaterialIdRegistry> MaterialIdRegistry;
 
 	/** Instance parameters at the time of creation. */
 	TSharedPtr<mu::FParameters> Parameters;
diff --git a/Engine/Plugins/Mutable/Source/CustomizableObjectEditor/Private/MuCOE/SMutableCodeViewer.cpp b/Engine/Plugins/Mutable/Source/CustomizableObjectEditor/Private/MuCOE/SMutableCodeViewer.cpp
--- a/Engine/Plugins/Mutable/Source/CustomizableObjectEditor/Private/MuCOE/SMutableCodeViewer.cpp
+++ b/Engine/Plugins/Mutable/Source/CustomizableObjectEditor/Private/MuCOE/SMutableCodeViewer.cpp
@@ -296,18 +296,17 @@ void SMutableCodeViewer::ClearSelectedTreeRow() const
 
 void SMutableCodeViewer::Construct(const FArguments& InArgs,
 	const TSharedPtr<mu::FModel>& InMutableModel,
-	const TSharedPtr<mu::FMeshIdRegistry>& InMeshIdRegistry,
-	const TSharedPtr<mu::FImageIdRegistry>& InImageIdRegistry,
 	const TArray<TSoftObjectPtr<const UTexture>>& InReferencedTextures,
-	const TArray<TSoftObjectPtr<const UStreamableRenderAsset>>& InReferencedMeshes )
+	const TArray<TSoftObjectPtr<const UStreamableRenderAsset>>& InReferencedMeshes)
 {
 	// Min width allowed for the column. Needed to avoid having issues with the constants space being to small
 	// and then getting too tall on the y axis crashing the UI drawer.
 	constexpr float MinParametersCollWidth = 200;
 
 	MutableModel = InMutableModel;
-	MeshIdRegistry = InMeshIdRegistry;
-	ImageIdRegistry = InImageIdRegistry;
+	MeshIdRegistry = InArgs._MeshIdRegistry;
+	ImageIdRegistry = InArgs._ImageIdRegistry;
+	MaterialIdRegistry = InArgs._MaterialIdRegistry;
 	ReferencedTextures = InReferencedTextures;
 	ReferencedMeshes = InReferencedMeshes;
 	PreviewParameters = mu::FModel::NewParameters(MutableModel);
@@ -2973,7 +2972,7 @@ void SMutableCodeViewer::Tick(const FGeometry& AllottedGeometry, const double In
 	ExternalResourceProvider->ReferencedMeshes = ReferencedMeshes;
 	System->SetExternalResourceProvider(ExternalResourceProvider);
 
-	System->GetPrivate()->BeginBuild(MutableModel, MeshIdRegistry, ImageIdRegistry);
+	System->GetPrivate()->BeginBuild(MutableModel, MeshIdRegistry, ImageIdRegistry, MaterialIdRegistry);
 
 	switch (OperationDataType)
 	{
diff --git a/Engine/Plugins/Mutable/Source/CustomizableObjectEditor/Private/MuCOE/SMutableCodeViewer.h b/Engine/Plugins/Mutable/Source/CustomizableObjectEditor/Private/MuCOE/SMutableCodeViewer.h
--- a/Engine/Plugins/Mutable/Source/CustomizableObjectEditor/Private/MuCOE/SMutableCodeViewer.h
+++ b/Engine/Plugins/Mutable/Source/CustomizableObjectEditor/Private/MuCOE/SMutableCodeViewer.h
@@ -55,12 +55,14 @@ public:
 
 	/** User-visible tag to identify the source of the data shown. */
 	SLATE_ARGUMENT(FString, DataTag)
+	SLATE_ARGUMENT(TSharedPtr<mu::FMeshIdRegistry>, MeshIdRegistry)
+	SLATE_ARGUMENT(TSharedPtr<mu::FImageIdRegistry>, ImageIdRegistry)
+	SLATE_ARGUMENT(TSharedPtr<mu::FMaterialIdRegistry>, MaterialIdRegistry)
 		
 	SLATE_END_ARGS()
 
-	void Construct(const FArguments& InArgs, const TSharedPtr<mu::FModel, ESPMode::ThreadSafe>& InMutableModel,
-		const TSharedPtr<mu::FMeshIdRegistry>& InMeshIdRegistry,
-		const TSharedPtr<mu::FImageIdRegistry>& InImageIdRegistry,
+	void Construct(const FArguments& InArgs,
+		const TSharedPtr<mu::FModel>& InMutableModel,
 		const TArray<TSoftObjectPtr<const UTexture>>& ReferencedTextures,
 		const TArray<TSoftObjectPtr<const UStreamableRenderAsset>>& ReferencedMeshes);
 
@@ -80,6 +82,7 @@ private:
 	TSharedPtr<mu::FModel> MutableModel;
 	TSharedPtr<mu::FMeshIdRegistry> MeshIdRegistry;
 	TSharedPtr<mu::FImageIdRegistry> ImageIdRegistry;
+	TSharedPtr<mu::FMaterialIdRegistry> MaterialIdRegistry;
 	
 	/** Array of external referenced textures in MutableModel, indexed by id. */
 	TArray<TSoftObjectPtr<const UTexture>> ReferencedTextures;
diff --git a/Engine/Plugins/Mutable/Source/CustomizableObjectEditor/Private/MuCOE/SMutableInstanceViewer.cpp b/Engine/Plugins/Mutable/Source/CustomizableObjectEditor/Private/MuCOE/SMutableInstanceViewer.cpp
--- a/Engine/Plugins/Mutable/Source/CustomizableObjectEditor/Private/MuCOE/SMutableInstanceViewer.cpp
+++ b/Engine/Plugins/Mutable/Source/CustomizableObjectEditor/Private/MuCOE/SMutableInstanceViewer.cpp
@@ -310,6 +310,7 @@ void SMutableInstanceViewer::SetInstance(const TSharedPtr<const mu::FInstance>&
 					NewLOD->MeshId = FText::FromString(MeshId.ToString());
 
 					// Generate the mesh to get some data
+					if (MeshId)
 					{
 						mu::OP::ADDRESS MeshRootAddress = MeshId.GetKey()->Address;
 						TSharedPtr<const mu::FMesh> MutableMesh = System.GetPrivate()->BuildMesh(
diff --git a/Engine/Plugins/Mutable/Source/CustomizableObjectEditor/Private/MuCOE/SMutableObjectViewer.cpp b/Engine/Plugins/Mutable/Source/CustomizableObjectEditor/Private/MuCOE/SMutableObjectViewer.cpp
--- a/Engine/Plugins/Mutable/Source/CustomizableObjectEditor/Private/MuCOE/SMutableObjectViewer.cpp
+++ b/Engine/Plugins/Mutable/Source/CustomizableObjectEditor/Private/MuCOE/SMutableObjectViewer.cpp
@@ -194,9 +194,12 @@ void SMutableObjectViewer::CompileMutableCodePressed()
 
 	FString DataTag = FString::Printf( TEXT("Mutable Code : %s opt %d"), *CompileOptions.TargetPlatform->PlatformName(), CompileOptions.OptimizationLevel );
 
-	TSharedRef<SMutableCodeViewer> CodeViewer = SNew(SMutableCodeViewer, CompileTask->Model,  CustomizableObject->GetPrivate()->MeshIdRegistry, CustomizableObject->GetPrivate()->ImageIdRegistry, RuntimeTextures, RuntimeMeshes)
-		.DataTag(DataTag);
-	
+	TSharedRef<SMutableCodeViewer> CodeViewer = SNew(SMutableCodeViewer, CompileTask->Model, RuntimeTextures, RuntimeMeshes)
+		.DataTag(DataTag)
+		.MeshIdRegistry(CustomizableObject->GetPrivate()->MeshIdRegistry)
+		.ImageIdRegistry(CustomizableObject->GetPrivate()->ImageIdRegistry)
+		.MaterialIdRegistry(CustomizableObject->GetPrivate()->MaterialIdRegistry);
+
 	DebuggerContentsBox->ClearContent();
 	DebuggerContentsBox->SetContent(CodeViewer);
 }
diff --git a/Engine/Plugins/Mutable/Source/MutableRuntime/Internal/MuR/ModelPrivate.h b/Engine/Plugins/Mutable/Source/MutableRuntime/Internal/MuR/ModelPrivate.h
--- a/Engine/Plugins/Mutable/Source/MutableRuntime/Internal/MuR/ModelPrivate.h
+++ b/Engine/Plugins/Mutable/Source/MutableRuntime/Internal/MuR/ModelPrivate.h
@@ -1052,7 +1052,7 @@ namespace mu
 			Result = Constant.Data;
 		}
 
-        inline EOpType GetOpType( OP::ADDRESS at ) const
+        FORCEINLINE EOpType GetOpType( OP::ADDRESS at ) const
         {
             if (at>=OP::ADDRESS(OpAddress.Num())) return EOpType::NONE;
 
diff --git a/Engine/Plugins/Mutable/Source/MutableRuntime/Internal/MuR/Operations.h b/Engine/Plugins/Mutable/Source/MutableRuntime/Internal/MuR/Operations.h
--- a/Engine/Plugins/Mutable/Source/MutableRuntime/Internal/MuR/Operations.h
+++ b/Engine/Plugins/Mutable/Source/MutableRuntime/Internal/MuR/Operations.h
@@ -383,6 +383,7 @@ namespace mu
     struct TOperation
     {
         typedef ADDRESS_TYPE ADDRESS;
+        typedef ADDRESS_TYPE CONSTANT_STRING_ADDRESS;
 
         /** Arguments for every operation type. */
         struct BoolConstantArgs
@@ -1028,8 +1029,8 @@ namespace mu
 
             // Index in the FProgram::ParameterLists with the parameters that are relevant
             // for whatever is added by this operation. This is used only for resources like
-            // images or meshes.
-            ADDRESS relevantParametersListIndex;
+            // images, meshes, or materials.
+            ADDRESS RelevantParametersListIndex;
         };
 
 		struct InstanceAddExtensionDataArgs
@@ -1051,6 +1052,11 @@ namespace mu
 			// be added to.
 			ADDRESS Instance;
 			ADDRESS Material;
+
+			// Index in the FProgram::ParameterLists with the parameters that are relevant
+			// for whatever is added by this operation. This is used only for resources like
+			// images, meshes, or materials.
+			ADDRESS RelevantParametersListIndex;
 		};
 
         struct LayoutPackArgs
diff --git a/Engine/Plugins/Mutable/Source/MutableRuntime/Internal/MuR/ResourceID.h b/Engine/Plugins/Mutable/Source/MutableRuntime/Internal/MuR/ResourceID.h
--- a/Engine/Plugins/Mutable/Source/MutableRuntime/Internal/MuR/ResourceID.h
+++ b/Engine/Plugins/Mutable/Source/MutableRuntime/Internal/MuR/ResourceID.h
@@ -9,6 +9,7 @@
 
 namespace mu
 {
+	// Mesh
 	struct FGeneratedMeshKey
 	{
 		OP::ADDRESS Address = 0;
@@ -34,6 +35,7 @@ namespace mu
 
 	typedef TRegistry<FGeneratedMeshKey, FGeneratedMeshData> FMeshIdRegistry;
 
+	// Image
 	struct FGeneratedImageKey
 	{
 		OP::ADDRESS Address = 0;
@@ -59,4 +61,25 @@ namespace mu
 	typedef TRegistry<FGeneratedImageKey, FGeneratedImageData>::FHandle FImageId;
 
 	typedef TRegistry<FGeneratedImageKey, FGeneratedImageData> FImageIdRegistry;
+
+	// Material
+	struct FGeneratedMaterialKey
+	{
+		OP::ADDRESS Address = 0;
+
+		TMemoryTrackedArray<uint8> ParameterValuesBlob;
+
+		bool operator==(const FGeneratedMaterialKey&) const = default;
+	};
+
+	
+	struct FGeneratedMaterialData
+	{
+	};
+	
+	MUTABLERUNTIME_API uint32 GetTypeHash(const FGeneratedMaterialKey& Key);
+
+	typedef TRegistry<FGeneratedMaterialKey, FGeneratedMaterialData>::FHandle FMaterialId;
+
+	typedef TRegistry<FGeneratedMaterialKey, FGeneratedMaterialData> FMaterialIdRegistry;
 }
diff --git a/Engine/Plugins/Mutable/Source/MutableRuntime/Internal/MuR/SystemPrivate.h b/Engine/Plugins/Mutable/Source/MutableRuntime/Internal/MuR/SystemPrivate.h
--- a/Engine/Plugins/Mutable/Source/MutableRuntime/Internal/MuR/SystemPrivate.h
+++ b/Engine/Plugins/Mutable/Source/MutableRuntime/Internal/MuR/SystemPrivate.h
@@ -147,8 +147,8 @@ namespace mu
 			//! Execute the operation to calculate the full result
 			Full,
 
-			//! Execute the operation to obtain the descriptor of an image.
-			ImageDesc
+			//! Execute the operation to obtain the descriptor of an image or mesh.
+			ImageDesc,
 		};
 		EType Type : 1;
 
@@ -158,12 +158,14 @@ namespace mu
 		uint64 ScheduleCallstack[CallstackMaxDepth];
 #endif
 	};
+	
 
 	inline uint32 GetTypeHash(const FScheduledOp& Op)
 	{
 		return HashCombineFast(::GetTypeHash(Op.At), HashCombineFast((uint32)Op.Stage), (uint32)Op.ExecutionIndex);
 	}
-
+	
+	
 	/** A cache address is the operation plus the context of execution(iteration indices, etc...). */
 	struct FCacheAddress
 	{
@@ -186,7 +188,7 @@ namespace mu
 			: At(InAt), ExecutionIndex(Item.ExecutionIndex), ExecutionOptions(Item.ExecutionOptions), Type(Item.Type)
 		{
 		}
-
+		
 		FCacheAddress(const FScheduledOp& Item)
 			: At(Item.At), ExecutionIndex(Item.ExecutionIndex), ExecutionOptions(Item.ExecutionOptions), Type(Item.Type)
 		{
@@ -195,10 +197,15 @@ namespace mu
 
 	inline uint32 GetTypeHash(const FCacheAddress& Address)
 	{
-		return HashCombineFast(::GetTypeHash(Address.At), (uint32)Address.ExecutionIndex);
-	}
+		uint32 Hash = Address.At;
 
+		Hash = HashCombineFast(Hash, static_cast<uint32>(Address.ExecutionIndex));
+		Hash = HashCombineFast(Hash, static_cast<uint32>(Address.ExecutionOptions));
+		Hash = HashCombineFast(Hash, static_cast<uint32>(Address.Type));
 
+		return Hash;
+	}
+	
 	inline uint32 GetTypeHash(const Ptr<const FResource>& V)
 	{
 		return ::GetTypeHash(V.get());
@@ -233,7 +240,7 @@ namespace mu
 
 		inline void erase(const FCacheAddress& at)
 		{
-			if (at.ExecutionIndex == 0 && at.ExecutionOptions == 0)
+			if (at.ExecutionIndex == 0 && at.ExecutionOptions == 0 && at.Type == FScheduledOp::EType::Full)
 			{
 				m_index0[at.At] = nullptr;
 			}
@@ -245,7 +252,7 @@ namespace mu
 
 		inline DATA get(const FCacheAddress& at) const
 		{
-			if (at.ExecutionIndex == 0 && at.ExecutionOptions == 0)
+			if (at.ExecutionIndex == 0 && at.ExecutionOptions == 0 && at.Type == FScheduledOp::EType::Full)
 			{
 				if (at.At < uint32(m_index0.Num()))
 				{
@@ -269,7 +276,7 @@ namespace mu
 
 		inline DATA* get_ptr(const FCacheAddress& at)
 		{
-			if (at.ExecutionIndex == 0 && at.ExecutionOptions == 0)
+			if (at.ExecutionIndex == 0 && at.ExecutionOptions == 0 && at.Type == FScheduledOp::EType::Full)
 			{
 				if (at.At < uint32(m_index0.Num()))
 				{
@@ -293,7 +300,7 @@ namespace mu
 
 		inline const DATA* get_ptr(const FCacheAddress& at) const
 		{
-			if (at.ExecutionIndex == 0 && at.ExecutionOptions == 0)
+			if (at.ExecutionIndex == 0 && at.ExecutionOptions == 0 && at.Type == FScheduledOp::EType::Full)
 			{
 				if (at.At < uint32(m_index0.Num()))
 				{
@@ -317,7 +324,7 @@ namespace mu
 
 		inline DATA& operator[](const FCacheAddress& at)
 		{
-			if (at.ExecutionIndex == 0 && at.ExecutionOptions == 0)
+			if (at.ExecutionIndex == 0 && at.ExecutionOptions == 0 && at.Type == FScheduledOp::EType::Full)
 			{
 				return m_index0[at.At];
 			}
@@ -329,7 +336,7 @@ namespace mu
 
 		inline const DATA& operator[](const FCacheAddress& at) const
 		{
-			if (at.ExecutionIndex == 0 && at.ExecutionOptions == 0)
+			if (at.ExecutionIndex == 0 && at.ExecutionOptions == 0 && at.Type == FScheduledOp::EType::Full)
 			{
 				return m_index0[at.At];
 			}
@@ -608,13 +615,18 @@ namespace mu
 			}
 
 			// Is it a desc data query?
-			if (at.Type == FScheduledOp::EType::ImageDesc)
+			switch (at.Type)
 			{
+			case FScheduledOp::EType::Full:
+				return Data->IsValueValid;
+
+			case FScheduledOp::EType::ImageDesc:
 				return Data->IsDescCacheValid;
+			
+			default:
+				unimplemented();
+				return false;
 			}
-
-			// It's a full data query.
-			return Data->IsValueValid != 0;
 		}
 
 
@@ -816,6 +828,31 @@ namespace mu
 		}
 
 
+		TSharedPtr<const FMaterial> GetMaterial(FCacheAddress at, bool& bIsLastReference)
+		{
+			bIsLastReference = false;
+
+			if (!at.At) return nullptr;
+			if (at.At >= OpExecutionData.size_code()) return nullptr;
+			FOpExecutionData* Data = OpExecutionData.get_ptr(at);
+			if (!Data) return nullptr;
+
+			check(Data->DataType == EDataType::Material);
+			TSharedPtr<const FMaterial> Result = MaterialResults[Data->DataTypeIndex];
+
+			// We need to decrease the hit-count even if the result is null.
+			check(Data->OpHitCount > 0);
+
+			--Data->OpHitCount;
+			if (Data->OpHitCount == 0 && !Data->IsCacheLocked)
+			{
+				SetUnused(*Data);
+				bIsLastReference = true;
+			}
+
+			return Result;
+		}
+
 		TSharedPtr<const FLayout> GetLayout(FCacheAddress at)
 		{
 			if (!at.At) return nullptr;
@@ -847,13 +884,13 @@ namespace mu
 			check(Data->DataType == EDataType::ExtensionData);
 			return ExtensionDataResults[Data->DataTypeIndex];
 		}
-
+		
 		void SetValidDesc(FCacheAddress at)
 		{
 			check(at.Type == FScheduledOp::EType::ImageDesc);
 			check(at.At < OpExecutionData.size_code());
-			FOpExecutionData* Data = OpExecutionData.get_ptr(at);
-			Data->IsDescCacheValid = true;
+			FOpExecutionData& Data = OpExecutionData[at];
+			Data.IsDescCacheValid = true;
 		}
 
 		void SetBool(FCacheAddress at, bool v)
@@ -1150,11 +1187,12 @@ namespace mu
 	{
 		FInstance::FID InstanceID;
 		int32 State = 0;
+		uint32 LODMask = 0;
 		TSharedPtr<const FInstance> Instance;
 		TSharedPtr<const FModel> Model;
 
-		TSharedPtr<FParameters> OldParameters;
-
+		TSharedPtr<const FParameters> OldParameters;
+		
 		/** Mask of the parameters that have changed since the last update.
 		* Every bit represents a state parameter.
 		*/
@@ -1204,6 +1242,8 @@ namespace mu
 		FMeshId GetMeshId(const TSharedPtr<const FModel>& Model, const FParameters* Params, uint32 ParamListIndex, OP::ADDRESS RootAt);
     	
 		FImageId GetImageId(const TSharedPtr<const FModel>& Model, const FParameters* Params, uint32 ParamListIndex, OP::ADDRESS RootAt);
+
+    	FMaterialId GetMaterialId(const TSharedPtr<const FModel>& Model, const FParameters* Params, uint32 ParamListIndex, OP::ADDRESS RootAt);
     	
 		/** Maximum working memory that mutable should be using. */
 		int64 BudgetBytes = 0;
@@ -1226,6 +1266,8 @@ namespace mu
     	TSharedPtr<FMeshIdRegistry> CurrentMeshIdRegistry;
     	
 		TSharedPtr<FImageIdRegistry> CurrentImageIdRegistry;
+
+    	TSharedPtr<FMaterialIdRegistry> CurrentMaterialIdRegistry;
     	
 		/** Resources that have been used in the past, but haven't been deallocated because they still fitted the memory budget and they could be reused. */
 		TMemoryTrackedArray<TSharedPtr<FImage>> PooledImages;
@@ -1522,6 +1564,31 @@ namespace mu
 			return Result;
 		}
 
+    	/** */
+    	[[nodiscard]] TSharedPtr<const FMaterial> LoadMaterial(const FCacheAddress& From)
+		{
+			bool bIsLastReference = false;
+			TSharedPtr<const FMaterial> Result = CurrentInstanceCache->GetMaterial(From, bIsLastReference);
+			if (!Result)
+			{
+				return nullptr;
+			}
+
+			// If we retrieved the last reference to this resource in "From" cache position (it could still be in other cache positions as well)
+			if (bIsLastReference)
+			{
+				int32* CountPtr = CacheResources.Find(Result);
+				check(CountPtr);
+				*CountPtr = (*CountPtr) - 1;
+				if (!*CountPtr)
+				{
+					CacheResources.FindAndRemoveChecked(Result);
+				}
+			}
+
+			return Result;
+		}
+
 		/** */
 		void StoreImage(const FCacheAddress& To, TSharedPtr<const FImage> Resource)
 		{
@@ -1554,6 +1621,17 @@ namespace mu
 			CurrentInstanceCache->SetMesh(To, Resource);
 		}
 
+    	void StoreMaterial(const FCacheAddress& To, TSharedPtr<const FMaterial> Resource)
+		{
+			if (Resource)
+			{
+				int32& Count = CacheResources.FindOrAdd(Resource, 0);
+				++Count;
+			}
+
+			CurrentInstanceCache->SetMaterial(To, Resource);
+		}
+
 		/** Return true if the resource is not in any cache (0,1,rom). */
 		bool IsBudgetTemp(const TSharedPtr<const FResource>& Resource)
 		{
@@ -1589,7 +1667,7 @@ namespace mu
 			{
 				Result += Value->GetDataSize();
 			}
-
+			
 			return Result;
 		}
 
@@ -1931,7 +2009,8 @@ namespace mu
 		/** This method can be used to internally prepare for code execution. */
 		MUTABLERUNTIME_API void BeginBuild(const TSharedPtr<const FModel>&,
 			const TSharedPtr<FMeshIdRegistry>&,
-			const TSharedPtr<FImageIdRegistry>&);
+			const TSharedPtr<FImageIdRegistry>&,
+			const TSharedPtr<FMaterialIdRegistry>&);
     	
 		MUTABLERUNTIME_API void EndBuild();
 
diff --git a/Engine/Plugins/Mutable/Source/MutableRuntime/Private/MuR/CodeRunner.cpp b/Engine/Plugins/Mutable/Source/MutableRuntime/Private/MuR/CodeRunner.cpp
--- a/Engine/Plugins/Mutable/Source/MutableRuntime/Private/MuR/CodeRunner.cpp
+++ b/Engine/Plugins/Mutable/Source/MutableRuntime/Private/MuR/CodeRunner.cpp
@@ -872,9 +872,15 @@ namespace mu
 				case 1:
 				{
 					// Assemble result
-					TSharedPtr<const FInstance> InstanceOpResult = LoadInstance(FCacheAddress(Args.Instance, Item));
-
-					TSharedPtr<FInstance> Result = mu::CloneOrTakeOver<FInstance>(InstanceOpResult);
+					TSharedPtr<FInstance> Instance;
+					if (Args.Instance)
+					{
+						Instance = mu::CloneOrTakeOver<FInstance>(LoadInstance(FCacheAddress(Args.Instance, Item)));
+					}
+					else
+					{
+						Instance = MakeShared<FInstance>();
+					}
 
 					if (TSharedPtr<const FExtensionData> ExtensionData = LoadExtensionData(FCacheAddress(Args.ExtensionData, Item)))
 					{
@@ -882,10 +888,10 @@ namespace mu
 						check(NameAddress < (uint32)Program.ConstantStrings.Num());
 						const FString& NameString = Program.ConstantStrings[NameAddress];
 
-						Result->AddExtensionData(ExtensionData, FName(NameString) );
+						Instance->SetExtensionData(ExtensionData.ToSharedRef(), FName(NameString));
 					}
 
-					StoreInstance(Item, Result);
+					StoreInstance(Item, Instance);
 					break;
 				}
 
@@ -895,92 +901,9 @@ namespace mu
 			
 			break;
 		}
-		case EOpType::IN_ADDOVERLAYMATERIAL:
-		{
-			OP::InstanceAddMaterialArgs Args = Program.GetOpArgs<OP::InstanceAddMaterialArgs>(Item.At);
-			switch (Item.Stage)
-			{
-			case 0:
-			{
-				// Must pass in a Material op
-				check(Args.Material);
-
-				AddOp(FScheduledOp(Item.At, Item, 1),
-					FScheduledOp(Args.Instance, Item),
-					FScheduledOp(Args.Material, Item));
-
-				break;
-			}
-
-			case 1:
-			{
-				TSharedPtr<const FInstance> Base = LoadInstance(FCacheAddress(Args.Instance, Item));
-				TSharedPtr<FInstance> Result;
-				if (!Base)
-				{
-					Result = MakeShared<FInstance>();
-				}
-				else
-				{
-					Result = mu::CloneOrTakeOver<FInstance>(Base);
-				}
-
-				// OverlayMaterial Material
-				TSharedPtr<FMaterial> LoadedMaterial = mu::CloneOrTakeOver<FMaterial>(LoadMaterial(FCacheAddress(Args.Material, Item)));
-				Result->AddOverlayMaterial(0, *LoadedMaterial.Get());
-
-				StoreInstance(Item, Result);
-				break;
-			}
-
-			default:
-				check(false);
-			}
-
-			break;
-		}
-
-		case EOpType::IN_ADDMATERIAL:
-		{
-			OP::InstanceAddMaterialArgs Args = Program.GetOpArgs<OP::InstanceAddMaterialArgs>(Item.At);
-			switch (Item.Stage)
-			{
-			case 0:
-				AddOp(FScheduledOp(Item.At, Item, 1),
-					FScheduledOp(Args.Instance, Item),
-					FScheduledOp(Args.Material, Item));
-				break;
-
-			case 1:
-			{
-				TSharedPtr<const FInstance> Base = LoadInstance(FCacheAddress(Args.Instance, Item));
-
-				TSharedPtr<FInstance> Result;
-				if (Base)
-				{
-					Result = mu::CloneOrTakeOver<FInstance>(Base);
-				}
-				else
-				{
-					Result = MakeShared<FInstance>();
-				}
-
-				// Material
-				TSharedPtr<FMaterial> LoadedMaterial = mu::CloneOrTakeOver<FMaterial>(LoadMaterial(FCacheAddress(Args.Material, Item)));
-				Result->AddMaterial(0, 0, 0, *LoadedMaterial.Get());
-
-				StoreInstance(Item, Result);
-				break;
-			}
-
-			default:
-				check(false);
-			}
-
-			break;
-		}
-        default:
-                check(false);
+        	
+		default:
+            check(false);
         }
     }
 
@@ -1000,7 +923,6 @@ namespace mu
 		EOpType type = Program.GetOpType(Item.At);
         switch (type)
         {
-
         case EOpType::IN_ADDMESH:
         {
 			OP::InstanceAddArgs Args = Program.GetOpArgs<OP::InstanceAddArgs>(Item.At);
@@ -1026,7 +948,7 @@ namespace mu
 
                 if ( Args.value )
                 {
-					FMeshId MeshId = System->WorkingMemoryManager.GetMeshId(InModel, InParams, Args.relevantParametersListIndex, Args.value);
+					FMeshId MeshId = System->WorkingMemoryManager.GetMeshId(InModel, InParams, Args.RelevantParametersListIndex, Args.value);
 					OP::ADDRESS NameAd = Args.name;
 					check(NameAd < (uint32)Program.ConstantStrings.Num());
 					const FString& Name = Program.ConstantStrings[NameAd];
@@ -1067,7 +989,7 @@ namespace mu
 
                 if ( Args.value )
                 {
-					FImageId ImageId = System->WorkingMemoryManager.GetImageId(InModel, InParams, Args.relevantParametersListIndex, Args.value);
+					FImageId ImageId = System->WorkingMemoryManager.GetImageId(InModel, InParams, Args.RelevantParametersListIndex, Args.value);
 					OP::ADDRESS NameAd = Args.name;
 					check(NameAd < (uint32)Program.ConstantStrings.Num());
 					const FString& Name = Program.ConstantStrings[NameAd];
@@ -1084,11 +1006,101 @@ namespace mu
             break;
         }
 
+        case EOpType::IN_ADDOVERLAYMATERIAL:
+        	{
+        		OP::InstanceAddMaterialArgs Args = Program.GetOpArgs<OP::InstanceAddMaterialArgs>(Item.At);
+        		switch (Item.Stage)
+        		{
+        		case 0:
+        			{
+        				// Must pass in a Material op
+        				check(Args.Material);
+
+        				AddOp(FScheduledOp(Item.At, Item, 1),
+							FScheduledOp(Args.Instance, Item),
+							FScheduledOp(Args.Material, Item));
+
+        				break;
+        			}
+
+        		case 1:
+        			{
+        				TSharedPtr<const FInstance> Base = LoadInstance(FCacheAddress(Args.Instance, Item));
+        				TSharedPtr<FInstance> Result;
+        				if (!Base)
+        				{
+        					Result = MakeShared<FInstance>();
+        				}
+        				else
+        				{
+        					Result = mu::CloneOrTakeOver<FInstance>(Base);
+        				}
+
+        				if (Args.Material)
+        				{
+        					FMaterialId MaterialId = System->WorkingMemoryManager.GetMaterialId(InModel, InParams, Args.RelevantParametersListIndex, Args.Material);
+        					Result->SetMaterialId(0, 0, 0, MaterialId);
+        				}
+
+        				StoreInstance(Item, Result);
+        				break;
+        			}
+
+        		default:
+        			check(false);
+        		}
+
+        		break;
+        	}
+
+        case EOpType::IN_ADDMATERIAL:
+        	{
+        		OP::InstanceAddMaterialArgs Args = Program.GetOpArgs<OP::InstanceAddMaterialArgs>(Item.At);
+        		switch (Item.Stage)
+        		{
+        		case 0:
+        			AddOp(FScheduledOp(Item.At, Item, 1),
+						FScheduledOp(Args.Instance, Item),
+						FScheduledOp(Args.Material, Item));
+        			break;
+
+        		case 1:
+        			{
+        				TSharedPtr<const FInstance> Base = LoadInstance(FCacheAddress(Args.Instance, Item));
+
+        				TSharedPtr<FInstance> Result;
+        				if (Base)
+        				{
+        					Result = mu::CloneOrTakeOver<FInstance>(Base);
+        				}
+        				else
+        				{
+        					Result = MakeShared<FInstance>();
+        				}
+
+        				if (Args.Material)
+        				{
+        					FMaterialId MaterialId = System->WorkingMemoryManager.GetMaterialId(InModel, InParams, Args.RelevantParametersListIndex, Args.Material);
+        					Result->SetMaterialId(0, 0, 0, MaterialId);
+        				}
+
+        				StoreInstance(Item, Result);
+        				break;
+        			}
+
+        		default:
+        			check(false);
+        		}
+
+        		break;
+        	}
+
+        	
         default:
 			check(false);
         }
     }
-
+	
 
     //---------------------------------------------------------------------------------------------
     bool CodeRunner::RunCode_ConstantResource(const FScheduledOp& Item, const FModel* InModel)
@@ -1436,9 +1448,9 @@ namespace mu
         {
 			const uint8* Data = Program.GetOpArgsPointer(Item.At);
 
-			OP::ADDRESS NameAt = 0;
-			FMemory::Memcpy(&NameAt, Data, sizeof(OP::ADDRESS)); 
-			Data += sizeof(OP::ADDRESS);
+			OP::CONSTANT_STRING_ADDRESS NameAt = 0;
+			FMemory::Memcpy(&NameAt, Data, sizeof(OP::CONSTANT_STRING_ADDRESS)); 
+			Data += sizeof(OP::CONSTANT_STRING_ADDRESS);
 				
 			OP::ADDRESS FactorAt = 0;
 			FMemory::Memcpy(&FactorAt, Data, sizeof(OP::ADDRESS)); 
@@ -5418,7 +5430,7 @@ namespace mu
     }	
 	
 	//---------------------------------------------------------------------------------------------
-    TSharedPtr<FRangeIndex> CodeRunner::BuildCurrentOpRangeIndex( const FScheduledOp& Item, const FParameters* pParams, const FModel* InModel, int32 parameterIndex )
+    TSharedPtr<FRangeIndex> BuildCurrentOpRangeIndex(const FScheduledOp& Item, const FParameters& Params, const FModel& InModel, FProgramCache& ProgramCache, int32 ParameterIndex)
     {
         if (!Item.ExecutionIndex)
         {
@@ -5427,27 +5439,31 @@ namespace mu
 
         // \todo: optimise to avoid allocating the index here, we could access internal
         // data directly.
-		TSharedPtr<FRangeIndex> Index = pParams->NewRangeIndex( parameterIndex );
+		TSharedPtr<FRangeIndex> Index = Params.NewRangeIndex( ParameterIndex );
         if (!Index)
         {
             return nullptr;
         }
 
-		const FProgram& Program = Model->GetPrivate()->Program;
-		const FParameterDesc& paramDesc = Program.Parameters[ parameterIndex ];
-        for( int32 rangeIndexInParam=0;
+		const FProgram& Program = InModel.GetPrivate()->Program;
+		const FParameterDesc& paramDesc = Program.Parameters[ ParameterIndex ];
+        for( int32 rangeIndexInParam = 0;
              rangeIndexInParam<paramDesc.Ranges.Num();
              ++rangeIndexInParam )
         {
             uint32 rangeIndexInModel = paramDesc.Ranges[rangeIndexInParam];
-            const ExecutionIndex& currentIndex = GetMemory().GetRangeIndex( Item.ExecutionIndex );
+            const ExecutionIndex& currentIndex = ProgramCache.GetRangeIndex( Item.ExecutionIndex );
             int32 Position = currentIndex.GetFromModelRangeIndex(rangeIndexInModel);
 			Index->Values[rangeIndexInParam] = Position;
         }
 
         return Index;
     }
-
+	
+	TSharedPtr<FRangeIndex> CodeRunner::BuildCurrentOpRangeIndex(const FScheduledOp& Item, int32 ParameterIndex)
+	{
+		return mu::BuildCurrentOpRangeIndex(Item, *Params, *Model.Get(), GetMemory(), ParameterIndex);
+	}
 
     //---------------------------------------------------------------------------------------------
     void CodeRunner::RunCode_Bool(const FScheduledOp& Item, const FParameters* pParams, const FModel* InModel )
@@ -5471,7 +5487,7 @@ namespace mu
         {
 			OP::ParameterArgs Args = Program.GetOpArgs<OP::ParameterArgs>(Item.At);
             bool result = false;
-			TSharedPtr<FRangeIndex> Index = BuildCurrentOpRangeIndex( Item, pParams, InModel, Args.variable );
+			TSharedPtr<FRangeIndex> Index = BuildCurrentOpRangeIndex( Item, Args.variable );
             result = pParams->GetBoolValue( Args.variable, Index.Get() );
             StoreBool( Item, result );
             break;
@@ -5684,7 +5700,7 @@ namespace mu
         case EOpType::NU_PARAMETER:
         {
 			OP::ParameterArgs Args = Program.GetOpArgs<OP::ParameterArgs>(Item.At);
-			TSharedPtr<FRangeIndex> Index = BuildCurrentOpRangeIndex( Item, pParams, InModel, Args.variable );
+			TSharedPtr<FRangeIndex> Index = BuildCurrentOpRangeIndex( Item, Args.variable );
             int32 result = pParams->GetIntValue( Args.variable, Index.Get());
 
             // Check that the value is actually valid. Otherwise set the default.
@@ -5731,7 +5747,7 @@ namespace mu
         case EOpType::SC_CONSTANT:
         {
 			OP::ScalarConstantArgs Args = Program.GetOpArgs<OP::ScalarConstantArgs>(Item.At);
-            float result = Args.Value;
+			float result = Args.Value;
             StoreScalar( Item, result );
             break;
         }
@@ -5739,7 +5755,7 @@ namespace mu
         case EOpType::SC_PARAMETER:
         {
 			OP::ParameterArgs Args = Program.GetOpArgs<OP::ParameterArgs>(Item.At);
-			TSharedPtr<FRangeIndex> Index = BuildCurrentOpRangeIndex( Item, pParams, InModel, Args.variable );
+			TSharedPtr<FRangeIndex> Index = BuildCurrentOpRangeIndex( Item, Args.variable );
             float result = pParams->GetFloatValue( Args.variable, Index.Get());
             StoreScalar( Item, result );
             break;
@@ -5853,7 +5869,7 @@ namespace mu
         case EOpType::ST_PARAMETER:
         {
 			OP::ParameterArgs Args = Program.GetOpArgs<OP::ParameterArgs>( Item.At );
-			TSharedPtr<FRangeIndex> Index = BuildCurrentOpRangeIndex( Item, pParams, InModel, Args.variable );
+			TSharedPtr<FRangeIndex> Index = BuildCurrentOpRangeIndex( Item, Args.variable );
 			FString result;
 			pParams->GetStringValue(Args.variable, result, Index.Get());
             StoreString( Item, MakeShared<String>(result) );
@@ -5889,7 +5905,7 @@ namespace mu
         case EOpType::CO_PARAMETER:
         {
 			OP::ParameterArgs Args = Program.GetOpArgs<OP::ParameterArgs>(Item.At);
-			TSharedPtr<FRangeIndex> Index = BuildCurrentOpRangeIndex( Item, pParams, InModel, Args.variable );
+			TSharedPtr<FRangeIndex> Index = BuildCurrentOpRangeIndex( Item, Args.variable );
             FVector4f V;
             pParams->GetColourValue( Args.variable, V, Index.Get());
             StoreColor( Item, V );
@@ -6134,7 +6150,7 @@ namespace mu
         case EOpType::PR_PARAMETER:
         {
 			OP::ParameterArgs Args = Program.GetOpArgs<OP::ParameterArgs>(Item.At);
-			TSharedPtr<FRangeIndex> Index = BuildCurrentOpRangeIndex( Item, pParams, InModel, Args.variable );
+			TSharedPtr<FRangeIndex> Index = BuildCurrentOpRangeIndex( Item, Args.variable );
             FProjector Result = pParams->GetPrivate()->GetProjectorValue(Args.variable,Index.Get());
 
             // The type cannot be changed, take it from the default value
@@ -6172,7 +6188,7 @@ namespace mu
 		case EOpType::MA_PARAMETER:
 			{
 				OP::ParameterArgs Args = Program.GetOpArgs<OP::ParameterArgs>(Item.At);
-				TSharedPtr<FRangeIndex> Index = BuildCurrentOpRangeIndex( Item, pParams, InModel, Args.variable );
+				TSharedPtr<FRangeIndex> Index = BuildCurrentOpRangeIndex( Item, Args.variable );
 				FMatrix44f Value;
 				pParams->GetMatrixValue( Args.variable, Value, Index.Get());
 				StoreMatrix( Item, Value );
@@ -6198,7 +6214,7 @@ namespace mu
 			check(Item.Stage == 0);
 
 			OP::ParameterArgs Args = Program.GetOpArgs<OP::ParameterArgs>(Item.At);
-			TSharedPtr<FRangeIndex> Index = BuildCurrentOpRangeIndex(Item, InParams, InModel, Args.variable);
+			TSharedPtr<FRangeIndex> Index = BuildCurrentOpRangeIndex(Item, Args.variable);
 
 			TSharedPtr<FMaterial> Material = MakeShared<FMaterial>();
 			Material->Material = TStrongObjectPtr(InParams->GetMaterialValue(Args.variable, Index.Get()));
@@ -6441,6 +6457,8 @@ namespace mu
 
         case EOpType::IN_ADDMESH:
         case EOpType::IN_ADDIMAGE:
+        case EOpType::IN_ADDMATERIAL:
+        case EOpType::IN_ADDOVERLAYMATERIAL:
             RunCode_InstanceAddResource(Item, InModel, pParams);
             break;
 
@@ -6506,8 +6524,8 @@ namespace mu
         }
     }
 
-	/** */
-	void CodeRunner::RunCodeImageDesc(const FScheduledOp& Item, const FParameters* pParams, const FModel* InModel,  uint32 LodMask )
+	
+	void CodeRunner::RunCodeImageDesc(const FScheduledOp& Item, const FParameters* pParams)
 	{
 		MUTABLE_CPUPROFILER_SCOPE(RunCodeImageDesc);
 		check(Item.Type == FScheduledOp::EType::ImageDesc);
diff --git a/Engine/Plugins/Mutable/Source/MutableRuntime/Private/MuR/CodeRunner.h b/Engine/Plugins/Mutable/Source/MutableRuntime/Private/MuR/CodeRunner.h
--- a/Engine/Plugins/Mutable/Source/MutableRuntime/Private/MuR/CodeRunner.h
+++ b/Engine/Plugins/Mutable/Source/MutableRuntime/Private/MuR/CodeRunner.h
@@ -105,9 +105,8 @@ namespace  mu
 
 		// Assertion to know when FScheduledOpData size changes. It is ok to modifiy if needed. 
 		static_assert(sizeof(FScheduledOpData) == 4*4+sizeof(TSharedPtr<FResource>), "FScheduledOpData size changed.");
-
-
-        TSharedPtr<FRangeIndex> BuildCurrentOpRangeIndex( const FScheduledOp&, const FParameters*, const FModel*, int32 ParameterIndex );
+    	
+        TSharedPtr<FRangeIndex> BuildCurrentOpRangeIndex( const FScheduledOp&, int32 ParameterIndex);
 
         void RunCode( const FScheduledOp&, const FParameters*, const TSharedPtr<const FModel>&, uint32 LodMask);
 
@@ -129,11 +128,11 @@ namespace  mu
         void RunCode_Colour(const FScheduledOp&, const FParameters*, const FModel* );
         void RunCode_Projector(const FScheduledOp&, const FParameters*, const FModel* );
         void RunCode_Matrix(const FScheduledOp&, const FParameters*, const FModel* );
-        void RunCode_Material(const FScheduledOp&, const FParameters*, const FModel* );
+    	void RunCode_Material(const FScheduledOp&, const FParameters*, const FModel* );
 
-		void RunCodeImageDesc(const FScheduledOp&, const FParameters*, const FModel*, uint32 LodMask);
+    	void RunCodeImageDesc(const FScheduledOp&, const FParameters*);
+    
 	public:
-
 		struct FExternalResourceId
 		{
 			/** If it is an image or mesh reference. */
@@ -718,6 +717,8 @@ namespace  mu
 	/** Helper function to create the memory-tracked image operator. */
 	extern FImageOperator MakeImageOperator(CodeRunner* Runner);
 
+	
+	TSharedPtr<FRangeIndex> BuildCurrentOpRangeIndex(const FScheduledOp& Item, const FParameters& Params, const FModel& InModel, FProgramCache& ProgramCache, int32 ParameterIndex);
 }
 
 #if MUTABLE_DEBUG_CODERUNNER_TASK_SCHEDULE_CALLSTACK
diff --git a/Engine/Plugins/Mutable/Source/MutableRuntime/Private/MuR/CodeRunnerBegin.cpp b/Engine/Plugins/Mutable/Source/MutableRuntime/Private/MuR/CodeRunnerBegin.cpp
new file mode 100644
--- /dev/null
+++ b/Engine/Plugins/Mutable/Source/MutableRuntime/Private/MuR/CodeRunnerBegin.cpp
@@ -0,0 +1,2349 @@
+// Copyright Epic Games, Inc. All Rights Reserved.
+
+#include "CodeRunnerBegin.h"
+
+#include "CodeRunner.h"
+#include "MuR/MutableTrace.h"
+#include "MuR/Operations.h"
+#include "MuR/SystemPrivate.h"
+
+
+namespace mu
+{
+	CodeRunnerBegin::CodeRunnerBegin(FSystem& InSystem, const TSharedPtr<const FModel>& InModel, const TSharedPtr<const FParameters>& InParams, uint32 InLodMask) :
+		Model(InModel),
+		Params(InParams),
+		System(InSystem),
+		Program(InModel->GetPrivate()->Program),
+		LodMask(InLodMask),
+		Executed(InModel->GetPrivate()->Program.OpAddress.Num())
+	{
+	}
+
+	
+	bool FScheduledOpInline::operator==(const FScheduledOpInline& Other) const
+	{
+		return At == Other.At && Stage == Other.Stage && ExecutionIndex == Other.ExecutionIndex;
+	}
+
+
+	uint32 GetTypeHash(const FScheduledOpInline& Op)
+	{
+		uint32 Hash = Op.At;
+		
+		Hash = HashCombineFast(Hash, Op.Stage);
+		Hash = HashCombineFast(Hash, Op.ExecutionIndex);
+
+		return Hash;
+	}
+
+
+	uint32 GetTypeHash(const FCacheAddressInline& Address)
+	{
+		return HashCombineFast(Address.At, static_cast<uint32>(Address.ExecutionIndex));
+	}
+
+
+	FOpSet::FOpSet(int32 NumElements)
+	{
+		Index0.SetNum(NumElements, false);
+	}
+
+
+	bool FOpSet::Contains(const FCacheAddressInline& Item)
+	{
+		if (Item.ExecutionIndex == 0)
+		{
+			return Index0[Item.At];
+		}
+		else
+		{
+			return IndexOther.Contains(Item);
+		}
+	}
+
+
+	void FOpSet::Add(const FCacheAddressInline& Item)
+	{
+		if (Item.ExecutionIndex == 0)
+		{
+			Index0[Item.At] = true;
+		}
+		else
+		{
+			IndexOther.Add(Item);
+		}
+	}
+	
+	
+	FStackValue FStack::Pop()
+	{
+		return TArray::Pop(EAllowShrinking::No);
+	}
+	
+
+	FProgramCache& CodeRunnerBegin::GetMemory() const
+	{
+		return *System.GetPrivate()->WorkingMemoryManager.CurrentInstanceCache;
+	}
+
+
+	FScheduledOpInline CodeRunnerBegin::PopOp()
+	{
+		return Items.Pop(EAllowShrinking::No);
+	}
+
+
+	void CodeRunnerBegin::PushOp(const FScheduledOpInline& Item)
+	{
+		check(Item.At < static_cast<uint32>(Program.OpAddress.Num()));
+
+		if (!Item.At)
+		{
+			return;
+		}
+
+		Items.Push(Item);
+	}
+
+
+	void CodeRunnerBegin::StoreInt(const FCacheAddressInline& To, int32 Value)
+	{
+		if (!To.At)
+		{
+			return;
+		}
+		
+		Executed.Add(To);
+
+		FStackValue StackValue;
+		StackValue.Set<int32>(Value);
+		Stack.Push(StackValue);
+
+		ResultsInt.Add(To, Value);
+	}
+
+
+	void CodeRunnerBegin::StoreFloat(const FCacheAddressInline& To, float Value)
+	{
+		if (!To.At)
+		{
+			return;
+		}
+		
+		Executed.Add(To);
+
+		FStackValue StackValue;
+		StackValue.Set<float>(Value);
+		Stack.Push(StackValue);
+
+		ResultsFloat.Add(To, Value);
+	}
+
+
+	void CodeRunnerBegin::StoreBool(const FCacheAddressInline& To, bool Value)
+	{
+		if (!To.At)
+		{
+			return;
+		}
+		
+		Executed.Add(To);
+
+		FStackValue StackValue;
+		StackValue.Set<bool>(Value);
+		Stack.Push(StackValue);
+			
+		ResultsBool.Add(To, Value);
+	}
+
+
+	void CodeRunnerBegin::StoreInstance(const FCacheAddressInline& To, const TSharedPtr<const FInstance>& Value)
+	{
+		if (!To.At)
+		{
+			return;
+		}
+
+		Executed.Add(To);
+
+		FStackValue StackValue;
+		StackValue.Set<TSharedPtr<const FInstance>>(Value);
+		Stack.Push(StackValue);
+
+		ResultsInstance.Add(To, Value);
+	}
+
+
+	void CodeRunnerBegin::StoreString(const FCacheAddressInline& To, const TSharedPtr<const String>& Value)
+	{
+		if (!To.At)
+		{
+			return;
+		}
+		
+		Executed.Add(To);
+
+		FStackValue StackValue;
+		StackValue.Set<TSharedPtr<const String>>(Value);
+		Stack.Push(StackValue);
+
+		ResultsString.Add(To, Value);
+	}
+
+
+	void CodeRunnerBegin::StoreExtensionData(const FCacheAddressInline& To, const TSharedPtr<const FExtensionData>& Value)
+	{
+		if (!To.At)
+		{
+			return;
+		}
+		
+		Executed.Add(To);
+
+		FStackValue StackValue;
+		StackValue.Set<TSharedPtr<const FExtensionData>>(Value);
+		Stack.Push(StackValue);
+
+		ResultsExtensionData.Add(To, Value);
+	}
+
+
+	void CodeRunnerBegin::StoreNone(const FCacheAddressInline& To)
+	{
+		if (!To.At)
+		{
+			return;
+		}
+		
+		Executed.Add(To);
+	}
+
+
+	int32 CodeRunnerBegin::LoadInt(const FCacheAddressInline& To)
+	{
+		if (!To.At)
+		{
+			return 0;	
+		}
+		
+		return Stack.Pop().Get<int32>();
+	}
+
+
+	float CodeRunnerBegin::LoadFloat(const FCacheAddressInline& To)
+	{
+		if (!To.At)
+		{
+			return 0.0f;	
+		}
+
+		return Stack.Pop().Get<float>();
+	}
+
+
+	bool CodeRunnerBegin::LoadBool(const FCacheAddressInline& To)
+	{
+		if (!To.At)
+		{
+			return false;	
+		}
+
+		return Stack.Pop().Get<bool>();
+	}
+
+
+	TSharedPtr<const FInstance> CodeRunnerBegin::LoadInstance(const FCacheAddressInline& To)
+	{
+		if (!To.At)
+		{
+			return nullptr;	
+		}
+			
+		return Stack.Pop().Get<TSharedPtr<const FInstance>>();
+	}
+
+
+	TSharedPtr<const String> CodeRunnerBegin::LoadString(const FCacheAddressInline& To)
+	{
+		if (!To.At)
+		{
+			return nullptr;	
+		}
+		
+		return Stack.Pop().Get<TSharedPtr<const String>>();
+	}
+
+
+	TSharedPtr<const FExtensionData> CodeRunnerBegin::LoadExtensionData(const FCacheAddressInline& To)
+	{
+		if (!To.At)
+		{
+			return nullptr;	
+		}
+		
+		return Stack.Pop().Get<TSharedPtr<const FExtensionData>>();
+	}
+
+
+	TSharedPtr<const FInstance> CodeRunnerBegin::RunCode(const FScheduledOpInline& Root)
+	{
+		MUTABLE_CPUPROFILER_SCOPE(CodeRunnerBegin::RunCode);
+		
+		PushOp(Root);
+		
+		while (!Items.IsEmpty())
+		{
+			const FScheduledOpInline Item = PopOp();
+			
+			const EOpType Type = Program.GetOpType(Item.At);
+
+			if (Executed.Contains(Item))
+			{
+				switch (GetOpDataType(Type))
+				{
+				case EDataType::Bool:
+					{
+						FStackValue Value;
+						Value.Set<bool>(ResultsBool[Item]);
+					
+						Stack.Push(Value);
+						break;
+					}
+
+				case EDataType::Int:
+					{
+						FStackValue Value;
+						Value.Set<int32>(ResultsInt[Item]);
+					
+						Stack.Push(Value);
+						break;
+					}
+					
+				case EDataType::Instance:
+					{
+						FStackValue Value;
+						Value.Set<TSharedPtr<const FInstance>>(ResultsInstance[Item]);
+					
+						Stack.Push(Value);
+						break;
+					}
+					
+				case EDataType::String:
+					{
+						FStackValue Value;
+						Value.Set<TSharedPtr<const String>>(ResultsString[Item]);
+					
+						Stack.Push(Value);
+						break;
+					}
+								
+				case EDataType::ExtensionData:
+					{
+						FStackValue Value;
+						Value.Set<TSharedPtr<const FExtensionData>>(ResultsExtensionData[Item]);
+					
+						Stack.Push(Value);
+						break;
+					}
+
+				case EDataType::Scalar:
+					{
+						FStackValue Value;
+						Value.Set<float>(ResultsFloat[Item]);
+					
+						Stack.Push(Value);
+						break;
+					}
+					
+				case EDataType::Color:		
+				case EDataType::Projector:	
+				case EDataType::Mesh:		
+				case EDataType::Image:		
+				case EDataType::Layout:		
+				case EDataType::Material:		
+					break;
+
+				default:
+					unimplemented();
+				}
+
+				continue;
+			}
+			
+			switch (Type)
+			{
+			case EOpType::CO_CONSTANT:
+			case EOpType::IM_CONSTANT:
+			case EOpType::LA_CONSTANT:
+			case EOpType::MA_CONSTANT:
+			case EOpType::PR_CONSTANT:
+			case EOpType::CO_PARAMETER:
+			case EOpType::MA_PARAMETER:
+			case EOpType::PR_PARAMETER:
+			case EOpType::ME_PARAMETER:
+			case EOpType::IM_PARAMETER:
+			case EOpType::MI_PARAMETER:
+				{
+					StoreNone(Item);
+					break;
+				}
+			
+			case EOpType::CO_CONDITIONAL:
+			case EOpType::ED_CONDITIONAL:
+			case EOpType::IM_CONDITIONAL:
+			case EOpType::IN_CONDITIONAL:
+			case EOpType::LA_CONDITIONAL:
+			case EOpType::ME_CONDITIONAL:
+			case EOpType::SC_CONDITIONAL:
+			case EOpType::MI_CONDITIONAL:
+				{
+					OP::ConditionalArgs Args = Program.GetOpArgs<OP::ConditionalArgs>(Item.At);
+
+					switch(Item.Stage)
+					{
+					case 0:
+						{
+							PushOp(FScheduledOpInline( Item,1));
+							PushOp(FScheduledOpInline( Args.condition, Item));
+
+							break;
+						}
+
+					case 1:
+						{
+							const bool Value = LoadBool(Args.condition);
+
+							OP::ADDRESS ResultAt = Value ? Args.yes : Args.no;
+
+							// Schedule the end of this instruction if necessary
+							FScheduledOpInline NextStage(Item, 2);
+							NextStage.CustomState = ResultAt;
+							PushOp(NextStage);
+							
+							PushOp(FScheduledOpInline(ResultAt, Item));
+							
+							break;
+						}
+
+					case 2:
+						{
+							switch (GetOpDataType(Type))
+							{
+							case EDataType::Int:
+								StoreInt(Item, LoadInt(Item.CustomState));
+								break;
+							
+							case EDataType::String:
+								StoreString(Item, LoadString(Item.CustomState));
+								break;
+
+							case EDataType::Instance:
+								StoreInstance(Item, LoadInstance(Item.CustomState));
+								break;
+
+							case EDataType::ExtensionData:
+								StoreExtensionData(Item, LoadExtensionData(Item.CustomState));
+								break;
+							
+							case EDataType::Scalar:		
+								StoreFloat(Item, LoadFloat(Item.CustomState));
+								break;
+
+							case EDataType::Color:		
+							case EDataType::Projector:	
+							case EDataType::Mesh:		
+							case EDataType::Image:		
+							case EDataType::Layout:		
+							case EDataType::Material:		
+								StoreNone(Item);
+								break;
+
+							case EDataType::Bool:
+								check(false);
+								
+							default:
+								unimplemented();
+							}
+
+							break;
+						}
+
+					default:
+						unimplemented();
+					}
+				
+					break;
+				}
+
+			case EOpType::NU_CONDITIONAL:
+				{
+					OP::ConditionalArgs Args = Program.GetOpArgs<OP::ConditionalArgs>(Item.At);
+
+					switch(Item.Stage)
+					{
+					case 0:
+						{
+							PushOp(FScheduledOpInline(Item,1));
+							PushOp(FScheduledOpInline(Args.condition, Item));
+
+							break;
+						}
+
+					case 1:
+						{
+							const bool Value = LoadBool(Args.condition);
+							
+							OP::ADDRESS ResultAt = Value ? Args.yes : Args.no;
+
+							PushOp(FScheduledOpInline(ResultAt, Item));
+
+							StoreNone(Item);
+							break;
+						}
+
+					default:
+						unimplemented();
+					}
+				
+					break;
+				}
+				
+			case EOpType::CO_SWITCH:
+			case EOpType::ED_SWITCH:
+			case EOpType::IM_SWITCH:
+			case EOpType::IN_SWITCH:
+			case EOpType::LA_SWITCH:
+			case EOpType::ME_SWITCH:
+			case EOpType::NU_SWITCH:
+			case EOpType::SC_SWITCH:
+			case EOpType::MI_SWITCH:
+				{
+					const uint8* Data = Program.GetOpArgsPointer(Item.At);
+
+					OP::ADDRESS VarAddress;
+					FMemory::Memcpy(&VarAddress, Data, sizeof(OP::ADDRESS));
+					Data += sizeof(OP::ADDRESS);
+
+					OP::ADDRESS DefAddress;
+					FMemory::Memcpy(&DefAddress, Data, sizeof(OP::ADDRESS));
+					Data += sizeof(OP::ADDRESS);
+
+					OP::FSwitchCaseDescriptor CaseDesc;
+					FMemory::Memcpy(&CaseDesc, Data, sizeof(OP::FSwitchCaseDescriptor));
+					Data += sizeof(OP::FSwitchCaseDescriptor);
+
+					switch (Item.Stage)
+					{
+					case 0:
+						{
+							if (VarAddress)
+							{
+								PushOp(FScheduledOpInline(Item, 1));
+								PushOp(FScheduledOpInline(VarAddress, Item));
+							}
+							else
+							{
+								switch (GetOpDataType(Type))
+								{
+								case EDataType::Bool:
+									StoreBool(Item, false);
+									break;
+								
+								case EDataType::Int:
+									StoreInt(Item, 0);
+									break;
+								
+								case EDataType::Instance:
+									StoreInstance(Item, nullptr);
+									break;
+
+								case EDataType::String:
+									StoreString(Item, nullptr);
+									break;
+								
+								case EDataType::ExtensionData:
+									StoreExtensionData(Item, MakeShared<FExtensionData>());
+									break;
+
+								case EDataType::Scalar:
+									StoreFloat(Item, 0.0f);
+									break;
+									
+								case EDataType::Color:			
+								case EDataType::Projector:		
+								case EDataType::Mesh:			
+								case EDataType::Image:			
+								case EDataType::Layout:			
+								case EDataType::Material:			
+									StoreNone(Item);
+									break;
+								
+								default:
+									unimplemented()
+								}
+							}
+							break;
+						}
+
+					case 1:
+						{
+							// Get the variable result
+							int32 Var = LoadInt(VarAddress);
+
+							OP::ADDRESS ValueAt = DefAddress;
+
+							if (!CaseDesc.bUseRanges)
+							{
+								for (uint32 C = 0; C < CaseDesc.Count; ++C)
+								{
+									int32 Condition;
+									FMemory::Memcpy(&Condition, Data, sizeof(int32));
+									Data += sizeof(int32);
+
+									OP::ADDRESS CaseAt;
+									FMemory::Memcpy(&CaseAt, Data, sizeof(OP::ADDRESS));
+									Data += sizeof(OP::ADDRESS);
+
+									if (CaseAt && Var == (int32)Condition)
+									{
+										ValueAt = CaseAt;
+										break;
+									}
+								}
+							}
+							else
+							{
+								for (uint32 C = 0; C < CaseDesc.Count; ++C)
+								{
+									int32 ConditionStart;
+									FMemory::Memcpy(&ConditionStart, Data, sizeof(int32));
+									Data += sizeof(int32);
+
+									uint32 RangeSize;
+									FMemory::Memcpy(&RangeSize, Data, sizeof(uint32));
+									Data += sizeof(uint32);
+
+									OP::ADDRESS CaseAt;
+									FMemory::Memcpy(&CaseAt, Data, sizeof(OP::ADDRESS));
+									Data += sizeof(OP::ADDRESS);
+
+									if (CaseAt && Var >= ConditionStart && Var < int32(ConditionStart + RangeSize))
+									{
+										ValueAt = CaseAt;
+										break;
+									}
+								}
+							}
+
+							// Schedule the end of this instruction if necessary
+							FScheduledOpInline NextStage(Item, 2);
+							NextStage.CustomState = ValueAt;
+							PushOp(NextStage);
+
+							PushOp(FScheduledOpInline(ValueAt, Item));
+							
+							break;
+						}
+						
+					case 2:
+						{
+							switch (GetOpDataType(Type))
+							{
+							case EDataType::Bool:
+								StoreBool(Item, LoadBool(Item.CustomState));
+								break;
+
+							case EDataType::Int:
+								StoreInt(Item, LoadInt(Item.CustomState));
+								break;
+
+							case EDataType::Instance:
+								StoreInstance(Item, LoadInstance(Item.CustomState));
+								break;
+								
+							case EDataType::String:
+								StoreString(Item, LoadString(Item.CustomState));
+								break;
+								
+							case EDataType::ExtensionData:
+								StoreExtensionData(Item, LoadExtensionData(Item.CustomState));
+								break;
+
+							case EDataType::Scalar:
+								StoreFloat(Item, LoadFloat(Item.CustomState));
+								break;
+								
+							case EDataType::Color:		
+							case EDataType::Projector:	
+							case EDataType::Mesh:		
+							case EDataType::Image:		
+							case EDataType::Layout:		
+							case EDataType::Material:		
+								StoreNone(Item);
+								break;
+
+							default:
+								unimplemented();
+							}
+
+							break;
+						}
+
+					default:
+						unimplemented();
+					}
+
+					break;
+				}
+
+			case EOpType::IN_ADDVECTOR:
+			case EOpType::IN_ADDSCALAR:
+				{
+					OP::InstanceAddArgs Args = Program.GetOpArgs<OP::InstanceAddArgs>(Item.At);
+					switch (Item.Stage)
+					{
+					case 0:
+						{
+							PushOp(FScheduledOpInline(Item, 1));
+							PushOp(FScheduledOpInline(Args.instance, Item));
+
+							break;
+						}
+
+					case 1:
+						{
+							TSharedPtr<const FInstance> InInstance;
+							if (Args.instance)
+							{
+								InInstance = LoadInstance(Args.instance);
+							}
+
+							TSharedPtr<FInstance> Instance;
+							if (InInstance)
+							{
+								Instance = mu::CloneOrTakeOver<FInstance>(InInstance);
+							}
+							else
+							{
+								Instance = MakeShared<FInstance>();
+							}
+
+							StoreInstance( Item, Instance);
+							break;
+						}
+
+					default:
+						unimplemented();
+					}
+
+					break;
+				}
+        		
+			case EOpType::IN_ADDMESH:
+				{
+					OP::InstanceAddArgs Args = Program.GetOpArgs<OP::InstanceAddArgs>(Item.At);
+					switch (Item.Stage)
+					{
+					case 0:
+						{
+							PushOp(FScheduledOpInline(Item, 1));
+							PushOp(FScheduledOpInline(Args.instance, Item));
+							PushOp(FScheduledOpInline(Args.value, Item));
+							
+							break;
+						}
+
+					case 1:
+						{
+							TSharedPtr<const FInstance> InInstance;
+							if (Args.instance)
+							{
+								InInstance = LoadInstance(Args.instance);
+							}
+
+							TSharedPtr<FInstance> Instance;
+							if (InInstance)
+							{
+								Instance = mu::CloneOrTakeOver<FInstance>(InInstance);
+							}
+							else
+							{
+								Instance = MakeShared<FInstance>();
+							}
+							
+							if (Args.value)
+							{
+								FMeshId MeshId = System.GetPrivate()->WorkingMemoryManager.GetMeshId(Model, Params.Get(), Args.RelevantParametersListIndex, Args.value);
+								OP::ADDRESS NameAd = Args.name;
+								check(NameAd < (uint32)Program.ConstantStrings.Num());
+								const FString& Name = Program.ConstantStrings[NameAd];
+								Instance->SetMesh(0, 0, MeshId, FName(Name));
+							}
+							
+							StoreInstance(Item, Instance);
+							break;
+						}
+
+					default:
+						unimplemented();
+					}
+					
+					break;
+				}
+
+			case EOpType::IN_ADDIMAGE:
+				{
+					OP::InstanceAddArgs Args = Program.GetOpArgs<OP::InstanceAddArgs>(Item.At);
+					switch (Item.Stage)
+					{
+					case 0:
+						{
+							PushOp(FScheduledOpInline(Item, 1));
+							PushOp(FScheduledOpInline(Args.instance, Item));
+							PushOp(FScheduledOpInline(Args.value, Item));
+
+							break;
+						}
+						
+					case 1:
+						{
+							TSharedPtr<const FInstance> InInstance;
+							if (Args.instance)
+							{
+								InInstance = LoadInstance(Args.instance);
+							}
+
+							TSharedPtr<FInstance> Instance;
+							if (InInstance)
+							{
+								Instance = mu::CloneOrTakeOver<FInstance>(InInstance);
+							}
+							else
+							{
+								Instance = MakeShared<FInstance>();
+							}
+							
+							if (Args.value)
+							{
+								FImageId ImageId = System.GetPrivate()->WorkingMemoryManager.GetImageId(Model, Params.Get(), Args.RelevantParametersListIndex, Args.value);
+								OP::ADDRESS NameAd = Args.name;
+								check(NameAd < (uint32)Program.ConstantStrings.Num());
+								const FString& Name = Program.ConstantStrings[NameAd];
+								Instance->AddImage(0, 0, 0, ImageId, FName(Name) );
+							}
+							
+							StoreInstance(Item, Instance);
+							break;
+						}
+
+					default:
+						unimplemented();
+					}
+
+					break;
+				}
+			case EOpType::IN_ADDMATERIAL:
+				{
+					OP::InstanceAddMaterialArgs Args = Program.GetOpArgs<OP::InstanceAddMaterialArgs>(Item.At);
+					switch (Item.Stage)
+					{
+					case 0:
+						{
+							PushOp(FScheduledOpInline(Item, 1));
+							PushOp(FScheduledOpInline(Args.Instance, Item));
+							PushOp(FScheduledOpInline(Args.Material, Item));
+							
+							break;
+						}
+
+					case 1:
+						{
+							TSharedPtr<const FInstance> InInstance;
+							if (Args.Instance)
+							{
+								InInstance = LoadInstance(Args.Instance);
+							}
+
+							TSharedPtr<FInstance> Instance;
+							if (InInstance)
+							{
+								Instance = mu::CloneOrTakeOver<FInstance>(InInstance);
+							}
+							else
+							{
+								Instance = MakeShared<FInstance>();
+							}
+
+
+							if (Args.Material)
+							{
+								FMaterialId MaterialId = System.GetPrivate()->WorkingMemoryManager.GetMaterialId(Model, Params.Get(), Args.RelevantParametersListIndex, Args.Material);
+								Instance->SetMaterialId(0, 0, 0, MaterialId);
+							}
+							
+							StoreInstance( Item, Instance);
+							break;
+						}
+
+					default:
+						unimplemented();
+					}
+
+					break;
+				}
+					
+			case EOpType::IN_ADDOVERLAYMATERIAL:
+				{
+					OP::InstanceAddMaterialArgs Args = Program.GetOpArgs<OP::InstanceAddMaterialArgs>(Item.At);
+					switch (Item.Stage)
+					{
+					case 0:
+						{
+							PushOp(FScheduledOpInline(Item, 1));
+							PushOp(FScheduledOpInline(Args.Instance, Item));
+							PushOp(FScheduledOpInline(Args.Material, Item));
+
+							break;
+						}
+
+					case 1:
+						{
+							TSharedPtr<const FInstance> InInstance;
+							if (Args.Instance)
+							{
+								InInstance = LoadInstance(Args.Instance);
+							}
+
+							TSharedPtr<FInstance> Instance;
+							if (InInstance)
+							{
+								Instance = mu::CloneOrTakeOver<FInstance>(InInstance);
+							}
+							else
+							{
+								Instance = MakeShared<FInstance>();
+							}
+
+							if (Args.Material)
+							{
+								FMaterialId MaterialId = System.GetPrivate()->WorkingMemoryManager.GetMaterialId(Model, Params.Get(), Args.RelevantParametersListIndex, Args.Material);
+								Instance->SetOverlayMaterialId(0, MaterialId);
+							}
+							
+							StoreInstance( Item, Instance);
+							break;
+						}
+
+					default:
+						unimplemented();
+					}
+
+					break;
+				}
+				
+			case EOpType::IN_ADDCOMPONENT:
+				{
+					OP::InstanceAddArgs Args = Program.GetOpArgs<OP::InstanceAddArgs>(Item.At);
+					switch (Item.Stage)
+					{
+					case 0:
+						{
+							PushOp(FScheduledOpInline(Item, 1));
+							PushOp(FScheduledOpInline(Args.instance, Item));
+							PushOp(FScheduledOpInline(Args.value, Item));
+							
+							break;
+						}
+
+					case 1:
+						{
+							TSharedPtr<const FInstance> InInstance;
+							if (Args.instance)
+							{
+								InInstance = LoadInstance(Args.instance);
+							}
+
+							TSharedPtr<FInstance> Instance;
+							if (InInstance)
+							{
+								Instance = mu::CloneOrTakeOver<FInstance>(InInstance);
+							}
+							else
+							{
+								Instance = MakeShared<FInstance>();
+							}
+
+							if (Args.value)
+							{
+								TSharedPtr<const FInstance> pComp = LoadInstance(Args.value);
+
+								int32 NewComponentIndex = Instance->AddComponent();
+
+								if ( !pComp->Components.IsEmpty() )
+								{
+									Instance->Components[NewComponentIndex] = pComp->Components[0];
+
+									// Id
+									Instance->Components[NewComponentIndex].Id = Args.ExternalId;
+								}
+							}
+
+							StoreInstance(Item, Instance);							
+							break;
+						}
+
+					default:
+						unimplemented();
+					}
+
+					break;
+				}
+				
+			case EOpType::IN_ADDSTRING:
+				{
+					OP::InstanceAddArgs Args = Program.GetOpArgs<OP::InstanceAddArgs>( Item.At );
+					switch (Item.Stage)
+					{
+					case 0:
+						{
+							PushOp(FScheduledOpInline(Item, 1));
+							PushOp(FScheduledOpInline(Args.instance, Item));
+							PushOp(FScheduledOpInline(Args.value, Item));
+
+							break;
+						}
+
+					case 1:
+						{
+							TSharedPtr<const FInstance> InInstance;
+							if (Args.instance)
+							{
+								InInstance = LoadInstance(Args.instance);
+							}
+
+							TSharedPtr<FInstance> Instance;
+							if (InInstance)
+							{
+								Instance = mu::CloneOrTakeOver<FInstance>(InInstance);
+							}
+							else
+							{
+								Instance = MakeShared<FInstance>();
+							}
+							
+							StoreInstance( Item, Instance );
+							break;
+						}
+
+					default:
+						unimplemented();
+					}
+
+					break;
+				}
+				
+			case EOpType::IN_ADDSURFACE:
+				{
+					OP::InstanceAddArgs Args = Program.GetOpArgs<OP::InstanceAddArgs>(Item.At);
+		            switch (Item.Stage)
+		            {
+		            case 0:
+			            {
+		            		PushOp(FScheduledOpInline(Item, 1));
+				            PushOp(FScheduledOpInline(Args.instance, Item));
+				            PushOp(FScheduledOpInline(Args.value, Item));
+
+				            break;
+			            }
+		            
+		            case 1:
+		            {
+		            	TSharedPtr<const FInstance> InInstance;
+		            	if (Args.instance)
+		            	{
+		            		InInstance = LoadInstance(Args.instance);
+		            	}
+
+		            	TSharedPtr<FInstance> Instance;
+		            	if (InInstance)
+		            	{
+		            		Instance = mu::CloneOrTakeOver<FInstance>(InInstance);
+		            	}
+		            	else
+		            	{
+		            		Instance = MakeShared<FInstance>();
+		            	}
+
+		                // Empty surfaces are ok, they still need to be created, because they may contain
+		                // additional information like internal or external IDs
+		            	TSharedPtr<const FInstance> Value;
+		            	if (Args.value)
+		            	{
+		            		Value = LoadInstance(Args.value);
+		            	}
+		            		
+	                    int32 SurfaceIndex = Instance->AddSurface(0, 0);
+
+	                    // Surface data
+	                    if (Value &&
+	                        Value->Components.Num() &&
+							Value->Components[0].LODs.Num() &&
+	                        Value->Components[0].LODs[0].Surfaces.Num())
+	                    {
+	                        Instance->Components[0].LODs[0].Surfaces[SurfaceIndex] = Value->Components[0].LODs[0].Surfaces[0];
+	                    }
+
+	                    // Name
+	                    OP::ADDRESS nameAd = Args.name;
+	                    check( nameAd < (uint32)Program.ConstantStrings.Num() );
+	                    const FString& Name = Program.ConstantStrings[ nameAd ];
+	                    Instance->SetSurfaceName( 0, 0, SurfaceIndex, FName(Name) );
+
+	                    // IDs
+	                    Instance->Components[0].LODs[0].Surfaces[SurfaceIndex].InternalId = Args.id;
+	                    Instance->Components[0].LODs[0].Surfaces[SurfaceIndex].ExternalId = Args.ExternalId;
+	                    Instance->Components[0].LODs[0].Surfaces[SurfaceIndex].SharedId = Args.SharedSurfaceId;
+		            		
+		                StoreInstance(Item, Instance);
+		                break;
+		            }
+
+		            default:
+		                check(false);
+		            }
+
+		            break;
+				}
+
+			case EOpType::IN_ADDLOD:
+				{
+					const uint8* Data = Program.GetOpArgsPointer(Item.At);
+
+					uint8 LODCount;
+					FMemory::Memcpy(&LODCount, Data, sizeof(uint8));
+					Data += sizeof(uint8);
+
+					switch (Item.Stage)
+		            {
+		            case 0:
+		            {
+		            	PushOp(FScheduledOpInline( Item, 1));
+
+		                for (uint8 LODIndex = 0; LODIndex < LODCount; ++LODIndex)
+		                {
+							OP::ADDRESS LODAddress;
+							FMemory::Memcpy(&LODAddress, Data, sizeof(OP::ADDRESS));
+							Data += sizeof(OP::ADDRESS);
+
+		                    if (LODAddress)
+		                    {
+		                        const bool bSelectedLod = ((1 << LODIndex) & LodMask) != 0;
+		                        if (bSelectedLod)
+		                        {
+		                        	PushOp(FScheduledOpInline(LODAddress, Item));
+		                        }
+		                    }
+		                }
+
+
+		                break;
+		            }
+
+		            case 1:
+		            {
+		                // Assemble result
+						TSharedPtr<FInstance> Result = MakeShared<FInstance>();
+						int32 ComponentIndex = Result->AddComponent();
+
+		                for (uint8 LODIndex = 0; LODIndex < LODCount; ++LODIndex)
+		                {
+							OP::ADDRESS LODAddress;
+							FMemory::Memcpy(&LODAddress, Data, sizeof(OP::ADDRESS));
+							Data += sizeof(OP::ADDRESS);
+							
+							if (LODAddress)
+		                    {						
+								bool bIsSelectedLod = ((1 << LODIndex) & LodMask ) != 0;
+
+								// Add an empty LOD even if not selected.
+								int32 InstanceLODIndex = Result->AddLOD(ComponentIndex);
+								
+								if (bIsSelectedLod)
+		                        {
+									TSharedPtr<const FInstance> LOD = LoadInstance(LODAddress);
+
+		                            // In a degenerated case, the returned pLOD may not have an LOD inside
+ 									if (!LOD->Components.IsEmpty() &&
+										!LOD->Components[0].LODs.IsEmpty())
+									{
+										Result->Components[ComponentIndex].LODs[InstanceLODIndex] = LOD->Components[0].LODs[0];
+									}
+								}
+		                    }
+		                }
+
+		                StoreInstance(Item, Result);
+		                break;
+		            }
+
+		            default:
+		                unimplemented();
+		            }
+
+		            break;
+		        }
+				
+			case EOpType::IN_ADDEXTENSIONDATA:
+				{
+					OP::InstanceAddExtensionDataArgs Args = Program.GetOpArgs<OP::InstanceAddExtensionDataArgs>(Item.At);
+
+					switch (Item.Stage)
+					{
+					case 0:
+						{
+							PushOp(FScheduledOpInline(Item, 1));
+							PushOp(FScheduledOpInline(Args.Instance, Item));
+							PushOp(FScheduledOpInline(Args.ExtensionData, Item));
+
+							break;
+						}
+
+					case 1:
+						{
+							TSharedPtr<const FInstance> InInstance;
+							if (Args.Instance)
+							{
+								InInstance = LoadInstance(Args.Instance);
+							}
+
+							TSharedPtr<FInstance> Instance;
+							if (InInstance)
+							{
+								Instance = mu::CloneOrTakeOver<FInstance>(InInstance);
+							}
+							else
+							{
+								Instance = MakeShared<FInstance>();
+							}
+							
+							if (TSharedPtr<const FExtensionData> ExtensionData = LoadExtensionData(Args.ExtensionData))
+							{
+								const OP::ADDRESS NameAddress = Args.ExtensionDataName;
+								check(NameAddress < (uint32)Program.ConstantStrings.Num());
+								const FString& NameString = Program.ConstantStrings[NameAddress];
+
+								Instance->SetExtensionData(ExtensionData.ToSharedRef(), FName(NameString));
+							}
+
+							StoreInstance(Item, Instance);
+							break;
+						}
+
+					default:
+						unimplemented();
+					}
+
+					break;
+				}
+				
+			case EOpType::BO_CONSTANT:
+				{
+					OP::BoolConstantArgs Args = Program.GetOpArgs<OP::BoolConstantArgs>(Item.At);
+
+					StoreBool(Item, Args.bValue);
+					break;
+				}
+
+			case EOpType::BO_PARAMETER:
+				{
+					OP::ParameterArgs Args = Program.GetOpArgs<OP::ParameterArgs>(Item.At);
+
+					FScheduledOp Op;
+					Op.At = Item.At;
+					Op.ExecutionIndex = Item.ExecutionIndex;
+        		
+					TSharedPtr<FRangeIndex> Index = BuildCurrentOpRangeIndex(Op, *Params, *Model.Get(), GetMemory(), Args.variable);
+					const bool Value = Params->GetBoolValue(Args.variable, Index.Get());
+					
+					StoreBool(Item, Value);
+					break;
+				}
+
+			case EOpType::BO_AND:
+				{
+					OP::BoolBinaryArgs Args = Program.GetOpArgs<OP::BoolBinaryArgs>(Item.At);
+
+					switch (Item.Stage)
+					{
+					case 0:
+						{
+							PushOp(FScheduledOpInline(Item, 1));
+							PushOp(FScheduledOpInline(Args.A, Item));
+							PushOp(FScheduledOpInline(Args.B, Item));
+
+							break;
+						}
+
+					case 1:
+						{
+							const bool ValueB = LoadBool(Args.A);
+							const bool ValueA = LoadBool(Args.B);
+
+							const bool Result = ValueA && ValueB;
+
+							StoreBool(Item, Result);
+							break;
+						}
+
+					default:
+						unimplemented();
+					}
+					
+					break;
+				}
+
+			case EOpType::BO_OR:
+				{
+					OP::BoolBinaryArgs Args = Program.GetOpArgs<OP::BoolBinaryArgs>(Item.At);
+
+					switch (Item.Stage)
+					{
+					case 0:
+						{
+							PushOp(FScheduledOpInline(Item, 1));
+							PushOp(FScheduledOpInline(Args.A, Item));
+							PushOp(FScheduledOpInline(Args.B, Item));
+
+							break;
+						}
+
+					case 1:
+						{
+							const bool ValueB = LoadBool(Args.A);
+							const bool ValueA = LoadBool(Args.B);
+
+							const bool Result = ValueA || ValueB;
+							
+							StoreBool(Item, Result);
+							break;
+						}
+
+					default:
+						unimplemented();
+					}
+					
+					break;
+				}
+
+			case EOpType::BO_NOT:
+				{
+					OP::BoolNotArgs Args = Program.GetOpArgs<OP::BoolNotArgs>(Item.At);
+
+					switch ( Item.Stage )
+					{
+					case 0:
+						{
+							PushOp(FScheduledOpInline(Item, 1));
+							PushOp(FScheduledOpInline(Args.A, Item));
+
+							break;
+						}
+					
+					case 1:
+						{
+							const bool Value = LoadBool(Args.A);
+							
+							const bool Result = !Value;
+							
+							StoreBool(Item, Result);
+							break;
+						}
+
+					default:
+						unimplemented();
+					}
+					
+					break;
+				}
+
+			case EOpType::BO_EQUAL_INT_CONST:
+				{
+					OP::BoolEqualScalarConstArgs Args = Program.GetOpArgs<OP::BoolEqualScalarConstArgs>(Item.At);
+
+					switch ( Item.Stage )
+					{
+					case 0:
+						{
+							PushOp(FScheduledOpInline(Item, 1));
+							PushOp(FScheduledOpInline(Args.Value, Item));
+
+							break;
+						}
+					
+					case 1:
+						{
+							const int32 Value = LoadInt(Args.Value);
+							
+							const bool Result = Value == Args.Constant;
+
+							StoreBool(Item, Result);
+							break;
+						}
+
+					default:
+						unimplemented();
+					}
+					
+					break;
+				}
+
+			case EOpType::SC_CONSTANT:
+				{
+					OP::ScalarConstantArgs Args = Program.GetOpArgs<OP::ScalarConstantArgs>(Item.At);
+
+					StoreFloat(Item, Args.Value);
+					break;
+				}
+			
+				
+			case EOpType::SC_PARAMETER:
+				{
+					OP::ParameterArgs Args = Program.GetOpArgs<OP::ParameterArgs>(Item.At);
+
+					FScheduledOp Op;
+					Op.At = Item.At;
+					Op.ExecutionIndex = Item.ExecutionIndex;
+					
+					TSharedPtr<FRangeIndex> Index = BuildCurrentOpRangeIndex(Op, *Params, *Model.Get(), GetMemory(), Args.variable);
+					const float Result = Params->GetFloatValue( Args.variable, Index.Get());
+
+					StoreFloat(Item, Result);
+					break;
+				}
+				
+			case EOpType::SC_ARITHMETIC:
+				{
+					OP::ArithmeticArgs Args = Program.GetOpArgs<OP::ArithmeticArgs>(Item.At);
+					switch (Item.Stage)
+					{
+					case 0:
+						{
+							PushOp(FScheduledOpInline(Item, 1));
+							PushOp(FScheduledOpInline(Args.A, Item));
+							PushOp(FScheduledOpInline(Args.B, Item));
+
+							break;
+						}
+
+					case 1:
+						{
+							float ValueA = LoadFloat(Args.A);
+							float ValueB = LoadFloat(Args.B);
+
+							float Result = 1.0f;
+							switch (Args.Operation)
+							{
+							case OP::ArithmeticArgs::ADD:
+								Result = ValueA + ValueB;
+								break;
+
+							case OP::ArithmeticArgs::MULTIPLY:
+								Result = ValueA * ValueB;
+								break;
+
+							case OP::ArithmeticArgs::SUBTRACT:
+								Result = ValueA - ValueB;
+								break;
+
+							case OP::ArithmeticArgs::DIVIDE:
+								Result = ValueA / ValueB;
+								break;
+
+							default:
+								unimplemented();
+							}
+
+							StoreFloat(Item, Result);
+							break;
+						}
+
+					default:
+						unimplemented();
+					}
+					
+					break;
+				}
+			
+			case EOpType::SC_CURVE:
+				{
+					OP::ScalarCurveArgs Args = Program.GetOpArgs<OP::ScalarCurveArgs>(Item.At);
+					switch (Item.Stage)
+					{
+					case 0:
+						{
+							PushOp(FScheduledOpInline(Item, 1));
+							PushOp(FScheduledOpInline(Args.time, Item));
+
+							break;
+						}
+
+					case 1:
+						{
+							const float Time = LoadFloat(Args.time);
+
+							const FRichCurve& Curve = Program.ConstantCurves[Args.curve];
+							float Result = Curve.Eval(Time);
+
+							StoreFloat(Item, Result);
+							break;
+						}
+
+					default:
+						unimplemented();
+					}
+					
+					break;
+				}
+				
+			case EOpType::ED_CONSTANT:
+				{
+					OP::ResourceConstantArgs Args = Program.GetOpArgs<OP::ResourceConstantArgs>(Item.At);
+
+					// Assume the ROM has been loaded previously
+					TSharedPtr<const FExtensionData> SourceConst;
+					Program.GetExtensionDataConstant(Args.value, SourceConst);
+
+					check(SourceConst);
+					
+					StoreExtensionData(Item, SourceConst);
+					break;
+				}
+
+			case EOpType::ST_CONSTANT:
+				{
+					OP::ResourceConstantArgs Args = Program.GetOpArgs<OP::ResourceConstantArgs>(Item.At);
+				
+					check(Args.value < static_cast<uint32>(Model->GetPrivate()->Program.ConstantStrings.Num()));
+
+					const FString& Result = Program.ConstantStrings[Args.value];
+
+					TSharedPtr<const String> Value = MakeShared<String>(Result);
+
+					StoreString(Item, Value);
+					break;
+				}
+
+			case EOpType::ST_PARAMETER:
+				{
+					OP::ParameterArgs Args = Program.GetOpArgs<OP::ParameterArgs>(Item.At);
+
+					FScheduledOp Op;
+					Op.At = Item.At;
+					Op.ExecutionIndex = Item.ExecutionIndex;
+				
+					TSharedPtr<FRangeIndex> Index = BuildCurrentOpRangeIndex(Op, *Params, *Model.Get(), GetMemory(), Args.variable);
+
+					FString Result;
+					Params->GetStringValue(Args.variable, Result, Index.Get());
+
+					TSharedPtr<const String> Value = MakeShared<String>(Result);
+					
+					StoreString(Item, Value);
+					break;
+				}
+
+			case EOpType::NU_CONSTANT:
+				{
+					OP::IntConstantArgs Args = Program.GetOpArgs<OP::IntConstantArgs>(Item.At);
+
+					const int32 Value = Args.Value;
+					
+					StoreInt(Item, Value);
+					break;
+				}
+
+			case EOpType::NU_PARAMETER:
+				{
+					OP::ParameterArgs Args = Program.GetOpArgs<OP::ParameterArgs>(Item.At);
+
+					const int32 Value = Params->GetIntValue(Args.variable);
+
+					StoreInt(Item, Value);
+					break;
+				}
+			
+			case EOpType::CO_ARITHMETIC:
+				{
+					OP::ArithmeticArgs Args = Program.GetOpArgs<OP::ArithmeticArgs>(Item.At);
+					PushOp(FScheduledOpInline(Args.A, Item));
+					PushOp(FScheduledOpInline(Args.B, Item));
+
+					StoreNone(Item);
+					break;	
+				}
+	
+			case EOpType::CO_FROMSCALARS:
+				{
+					StoreNone(Item);
+					break;
+				}
+				
+			case EOpType::CO_SAMPLEIMAGE:
+				{
+					OP::ColourSampleImageArgs Args = Program.GetOpArgs<OP::ColourSampleImageArgs>(Item.At);
+					PushOp(FScheduledOpInline(Args.Image, Item));
+
+					StoreNone(Item);
+					break;
+				}
+	
+			case EOpType::CO_SWIZZLE:
+				{
+					OP::ColourSwizzleArgs Args = Program.GetOpArgs<OP::ColourSwizzleArgs>(Item.At);
+					PushOp(FScheduledOpInline(Args.sources[0], Item));
+					PushOp(FScheduledOpInline(Args.sources[1], Item));
+					PushOp(FScheduledOpInline(Args.sources[2], Item));
+					PushOp(FScheduledOpInline(Args.sources[3], Item));
+
+					StoreNone(Item);
+					break;
+				}
+
+			case EOpType::ME_CONSTANT:
+				{
+					OP::MeshConstantArgs Args = Program.GetOpArgs<OP::MeshConstantArgs>(Item.At);
+					ME_SKELETON_ID.Push(Args.Skeleton);
+
+					StoreNone(Item);
+					break;
+				}
+				
+			case EOpType::ME_REFERENCE:
+				{
+					OP::ResourceReferenceArgs Args = Program.GetOpArgs<OP::ResourceReferenceArgs>(Item.At);
+					ME_REFERENCE_ID.Push(Args.ID);
+
+					StoreNone(Item);
+					break;
+				}
+
+			case EOpType::ME_APPLYLAYOUT:
+				{
+					OP::MeshApplyLayoutArgs Args = Program.GetOpArgs<OP::MeshApplyLayoutArgs>(Item.At);
+					PushOp(FScheduledOpInline(Args.Mesh, Item));
+
+					StoreNone(Item);
+					break;
+				}
+	
+	
+			case EOpType::ME_PREPARELAYOUT:
+				{
+					OP::MeshPrepareLayoutArgs Args = Program.GetOpArgs<OP::MeshPrepareLayoutArgs>(Item.At);
+					PushOp(FScheduledOpInline(Args.Mesh, Item));
+
+					StoreNone(Item);
+					break;
+				}
+	
+			case EOpType::ME_DIFFERENCE:
+				{
+					const uint8* data = Program.GetOpArgsPointer(Item.At);
+
+					OP::ADDRESS BaseAt = 0;
+					FMemory::Memcpy(&BaseAt, data, sizeof(OP::ADDRESS)); 
+					data += sizeof(OP::ADDRESS);
+
+					OP::ADDRESS TargetAt = 0;
+					FMemory::Memcpy(&TargetAt, data, sizeof(OP::ADDRESS)); 
+					data += sizeof(OP::ADDRESS);
+
+					PushOp(FScheduledOpInline(BaseAt, Item));
+					PushOp(FScheduledOpInline(TargetAt, Item));
+
+					StoreNone(Item);	
+					break;
+				}
+	
+			case EOpType::ME_MORPH:
+				{
+					const uint8* Data = Program.GetOpArgsPointer(Item.At);
+
+					OP::CONSTANT_STRING_ADDRESS NameAt = 0;
+					FMemory::Memcpy(&NameAt, Data, sizeof(OP::CONSTANT_STRING_ADDRESS)); 
+					Data += sizeof(OP::CONSTANT_STRING_ADDRESS);
+					
+					OP::ADDRESS FactorAt = 0;
+					FMemory::Memcpy(&FactorAt, Data, sizeof(OP::ADDRESS)); 
+					Data += sizeof(OP::ADDRESS);
+			
+					OP::ADDRESS BaseAt = 0;
+					FMemory::Memcpy(&BaseAt, Data, sizeof(OP::ADDRESS)); 
+					Data += sizeof(OP::ADDRESS);
+
+					OP::ADDRESS TargetAt = 0;
+					FMemory::Memcpy(&TargetAt, Data, sizeof(OP::ADDRESS)); 
+					Data += sizeof(OP::ADDRESS);
+
+					PushOp(FScheduledOpInline(BaseAt, Item));
+					PushOp(FScheduledOpInline(TargetAt, Item));
+
+					StoreNone(Item);
+					break;
+				}
+	
+			case EOpType::ME_MERGE:
+				{
+					OP::MeshMergeArgs Args = Program.GetOpArgs<OP::MeshMergeArgs>(Item.At);
+					PushOp(FScheduledOpInline(Args.Base, Item));
+					PushOp(FScheduledOpInline(Args.Added, Item));
+
+					StoreNone(Item);
+					break;
+				}
+	
+			case EOpType::ME_MASKCLIPMESH:
+				{
+					OP::MeshMaskClipMeshArgs Args = Program.GetOpArgs<OP::MeshMaskClipMeshArgs>(Item.At);
+					PushOp(FScheduledOpInline(Args.source, Item));
+					PushOp(FScheduledOpInline(Args.clip, Item));
+
+					StoreNone(Item);
+					break;
+				}
+	
+			case EOpType::ME_MASKCLIPUVMASK:
+				{
+					OP::MeshMaskClipUVMaskArgs Args = Program.GetOpArgs<OP::MeshMaskClipUVMaskArgs>(Item.At);
+					PushOp(FScheduledOpInline(Args.Source, Item));
+					PushOp(FScheduledOpInline(Args.UVSource, Item));
+					PushOp(FScheduledOpInline(Args.MaskImage, Item));
+					PushOp(FScheduledOpInline(Args.LayoutIndex, Item));
+
+					StoreNone(Item);
+					break;
+				}
+	
+			case EOpType::ME_MASKDIFF:
+				{
+					OP::MeshMaskDiffArgs Args = Program.GetOpArgs<OP::MeshMaskDiffArgs>(Item.At);
+					PushOp(FScheduledOpInline(Args.Source, Item));
+					PushOp(FScheduledOpInline(Args.Fragment, Item));
+
+					StoreNone(Item);
+					break;
+				}
+
+			case EOpType::ME_FORMAT:
+				{
+					OP::MeshFormatArgs Args = Program.GetOpArgs<OP::MeshFormatArgs>(Item.At);
+					PushOp(FScheduledOpInline(Args.source, Item));
+					PushOp(FScheduledOpInline(Args.format, Item));
+
+					StoreNone(Item);
+					break;
+				}
+	
+			case EOpType::ME_EXTRACTLAYOUTBLOCK:
+				{
+					const uint8* Data = Program.GetOpArgsPointer(Item.At);
+
+					OP::ADDRESS Source;
+					FMemory::Memcpy( &Source, Data, sizeof(OP::ADDRESS) );
+					Data += sizeof(OP::ADDRESS);
+					
+					PushOp(FScheduledOpInline(Source, Item));
+
+					StoreNone(Item);
+					break;
+				}
+	
+			case EOpType::ME_TRANSFORM:
+				{
+					OP::MeshTransformArgs Args = Program.GetOpArgs<OP::MeshTransformArgs>(Item.At);
+					PushOp(FScheduledOpInline(Args.source, Item));
+
+					StoreNone(Item);
+					break;
+				}
+	
+			case EOpType::ME_CLIPMORPHPLANE:
+				{
+					OP::MeshClipMorphPlaneArgs Args = Program.GetOpArgs<OP::MeshClipMorphPlaneArgs>(Item.At);
+					PushOp(FScheduledOpInline(Args.Source, Item));
+
+					StoreNone(Item);
+					break;
+				}
+			
+			case EOpType::ME_CLIPWITHMESH:
+				{
+					OP::MeshClipWithMeshArgs Args = Program.GetOpArgs<OP::MeshClipWithMeshArgs>(Item.At);
+					PushOp(FScheduledOpInline(Args.Source, Item));
+					PushOp(FScheduledOpInline(Args.ClipMesh, Item));
+
+					StoreNone(Item);
+					break;
+				}
+				
+			case EOpType::ME_CLIPDEFORM:
+				{
+					OP::MeshClipDeformArgs Args = Program.GetOpArgs<OP::MeshClipDeformArgs>(Item.At);
+					PushOp(FScheduledOpInline(Args.mesh, Item));
+					PushOp(FScheduledOpInline(Args.clipShape, Item));
+
+					StoreNone(Item);
+					break;
+				}
+
+			case EOpType::ME_APPLYPOSE:
+				{
+					OP::MeshApplyPoseArgs Args = Program.GetOpArgs<OP::MeshApplyPoseArgs>(Item.At);
+					PushOp(FScheduledOpInline(Args.base, Item));
+					PushOp(FScheduledOpInline(Args.pose, Item));
+
+					StoreNone(Item);
+					break;
+				}
+	
+			case EOpType::ME_BINDSHAPE:
+				{
+					OP::MeshBindShapeArgs Args = Program.GetOpArgs<OP::MeshBindShapeArgs>(Item.At);
+					PushOp(FScheduledOpInline(Args.mesh, Item));
+					PushOp(FScheduledOpInline(Args.shape, Item));
+
+					StoreNone(Item);
+					break;
+				}
+	
+			case EOpType::ME_APPLYSHAPE:
+				{
+					OP::MeshApplyShapeArgs Args = Program.GetOpArgs<OP::MeshApplyShapeArgs>(Item.At);
+					PushOp(FScheduledOpInline(Args.mesh, Item));
+					PushOp(FScheduledOpInline(Args.shape, Item));
+
+					StoreNone(Item);
+					break;
+				}
+	
+			case EOpType::ME_MORPHRESHAPE:
+				{
+					OP::MeshMorphReshapeArgs Args = Program.GetOpArgs<OP::MeshMorphReshapeArgs>(Item.At);
+					PushOp(FScheduledOpInline(Args.Morph, Item));
+					PushOp(FScheduledOpInline(Args.Reshape, Item));
+
+					StoreNone(Item);
+					break;
+				}
+	
+			case EOpType::ME_SETSKELETON:
+				{
+					OP::MeshSetSkeletonArgs Args = Program.GetOpArgs<OP::MeshSetSkeletonArgs>(Item.At);
+					PushOp(FScheduledOpInline(Args.Source, Item));
+					PushOp(FScheduledOpInline(Args.Skeleton, Item));
+
+					StoreNone(Item);
+					break;
+				}
+	
+			case EOpType::ME_REMOVEMASK:
+				{
+					const uint8* Data = Program.GetOpArgsPointer(Item.At);
+
+					OP::ADDRESS Source;
+					FMemory::Memcpy(&Source,Data,sizeof(OP::ADDRESS)); 
+					Data += sizeof(OP::ADDRESS);
+
+					PushOp(FScheduledOpInline(Source, Item));
+
+					EFaceCullStrategy FaceCullStrategy;
+					FMemory::Memcpy(&FaceCullStrategy, Data, sizeof(EFaceCullStrategy));
+					Data += sizeof(EFaceCullStrategy);
+
+					uint16 NumRemoves;
+					FMemory::Memcpy(&NumRemoves,Data,sizeof(uint16)); 
+					Data += sizeof(uint16);
+
+					switch (Item.Stage)
+					{
+					case 0:
+						{
+							PushOp(FScheduledOpInline(Item, 1));
+
+							for (uint16 Index = 0; Index < NumRemoves; ++Index)
+							{
+								OP::ADDRESS Condition;
+								FMemory::Memcpy(&Condition, Data, sizeof(OP::ADDRESS)); 
+								Data += sizeof(OP::ADDRESS);
+
+								if (Condition)
+								{
+									PushOp(FScheduledOpInline(Condition, Item));								
+								}
+
+								OP::ADDRESS Mask;
+								FMemory::Memcpy(&Mask, Data,sizeof(OP::ADDRESS)); 
+								Data += sizeof(OP::ADDRESS);
+							}
+
+							break;
+						}
+
+					case 1:
+						{
+							for (uint16 Index = 0; Index < NumRemoves; ++Index)
+							{
+								OP::ADDRESS Condition;
+								FMemory::Memcpy(&Condition,Data,sizeof(OP::ADDRESS)); 
+								Data += sizeof(OP::ADDRESS);
+								
+								OP::ADDRESS Mask;
+								FMemory::Memcpy(&Mask,Data,sizeof(OP::ADDRESS)); 
+								Data += sizeof(OP::ADDRESS);
+
+								bool bValue = true;
+								if (Condition)
+								{
+								 	bValue = LoadBool(Condition);
+								}
+								
+								if (bValue)
+								{
+									PushOp(FScheduledOpInline(Mask, Item));
+								}
+							}
+
+							StoreNone(Item);
+							break;
+						}
+
+					default:
+						unimplemented()		
+					}
+					
+					break;
+				}
+
+			case EOpType::ME_ADDMETADATA:
+				{
+					const OP::MeshAddMetadataArgs OpArgs = Program.GetOpArgs<OP::MeshAddMetadataArgs>(Item.At);
+					PushOp(FScheduledOpInline(OpArgs.Source, Item));
+
+					using OpEnumFlags = OP::MeshAddMetadataArgs::EnumFlags;
+					const bool bIsSkeletonList = EnumHasAnyFlags(OpArgs.Flags, OpEnumFlags::IsSkeletonList);
+
+					if (bIsSkeletonList)
+					{
+						if (Program.ConstantUInt32Lists.IsValidIndex(OpArgs.SkeletonIds.ListAddress))
+						{
+							ME_SKELETON_ID.Append(Program.ConstantUInt32Lists[OpArgs.SkeletonIds.ListAddress]);
+						}
+						else
+						{
+							check(false);
+						}
+					}
+					else
+					{
+						ME_SKELETON_ID.Push(OpArgs.SkeletonIds.SkeletonId);
+					}
+
+					StoreNone(Item);
+					break;
+				}
+				
+			case EOpType::ME_PROJECT:
+				{
+					OP::MeshProjectArgs Args = Program.GetOpArgs<OP::MeshProjectArgs>(Item.At);
+					PushOp(FScheduledOpInline(Args.Mesh, Item));
+					PushOp(FScheduledOpInline(Args.Projector, Item));
+
+					StoreNone(Item);
+					break;
+				}
+
+			case EOpType::ME_OPTIMIZESKINNING:
+				{
+					OP::MeshOptimizeSkinningArgs Args = Program.GetOpArgs<OP::MeshOptimizeSkinningArgs>(Item.At);
+					PushOp(FScheduledOpInline(Args.source, Item));
+
+					StoreNone(Item);
+					break;
+				}
+
+			case EOpType::ME_TRANSFORMWITHMESH:
+				{
+					OP::MeshTransformWithinMeshArgs Args = Program.GetOpArgs<OP::MeshTransformWithinMeshArgs>(Item.At);
+					PushOp(FScheduledOpInline(Args.sourceMesh, Item));
+					PushOp(FScheduledOpInline(Args.boundingMesh, Item));
+
+					StoreNone(Item);
+					break;
+				}
+			
+			case EOpType::IM_REFERENCE:
+				{
+					OP::ResourceReferenceArgs Args = Program.GetOpArgs<OP::ResourceReferenceArgs>(Item.At);
+					IM_REFERENCE_ID.Push(Args.ID);
+
+					StoreNone(Item);
+					break;
+				}
+				
+			case EOpType::IM_LAYERCOLOUR:
+				{
+					OP::ImageLayerColourArgs Args = Program.GetOpArgs<OP::ImageLayerColourArgs>(Item.At);
+					PushOp(FScheduledOpInline(Args.base, Item));
+					PushOp(FScheduledOpInline(Args.mask, Item));
+
+					StoreNone(Item);
+					break;
+				}
+
+			case EOpType::IM_LAYER:
+				{
+					OP::ImageLayerArgs Args = Program.GetOpArgs<OP::ImageLayerArgs>(Item.At);	
+					PushOp(FScheduledOpInline(Args.base, Item));
+					PushOp(FScheduledOpInline(Args.mask, Item));
+					PushOp(FScheduledOpInline(Args.blended, Item));
+
+					StoreNone(Item);
+					break;
+				}
+
+			case EOpType::IM_MULTILAYER:
+				{
+					OP::ImageMultiLayerArgs Args = Program.GetOpArgs<OP::ImageMultiLayerArgs>(Item.At);
+
+					switch (Item.Stage)
+					{
+					case 0:
+						{
+							PushOp(FScheduledOpInline(Item, 1));
+							PushOp(FScheduledOpInline(Args.rangeSize, Item));
+							PushOp(FScheduledOpInline(Args.base, Item));
+
+							break;
+						}
+
+					case 1:
+						{
+							int32 NumIterations = 0;
+							if (Args.rangeSize)
+							{
+								const EDataType RangeSizeType = GetOpDataType(Model->GetPrivate()->Program.GetOpType(Args.rangeSize) );
+								if (RangeSizeType == EDataType::Int)
+								{
+									NumIterations = LoadInt(Args.rangeSize);
+								}
+								else if (RangeSizeType == EDataType::Scalar)
+								{
+									NumIterations = LoadFloat(Args.rangeSize);
+								}
+							}
+
+							for (int32 IterationIndex = 0; IterationIndex < NumIterations; ++IterationIndex)
+							{
+								ExecutionIndex Index = GetMemory().GetRangeIndex(Item.ExecutionIndex);
+								Index.SetFromModelRangeIndex(Args.rangeId, IterationIndex);
+
+								FScheduledOpInline ItemCopy = Item;
+								ItemCopy.ExecutionIndex = GetMemory().GetRangeIndexIndex(Index);
+
+								PushOp(FScheduledOpInline(Args.mask, Item));
+								PushOp(FScheduledOpInline(Args.blended, Item));
+							}
+
+							StoreNone(Item);
+							break;
+						}
+
+					default:
+						unimplemented()
+					}
+
+					break;
+				}
+
+			case EOpType::IM_NORMALCOMPOSITE:
+				{
+					OP::ImageNormalCompositeArgs Args = Program.GetOpArgs<OP::ImageNormalCompositeArgs>(Item.At);
+					PushOp(FScheduledOpInline(Args.base, Item));
+					PushOp(FScheduledOpInline(Args.normal, Item));
+
+					StoreNone(Item);
+					break;
+				}
+	
+			case EOpType::IM_PIXELFORMAT:
+				{
+					OP::ImagePixelFormatArgs Args = Program.GetOpArgs<OP::ImagePixelFormatArgs>(Item.At);
+					PushOp(FScheduledOpInline(Args.source, Item));
+
+					StoreNone(Item);
+					break;
+				}
+
+			case EOpType::IM_MIPMAP:
+				{
+					OP::ImageMipmapArgs Args = Program.GetOpArgs<OP::ImageMipmapArgs>(Item.At);
+					PushOp(FScheduledOpInline(Args.source, Item));
+
+					StoreNone(Item);
+					break;
+				}
+	
+			case EOpType::IM_RESIZE:
+				{
+					OP::ImageResizeArgs Args = Program.GetOpArgs<OP::ImageResizeArgs>(Item.At);
+					PushOp(FScheduledOpInline(Args.Source, Item));
+
+					StoreNone(Item);
+					break;
+				}
+	
+			case EOpType::IM_RESIZELIKE:
+				{
+					OP::ImageResizeLikeArgs Args = Program.GetOpArgs<OP::ImageResizeLikeArgs>(Item.At);
+					PushOp(FScheduledOpInline(Args.Source, Item));
+					PushOp(FScheduledOpInline(Args.SizeSource, Item));
+
+					StoreNone(Item);
+					break;
+				}
+
+			case EOpType::IM_RESIZEREL:
+				{
+					OP::ImageResizeRelArgs Args = Program.GetOpArgs<OP::ImageResizeRelArgs>(Item.At);
+					PushOp(FScheduledOpInline(Args.Source, Item));
+
+					StoreNone(Item);
+					break;
+				}
+	
+			case EOpType::IM_BLANKLAYOUT:
+				{
+					OP::ImageBlankLayoutArgs Args = Program.GetOpArgs<OP::ImageBlankLayoutArgs>(Item.At);
+					PushOp(FScheduledOpInline(Args.Layout, Item));
+
+					StoreNone(Item);
+					break;
+				}
+
+			case EOpType::IM_COMPOSE:
+				{
+					OP::ImageComposeArgs Args = Program.GetOpArgs<OP::ImageComposeArgs>(Item.At);
+					PushOp(FScheduledOpInline(Args.base, Item));
+					PushOp(FScheduledOpInline(Args.blockImage, Item)); 
+
+					StoreNone(Item);
+					break;
+				}
+	
+			case EOpType::IM_INTERPOLATE:
+				{
+					OP::ImageInterpolateArgs Args = Program.GetOpArgs<OP::ImageInterpolateArgs>(Item.At);
+					
+					for (int32 ImageIndex = 0; ImageIndex < MUTABLE_OP_MAX_INTERPOLATE_COUNT; ++ImageIndex)
+					{
+						if (!Args.Targets[ImageIndex])
+						{
+							break;
+						}
+
+						PushOp(FScheduledOpInline(Args.Targets[ImageIndex], Item));
+					}
+
+					StoreNone(Item);
+					break;
+				}
+	
+			case EOpType::IM_SATURATE:
+				{
+					OP::ImageSaturateArgs Args = Program.GetOpArgs<OP::ImageSaturateArgs>(Item.At);
+					PushOp(FScheduledOpInline(Args.Base, Item));
+
+					StoreNone(Item);
+					break;
+				}
+
+			case EOpType::IM_LUMINANCE:
+				{
+					OP::ImageLuminanceArgs Args = Program.GetOpArgs<OP::ImageLuminanceArgs>(Item.At);
+					PushOp(FScheduledOpInline(Args.Base, Item));
+
+					StoreNone(Item);
+					break;
+				}
+	
+			case EOpType::IM_SWIZZLE:
+				{
+					OP::ImageSwizzleArgs Args = Program.GetOpArgs<OP::ImageSwizzleArgs>(Item.At);
+				
+					TArray<OP::ADDRESS, TFixedAllocator<4>> ValidArgs;
+					for (int32 SourceIndex = 0; SourceIndex < 4; ++SourceIndex)
+					{
+						if (Args.sources[SourceIndex])
+						{
+							PushOp(FScheduledOpInline(Args.sources[SourceIndex], Item));
+						}
+					}
+
+					StoreNone(Item);
+					break;
+				}
+	
+			case EOpType::IM_COLOURMAP:
+				{
+					OP::ImageColourMapArgs Args = Program.GetOpArgs<OP::ImageColourMapArgs>(Item.At);
+					PushOp(FScheduledOpInline(Args.Base, Item));
+
+					StoreNone(Item);
+					break;
+				}
+	
+			case EOpType::IM_BINARISE:
+				{
+					OP::ImageBinariseArgs Args = Program.GetOpArgs<OP::ImageBinariseArgs>(Item.At);
+					PushOp(FScheduledOpInline(Args.Base, Item));
+
+					StoreNone(Item);
+					break;
+				}
+	
+			case EOpType::IM_INVERT:
+				{
+					OP::ImageInvertArgs Args = Program.GetOpArgs<OP::ImageInvertArgs>(Item.At);
+					PushOp(FScheduledOpInline(Args.Base, Item));
+
+					StoreNone(Item);
+					break;
+				}
+
+			case EOpType::IM_PLAINCOLOUR:
+				{
+					StoreNone(Item);
+					break;
+				}
+	
+			case EOpType::IM_CROP:
+				{
+					OP::ImageCropArgs Args = Program.GetOpArgs<OP::ImageCropArgs>(Item.At);
+					PushOp(FScheduledOpInline(Args.source, Item));
+
+					StoreNone(Item);
+					break;
+				}
+	
+			case EOpType::IM_PATCH:
+				{
+					OP::ImagePatchArgs Args = Program.GetOpArgs<OP::ImagePatchArgs>(Item.At);
+					PushOp(FScheduledOpInline(Args.base, Item));
+					PushOp(FScheduledOpInline(Args.patch, Item));
+
+					StoreNone(Item);
+					break;
+				}
+	
+			case EOpType::IM_RASTERMESH:
+				{
+					OP::ImageRasterMeshArgs Args = Program.GetOpArgs<OP::ImageRasterMeshArgs>(Item.At);
+					PushOp(FScheduledOpInline(Args.mesh, Item));
+
+					StoreNone(Item);
+					break;
+				}
+
+			case EOpType::IM_MAKEGROWMAP:
+				{
+					OP::ImageMakeGrowMapArgs Args = Program.GetOpArgs<OP::ImageMakeGrowMapArgs>(Item.At);
+					PushOp(FScheduledOpInline(Args.mask, Item));
+
+					StoreNone(Item);
+					break;
+				}
+	
+			case EOpType::IM_DISPLACE:
+				{
+					OP::ImageDisplaceArgs Args = Program.GetOpArgs<OP::ImageDisplaceArgs>(Item.At);
+					PushOp(FScheduledOpInline(Args.Source, Item));
+
+					StoreNone(Item);
+					break;
+				}
+	
+			case EOpType::IM_TRANSFORM:
+				{
+					OP::ImageTransformArgs Args = Program.GetOpArgs<OP::ImageTransformArgs>(Item.At);
+					PushOp(FScheduledOpInline(Args.Base, Item));
+
+					StoreNone(Item);
+					break;
+				}
+	
+			case EOpType::LA_MERGE:
+				{
+					OP::LayoutMergeArgs Args = Program.GetOpArgs<OP::LayoutMergeArgs>(Item.At);
+					PushOp(FScheduledOpInline(Args.Base, Item));
+					PushOp(FScheduledOpInline(Args.Added, Item));
+
+					StoreNone(Item);
+					break;
+				}
+	
+			case EOpType::LA_PACK:
+				{
+					OP::LayoutPackArgs Args = Program.GetOpArgs<OP::LayoutPackArgs>(Item.At);
+					PushOp(FScheduledOpInline( Args.Source, Item));
+
+					StoreNone(Item);
+					break;
+				}
+
+			case EOpType::LA_FROMMESH:
+				{
+					OP::LayoutFromMeshArgs Args = Program.GetOpArgs<OP::LayoutFromMeshArgs>(Item.At);
+					PushOp(FScheduledOpInline(Args.Mesh, Item));
+
+					StoreNone(Item);
+					break;
+				}
+
+			case EOpType::LA_REMOVEBLOCKS:
+				{
+					OP::LayoutRemoveBlocksArgs Args = Program.GetOpArgs<OP::LayoutRemoveBlocksArgs>(Item.At);
+					PushOp(FScheduledOpInline(Args.Source, Item));
+					PushOp(FScheduledOpInline(Args.ReferenceLayout, Item));
+
+					StoreNone(Item);
+					break;
+				}
+			
+			case EOpType::NONE:
+				{
+					check(Item.At == 0);
+					
+					StoreNone(Item);
+					break;
+				}
+
+			case EOpType::MI_CONSTANT:
+				{
+					OP::ResourceConstantArgs Args = Program.GetOpArgs<OP::ResourceConstantArgs>(Item.At);
+					MI_REFERENCE_ID.Push(Args.value);
+
+					StoreNone(Item);
+					break;
+				}
+				
+			default:
+				unimplemented();
+			}
+		}
+
+		check(Stack.Num() <= 1);
+
+		return Stack.Num() ? LoadInstance(Root) : MakeShared<FInstance>();
+	}
+}
diff --git a/Engine/Plugins/Mutable/Source/MutableRuntime/Private/MuR/CodeRunnerBegin.h b/Engine/Plugins/Mutable/Source/MutableRuntime/Private/MuR/CodeRunnerBegin.h
new file mode 100644
--- /dev/null
+++ b/Engine/Plugins/Mutable/Source/MutableRuntime/Private/MuR/CodeRunnerBegin.h
@@ -0,0 +1,164 @@
+// Copyright Epic Games, Inc. All Rights Reserved.
+
+#pragma once
+
+#include "MuR/SystemPrivate.h"
+
+
+namespace mu
+{
+	class FProgramCache;
+	class FInstance;
+	class FExtensionData;
+	struct FProgram;
+	class FModel;
+	class FParameters;
+	class FSystem;
+	class String;
+
+	struct FScheduledOpInline
+	{
+		explicit FScheduledOpInline(OP::ADDRESS InAt = 0, const uint16 InStage = 0, const uint16 InExecutionIndex = 0) :
+			At(InAt),
+			Stage(InStage),
+		 	ExecutionIndex(InExecutionIndex)
+		{
+		}
+		
+		explicit FScheduledOpInline(OP::ADDRESS InAt, const FScheduledOpInline& InOpTemplate) :
+			At(InAt),
+     		ExecutionIndex(InOpTemplate.ExecutionIndex)
+		{
+		}
+		
+		FScheduledOpInline(const FScheduledOpInline& InOpTemplate, uint16 InStage) :
+			At(InOpTemplate.At),
+         	Stage(InStage),
+            ExecutionIndex(InOpTemplate.ExecutionIndex)
+		{
+		}
+
+		bool operator==(const FScheduledOpInline& Other) const;
+
+		OP::ADDRESS At = 0;
+		uint16 Stage = 0;
+		uint16 ExecutionIndex = 0;
+		OP::ADDRESS CustomState = 0;
+	};
+
+
+	uint32 GetTypeHash(const FScheduledOpInline& Op);
+
+
+	struct FCacheAddressInline
+	{
+		FCacheAddressInline(const OP::ADDRESS& InAt)
+		{
+			At = InAt;
+		}
+		
+		FCacheAddressInline(const FScheduledOpInline& Item)
+		{
+			At = Item.At;
+			ExecutionIndex = Item.ExecutionIndex;
+		}
+		
+		OP::ADDRESS At = 0;
+		uint16 ExecutionIndex = 0;
+
+		bool operator==(const FCacheAddressInline&) const = default;
+	};
+
+
+	uint32 GetTypeHash(const FCacheAddressInline& Address);
+
+
+	class FOpSet
+	{
+	public:
+		FOpSet(int32 NumElements);
+
+		bool Contains(const FCacheAddressInline& Item);
+
+		void Add(const FCacheAddressInline& Item);
+
+	private:
+		TBitArray<> Index0;
+		TSet<FCacheAddressInline> IndexOther;
+	};
+	
+	using FStackValue = TVariant<bool, int32, float, TSharedPtr<const String>, TSharedPtr<const FExtensionData>, TSharedPtr<const FInstance>>;
+
+	class FStack : public TArray<FStackValue, TInlineAllocator<32>>
+	{
+	public:
+		FStackValue Pop();
+	};
+	
+	class CodeRunnerBegin
+	{
+	public:
+		CodeRunnerBegin(FSystem& InSystem, const TSharedPtr<const FModel>& InModel, const TSharedPtr<const FParameters>& InParams, uint32 InLodMask);
+
+		TSharedPtr<const FInstance> RunCode(const FScheduledOpInline& Root);
+
+		FProgramCache& GetMemory() const;
+		
+		FScheduledOpInline PopOp();
+		
+		void PushOp(const FScheduledOpInline& Item);
+		
+		void StoreInt(const FCacheAddressInline& To, int32 Value);
+		
+		void StoreFloat(const FCacheAddressInline& To, float Value);
+		
+		void StoreBool(const FCacheAddressInline& To, bool Value);
+		
+		void StoreInstance(const FCacheAddressInline& To, const TSharedPtr<const FInstance>& Value);
+		
+		void StoreString(const FCacheAddressInline& To, const TSharedPtr<const String>& Value);
+		
+		void StoreExtensionData(const FCacheAddressInline& To, const TSharedPtr<const FExtensionData>& Value);
+
+		void StoreNone(const FCacheAddressInline& To);
+
+		int32 LoadInt(const FCacheAddressInline& To);
+
+		float LoadFloat(const FCacheAddressInline& To);
+
+		bool LoadBool(const FCacheAddressInline& To);
+		
+		TSharedPtr<const FInstance> LoadInstance(const FCacheAddressInline& To);
+
+		TSharedPtr<const String> LoadString(const FCacheAddressInline& To);
+
+		TSharedPtr<const FExtensionData> LoadExtensionData(const FCacheAddressInline& To);
+		
+		TSharedPtr<const FModel> Model;
+		TSharedPtr<const FParameters> Params;
+
+		FSystem& System;
+		FProgram& Program;
+		
+		uint32 LodMask;
+
+		TArray<FScheduledOpInline, TInlineAllocator<256, TAlignedHeapAllocator<PLATFORM_CACHE_LINE_SIZE>>> Items;
+
+		FStack Stack;
+		
+		FOpSet Executed;
+
+		TMap<FCacheAddressInline, float> ResultsFloat;
+		TMap<FCacheAddressInline, bool> ResultsBool;
+		TMap<FCacheAddressInline, int32> ResultsInt;
+		TMap<FCacheAddressInline, TSharedPtr<const FInstance>> ResultsInstance;
+		TMap<FCacheAddressInline, TSharedPtr<const String>> ResultsString;
+		TMap<FCacheAddressInline, TSharedPtr<const FExtensionData>> ResultsExtensionData;
+		
+		TArray<int32> MI_REFERENCE_ID;
+		TArray<int32> IM_REFERENCE_ID;
+		TArray<int32> ME_REFERENCE_ID;
+    	
+		TArray<int32> ME_SKELETON_ID;
+	};
+}
diff --git a/Engine/Plugins/Mutable/Source/MutableRuntime/Private/MuR/CodeRunnerTasked.cpp b/Engine/Plugins/Mutable/Source/MutableRuntime/Private/MuR/CodeRunnerTasked.cpp
--- a/Engine/Plugins/Mutable/Source/MutableRuntime/Private/MuR/CodeRunnerTasked.cpp
+++ b/Engine/Plugins/Mutable/Source/MutableRuntime/Private/MuR/CodeRunnerTasked.cpp
@@ -271,7 +271,7 @@ namespace mu
 				// Special processing in case it is an ImageDesc operation
 				if (item.Type == FScheduledOp::EType::ImageDesc)
 				{
-					RunCodeImageDesc(item, Params, Model.Get(), LODMask);
+					RunCodeImageDesc(item, Params);
 					continue;
 				}
 
@@ -2589,7 +2589,7 @@ namespace mu
 		case EOpType::IM_PARAMETER:
 		{
 			OP::ParameterArgs args = Program.GetOpArgs<OP::ParameterArgs>(Item.At);
-			TSharedPtr<FRangeIndex> Index = BuildCurrentOpRangeIndex(Item, Params, Model.Get(), args.variable);
+			TSharedPtr<FRangeIndex> Index = BuildCurrentOpRangeIndex(Item, args.variable);
 
 			UTexture* Image = Params->GetImageValue(args.variable, Index.Get());
 
@@ -2627,7 +2627,7 @@ namespace mu
 		case EOpType::ME_PARAMETER:
 		{
 			OP::MeshParameterArgs Args = Program.GetOpArgs<OP::MeshParameterArgs>(Item.At);
-			TSharedPtr<FRangeIndex> Index = BuildCurrentOpRangeIndex(Item, Params, Model.Get(), Args.variable);
+			TSharedPtr<FRangeIndex> Index = BuildCurrentOpRangeIndex(Item, Args.variable);
 
 			USkeletalMesh* Id = Params->GetMeshValue(Args.variable, Index.Get());
 
diff --git a/Engine/Plugins/Mutable/Source/MutableRuntime/Private/MuR/Instance.cpp b/Engine/Plugins/Mutable/Source/MutableRuntime/Private/MuR/Instance.cpp
--- a/Engine/Plugins/Mutable/Source/MutableRuntime/Private/MuR/Instance.cpp
+++ b/Engine/Plugins/Mutable/Source/MutableRuntime/Private/MuR/Instance.cpp
@@ -319,19 +319,19 @@ namespace mu
         return NAME_None;
     }
 
-	const FMaterial& FInstance::GetOverlayMaterial(int32 ComponentIndex) const
+	FMaterialId FInstance::GetOverlayMaterialId(int32 ComponentIndex) const
 	{
 		check(Components.IsValidIndex(ComponentIndex));
-		return Components[ComponentIndex].OverlayMaterial;
+		return Components[ComponentIndex].OverlayMaterialId;
 	}
 
-	const FMaterial& FInstance::GetMaterial(int32 ComponentIndex, int32 LODIndex, int32 SurfaceIndex) const
+	FMaterialId FInstance::GetMaterialId(int32 ComponentIndex, int32 LODIndex, int32 SurfaceIndex) const
 	{
 		check(Components.IsValidIndex(ComponentIndex));
 		check(Components[ComponentIndex].LODs.IsValidIndex(LODIndex));
 		check(Components[ComponentIndex].LODs[LODIndex].Surfaces.IsValidIndex(SurfaceIndex));
 
-		return Components[ComponentIndex].LODs[LODIndex].Surfaces[SurfaceIndex].Material;
+		return Components[ComponentIndex].LODs[LODIndex].Surfaces[SurfaceIndex].MaterialId;
 	}
 
     
@@ -406,8 +406,8 @@ namespace mu
     }
 
 
-	
-	void FInstance::SetMesh(int32 ComponentIndex, int32 LODIndex, FMeshId meshId, FName Name)
+	//---------------------------------------------------------------------------------------------
+	void FInstance::SetMesh(int32 ComponentIndex, int32 LODIndex, const FMeshId& MeshId, FName Name)
 	{
 		// Automatically create the necessary lods and components
 		while (ComponentIndex >= Components.Num())
@@ -420,12 +420,12 @@ namespace mu
 		}
 
 		FInstanceLOD& LOD = Components[ComponentIndex].LODs[LODIndex];
-		LOD.MeshId = meshId;
+		LOD.MeshId = MeshId;
 		LOD.MeshName = Name;
 	}
 
 
-    int32 FInstance::AddImage( int32 ComponentIndex, int32 LODIndex, int32 SurfaceIndex, FImageId imageId, FName Name)
+    int32 FInstance::AddImage( int32 ComponentIndex, int32 LODIndex, int32 SurfaceIndex, const FImageId& ImageId, FName Name)
 	{
 		// Automatically create the necessary lods and components
 		while (ComponentIndex >= Components.Num())
@@ -442,7 +442,7 @@ namespace mu
 		}
 
 		FInstanceSurface& Surface = Components[ComponentIndex].LODs[LODIndex].Surfaces[SurfaceIndex];
-		return Surface.Images.Add({ imageId, Name });
+		return Surface.Images.Add({ ImageId, Name });
 	}
 	
 	
@@ -509,7 +509,7 @@ namespace mu
     }
 	
 
-	void FInstance::AddMaterial(int32 ComponentIndex, int32 LODIndex, int32 SurfaceIndex, FMaterial Material)
+	void FInstance::SetMaterialId(int32 ComponentIndex, int32 LODIndex, int32 SurfaceIndex, const FMaterialId& MaterialId)
 	{
 		// Automatically create the necessary lods and components
 		while (ComponentIndex >= Components.Num())
@@ -525,25 +525,23 @@ namespace mu
 			AddSurface(LODIndex, ComponentIndex);
 		}
 		
-		Components[ComponentIndex].LODs[LODIndex].Surfaces[SurfaceIndex].Material = Material;
+		Components[ComponentIndex].LODs[LODIndex].Surfaces[SurfaceIndex].MaterialId = MaterialId;
 	}
 
 
-	void FInstance::AddOverlayMaterial(int32 ComponentIndex, FMaterial OverlayMaterial)
+	void FInstance::SetOverlayMaterialId(int32 ComponentIndex, const FMaterialId& MaterialId)
 	{
 		while (ComponentIndex >= Components.Num())
 		{
 			AddComponent();
 		}
 
-		Components[ComponentIndex].OverlayMaterial = OverlayMaterial;
+		Components[ComponentIndex].OverlayMaterialId = MaterialId;
 	}
+
 	
-    
-	void FInstance::AddExtensionData(const TSharedPtr<const FExtensionData>& Data, FName Name)
+	void FInstance::SetExtensionData(const TSharedRef<const FExtensionData>& Data, FName Name)
 	{
-		check(Data);
-
 		NamedExtensionData& Entry = ExtensionData.AddDefaulted_GetRef();
 		Entry.Data = Data;
 		Entry.Name = Name;
diff --git a/Engine/Plugins/Mutable/Source/MutableRuntime/Private/MuR/ResourceID.cpp b/Engine/Plugins/Mutable/Source/MutableRuntime/Private/MuR/ResourceID.cpp
--- a/Engine/Plugins/Mutable/Source/MutableRuntime/Private/MuR/ResourceID.cpp
+++ b/Engine/Plugins/Mutable/Source/MutableRuntime/Private/MuR/ResourceID.cpp
@@ -37,3 +37,13 @@ uint32 mu::GetTypeHash(const FGeneratedImageKey& Key, const bool bGenerateDeterm
 
 	return Hash;
 }
+
+
+uint32 mu::GetTypeHash(const FGeneratedMaterialKey& Key)
+{
+	uint32 Hash = Key.Address;
+
+	Hash = HashCombineFast(Hash, static_cast<uint32>(FXxHash64::HashBuffer(Key.ParameterValuesBlob.GetData(), Key.ParameterValuesBlob.Num()).Hash));
+
+	return Hash;
+}
diff --git a/Engine/Plugins/Mutable/Source/MutableRuntime/Private/MuR/System.cpp b/Engine/Plugins/Mutable/Source/MutableRuntime/Private/MuR/System.cpp
--- a/Engine/Plugins/Mutable/Source/MutableRuntime/Private/MuR/System.cpp
+++ b/Engine/Plugins/Mutable/Source/MutableRuntime/Private/MuR/System.cpp
@@ -2,6 +2,7 @@
 
 #include "MuR/System.h"
 
+#include "CodeRunnerBegin.h"
 #include "Containers/Array.h"
 #include "Containers/Map.h"
 #include "Containers/Set.h"
@@ -33,6 +34,7 @@
 #include "Trace/Detail/Channel.h"
 #include "ProfilingDebugging/CountersTrace.h"
 #include "PackedNormal.h"
+#include "MuR/SystemPrivate.h"
 
 #include UE_INLINE_GENERATED_CPP_BY_NAME(System)
 
@@ -225,15 +227,61 @@ namespace mu
 
 
     //---------------------------------------------------------------------------------------------
-    TSharedPtr<const FInstance> FSystem::BeginUpdate( FInstance::FID InInstanceID,
+    TSharedPtr<const FInstance> FSystem::BeginUpdate_GameThread(FInstance::FID InInstanceID,
                                      const TSharedPtr<const FParameters>& InParams,
                                      const TSharedPtr<FMeshIdRegistry>& InMeshIdRegistry,
 									 const TSharedPtr<FImageIdRegistry>& InImageIdRegistry,
+									 const TSharedPtr<FMaterialIdRegistry>& InMaterialIdRegistry,
                                      int32 InStateIndex,
-                                     uint32 InLodMask )
+                                     uint32 InLodMask)
     {
 		LLM_SCOPE_BYNAME(TEXT("MutableRuntime"));
-		MUTABLE_CPUPROFILER_SCOPE(SystemBeginUpdate);
+		MUTABLE_CPUPROFILER_SCOPE(BeginUpdate_GameThread);
+
+    	FLiveInstance* LiveInstance = m_pD->FindLiveInstance(InInstanceID);
+    	if (!LiveInstance)
+    	{
+    		UE_LOG(LogMutableCore, Error, TEXT("Invalid instance id in mutable update."));
+    		return nullptr;
+    	}
+    	
+		if (!InParams)
+		{
+			UE_LOG(LogMutableCore, Error, TEXT("Invalid parameters in mutable update."));
+			return nullptr;
+		}
+
+    	m_pD->WorkingMemoryManager.CurrentInstanceCache = LiveInstance->Cache;
+
+		m_pD->WorkingMemoryManager.CurrentMeshIdRegistry = InMeshIdRegistry;
+    	m_pD->WorkingMemoryManager.CurrentImageIdRegistry = InImageIdRegistry;
+    	m_pD->WorkingMemoryManager.CurrentMaterialIdRegistry = InMaterialIdRegistry;
+    	
+		OP::ADDRESS RootAt = LiveInstance->Model->GetPrivate()->Program.States[InStateIndex].Root;
+
+    	CodeRunnerBegin Runner(*this, LiveInstance->Model, InParams, InLodMask);
+    	
+    	TSharedPtr<const FInstance> Result = Runner.RunCode(FScheduledOpInline(RootAt, 0));
+    	if (!Result)
+		{
+			// In case of failure return an empty instance, to prevent following code to have to check it every time
+			Result = MakeShared<FInstance>();
+		}
+    	
+		return Result;
+	}
+
+	
+    TSharedPtr<const FInstance> FSystem::BeginUpdate_MutableThread(FInstance::FID InInstanceID,
+		const TSharedPtr<const FParameters>& InParams,
+		const TSharedPtr<FMeshIdRegistry>& InMeshIdRegistry,
+		const TSharedPtr<FImageIdRegistry>& InImageIdRegistry,
+		const TSharedPtr<FMaterialIdRegistry>& InMaterialIdRegistry,
+		int32 InStateIndex,
+		uint32 InLodMask)
+    {
+		LLM_SCOPE_BYNAME(TEXT("MutableRuntime"));
+		MUTABLE_CPUPROFILER_SCOPE(BeginUpdate_MutableThread);
 		TRACE_COUNTER_INCREMENT(MutableRuntime_Updates);
 
 		if (!InParams)
@@ -253,6 +301,7 @@ namespace mu
 
     	m_pD->WorkingMemoryManager.CurrentMeshIdRegistry = InMeshIdRegistry;
     	m_pD->WorkingMemoryManager.CurrentImageIdRegistry = InImageIdRegistry;
+    	m_pD->WorkingMemoryManager.CurrentMaterialIdRegistry = InMaterialIdRegistry;
 
 		FProgram& Program = LiveInstance->Model->GetPrivate()->Program;
 
@@ -364,8 +413,6 @@ namespace mu
 			return UE::Tasks::MakeCompletedTask<TSharedPtr<const FImage>>(BlackImage());
     	}
     	
-		TSharedPtr<const FImage> pResult;
-
 		// Find the live instance
 		FLiveInstance* LiveInstance = m_pD->FindLiveInstance(instanceID);
 		check(LiveInstance);
@@ -574,8 +621,9 @@ namespace mu
 		m_pD->WorkingMemoryManager.CurrentInstanceCache = nullptr;
 		return Result;
 	}
+	
 
-	UE::Tasks::TTask<TSharedPtr<const FMesh>> FSystem::GetMesh(FInstance::FID instanceID, const FMeshId& MeshId, EMeshContentFlags MeshContentFilter)
+    UE::Tasks::TTask<TSharedPtr<const FMesh>> FSystem::GetMesh(FInstance::FID instanceID, const FMeshId& MeshId, EMeshContentFlags MeshContentFilter)
     {
 		LLM_SCOPE_BYNAME(TEXT("MutableRuntime"));
 		MUTABLE_CPUPROFILER_SCOPE(SystemGetMesh);
@@ -584,9 +632,7 @@ namespace mu
 		{
 			return UE::Tasks::MakeCompletedTask<TSharedPtr<const FMesh>>(new FMesh());
 		}
-
-		TSharedPtr<const FMesh> ResultMesh;
-
+		
 		// Find the live instance
 		FLiveInstance* LiveInstance = m_pD->FindLiveInstance(instanceID);
 		check(LiveInstance);
@@ -641,8 +687,71 @@ namespace mu
 				UE::Tasks::EExtendedTaskPriority::Inline);
 	}
 
+	
+	UE::Tasks::TTask<TSharedPtr<const FMaterial>> FSystem::GetMaterial(FInstance::FID InstanceID, const FMaterialId& MaterialId)
+	{
+		LLM_SCOPE_BYNAME(TEXT("MutableRuntime"));
+		MUTABLE_CPUPROFILER_SCOPE(SystemGetMaterial);
+		
+		if (!MaterialId)
+		{
+			return UE::Tasks::MakeCompletedTask<TSharedPtr<const FMaterial>>(new FMaterial());
+		}
+		
+		// Find the live instance
+		FLiveInstance* LiveInstance = m_pD->FindLiveInstance(InstanceID);
+		check(LiveInstance);
+		m_pD->WorkingMemoryManager.CurrentInstanceCache = LiveInstance->Cache;
+		m_pD->WorkingMemoryManager.BeginRunnerThread();
+		
+		OP::ADDRESS RootAddress = MaterialId.GetKey()->Address;
+			
+		mu::EOpType OpType = LiveInstance->Model->GetPrivate()->Program.GetOpType(RootAddress);
+		if (GetOpDataType(OpType) != EDataType::Material)
+		{
+			m_pD->WorkingMemoryManager.EndRunnerThread();
+			m_pD->WorkingMemoryManager.CurrentInstanceCache = nullptr;
+
+			return UE::Tasks::MakeCompletedTask<TSharedPtr<const FMaterial>>(new FMaterial());
+		}
+	
+		TSharedRef<CodeRunner> Runner = CodeRunner::Create(
+				m_pD->Settings,
+				m_pD, 
+				EExecutionStrategy::MinimizeMemory, 
+				LiveInstance->Model, 
+				LiveInstance->OldParameters.Get(), 
+				RootAddress, 
+				FSystem::AllLODs, 
+				0, 
+				0, 
+				FScheduledOp::EType::Full);
+		
+		constexpr bool bForceInlineExecution = false;
+		UE::Tasks::FTask RunnerCompletionEvent = Runner->StartRun(bForceInlineExecution);
+
+		return UE::Tasks::Launch(TEXT("FSystem::GetMaterialResultTask"),
+				[SystemPrivate = m_pD, Runner, RootAddress]() -> TSharedPtr<const FMaterial>
+				{
+					TSharedPtr<const FMaterial> Result;
+
+					SystemPrivate->bUnrecoverableError = Runner->bUnrecoverableError;
+					if (!Runner->bUnrecoverableError)
+					{
+						Result = SystemPrivate->WorkingMemoryManager.LoadMaterial(FCacheAddress(RootAddress, 0, 0));
+					}
+					
+					SystemPrivate->WorkingMemoryManager.EndRunnerThread();
+					SystemPrivate->WorkingMemoryManager.CurrentInstanceCache = nullptr;
+					
+					return Result;
+				},
+				UE::Tasks::Prerequisites(RunnerCompletionEvent),
+				UE::Tasks::ETaskPriority::Inherit,
+				UE::Tasks::EExtendedTaskPriority::Inline);
+	}
+	
 
-    //---------------------------------------------------------------------------------------------
 	void FSystem::EndUpdate(FInstance::FID instanceID)
 	{
 		LLM_SCOPE_BYNAME(TEXT("MutableRuntime"));
@@ -871,7 +980,8 @@ namespace mu
 	//---------------------------------------------------------------------------------------------
 	void FSystem::Private::BeginBuild(const TSharedPtr<const FModel>& InModel,
 		const TSharedPtr<FMeshIdRegistry>& InMeshIdRegistry,
-		const TSharedPtr<FImageIdRegistry>& InImageIdRegistry)
+		const TSharedPtr<FImageIdRegistry>& InImageIdRegistry,
+		const TSharedPtr<FMaterialIdRegistry>& InMaterialIdRegistry)
 	{
 		// We don't have a FLiveInstance, let's create the memory
 		// \TODO: There is no clear moment to remove this... EndBuild?
@@ -880,6 +990,7 @@ namespace mu
 
     	WorkingMemoryManager.CurrentMeshIdRegistry = InMeshIdRegistry;
     	WorkingMemoryManager.CurrentImageIdRegistry = InImageIdRegistry;
+    	WorkingMemoryManager.CurrentMaterialIdRegistry = InMaterialIdRegistry;
     	
 		// Ensure the model cache has been created
 		WorkingMemoryManager.FindOrAddModelCache(InModel);
@@ -2021,7 +2132,17 @@ namespace mu
 
     	GetResourceKey(Model, Params, ParamListIndex, RootAt, NewKey.ParameterValuesBlob);
 
-    	return CurrentImageIdRegistry->Add(NewKey, { });
+    	return CurrentImageIdRegistry->Add(NewKey, {});
     }
 
+	
+	FMaterialId FWorkingMemoryManager::GetMaterialId(const TSharedPtr<const FModel>& Model, const FParameters* Params, uint32 ParamListIndex, OP::ADDRESS RootAt)
+    {
+    	FGeneratedMaterialKey NewKey;
+    	NewKey.Address = RootAt;
+
+    	GetResourceKey(Model, Params, ParamListIndex, RootAt, NewKey.ParameterValuesBlob);
+
+    	return CurrentMaterialIdRegistry->Add(NewKey, {});
+    }
 }
diff --git a/Engine/Plugins/Mutable/Source/MutableRuntime/Public/MuR/Instance.h b/Engine/Plugins/Mutable/Source/MutableRuntime/Public/MuR/Instance.h
--- a/Engine/Plugins/Mutable/Source/MutableRuntime/Public/MuR/Instance.h
+++ b/Engine/Plugins/Mutable/Source/MutableRuntime/Public/MuR/Instance.h
@@ -18,11 +18,11 @@ namespace mu
 	struct FInstanceSurface
 	{
 		FName Name;
-		uint32 InternalId=0;
-		uint32 ExternalId =0;
-		uint32 SharedId =0;
+		uint32 InternalId = 0;
+		uint32 ExternalId = 0;
+		uint32 SharedId = 0;
 
-		FMaterial Material;
+		FMaterialId MaterialId;
 
 		struct FInstanceImage
 		{
@@ -30,7 +30,7 @@ namespace mu
 			FName Name;
 		};
 
-		TArray<FInstanceImage, TInlineAllocator<4>> Images;
+		TArray<FInstanceImage> Images;
 
 		struct FInstanceVector
 		{
@@ -64,7 +64,7 @@ namespace mu
 		FName MeshName;
 
 		// The order must match the meshes surfaces
-		TArray<FInstanceSurface, TInlineAllocator<4>> Surfaces;
+		TArray<FInstanceSurface, TInlineAllocator<1>> Surfaces;
 	};
 
 
@@ -72,9 +72,9 @@ namespace mu
 	{
 		uint16 Id;
 
-		FMaterial OverlayMaterial;
+		FMaterialId OverlayMaterialId;
 
-		TArray<FInstanceLOD, TInlineAllocator<4>> LODs;
+		TArray<FInstanceLOD, TInlineAllocator<1>> LODs;
 	};
 
 	
@@ -248,9 +248,9 @@ namespace mu
 
 		//! Get the index of the overlay material.
 		//! \param ComponentIndex Index of the component, from 0 to GetComponentCount()-1
-		UE_API const FMaterial& GetOverlayMaterial( int32 ComponentIndex ) const;
+		UE_API FMaterialId GetOverlayMaterialId(int32 ComponentIndex) const;
 
-		UE_API const FMaterial& GetMaterial(int32 ComponentIndex, int32 LODIndex, int32 SurfaceIndex) const;
+		UE_API FMaterialId GetMaterialId(int32 ComponentIndex, int32 LODIndex, int32 SurfaceIndex) const;
 
 		//! Get the number of ExtensionData values in a component
 		UE_API int32 GetExtensionDataCount() const;
@@ -265,13 +265,13 @@ namespace mu
 		
 		int32 AddLOD(int32 ComponentIndex);
 
-		void SetMesh(int32 ComponentIndex, int32 LODIndex, FMeshId, FName Name);
+		void SetMesh(int32 ComponentIndex, int32 LODIndex, const FMeshId& MeshId, FName Name);
 
 		int32 AddSurface(int32 ComponentIndex, int32 LODIndex);
 
 		void SetSurfaceName(int32 ComponentIndex, int32 LODIndex, int32 SurfaceIndex, FName Name);
 
-		int32 AddImage(int32 ComponentIndex, int32 LODIndex, int32 SurfaceIndex, FImageId, FName Name);
+		int32 AddImage(int32 ComponentIndex, int32 LODIndex, int32 SurfaceIndex, const FImageId& Imag, FName Name);
 
 		int32 AddVector(int32 ComponentIndex, int32 LODIndex, int32 SurfaceIndex, const FVector4f&, FName Name);
 
@@ -279,14 +279,13 @@ namespace mu
 
 		int32 AddString(int32 ComponentIndex, int32 LODIndex, int32 SurfaceIndex, const FString& Value, FName Name);
 
-		void AddMaterial(int32 ComponentIndex, int32 LODIndex, int32 SurfaceIndex, FMaterial OverlayMaterial);
+		void SetMaterialId(int32 ComponentIndex, int32 LODIndex, int32 SurfaceIndex, const FMaterialId& MaterialId);
 
-		void AddOverlayMaterial(int32 ComponentIndex, FMaterial OverlayMaterial);
+		void SetOverlayMaterialId(int32 ComponentIndex, const FMaterialId& MaterialId);
 
-		// Data must be non-null
-		void AddExtensionData(const TSharedPtr<const FExtensionData>& Data, FName Name);
+		void SetExtensionData(const TSharedRef<const FExtensionData>& Data, FName Name);
 		
-		TArray<FInstanceComponent,TInlineAllocator<4>> Components;
+		TArray<FInstanceComponent, TInlineAllocator<1>> Components;
 
 		// Every entry must have a valid ExtensionData and name
 		TArray<NamedExtensionData> ExtensionData;
diff --git a/Engine/Plugins/Mutable/Source/MutableRuntime/Public/MuR/System.h b/Engine/Plugins/Mutable/Source/MutableRuntime/Public/MuR/System.h
--- a/Engine/Plugins/Mutable/Source/MutableRuntime/Public/MuR/System.h
+++ b/Engine/Plugins/Mutable/Source/MutableRuntime/Public/MuR/System.h
@@ -150,13 +150,25 @@ namespace mu
         //! \return the instance data with all the LOD, components, and ids to generate the meshes 
         //! and images. The returned Instance is only valid until the next call to EndUpdate with 
         //! the same instanceID parameter.
-		UE_API TSharedPtr<const FInstance> BeginUpdate(FInstance::FID InstanceID,
+		//!
+		//! Very fast. Can be called in the Game Thread.
+		UE_API TSharedPtr<const FInstance> BeginUpdate_GameThread(FInstance::FID InstanceID,
 			const TSharedPtr<const FParameters>& Params,
 			const TSharedPtr<FMeshIdRegistry>& InMeshIdRegistry,
 			const TSharedPtr<FImageIdRegistry>& InImageIdRegistry,
+			const TSharedPtr<FMaterialIdRegistry>& InMaterialIdRegistry,
 			int32 StateIndex,
 			uint32 LodMask);
 
+		/** Potentially expensive. Must be called in the Mutable Thread. */
+		UE_API TSharedPtr<const FInstance> BeginUpdate_MutableThread(FInstance::FID InInstanceID,
+			const TSharedPtr<const FParameters>& InParams,
+			const TSharedPtr<FMeshIdRegistry>& InMeshIdRegistry,
+			const TSharedPtr<FImageIdRegistry>& InImageIdRegistry,
+			const TSharedPtr<FMaterialIdRegistry>& InMaterialIdRegistry,
+			int32 InStateIndex,
+			uint32 InLodMask);
+		
 		//! Only valid between BeginUpdate and EndUpdate
 		//! Calculate the description of an image, without generating it.
 		UE_API UE::Tasks::TTask<FExtendedImageDesc> GetImageDesc(FInstance::FID InstanceID, const FImageId& ImageId);
@@ -182,7 +194,9 @@ namespace mu
 
         //! Only valid between BeginUpdate and EndUpdate
 		UE_API TSharedPtr<const FMesh> GetMeshInline(FInstance::FID InstanceID, const FImageId& MeshId, EMeshContentFlags MeshContentFilter = EMeshContentFlags::AllFlags);
-
+		
+		UE_API UE::Tasks::TTask<TSharedPtr<const FMaterial>> GetMaterial(FInstance::FID InstanceID, const FMaterialId& MaterialId);
+		
         //! Invalidate and free the last Instance data returned by a call to BeginUpdate with
         //! the same instance index. After a call to this method, that Instance cannot be used any
         //! more and its content is undefined.
diff --git a/Engine/Plugins/Mutable/Source/MutableTools/Private/MuT/ASTOpAddMaterial.cpp b/Engine/Plugins/Mutable/Source/MutableTools/Private/MuT/ASTOpAddMaterial.cpp
new file mode 100644
--- /dev/null
+++ b/Engine/Plugins/Mutable/Source/MutableTools/Private/MuT/ASTOpAddMaterial.cpp
@@ -0,0 +1,132 @@
+// Copyright Epic Games, Inc. All Rights Reserved.
+
+#include "MuT/ASTOpAddMaterial.h"
+
+#include "HAL/PlatformMath.h"
+#include "MuR/ModelPrivate.h"
+#include "MuR/RefCounted.h"
+#include "MuT/CodeOptimiser.h"
+
+namespace mu
+{
+
+
+//-------------------------------------------------------------------------------------------------
+//-------------------------------------------------------------------------------------------------
+//-------------------------------------------------------------------------------------------------
+	ASTOpAddMaterial::ASTOpAddMaterial()
+	: Instance(this)
+	, Material(this)
+{
+}
+
+
+	ASTOpAddMaterial::~ASTOpAddMaterial()
+{
+	// Explicit call needed to avoid recursive destruction
+	ASTOp::RemoveChildren();
+}
+
+bool ASTOpAddMaterial::IsEqual(const ASTOp& OtherUntyped) const
+{
+	if (OtherUntyped.GetOpType() == GetOpType())
+	{
+		const ASTOpAddMaterial* Other = static_cast<const ASTOpAddMaterial*>(&OtherUntyped);
+		return Instance == Other->Instance
+			&& Material == Other->Material;
+	}
+
+	return false;
+}
+
+uint64 ASTOpAddMaterial::Hash() const
+{
+	uint64 Result = std::hash<uint64>()(uint64(Type));
+	
+	hash_combine(Result, Instance.child().get());
+	hash_combine(Result, Material.child().get());
+	
+	return Result;
+}
+
+mu::Ptr<ASTOp> ASTOpAddMaterial::Clone(MapChildFuncRef MapChild) const
+{
+	Ptr<ASTOpAddMaterial> NewInstance = new ASTOpAddMaterial();
+
+	NewInstance->Type = Type;
+	NewInstance->Instance = MapChild(Instance.child());
+	NewInstance->Material = MapChild(Material.child());
+
+	return NewInstance;
+}
+
+void ASTOpAddMaterial::ForEachChild(const TFunctionRef<void(ASTChild&)> F)
+{
+	F(Instance);
+	F(Material);
+}
+
+void ASTOpAddMaterial::Link(FProgram& Program, FLinkerOptions*)
+{
+	check(Type != EOpType::NONE);
+
+	// Already linked?
+	if (linkedAddress)
+	{
+		return;
+	}
+
+	OP::InstanceAddMaterialArgs Args;
+	FMemory::Memzero(Args);
+
+	if (Instance)
+	{
+		Args.Instance = Instance->linkedAddress;
+	}
+
+	if (Material)
+	{
+		Args.Material = Material->linkedAddress;
+	}
+
+	// Find out relevant parameters. \todo: this may be optimised by reusing partial
+	// values in a LINK_CONTEXT or similar
+	SubtreeRelevantParametersVisitorAST visitor;
+	visitor.Run(Material.child());
+
+	TArray<uint16> params;
+	for (const FString& paramName : visitor.Parameters)
+	{
+		for (int32 i = 0; i < Program.Parameters.Num(); ++i)
+		{
+			const auto& param = Program.Parameters[i];
+			if (param.Name == paramName)
+			{
+				params.Add(uint16(i));
+				break;
+			}
+		}
+	}
+
+	params.Sort();
+
+	auto it = Program.ParameterLists.Find(params);
+
+	if (it != INDEX_NONE)
+	{
+		Args.RelevantParametersListIndex = it;
+	}
+	else
+	{
+		Args.RelevantParametersListIndex = uint32_t(Program.ParameterLists.Num());
+		Program.ParameterLists.Add(params);
+	}
+
+	
+	linkedAddress = (OP::ADDRESS)Program.OpAddress.Num();
+	Program.OpAddress.Add((uint32_t)Program.ByteCode.Num());
+	AppendCode(Program.ByteCode, Type);
+	AppendCode(Program.ByteCode, Args);
+}
+
+}
diff --git a/Engine/Plugins/Mutable/Source/MutableTools/Private/MuT/ASTOpAddMaterial.h b/Engine/Plugins/Mutable/Source/MutableTools/Private/MuT/ASTOpAddMaterial.h
new file mode 100644
--- /dev/null
+++ b/Engine/Plugins/Mutable/Source/MutableTools/Private/MuT/ASTOpAddMaterial.h
@@ -0,0 +1,46 @@
+// Copyright Epic Games, Inc. All Rights Reserved.
+
+#pragma once
+
+#include "HAL/Platform.h"
+#include "Math/NumericLimits.h"
+#include "MuR/Operations.h"
+#include "MuR/Ptr.h"
+#include "MuT/AST.h"
+
+
+namespace mu
+{
+struct FProgram;
+
+
+//---------------------------------------------------------------------------------------------
+//! Adds a material to an Instance
+//---------------------------------------------------------------------------------------------
+class ASTOpAddMaterial final : public ASTOp
+{
+public:
+
+	//! Type of switch
+	EOpType Type = EOpType::NONE;
+
+	ASTChild Instance;
+	ASTChild Material;
+
+public:
+
+	ASTOpAddMaterial();
+	ASTOpAddMaterial(const ASTOpAddMaterial&) = delete;
+	~ASTOpAddMaterial();
+
+	EOpType GetOpType() const override { return Type; }
+	uint64 Hash() const override;
+	void ForEachChild(const TFunctionRef<void(ASTChild&)> F) override;
+	bool IsEqual(const ASTOp& OtherUntyped) const override;
+	Ptr<ASTOp> Clone(MapChildFuncRef MapChild) const override;
+	void Link(FProgram& Program, FLinkerOptions* Options) override;
+};
+
+
+}
+
diff --git a/Engine/Plugins/Mutable/Source/MutableTools/Private/MuT/ASTOpAddOverlayMaterial.cpp b/Engine/Plugins/Mutable/Source/MutableTools/Private/MuT/ASTOpAddOverlayMaterial.cpp
deleted file mode 100644
diff --git a/Engine/Plugins/Mutable/Source/MutableTools/Private/MuT/ASTOpAddOverlayMaterial.h b/Engine/Plugins/Mutable/Source/MutableTools/Private/MuT/ASTOpAddOverlayMaterial.h
deleted file mode 100644
diff --git a/Engine/Plugins/Mutable/Source/MutableTools/Private/MuT/ASTOpInstanceAdd.cpp b/Engine/Plugins/Mutable/Source/MutableTools/Private/MuT/ASTOpInstanceAdd.cpp
--- a/Engine/Plugins/Mutable/Source/MutableTools/Private/MuT/ASTOpInstanceAdd.cpp
+++ b/Engine/Plugins/Mutable/Source/MutableTools/Private/MuT/ASTOpInstanceAdd.cpp
@@ -100,6 +100,8 @@ namespace mu
 
 	void ASTOpInstanceAdd::Link(FProgram& program, FLinkerOptions*)
 	{
+		check(type != EOpType::NONE);
+		
 		// Already linked?
 		if (!linkedAddress)
 		{
@@ -141,11 +143,11 @@ namespace mu
 
 				if (it != INDEX_NONE)
 				{
-					Args.relevantParametersListIndex = it;
+					Args.RelevantParametersListIndex = it;
 				}
 				else
 				{
-					Args.relevantParametersListIndex = uint32_t(program.ParameterLists.Num());
+					Args.RelevantParametersListIndex = uint32_t(program.ParameterLists.Num());
 					program.ParameterLists.Add(params);
 				}
 			}
diff --git a/Engine/Plugins/Mutable/Source/MutableTools/Private/MuT/CodeGenerator.cpp b/Engine/Plugins/Mutable/Source/MutableTools/Private/MuT/CodeGenerator.cpp
--- a/Engine/Plugins/Mutable/Source/MutableTools/Private/MuT/CodeGenerator.cpp
+++ b/Engine/Plugins/Mutable/Source/MutableTools/Private/MuT/CodeGenerator.cpp
@@ -16,7 +16,7 @@
 #include "MuR/Platform.h"
 #include "MuT/ASTOpAddLOD.h"
 #include "MuT/ASTOpAddExtensionData.h"
-#include "MuT/ASTOpAddOverlayMaterial.h"
+#include "MuT/ASTOpAddMaterial.h"
 #include "MuT/ASTOpConditional.h"
 #include "MuT/ASTOpConstantBool.h"
 #include "MuT/ASTOpConstantResource.h"
@@ -611,14 +611,15 @@ namespace mu
 
 				if (Ptr<NodeMaterial> MaterialNode = InNode->OverlayMaterial)
 				{
-					Ptr<ASTOpAddOverlayMaterial> AddOverlayMaterialOp = new ASTOpAddOverlayMaterial();
+					Ptr<ASTOpAddMaterial> AddOverlayMaterialOp = new ASTOpAddMaterial();
+					AddOverlayMaterialOp->Type = EOpType::IN_ADDOVERLAYMATERIAL;
 					AddOverlayMaterialOp->Instance = LastInstOp;
 
 					// Material
 					FMaterialGenerationOptions MaterialOptions;
 					FMaterialGenerationResult MaterialResult;
 					GenerateMaterial(MaterialOptions, MaterialResult, MaterialNode);
-					AddOverlayMaterialOp->OverlayMaterial = MaterialResult.op;
+					AddOverlayMaterialOp->Material = MaterialResult.op;
 
 					LastInstOp = AddOverlayMaterialOp;
 				}
@@ -1651,14 +1652,14 @@ namespace mu
 					//------------------------------------------------------------------------
 					if (SurfaceNode->Material)
 					{
-						Ptr<ASTOpInstanceAdd> Op = new ASTOpInstanceAdd();
-						Op->type = EOpType::IN_ADDMATERIAL;
-						Op->instance = LastSurfOp;
+						Ptr<ASTOpAddMaterial> Op = new ASTOpAddMaterial();
+						Op->Type = EOpType::IN_ADDMATERIAL;
+						Op->Instance = LastSurfOp;
 
 						FMaterialGenerationOptions MaterialOptions;
 						FMaterialGenerationResult MaterialResult;
 						GenerateMaterial(MaterialOptions, MaterialResult, SurfaceNode->Material);
-						Op->value = MaterialResult.op;
+						Op->Material = MaterialResult.op;
 						LastSurfOp = Op;
 					}
 				}
diff --git a/Engine/Plugins/Mutable/Source/MutableTools/Private/MuT/CodeOptimiser.cpp b/Engine/Plugins/Mutable/Source/MutableTools/Private/MuT/CodeOptimiser.cpp
--- a/Engine/Plugins/Mutable/Source/MutableTools/Private/MuT/CodeOptimiser.cpp
+++ b/Engine/Plugins/Mutable/Source/MutableTools/Private/MuT/CodeOptimiser.cpp
@@ -414,7 +414,7 @@ namespace mu
 			Model->GetPrivate()->Program.States.Add(state);
 
 			TSharedPtr<FParameters> LocalParams = FModel::NewParameters(Model);
-			System->GetPrivate()->BeginBuild(Model, nullptr, nullptr);
+			System->GetPrivate()->BeginBuild(Model, nullptr, nullptr, nullptr);
 
 			// Calculate the value and replace this op by a constant
 			switch( DataType )
