[Mutable] Move NodeLayout blocks generation to core.

#rnx

diff --git a/Engine/Plugins/Mutable/Source/CustomizableObjectEditor/Private/MuCOE/CustomizableObjectLayout.cpp b/Engine/Plugins/Mutable/Source/CustomizableObjectEditor/Private/MuCOE/CustomizableObjectLayout.cpp
--- a/Engine/Plugins/Mutable/Source/CustomizableObjectEditor/Private/MuCOE/CustomizableObjectLayout.cpp
+++ b/Engine/Plugins/Mutable/Source/CustomizableObjectEditor/Private/MuCOE/CustomizableObjectLayout.cpp
@@ -7,6 +7,7 @@
 #include "MuCO/CustomizableObjectInstancePrivate.h"
 #include "MuCOE/CustomizableObjectCompiler.h"
 #include "MuCOE/GenerateMutableSource/GenerateMutableSourceLayout.h"
+#include "MuCOE/GenerateMutableSource/GenerateMutableSourceMesh.h"
 #include "MuCOE/GraphTraversal.h"
 #include "MuCOE/MutableUtils.h"
 #include "MuCOE/CustomizableObjectEditor.h"
@@ -38,6 +39,32 @@ void UCustomizableObjectLayout::SetLayout(int32 LODIndex, int32 MatIndex, int32
 	UVChannel = UVIndex;
 }
 
+mu::EAutoBlocksStrategy ConvertAutomaticBlocksStrategy(const ECustomizableObjectLayoutAutomaticBlocksStrategy InAutomaticBlockStrategy)
+{
+	mu::EAutoBlocksStrategy AutoBlockStrategy = mu::EAutoBlocksStrategy::Rectangles;
+
+	switch (InAutomaticBlockStrategy)
+	{
+	case ECustomizableObjectLayoutAutomaticBlocksStrategy::Rectangles:
+		AutoBlockStrategy = mu::EAutoBlocksStrategy::Rectangles;
+		break;
+
+	case ECustomizableObjectLayoutAutomaticBlocksStrategy::UVIslands:
+		AutoBlockStrategy = mu::EAutoBlocksStrategy::UVIslands;
+		break;
+
+	case ECustomizableObjectLayoutAutomaticBlocksStrategy::Ignore:
+		AutoBlockStrategy = mu::EAutoBlocksStrategy::Ignore;
+		break;
+
+	default:
+		checkNoEntry();
+	}
+
+	return AutoBlockStrategy;
+}
+
+
 
 mu::EPackStrategy ConvertLayoutStrategy(const ECustomizableObjectTextureLayoutPackingStrategy LayoutPackStrategy)
 {
@@ -106,64 +133,121 @@ void UCustomizableObjectLayout::GenerateAutomaticBlocksFromUVs()
 	check(Object);
 
 	FCompilationOptions Options = Object->GetPrivate()->GetCompileOptions();
-	
+
 	FMutableCompilationContext CompilationContext(Object, Compiler, Options);
 	FMutableGraphGenerationContext GenerationContext(CompilationContext);
 
-	bool bOutWasEmpty = false;
-	mu::Ptr<mu::NodeLayout> LayoutNode = CreateMutableLayoutNode(GenerationContext, this, false, bOutWasEmpty );
+	mu::Ptr<mu::NodeLayout> LayoutNode = CreateMutableLayoutNode(GenerationContext, this, false);
+	check(LayoutNode);
 
-	if (LayoutNode)
+	if (PackingStrategy == ECustomizableObjectTextureLayoutPackingStrategy::Overlay)
+		{
+			// Legacy behavior
+			if (Blocks.IsEmpty())
+			{
+				LayoutNode->Blocks.SetNum(1);
+				LayoutNode->Blocks[0].Min = { 0,0 };
+				LayoutNode->Blocks[0].Size = LayoutNode->Size;
+				LayoutNode->Blocks[0].Priority = 0;
+				LayoutNode->Blocks[0].bReduceBothAxes = false;
+				LayoutNode->Blocks[0].bReduceByTwo = false;
+			}
+		}
+	else // Convert Mesh and generate blocks
 	{
-		AutomaticBlocks.Empty();
+		TSharedPtr<mu::FMesh> MutableMesh;
 
-		// Generate the editor layout blocks from the mutable layout
-		for (int32 i = 0; i < LayoutNode->Blocks.Num(); ++i)
+		if (USkeletalMesh* SkeletalMesh = Cast<USkeletalMesh>(GenerationContext.LoadObject(GetMesh())))
+		{
+			// We don't need all the data to generate the blocks
+			const EMutableMeshConversionFlags ShapeFlags =
+				EMutableMeshConversionFlags::IgnoreSkinning |
+				EMutableMeshConversionFlags::IgnorePhysics |
+				EMutableMeshConversionFlags::IgnoreMorphs |
+				EMutableMeshConversionFlags::DoNotCreateMeshMetadata;
+
+			GenerationContext.MeshGenerationFlags.Push(ShapeFlags);
+
+			bool bMeshMustExist = true;
+			bool bForceImmediateConversion = true;
+			FMutableSourceMeshData Source;
+			Source.Mesh = SkeletalMesh;
+			MutableMesh = ConvertSkeletalMeshToMutable(Source, bMeshMustExist, GetLOD(), GetMaterial(), GenerationContext, nullptr, bForceImmediateConversion);
+
+			GenerationContext.MeshGenerationFlags.Pop();
+		}
+		else if (UStaticMesh* StaticMesh = Cast<UStaticMesh>(GenerationContext.LoadObject(GetMesh())))
 		{
-			FIntPoint Min = FIntPoint(LayoutNode->Blocks[i].Min.X, LayoutNode->Blocks[i].Min.Y);
-			FIntPoint Size = FIntPoint(LayoutNode->Blocks[i].Size.X, LayoutNode->Blocks[i].Size.Y);
+			MutableMesh = ConvertStaticMeshToMutable(StaticMesh, GetLOD(), GetMaterial(), GenerationContext, nullptr);
+		}
 
-			// Ignore blocks contained inside any block in the initial block set.
-			bool bContainedInInitialSet = false;
-			for (const FCustomizableObjectLayoutBlock& Block : Blocks)
+		if (MutableMesh)
+		{
+			// Generating blocks with the mutable mesh
+			if (AutomaticBlocksStrategy == ECustomizableObjectLayoutAutomaticBlocksStrategy::Rectangles)
 			{
-				FInt32Rect ExistingRect(Block.Min, Block.Max + FIntPoint(1));
-				if (ExistingRect.Contains(Min) && ExistingRect.Contains(Min + Size))
-				{
-					bContainedInInitialSet = true;
-					break;
-				}
+				LayoutNode->GenerateLayoutBlocks(MutableMesh, GetUVChannel());
 			}
-			if (bContainedInInitialSet)
+			else if (AutomaticBlocksStrategy == ECustomizableObjectLayoutAutomaticBlocksStrategy::UVIslands)
 			{
-				continue;
+				LayoutNode->GenerateLayoutBlocksFromUVIslands(MutableMesh, GetUVChannel());
 			}
-
-			FCustomizableObjectLayoutBlock Block(FIntPoint(Min.X, Min.Y), FIntPoint(Min.X + Size.X, Min.Y + Size.Y));
-			Block.bIsAutomatic = true;
-			Block.Id = FGuid::NewGuid();
-
-			TSharedPtr<mu::FImage> Mask = LayoutNode->Blocks[i].Mask;
-			if (Mask)
+			else
 			{
-				UTexture2D* UnrealImage = NewObject<UTexture2D>(UTexture2D::StaticClass());
+				// Unimplemented
+				check(false);
+			}
+		}
+	}
 
-				FMutableModelImageProperties Props;
-				Props.Filter = TF_Nearest;
-				Props.SRGB = true;
-				Props.LODBias = 0;
-				ConvertImage(UnrealImage, Mask, Props );
-				UnrealImage->NeverStream = true;
-				UnrealImage->UpdateResource();
+	AutomaticBlocks.Empty();
+
+	// Generate the editor layout blocks from the mutable layout
+	for (int32 i = 0; i < LayoutNode->Blocks.Num(); ++i)
+	{
+		FIntPoint Min = FIntPoint(LayoutNode->Blocks[i].Min.X, LayoutNode->Blocks[i].Min.Y);
+		FIntPoint Size = FIntPoint(LayoutNode->Blocks[i].Size.X, LayoutNode->Blocks[i].Size.Y);
 
-				Block.Mask = UnrealImage;
+		// Ignore blocks contained inside any block in the initial block set.
+		bool bContainedInInitialSet = false;
+		for (const FCustomizableObjectLayoutBlock& Block : Blocks)
+		{
+			FInt32Rect ExistingRect(Block.Min, Block.Max + FIntPoint(1));
+			if (ExistingRect.Contains(Min) && ExistingRect.Contains(Min + Size))
+			{
+				bContainedInInitialSet = true;
+				break;
 			}
-			AutomaticBlocks.Add(Block);
 		}
+		if (bContainedInInitialSet)
+		{
+			continue;
+		}
+
+		FCustomizableObjectLayoutBlock Block(FIntPoint(Min.X, Min.Y), FIntPoint(Min.X + Size.X, Min.Y + Size.Y));
+		Block.bIsAutomatic = true;
+		Block.Id = FGuid::NewGuid();
 
-		Node->PostEditChange();
-		Node->GetGraph()->MarkPackageDirty();
+		TSharedPtr<mu::FImage> Mask = LayoutNode->Blocks[i].Mask;
+		if (Mask)
+		{
+			UTexture2D* UnrealImage = NewObject<UTexture2D>(UTexture2D::StaticClass());
+
+			FMutableModelImageProperties Props;
+			Props.Filter = TF_Nearest;
+			Props.SRGB = true;
+			Props.LODBias = 0;
+			ConvertImage(UnrealImage, Mask, Props);
+			UnrealImage->NeverStream = true;
+			UnrealImage->UpdateResource();
+
+			Block.Mask = UnrealImage;
+		}
+		AutomaticBlocks.Add(Block);
 	}
+
+	Node->PostEditChange();
+	Node->GetGraph()->MarkPackageDirty();
 }
 
 
@@ -221,7 +305,7 @@ void UCustomizableObjectLayout::SetIgnoreWarningsLOD(int32 LODValue)
 }
 
 
-TSoftObjectPtr<UObject> UCustomizableObjectLayout::GetMesh() const
+TSoftObjectPtr<UStreamableRenderAsset> UCustomizableObjectLayout::GetMesh() const
 {
 	if (UObject* Node = GetOuter())
 	{
diff --git a/Engine/Plugins/Mutable/Source/CustomizableObjectEditor/Private/MuCOE/CustomizableObjectLayout.h b/Engine/Plugins/Mutable/Source/CustomizableObjectEditor/Private/MuCOE/CustomizableObjectLayout.h
--- a/Engine/Plugins/Mutable/Source/CustomizableObjectEditor/Private/MuCOE/CustomizableObjectLayout.h
+++ b/Engine/Plugins/Mutable/Source/CustomizableObjectEditor/Private/MuCOE/CustomizableObjectLayout.h
@@ -7,6 +7,7 @@
 #include "MuR/Image.h"
 #include "CustomizableObjectLayout.generated.h"
 
+class UStreamableRenderAsset;
 
 UENUM()
 enum class ECustomizableObjectTextureLayoutPackingStrategy : uint8
@@ -32,6 +33,8 @@ enum class ECustomizableObjectLayoutAutomaticBlocksStrategy : uint8
 	Ignore = 2 UMETA(DisplayName = "Ignore"),
 };
 
+mu::EAutoBlocksStrategy ConvertAutomaticBlocksStrategy(const ECustomizableObjectLayoutAutomaticBlocksStrategy LayoutPackStrategy);
+
 UENUM()
 enum class ECustomizableObjectLayoutAutomaticBlocksMergeStrategy : uint8
 {
@@ -122,7 +125,7 @@ public:
 	int32 GetMaterial() const { return Material; }
 	int32 GetUVChannel() const { return UVChannel; }
 	FString GetLayoutName() const { return LayoutName; }
-	TSoftObjectPtr<UObject> GetMesh() const;
+	TSoftObjectPtr<UStreamableRenderAsset> GetMesh() const;
 	FIntPoint GetGridSize() const { return GridSize; }
 	FIntPoint GetMaxGridSize() const { return MaxGridSize; }
 	bool GetIgnoreVertexLayoutWarnings() const { return bIgnoreUnassignedVertexWarning; };
diff --git a/Engine/Plugins/Mutable/Source/CustomizableObjectEditor/Private/MuCOE/GenerateMutableSource/GenerateMutableSourceLayout.cpp b/Engine/Plugins/Mutable/Source/CustomizableObjectEditor/Private/MuCOE/GenerateMutableSource/GenerateMutableSourceLayout.cpp
--- a/Engine/Plugins/Mutable/Source/CustomizableObjectEditor/Private/MuCOE/GenerateMutableSource/GenerateMutableSourceLayout.cpp
+++ b/Engine/Plugins/Mutable/Source/CustomizableObjectEditor/Private/MuCOE/GenerateMutableSource/GenerateMutableSourceLayout.cpp
@@ -19,8 +19,10 @@ mu::Ptr<mu::NodeLayout> CreateDefaultLayout()
 	mu::Ptr<mu::NodeLayout> LayoutNode = new mu::NodeLayout();
 	LayoutNode->Size = { GridSize, GridSize };
 	LayoutNode->MaxSize = { GridSize, GridSize };
+	LayoutNode->AutoBlockStrategy = mu::EAutoBlocksStrategy::Ignore;
 	LayoutNode->Strategy = mu::EPackStrategy::Resizeable;
 	LayoutNode->ReductionMethod = mu::EReductionMethod::Halve;
+	LayoutNode->TexCoordsIndex = 0;
 	LayoutNode->Blocks.SetNum(1);
 	LayoutNode->Blocks[0].Min = { 0, 0 };
 	LayoutNode->Blocks[0].Size = { GridSize, GridSize };
@@ -32,18 +34,18 @@ mu::Ptr<mu::NodeLayout> CreateDefaultLayout()
 }
 
 
-mu::Ptr<mu::NodeLayout> CreateMutableLayoutNode(FMutableGraphGenerationContext& GenerationContext, const UCustomizableObjectLayout* UnrealLayout, bool bIgnoreLayoutWarnings, bool& bWasEmpty )
+mu::Ptr<mu::NodeLayout> CreateMutableLayoutNode(FMutableGraphGenerationContext& GenerationContext, const UCustomizableObjectLayout* UnrealLayout, bool bIgnoreLayoutWarnings)
 {
-	bWasEmpty = false;
 	mu::Ptr<mu::NodeLayout> LayoutNode = new mu::NodeLayout;
 
 	LayoutNode->Size = FIntVector2(UnrealLayout->GetGridSize().X, UnrealLayout->GetGridSize().Y);
 	LayoutNode->MaxSize = FIntVector2(UnrealLayout->GetMaxGridSize().X, UnrealLayout->GetMaxGridSize().Y);
+	LayoutNode->TexCoordsIndex = UnrealLayout->GetUVChannel();
 
-	mu::EPackStrategy PackStrategy = ConvertLayoutStrategy(UnrealLayout->PackingStrategy);
-	LayoutNode->Strategy = PackStrategy;
-
-	LayoutNode->ReductionMethod = (UnrealLayout->BlockReductionMethod == ECustomizableObjectLayoutBlockReductionMethod::Halve ? mu::EReductionMethod::Halve : mu::EReductionMethod::Unitary);
+	LayoutNode->AutoBlockStrategy = ConvertAutomaticBlocksStrategy(UnrealLayout->AutomaticBlocksStrategy);
+	LayoutNode->bMergeChildBlocks = UnrealLayout->AutomaticBlocksMergeStrategy == ECustomizableObjectLayoutAutomaticBlocksMergeStrategy::MergeChildBlocks;
+	LayoutNode->Strategy = ConvertLayoutStrategy(UnrealLayout->PackingStrategy);
+	LayoutNode->ReductionMethod = UnrealLayout->BlockReductionMethod == ECustomizableObjectLayoutBlockReductionMethod::Halve ? mu::EReductionMethod::Halve : mu::EReductionMethod::Unitary;
 
 	if (bIgnoreLayoutWarnings)
 	{
@@ -61,75 +63,6 @@ mu::Ptr<mu::NodeLayout> CreateMutableLayoutNode(FMutableGraphGenerationContext&
 		LayoutNode->Blocks[BlockIndex] = ToMutable(GenerationContext, UnrealLayout->Blocks[BlockIndex]);
 	}
 
-
-	ECustomizableObjectLayoutAutomaticBlocksStrategy AutomaticBlockStrategy = UnrealLayout->AutomaticBlocksStrategy;
-	ECustomizableObjectTextureLayoutPackingStrategy PackingStrategy = UnrealLayout->PackingStrategy;
-	
-	if (AutomaticBlockStrategy == ECustomizableObjectLayoutAutomaticBlocksStrategy::Ignore || 
-		PackingStrategy == ECustomizableObjectTextureLayoutPackingStrategy::Overlay)
- 	{
-		// Legacy behavior
-		if (UnrealLayout->Blocks.IsEmpty())
-		{
-			bWasEmpty = true;
-			LayoutNode->Blocks.SetNum(1);
-			LayoutNode->Blocks[0].Min = { 0,0 };
-			LayoutNode->Blocks[0].Size = LayoutNode->Size;
-			LayoutNode->Blocks[0].Priority = 0;
-			LayoutNode->Blocks[0].bReduceBothAxes = false;
-			LayoutNode->Blocks[0].bReduceByTwo = false;
-		}
-	}
-	else
-	{
-		// Convert the UE mesh in the layout into a Mutable mesh.
-		TSharedPtr<mu::FMesh> MutableMesh;
-
-		// TODO: Don't force load the mesh by moving the creation of automatic blocks to the core compilation stage.
-		if (USkeletalMesh* SkeletalMesh = Cast<USkeletalMesh>(GenerationContext.LoadObject(UnrealLayout->GetMesh())))
-		{
-			// We don't need all the data to generate the blocks
-			const EMutableMeshConversionFlags ShapeFlags =
-				EMutableMeshConversionFlags::IgnoreSkinning |
-				EMutableMeshConversionFlags::IgnorePhysics |
-				EMutableMeshConversionFlags::IgnoreMorphs |
-				EMutableMeshConversionFlags::DoNotCreateMeshMetadata;
-
-			GenerationContext.MeshGenerationFlags.Push(ShapeFlags);
-
-			bool bMeshMustExist = true;
-			bool bForceImmediateConversion = true;
-			FMutableSourceMeshData Source;
-			Source.Mesh = SkeletalMesh;
-			MutableMesh = ConvertSkeletalMeshToMutable(Source, bMeshMustExist, UnrealLayout->GetLOD(), UnrealLayout->GetMaterial(), GenerationContext, nullptr, bForceImmediateConversion);
-
-			GenerationContext.MeshGenerationFlags.Pop();
-		}
-		else if (UStaticMesh* StaticMesh = Cast<UStaticMesh>(GenerationContext.LoadObject(UnrealLayout->GetMesh())))
-		{
-			MutableMesh = ConvertStaticMeshToMutable(StaticMesh, UnrealLayout->GetLOD(), UnrealLayout->GetMaterial(), GenerationContext, nullptr);
-		}
-
-		if (MutableMesh)
-		{
-			// Generating blocks with the mutable mesh
-			if (AutomaticBlockStrategy == ECustomizableObjectLayoutAutomaticBlocksStrategy::Rectangles)
-			{
-				LayoutNode->GenerateLayoutBlocks(MutableMesh, UnrealLayout->GetUVChannel());
-			}
-			else if (AutomaticBlockStrategy == ECustomizableObjectLayoutAutomaticBlocksStrategy::UVIslands)
-			{
-				bool bMergeChildBlocks = UnrealLayout->AutomaticBlocksMergeStrategy == ECustomizableObjectLayoutAutomaticBlocksMergeStrategy::MergeChildBlocks;
-				LayoutNode->GenerateLayoutBlocksFromUVIslands(MutableMesh, UnrealLayout->GetUVChannel(), bMergeChildBlocks);
-			}
-			else
-			{
-				// Unimplemented
-				check(false);
-			}
-		}
-	}
-
 	return LayoutNode;
 }
 
diff --git a/Engine/Plugins/Mutable/Source/CustomizableObjectEditor/Private/MuCOE/GenerateMutableSource/GenerateMutableSourceLayout.h b/Engine/Plugins/Mutable/Source/CustomizableObjectEditor/Private/MuCOE/GenerateMutableSource/GenerateMutableSourceLayout.h
--- a/Engine/Plugins/Mutable/Source/CustomizableObjectEditor/Private/MuCOE/GenerateMutableSource/GenerateMutableSourceLayout.h
+++ b/Engine/Plugins/Mutable/Source/CustomizableObjectEditor/Private/MuCOE/GenerateMutableSource/GenerateMutableSourceLayout.h
@@ -11,6 +11,6 @@ struct FMutableGraphGenerationContext;
 
 extern mu::Ptr<mu::NodeLayout> CreateDefaultLayout();
 
-extern mu::Ptr<mu::NodeLayout> CreateMutableLayoutNode(FMutableGraphGenerationContext& GenerationContext, const UCustomizableObjectLayout*, bool bIgnoreLayoutWarnings, bool& bWasEmpty);
+extern mu::Ptr<mu::NodeLayout> CreateMutableLayoutNode(FMutableGraphGenerationContext& GenerationContext, const UCustomizableObjectLayout*, bool bIgnoreLayoutWarnings);
 
 extern mu::FSourceLayoutBlock ToMutable(FMutableGraphGenerationContext& GenerationContext, const FCustomizableObjectLayoutBlock&);
diff --git a/Engine/Plugins/Mutable/Source/CustomizableObjectEditor/Private/MuCOE/GenerateMutableSource/GenerateMutableSourceMesh.cpp b/Engine/Plugins/Mutable/Source/CustomizableObjectEditor/Private/MuCOE/GenerateMutableSource/GenerateMutableSourceMesh.cpp
--- a/Engine/Plugins/Mutable/Source/CustomizableObjectEditor/Private/MuCOE/GenerateMutableSource/GenerateMutableSourceMesh.cpp
+++ b/Engine/Plugins/Mutable/Source/CustomizableObjectEditor/Private/MuCOE/GenerateMutableSource/GenerateMutableSourceMesh.cpp
@@ -3868,8 +3868,6 @@ mu::Ptr<mu::NodeMesh> GenerateMutableSourceMesh(const UEdGraphPin* Pin,
 					}
 				}
 
-				FSkeletalMeshModel* ImportedModel = SkeletalMesh->GetImportedModel();
-
 				FMutableSourceMeshData Source = BaseMeshData;
 				Source.Mesh = SkeletalMesh;
 				Source.AnimInstance = TypedNodeSkel->AnimInstance;
@@ -3916,16 +3914,14 @@ mu::Ptr<mu::NodeMesh> GenerateMutableSourceMesh(const UEdGraphPin* Pin,
 						SectionIndexConnected :
 						SectionIndex;
 
-					int32 NumLayouts = 0;
-					if (ImportedModel && ImportedModel->LODModels.IsValidIndex(LODIndexLayout))
-					{
-						NumLayouts = ImportedModel->LODModels[LODIndexLayout].NumTexCoords;
-					}
-					MeshNode->Layouts.SetNum(NumLayouts);
 
 					const FLayoutGenerationFlags& LayoutFlags = GenerationContext.LayoutGenerationFlags.Last();
 
 					const TArray<UCustomizableObjectLayout*>& Layouts = TypedNodeSkel->GetLayouts(*Pin);
+					const int32 NumLayouts = Layouts.Num();
+
+					MeshNode->Layouts.SetNum(NumLayouts);
+
 					for (int32 LayoutIndex = 0; LayoutIndex < NumLayouts; ++LayoutIndex)
 					{
 						if (!LayoutFlags.TexturePinModes.IsValidIndex(LayoutIndex) ||
@@ -3933,10 +3929,11 @@ mu::Ptr<mu::NodeMesh> GenerateMutableSourceMesh(const UEdGraphPin* Pin,
 						{
 							MeshNode->Layouts[LayoutIndex] = CreateDefaultLayout();
 
-							// Keep packing strategy if possible, Overlay can be valid with EPinMode == Passthrough	
 							if (Layouts.IsValidIndex(LayoutIndex) && Layouts[LayoutIndex])
 							{
+								// Keep packing strategy if possible, Overlay can be valid with EPinMode == Passthrough	
 								MeshNode->Layouts[LayoutIndex]->Strategy = ConvertLayoutStrategy(Layouts[LayoutIndex]->PackingStrategy);
+								MeshNode->Layouts[LayoutIndex]->TexCoordsIndex = Layouts[LayoutIndex]->GetUVChannel();
 							}
 
 							// Ignore layout
@@ -3946,8 +3943,7 @@ mu::Ptr<mu::NodeMesh> GenerateMutableSourceMesh(const UEdGraphPin* Pin,
 						const UCustomizableObjectLayout* Layout = Layouts.IsValidIndex(LayoutIndex) ? Layouts[LayoutIndex] : nullptr;
 						if (ensure(Layout))
 						{
-							bool bWasEmpty = false;
-							mu::Ptr<mu::NodeLayout> LayoutNode = CreateMutableLayoutNode(GenerationContext, Layout, bLinkedToExtendMaterial, bWasEmpty); // TODO PERE: Figure out
+							mu::Ptr<mu::NodeLayout> LayoutNode = CreateMutableLayoutNode(GenerationContext, Layout, bLinkedToExtendMaterial); // TODO PERE: Figure out
 							LayoutNode->SetMessageContext(Node);
 
 							MeshNode->Layouts[LayoutIndex] = LayoutNode;
@@ -4037,16 +4033,9 @@ mu::Ptr<mu::NodeMesh> GenerateMutableSourceMesh(const UEdGraphPin* Pin,
 				
 				if (Layout)
 				{
-					bool bWasEmpty = false;
-					LayoutNode = CreateMutableLayoutNode(GenerationContext, Layout, false, bWasEmpty);
-					/*if (bWasEmpty)
-					{
-						FString msg = "Layout without any block found. A grid sized block will be used instead.";
-						GenerationContext.Log(FText::FromString(msg), Node, EMessageSeverity::Warning);
-					}*/
+					LayoutNode = CreateMutableLayoutNode(GenerationContext, Layout, false);
 				}
-
-				if (!LayoutNode)
+				else
 				{
 					LayoutNode = CreateDefaultLayout();
 				}
@@ -4556,43 +4545,36 @@ mu::Ptr<mu::NodeMesh> GenerateMutableSourceMesh(const UEdGraphPin* Pin,
 						// Pose Assets do not need this part of the code
 						if (Pin->PinType.PinCategory == Schema->PC_Mesh)
 						{
+							const FLayoutGenerationFlags& LayoutFlags = GenerationContext.LayoutGenerationFlags.Last();
+
 							TArray<UCustomizableObjectLayout*> Layouts = TypedNodeTable->GetLayouts(Pin);
-							MeshTableNode->Layouts.SetNum(Layouts.Num());
+							const int32 NumLayouts = Layouts.Num();
 
-							if (Layouts.Num())
+							// Generating node Layouts
+							MeshTableNode->Layouts.SetNum(NumLayouts);
+							for (int32 LayoutIndex = 0; LayoutIndex < NumLayouts; ++LayoutIndex)
 							{
-								// Generating node Layouts
-								const FLayoutGenerationFlags& LayoutFlags = GenerationContext.LayoutGenerationFlags.Last();
-								
-								for (int32 LayoutIndex = 0; LayoutIndex < Layouts.Num(); ++LayoutIndex)
+								if (!LayoutFlags.TexturePinModes.IsValidIndex(LayoutIndex) ||
+									LayoutFlags.TexturePinModes[LayoutIndex] != EPinMode::Mutable)
 								{
-									if (!LayoutFlags.TexturePinModes.IsValidIndex(LayoutIndex) ||
-										LayoutFlags.TexturePinModes[LayoutIndex] != EPinMode::Mutable)
-									{
-										MeshTableNode->Layouts[LayoutIndex] = CreateDefaultLayout();
-
-										// Keep packing strategy if possible, Overlay can be valid with EPinMode == Passthrough	
-										if (Layouts[LayoutIndex])
-										{
-											MeshTableNode->Layouts[LayoutIndex]->Strategy = ConvertLayoutStrategy(Layouts[LayoutIndex]->PackingStrategy);
-										}
-
-										// Ignore layouts
-										continue;
-									}
-
-									bool bWasEmpty = false;
-									// In tables, mimic the legacy behaviour and ignore all layout warnings beyond LOD 0.
-									bool bIgnoreLayoutWarnings = true;
-									mu::Ptr<mu::NodeLayout> LayoutNode = CreateMutableLayoutNode(GenerationContext, Layouts[LayoutIndex], bIgnoreLayoutWarnings, bWasEmpty);
-									if (bWasEmpty)
+									MeshTableNode->Layouts[LayoutIndex] = CreateDefaultLayout();
+							
+									// Keep packing strategy if possible, Overlay can be valid with EPinMode == Passthrough	
+									if (Layouts[LayoutIndex])
 									{
-										FString msg = "Mesh Column [" + MutableColumnName + "] Layout doesn't has any block. A grid sized block will be used instead.";
-										GenerationContext.Log(FText::FromString(msg), Node, EMessageSeverity::Warning);
+										MeshTableNode->Layouts[LayoutIndex]->Strategy = ConvertLayoutStrategy(Layouts[LayoutIndex]->PackingStrategy);
+										MeshTableNode->Layouts[LayoutIndex]->TexCoordsIndex = Layouts[LayoutIndex]->GetUVChannel();
 									}
-
-									MeshTableNode->Layouts[LayoutIndex] = LayoutNode;
+							
+									// Ignore layouts
+									continue;
 								}
+							
+								// In tables, mimic the legacy behaviour and ignore all layout warnings beyond LOD 0.
+								bool bIgnoreLayoutWarnings = true;
+								mu::Ptr<mu::NodeLayout> LayoutNode = CreateMutableLayoutNode(GenerationContext, Layouts[LayoutIndex], bIgnoreLayoutWarnings);
+							
+								MeshTableNode->Layouts[LayoutIndex] = LayoutNode;
 							}
 
 							// Applying Mesh Morph Nodes
@@ -4652,7 +4634,8 @@ mu::Ptr<mu::NodeMesh> GenerateMutableSourceMesh(const UEdGraphPin* Pin,
 			const FLayoutGenerationFlags& LayoutFlags = GenerationContext.LayoutGenerationFlags.Last();
 
 			const TArray<UCustomizableObjectLayout*>& Layouts = TypedNodeParam->GetLayouts(*Pin);
-			int32 NumLayouts = Layouts.Num();
+			const int32 NumLayouts = Layouts.Num();
+
 			ParameterNode->Layouts.SetNum(NumLayouts);
 			for (int32 LayoutIndex = 0; LayoutIndex < NumLayouts; ++LayoutIndex)
 			{
@@ -4661,10 +4644,11 @@ mu::Ptr<mu::NodeMesh> GenerateMutableSourceMesh(const UEdGraphPin* Pin,
 				{
 					ParameterNode->Layouts[LayoutIndex] = CreateDefaultLayout();
 
-					// Keep packing strategy if possible, Overlay can be valid with EPinMode == Passthrough	
 					if (Layouts.IsValidIndex(LayoutIndex) && Layouts[LayoutIndex])
 					{
+						// Keep packing strategy if possible, Overlay can be valid with EPinMode == Passthrough	
 						ParameterNode->Layouts[LayoutIndex]->Strategy = ConvertLayoutStrategy(Layouts[LayoutIndex]->PackingStrategy);
+						ParameterNode->Layouts[LayoutIndex]->TexCoordsIndex = Layouts[LayoutIndex]->GetUVChannel();
 					}
 
 					// Ignore layout
@@ -4674,13 +4658,31 @@ mu::Ptr<mu::NodeMesh> GenerateMutableSourceMesh(const UEdGraphPin* Pin,
 				const UCustomizableObjectLayout* Layout = Layouts.IsValidIndex(LayoutIndex) ? Layouts[LayoutIndex] : nullptr;
 				if (ensure(Layout))
 				{
-					bool bWasEmpty = false;
-					mu::Ptr<mu::NodeLayout> LayoutNode = CreateMutableLayoutNode(GenerationContext, Layout, bLinkedToExtendMaterial, bWasEmpty);
+					
+					mu::Ptr<mu::NodeLayout> LayoutNode = CreateMutableLayoutNode(GenerationContext, Layout, bLinkedToExtendMaterial);
 					LayoutNode->SetMessageContext(Node);
 
 					ParameterNode->Layouts[LayoutIndex] = LayoutNode;
 				}
 			}
+
+			TSoftObjectPtr<UStreamableRenderAsset> LayoutMesh = TypedNodeParam->GetMesh();
+			if (!LayoutMesh.IsNull())
+			{
+				const EMutableMeshConversionFlags ModifiersMeshFlags =
+					EMutableMeshConversionFlags::IgnoreSkinning |
+					EMutableMeshConversionFlags::IgnorePhysics |
+					EMutableMeshConversionFlags::IgnoreMorphs |
+					EMutableMeshConversionFlags::DoNotCreateMeshMetadata;
+				GenerationContext.MeshGenerationFlags.Push(ModifiersMeshFlags);
+
+				// Get the base mesh
+				FMutableSourceMeshData Source;
+				Source.Mesh = LayoutMesh;
+				ParameterNode->ReferenceMesh = ConvertSkeletalMeshToMutable(Source, true, LODIndex, SectionIndex, GenerationContext, Node);
+
+				GenerationContext.MeshGenerationFlags.Pop();
+			}
 		}
 
 		// For constants we also have:
diff --git a/Engine/Plugins/Mutable/Source/CustomizableObjectEditor/Private/MuCOE/GenerateMutableSource/GenerateMutableSourceModifier.cpp b/Engine/Plugins/Mutable/Source/CustomizableObjectEditor/Private/MuCOE/GenerateMutableSource/GenerateMutableSourceModifier.cpp
--- a/Engine/Plugins/Mutable/Source/CustomizableObjectEditor/Private/MuCOE/GenerateMutableSource/GenerateMutableSourceModifier.cpp
+++ b/Engine/Plugins/Mutable/Source/CustomizableObjectEditor/Private/MuCOE/GenerateMutableSource/GenerateMutableSourceModifier.cpp
@@ -414,13 +414,6 @@ mu::Ptr<mu::NodeModifier> GenerateMutableSourceModifier(const UEdGraphPin * Pin,
 
 	else if (const UCustomizableObjectNodeModifierRemoveMeshBlocks* TypedNodeRemBlocks = Cast<UCustomizableObjectNodeModifierRemoveMeshBlocks>(Node))
 	{
-		const EMutableMeshConversionFlags ModifiersMeshFlags =
-			EMutableMeshConversionFlags::IgnoreSkinning |
-			EMutableMeshConversionFlags::IgnorePhysics |
-			EMutableMeshConversionFlags::IgnoreMorphs |
-			EMutableMeshConversionFlags::DoNotCreateMeshMetadata;
-		GenerationContext.MeshGenerationFlags.Push(ModifiersMeshFlags);
-
 		mu::Ptr<mu::NodeModifierMeshClipWithUVMask> ClipNode = new mu::NodeModifierMeshClipWithUVMask();
 		Result = ClipNode;
 
@@ -429,12 +422,10 @@ mu::Ptr<mu::NodeModifier> GenerateMutableSourceModifier(const UEdGraphPin * Pin,
 		ClipNode->MultipleTagsPolicy = TypedNodeRemBlocks->MultipleTagPolicy;
 		ClipNode->RequiredTags = TypedNodeRemBlocks->GetNodeRequiredTags(&GenerationContext.MacroNodesStack);
 
-		bool bWasEmpty = false;
-		mu::Ptr<mu::NodeLayout> SourceLayout = CreateMutableLayoutNode(GenerationContext, TypedNodeRemBlocks->Layout, true, bWasEmpty);
+		check(TypedNodeRemBlocks->Layout->GetMesh().IsNull());
+		mu::Ptr<mu::NodeLayout> SourceLayout = CreateMutableLayoutNode(GenerationContext, TypedNodeRemBlocks->Layout, true);
 		ClipNode->ClipLayout = SourceLayout;
 		ClipNode->LayoutIndex = TypedNodeRemBlocks->ParentLayoutIndex;
-
-		GenerationContext.MeshGenerationFlags.Pop();
 	}
 
 	else if (UCustomizableObjectNodeModifierEditMeshSection* TypedNodeEdit = Cast<UCustomizableObjectNodeModifierEditMeshSection>(Node))
diff --git a/Engine/Plugins/Mutable/Source/CustomizableObjectEditor/Private/MuCOE/Nodes/CustomizableObjectNodeMeshInterface.h b/Engine/Plugins/Mutable/Source/CustomizableObjectEditor/Private/MuCOE/Nodes/CustomizableObjectNodeMeshInterface.h
--- a/Engine/Plugins/Mutable/Source/CustomizableObjectEditor/Private/MuCOE/Nodes/CustomizableObjectNodeMeshInterface.h
+++ b/Engine/Plugins/Mutable/Source/CustomizableObjectEditor/Private/MuCOE/Nodes/CustomizableObjectNodeMeshInterface.h
@@ -31,7 +31,7 @@ public:
 	virtual UTexture2D* FindTextureForPin(const UEdGraphPin* Pin) const PURE_VIRTUAL(UCustomizableObjectNodeMesh::FindTextureForPin, return {};);
 
 	/** Returns the Unreal mesh (e.g., USkeletalMesh, UStaticMesh...). */
-	virtual TSoftObjectPtr<UObject> GetMesh() const PURE_VIRTUAL(UCustomizableObjectNodeMesh::GetMesh, return {};);
+	virtual TSoftObjectPtr<UStreamableRenderAsset> GetMesh() const PURE_VIRTUAL(UCustomizableObjectNodeMesh::GetMesh, return {};);
 
 	/** Returns the output Mesh pin associated to the given LODIndex and SectionIndex. Override. */
 	virtual UEdGraphPin* GetMeshPin(int32 LOD, int32 SectionIndex) const PURE_VIRTUAL(UCustomizableObjectNodeMesh::GetMeshPin, return {};);
diff --git a/Engine/Plugins/Mutable/Source/CustomizableObjectEditor/Private/MuCOE/Nodes/CustomizableObjectNodeSkeletalMesh.cpp b/Engine/Plugins/Mutable/Source/CustomizableObjectEditor/Private/MuCOE/Nodes/CustomizableObjectNodeSkeletalMesh.cpp
--- a/Engine/Plugins/Mutable/Source/CustomizableObjectEditor/Private/MuCOE/Nodes/CustomizableObjectNodeSkeletalMesh.cpp
+++ b/Engine/Plugins/Mutable/Source/CustomizableObjectEditor/Private/MuCOE/Nodes/CustomizableObjectNodeSkeletalMesh.cpp
@@ -233,7 +233,7 @@ TArray<UCustomizableObjectLayout*> UCustomizableObjectNodeSkeletalMesh::GetLayou
 }
 
 
-TSoftObjectPtr<UObject> UCustomizableObjectNodeSkeletalMesh::GetMesh() const
+TSoftObjectPtr<UStreamableRenderAsset> UCustomizableObjectNodeSkeletalMesh::GetMesh() const
 {
 	return SkeletalMesh;
 }
diff --git a/Engine/Plugins/Mutable/Source/CustomizableObjectEditor/Private/MuCOE/Nodes/CustomizableObjectNodeSkeletalMesh.h b/Engine/Plugins/Mutable/Source/CustomizableObjectEditor/Private/MuCOE/Nodes/CustomizableObjectNodeSkeletalMesh.h
--- a/Engine/Plugins/Mutable/Source/CustomizableObjectEditor/Private/MuCOE/Nodes/CustomizableObjectNodeSkeletalMesh.h
+++ b/Engine/Plugins/Mutable/Source/CustomizableObjectEditor/Private/MuCOE/Nodes/CustomizableObjectNodeSkeletalMesh.h
@@ -232,7 +232,7 @@ public:
 	// UCustomizableObjectNodeMesh interface
 	UE_API virtual UTexture2D* FindTextureForPin(const UEdGraphPin* Pin) const override;
 	UE_API virtual TArray<UCustomizableObjectLayout*> GetLayouts(const UEdGraphPin& MeshPin) const override;
-	UE_API virtual TSoftObjectPtr<UObject> GetMesh() const override;
+	UE_API virtual TSoftObjectPtr<UStreamableRenderAsset> GetMesh() const override;
 	UE_API virtual UEdGraphPin* GetMeshPin(int32 LOD, int32 SectionIndex) const override;
 	UE_API virtual void GetPinSection(const UEdGraphPin& Pin, int32& OutLODIndex, int32& OutSectionIndex, int32& OutLayoutIndex) const override;
 
diff --git a/Engine/Plugins/Mutable/Source/CustomizableObjectEditor/Private/MuCOE/Nodes/CustomizableObjectNodeSkeletalMeshParameter.cpp b/Engine/Plugins/Mutable/Source/CustomizableObjectEditor/Private/MuCOE/Nodes/CustomizableObjectNodeSkeletalMeshParameter.cpp
--- a/Engine/Plugins/Mutable/Source/CustomizableObjectEditor/Private/MuCOE/Nodes/CustomizableObjectNodeSkeletalMeshParameter.cpp
+++ b/Engine/Plugins/Mutable/Source/CustomizableObjectEditor/Private/MuCOE/Nodes/CustomizableObjectNodeSkeletalMeshParameter.cpp
@@ -141,7 +141,7 @@ TArray<UCustomizableObjectLayout*> UCustomizableObjectNodeSkeletalMeshParameter:
 }
 
 
-TSoftObjectPtr<UObject> UCustomizableObjectNodeSkeletalMeshParameter::GetMesh() const
+TSoftObjectPtr<UStreamableRenderAsset> UCustomizableObjectNodeSkeletalMeshParameter::GetMesh() const
 {
 	return ReferenceValue;
 }
diff --git a/Engine/Plugins/Mutable/Source/CustomizableObjectEditor/Private/MuCOE/Nodes/CustomizableObjectNodeSkeletalMeshParameter.h b/Engine/Plugins/Mutable/Source/CustomizableObjectEditor/Private/MuCOE/Nodes/CustomizableObjectNodeSkeletalMeshParameter.h
--- a/Engine/Plugins/Mutable/Source/CustomizableObjectEditor/Private/MuCOE/Nodes/CustomizableObjectNodeSkeletalMeshParameter.h
+++ b/Engine/Plugins/Mutable/Source/CustomizableObjectEditor/Private/MuCOE/Nodes/CustomizableObjectNodeSkeletalMeshParameter.h
@@ -80,7 +80,7 @@ public:
 	// ICustomizableObjectNodeMeshInterface interface
 	UE_API virtual UTexture2D* FindTextureForPin(const UEdGraphPin* Pin) const override;
 	UE_API virtual TArray<UCustomizableObjectLayout*> GetLayouts(const UEdGraphPin& MeshPin) const override;
-	UE_API virtual TSoftObjectPtr<UObject> GetMesh() const override;
+	UE_API virtual TSoftObjectPtr<UStreamableRenderAsset> GetMesh() const override;
 	UE_API virtual UEdGraphPin* GetMeshPin(int32 LOD, int32 SectionIndex) const override;
 	UE_API virtual void GetPinSection(const UEdGraphPin& Pin, int32& OutLODIndex, int32& OutSectionIndex, int32& OutLayoutIndex) const override;
 
diff --git a/Engine/Plugins/Mutable/Source/CustomizableObjectEditor/Private/MuCOE/Nodes/CustomizableObjectNodeStaticMesh.cpp b/Engine/Plugins/Mutable/Source/CustomizableObjectEditor/Private/MuCOE/Nodes/CustomizableObjectNodeStaticMesh.cpp
--- a/Engine/Plugins/Mutable/Source/CustomizableObjectEditor/Private/MuCOE/Nodes/CustomizableObjectNodeStaticMesh.cpp
+++ b/Engine/Plugins/Mutable/Source/CustomizableObjectEditor/Private/MuCOE/Nodes/CustomizableObjectNodeStaticMesh.cpp
@@ -250,7 +250,7 @@ TArray<class UCustomizableObjectLayout*> UCustomizableObjectNodeStaticMesh::GetL
 }
 
 
-TSoftObjectPtr<UObject> UCustomizableObjectNodeStaticMesh::GetMesh() const
+TSoftObjectPtr<UStreamableRenderAsset> UCustomizableObjectNodeStaticMesh::GetMesh() const
 {
 	return StaticMesh;
 }
diff --git a/Engine/Plugins/Mutable/Source/CustomizableObjectEditor/Private/MuCOE/Nodes/CustomizableObjectNodeStaticMesh.h b/Engine/Plugins/Mutable/Source/CustomizableObjectEditor/Private/MuCOE/Nodes/CustomizableObjectNodeStaticMesh.h
--- a/Engine/Plugins/Mutable/Source/CustomizableObjectEditor/Private/MuCOE/Nodes/CustomizableObjectNodeStaticMesh.h
+++ b/Engine/Plugins/Mutable/Source/CustomizableObjectEditor/Private/MuCOE/Nodes/CustomizableObjectNodeStaticMesh.h
@@ -152,7 +152,7 @@ public:
 	// UCustomizableObjectNodeMesh interface
 	UE_API virtual UTexture2D* FindTextureForPin(const UEdGraphPin* Pin) const override;
 	UE_API virtual TArray<UCustomizableObjectLayout*> GetLayouts(const UEdGraphPin& OutPin) const override;
-	UE_API virtual TSoftObjectPtr<UObject> GetMesh() const override;
+	UE_API virtual TSoftObjectPtr<UStreamableRenderAsset> GetMesh() const override;
 	UE_API virtual UEdGraphPin* GetMeshPin(int32 LOD, int32 SectionIndex) const override;
 	UE_API virtual void GetPinSection(const UEdGraphPin& Pin, int32& OutLODIndex, int32& OutSectionIndex, int32& OutLayoutIndex) const override;
 
diff --git a/Engine/Plugins/Mutable/Source/CustomizableObjectEditor/Private/MuCOE/Nodes/CustomizableObjectNodeTable.cpp b/Engine/Plugins/Mutable/Source/CustomizableObjectEditor/Private/MuCOE/Nodes/CustomizableObjectNodeTable.cpp
--- a/Engine/Plugins/Mutable/Source/CustomizableObjectEditor/Private/MuCOE/Nodes/CustomizableObjectNodeTable.cpp
+++ b/Engine/Plugins/Mutable/Source/CustomizableObjectEditor/Private/MuCOE/Nodes/CustomizableObjectNodeTable.cpp
@@ -1024,7 +1024,7 @@ UTexture2D* UCustomizableObjectNodeTable::FindReferenceTextureParameter(const UE
 }
 
 
-UObject* UCustomizableObjectNodeTable::GetDefaultMeshForLayout(const UCustomizableObjectLayout* InLayout) const
+UStreamableRenderAsset* UCustomizableObjectNodeTable::GetDefaultMeshForLayout(const UCustomizableObjectLayout* InLayout) const
 {
 	for (const UEdGraphPin* Pin : Pins)
 	{
@@ -1034,7 +1034,7 @@ UObject* UCustomizableObjectNodeTable::GetDefaultMeshForLayout(const UCustomizab
 			{
 				if (Layout == InLayout)
 				{
-					return GetColumnDefaultAssetByType<UObject>(MeshPinData->ColumnPropertyName);
+					return GetColumnDefaultAssetByType<UStreamableRenderAsset>(MeshPinData->ColumnPropertyName);
 				}
 			}
 		}
diff --git a/Engine/Plugins/Mutable/Source/CustomizableObjectEditor/Private/MuCOE/Nodes/CustomizableObjectNodeTable.h b/Engine/Plugins/Mutable/Source/CustomizableObjectEditor/Private/MuCOE/Nodes/CustomizableObjectNodeTable.h
--- a/Engine/Plugins/Mutable/Source/CustomizableObjectEditor/Private/MuCOE/Nodes/CustomizableObjectNodeTable.h
+++ b/Engine/Plugins/Mutable/Source/CustomizableObjectEditor/Private/MuCOE/Nodes/CustomizableObjectNodeTable.h
@@ -319,7 +319,7 @@ public:
 	UE_API UTexture2D* FindReferenceTextureParameter(const UEdGraphPin* Pin, FString ParameterImageName) const;
 
 	// Methods to get the UVs of the skeletal mesh 
-	UE_API UObject* GetDefaultMeshForLayout(const UCustomizableObjectLayout* Layout) const;
+	UE_API UStreamableRenderAsset* GetDefaultMeshForLayout(const UCustomizableObjectLayout* Layout) const;
 
 	// Methods to provide the Layouts to the Layout block editors
 	UE_API TArray<UCustomizableObjectLayout*> GetLayouts(const UEdGraphPin* Pin) const;
diff --git a/Engine/Plugins/Mutable/Source/MutableRuntime/Public/MuR/Layout.h b/Engine/Plugins/Mutable/Source/MutableRuntime/Public/MuR/Layout.h
--- a/Engine/Plugins/Mutable/Source/MutableRuntime/Public/MuR/Layout.h
+++ b/Engine/Plugins/Mutable/Source/MutableRuntime/Public/MuR/Layout.h
@@ -14,6 +14,13 @@
 
 namespace mu
 {
+	//! Types of automatic block strategies 
+	enum class EAutoBlocksStrategy : uint32
+	{
+		Rectangles,
+		UVIslands,
+		Ignore,
+	};
 
 	//! Types of layout packing strategies 
 	enum class EPackStrategy : uint32
diff --git a/Engine/Plugins/Mutable/Source/MutableTools/Private/MuT/CodeGenerator_Mesh.cpp b/Engine/Plugins/Mutable/Source/MutableTools/Private/MuT/CodeGenerator_Mesh.cpp
--- a/Engine/Plugins/Mutable/Source/MutableTools/Private/MuT/CodeGenerator_Mesh.cpp
+++ b/Engine/Plugins/Mutable/Source/MutableTools/Private/MuT/CodeGenerator_Mesh.cpp
@@ -625,6 +625,50 @@ namespace mu
 	}
 
 
+	void GenerateLayoutsForMesh(const TSharedPtr<FMesh>& Mesh, const TArray<Ptr<NodeLayout>>& Layouts)
+	{
+		if (!Mesh)
+		{
+			return;
+		}
+		
+		for (const Ptr<NodeLayout>& Layout : Layouts)
+		{
+			if (!Layout)
+			{
+				continue;
+			}
+
+			// Legacy behavior
+			if ((Layout->AutoBlockStrategy == EAutoBlocksStrategy::Ignore || Layout->Strategy == EPackStrategy::Overlay) &&
+				Layout->Blocks.IsEmpty())
+			{
+				Layout->Blocks.SetNum(1);
+				Layout->Blocks[0].Min = { 0,0 };
+				Layout->Blocks[0].Size = Layout->Size;
+				Layout->Blocks[0].Priority = 0;
+				Layout->Blocks[0].bReduceBothAxes = false;
+				Layout->Blocks[0].bReduceByTwo = false;
+				continue;
+			}
+
+			UntypedMeshBufferIteratorConst TexIt(Mesh->GetVertexBuffers(), EMeshBufferSemantic::TexCoords, Layout->TexCoordsIndex);
+			if (TexIt.ptr() == nullptr)
+			{
+				continue;
+			}
+
+			if (Layout->AutoBlockStrategy == EAutoBlocksStrategy::Rectangles)
+			{
+				Layout->GenerateLayoutBlocks(Mesh, Layout->TexCoordsIndex);
+			}
+			else if (Layout->AutoBlockStrategy == EAutoBlocksStrategy::UVIslands)
+			{
+				Layout->GenerateLayoutBlocksFromUVIslands(Mesh, Layout->TexCoordsIndex);
+			}
+		}
+	}
+
 	Ptr<ASTOp> CodeGenerator::GenerateLayoutOpsAndResult(
 		const FMeshGenerationDynamicOptions& Options, 
 		Ptr<ASTOp> LastMeshOp, 
@@ -865,6 +909,11 @@ namespace mu
 
 				if (Options.bLayouts)
 				{
+					if (Options.OverrideLayouts.IsEmpty())
+					{
+						GenerateLayoutsForMesh(Mesh, NodeCopy->Layouts);
+					}
+
 					LastMeshOp = GenerateLayoutOpsAndResult(Options, LastMeshOp, NodeCopy->Layouts, MeshIDPrefix, Result);
 				}
 
@@ -1324,9 +1373,30 @@ namespace mu
 
 	FMeshTask CodeGenerator::GenerateMesh_Parameter(const FMeshGenerationStaticOptions& StaticOptions, FMeshOptionsTask InOptions, const NodeMeshParameter* Node)
 	{
+		TArray<UE::Tasks::FTask, TInlineAllocator<2>> Requisites;
+		Requisites.Emplace(InOptions);
+
+		// Reference Mesh used to generate the UV layouts.
+		TSharedPtr< TSharedPtr<FMesh> > ReferenceMesh = MakeShared<TSharedPtr<FMesh>>();
+
+		// Compile-time reference? Add task to resolve as requisite.
+		if (Node->ReferenceMesh && Node->ReferenceMesh->IsReference() && Node->ReferenceMesh->IsForceLoad())
+		{
+			uint32 MeshID = Node->ReferenceMesh->GetReferencedMesh();
+
+			bool bRunImmediatlyIfPossible = IsInGameThread();
+			FString Morph = Node->ReferenceMesh->GetReferencedMorph();
+			UE::Tasks::FTask ReferenceCompletion = CompilerOptions->OptimisationOptions.ReferencedMeshResourceProvider(MeshID, Morph, ReferenceMesh, bRunImmediatlyIfPossible);
+			Requisites.Emplace(ReferenceCompletion);
+		}
+		else
+		{
+			*ReferenceMesh = Node->ReferenceMesh;
+		}
+
 		// Local pipe because we call GenerateRange in there.
 		FMeshTask MeshParamTask = this->LocalPipe.Launch(TEXT("MutableMeshParameter"),
-			[InOptions, Node, StaticOptions, this]() mutable
+			[InOptions, Node, ReferenceMesh, StaticOptions, this]() mutable
 			{
 				FMeshGenerationDynamicOptions Options = InOptions.GetResult();
 
@@ -1352,6 +1422,12 @@ namespace mu
 				MeshID = UniqueMeshIds.EnsureUnique(MeshID);
 				Op->MeshID = MeshID;
 
+				if (Options.OverrideLayouts.IsEmpty())
+				{
+					TSharedPtr<FMesh> Mesh = *ReferenceMesh;
+					GenerateLayoutsForMesh(Mesh, Node->Layouts);
+				}
+
 				{
 					UE::TUniqueLock Lock(FirstPass.ParameterNodes.Mutex);
 
@@ -1385,7 +1461,7 @@ namespace mu
 
 				return Result;
 			},
-			InOptions
+			Requisites
 		);
 
 		// Apply the modifier for the pre-normal operations stage.
diff --git a/Engine/Plugins/Mutable/Source/MutableTools/Private/MuT/NodeLayout.cpp b/Engine/Plugins/Mutable/Source/MutableTools/Private/MuT/NodeLayout.cpp
--- a/Engine/Plugins/Mutable/Source/MutableTools/Private/MuT/NodeLayout.cpp
+++ b/Engine/Plugins/Mutable/Source/MutableTools/Private/MuT/NodeLayout.cpp
@@ -293,7 +293,7 @@ namespace mu
 	}
 
 
-	void NodeLayout::GenerateLayoutBlocksFromUVIslands(const TSharedPtr<FMesh>& Mesh, int32 LayoutIndex, bool bMergeChildBlocks)
+	void NodeLayout::GenerateLayoutBlocksFromUVIslands(const TSharedPtr<FMesh>& Mesh, int32 LayoutIndex)
 	{
 		if (!(Mesh && LayoutIndex >= 0 && Size.X > 0 && Size.Y > 0))
 		{
diff --git a/Engine/Plugins/Mutable/Source/MutableTools/Public/MuT/NodeLayout.h b/Engine/Plugins/Mutable/Source/MutableTools/Public/MuT/NodeLayout.h
--- a/Engine/Plugins/Mutable/Source/MutableTools/Public/MuT/NodeLayout.h
+++ b/Engine/Plugins/Mutable/Source/MutableTools/Public/MuT/NodeLayout.h
@@ -50,6 +50,7 @@ namespace mu
 		FIntVector2 MaxSize = FIntVector2(0, 0);
 
 		//!
+		EAutoBlocksStrategy AutoBlockStrategy = EAutoBlocksStrategy::Rectangles;
 		TArray<FSourceLayoutBlock> Blocks;
 
 		//! Packing strategy
@@ -61,6 +62,10 @@ namespace mu
 		*/
 		int32 FirstLODToIgnoreWarnings = 0;
 
+		int32 TexCoordsIndex = INDEX_NONE;
+
+		bool bMergeChildBlocks = false;
+
 	public:
 
 		// Node Interface
@@ -76,7 +81,7 @@ namespace mu
 		* A list of existing blocks may be provided to specify starting blocks that shouldn't be modified unless it is mandatory because they contain partial face islands. 
 		*/
 		UE_API void GenerateLayoutBlocks(const TSharedPtr<FMesh>&, int32 LayoutIndex );
-		UE_API void GenerateLayoutBlocksFromUVIslands(const TSharedPtr<FMesh>&, int32 LayoutIndex, bool bMergeChildBlocks);
+		UE_API void GenerateLayoutBlocksFromUVIslands(const TSharedPtr<FMesh>&, int32 LayoutIndex);
 
 	protected:
 
diff --git a/Engine/Plugins/Mutable/Source/MutableTools/Public/MuT/NodeMeshParameter.h b/Engine/Plugins/Mutable/Source/MutableTools/Public/MuT/NodeMeshParameter.h
--- a/Engine/Plugins/Mutable/Source/MutableTools/Public/MuT/NodeMeshParameter.h
+++ b/Engine/Plugins/Mutable/Source/MutableTools/Public/MuT/NodeMeshParameter.h
@@ -18,6 +18,9 @@ namespace mu
 	class NodeMeshParameter : public NodeMesh
 	{
 	public:
+		
+		TSharedPtr<FMesh> ReferenceMesh;
+
 		TArray<Ptr<NodeLayout>> Layouts;
 
 		uint8 LODIndex = 0;
