[Mutable] Add check to fix CIS warning.

#rnx

diff --git a/Engine/Plugins/Mutable/Source/MutableTools/Private/MuT/CodeGenerator_Mesh.cpp b/Engine/Plugins/Mutable/Source/MutableTools/Private/MuT/CodeGenerator_Mesh.cpp
--- a/Engine/Plugins/Mutable/Source/MutableTools/Private/MuT/CodeGenerator_Mesh.cpp
+++ b/Engine/Plugins/Mutable/Source/MutableTools/Private/MuT/CodeGenerator_Mesh.cpp
@@ -1380,18 +1380,22 @@ namespace mu
 		TSharedPtr< TSharedPtr<FMesh> > ReferenceMesh = MakeShared<TSharedPtr<FMesh>>();
 
 		// Compile-time reference? Add task to resolve as requisite.
-		if (Node->ReferenceMesh && Node->ReferenceMesh->IsReference() && Node->ReferenceMesh->IsForceLoad())
+		TSharedPtr<FMesh> Mesh = Node->ReferenceMesh;
+		if (Mesh)
 		{
-			uint32 MeshID = Node->ReferenceMesh->GetReferencedMesh();
+			if (Mesh->IsReference() && Mesh->IsForceLoad())
+			{
+				uint32 MeshID = Mesh->GetReferencedMesh();
 
-			bool bRunImmediatlyIfPossible = IsInGameThread();
-			FString Morph = Node->ReferenceMesh->GetReferencedMorph();
-			UE::Tasks::FTask ReferenceCompletion = CompilerOptions->OptimisationOptions.ReferencedMeshResourceProvider(MeshID, Morph, ReferenceMesh, bRunImmediatlyIfPossible);
-			Requisites.Emplace(ReferenceCompletion);
-		}
-		else
-		{
-			*ReferenceMesh = Node->ReferenceMesh;
+				bool bRunImmediatlyIfPossible = IsInGameThread();
+				FString Morph = Mesh->GetReferencedMorph();
+				UE::Tasks::FTask ReferenceCompletion = CompilerOptions->OptimisationOptions.ReferencedMeshResourceProvider(MeshID, Morph, ReferenceMesh, bRunImmediatlyIfPossible);
+				Requisites.Emplace(ReferenceCompletion);
+			}
+			else
+			{
+				*ReferenceMesh = Mesh;
+			}
 		}
 
 		// Local pipe because we call GenerateRange in there.
@@ -1424,8 +1428,7 @@ namespace mu
 
 				if (Options.OverrideLayouts.IsEmpty())
 				{
-					TSharedPtr<FMesh> Mesh = *ReferenceMesh;
-					GenerateLayoutsForMesh(Mesh, Node->Layouts);
+					GenerateLayoutsForMesh(*ReferenceMesh, Node->Layouts);
 				}
 
 				{
