[Mutable] Fix Material Parameter asert.

diff --git a/Engine/Plugins/Mutable/Source/CustomizableObject/Internal/MuCO/CustomizableObjectSystemPrivate.h b/Engine/Plugins/Mutable/Source/CustomizableObject/Internal/MuCO/CustomizableObjectSystemPrivate.h
--- a/Engine/Plugins/Mutable/Source/CustomizableObject/Internal/MuCO/CustomizableObjectSystemPrivate.h
+++ b/Engine/Plugins/Mutable/Source/CustomizableObject/Internal/MuCO/CustomizableObjectSystemPrivate.h
@@ -427,7 +427,7 @@ struct FInstanceUpdateData
 
 	struct FLOD
 	{		
-		mu::FMeshId MeshID;
+		mu::FMeshId MeshId;
 		TSharedPtr<const mu::FMesh> Mesh;
 
 		/** Range in the Surfaces array */
diff --git a/Engine/Plugins/Mutable/Source/CustomizableObject/Private/MuCO/CustomizableObjectInstance.cpp b/Engine/Plugins/Mutable/Source/CustomizableObject/Private/MuCO/CustomizableObjectInstance.cpp
--- a/Engine/Plugins/Mutable/Source/CustomizableObject/Private/MuCO/CustomizableObjectInstance.cpp
+++ b/Engine/Plugins/Mutable/Source/CustomizableObject/Private/MuCO/CustomizableObjectInstance.cpp
@@ -4926,8 +4926,13 @@ UE::Tasks::FTask UCustomizableInstancePrivate::LoadAdditionalAssetsAndData(const
 		{
 			continue;
 		}
+
+		if (Surface.Material->ReferenceID == INDEX_NONE)
+		{
+			continue;
+		}
 		
-		const int32 MaterialIndex = Surface.Material->MaterialId;
+		const int32 MaterialIndex = Surface.Material->ReferenceID;
 		if (ObjectToInstanceIndexMap.Contains(MaterialIndex))
 		{
 			continue;
@@ -4951,14 +4956,19 @@ UE::Tasks::FTask UCustomizableInstancePrivate::LoadAdditionalAssetsAndData(const
 		{
 			continue;
 		}
+
+		if (Component.OverlayMaterial->ReferenceID == INDEX_NONE)
+		{
+			continue;
+		}
 		
-		const int32 MaterialIndex = Component.OverlayMaterial->MaterialId;
-		if (ObjectToInstanceIndexMap.Contains(Component.OverlayMaterial->MaterialId))
+		const int32 MaterialIndex = Component.OverlayMaterial->ReferenceID;
+		if (ObjectToInstanceIndexMap.Contains(Component.OverlayMaterial->ReferenceID))
 		{
 			continue;
 		}
 
-		TSoftObjectPtr<UMaterialInterface> AssetPtr = ModelResources.Materials.IsValidIndex(MaterialIndex) ? ModelResources.Materials[Component.OverlayMaterial->MaterialId] : nullptr;
+		TSoftObjectPtr<UMaterialInterface> AssetPtr = ModelResources.Materials.IsValidIndex(MaterialIndex) ? ModelResources.Materials[Component.OverlayMaterial->ReferenceID] : nullptr;
 		UMaterialInterface* LoadedMaterial = AssetPtr.Get();
 
 		const int32 ReferencedMaterialsIndex = ReferencedMaterials.Add(LoadedMaterial);
@@ -5830,9 +5840,9 @@ void UCustomizableInstancePrivate::BuildMaterials(const TSharedRef<FUpdateContex
 
 			if (Component.OverlayMaterial)
 			{
-				if (Component.OverlayMaterial->MaterialId != INDEX_NONE)
+				if (Component.OverlayMaterial->ReferenceID != INDEX_NONE)
 				{
-					if (const uint32* ReferencedMaterialIndex = ObjectToInstanceIndexMap.Find(Component.OverlayMaterial->MaterialId))
+					if (const uint32* ReferencedMaterialIndex = ObjectToInstanceIndexMap.Find(Component.OverlayMaterial->ReferenceID))
 					{
 						if (ReferencedMaterials.IsValidIndex(*ReferencedMaterialIndex))
 						{
@@ -5914,11 +5924,11 @@ void UCustomizableInstancePrivate::BuildMaterials(const TSharedRef<FUpdateContex
 
 				UMaterialInterface* MaterialTemplate = nullptr;
 
-				if (Surface.Material->MaterialId == INDEX_NONE)
+				if (Surface.Material->ReferenceID == INDEX_NONE)
 				{
 					MaterialTemplate = Surface.Material->Material.Get();
 				}
-				else if (const uint32* ReferencedMaterialIndex = ObjectToInstanceIndexMap.Find(Surface.Material->MaterialId))
+				else if (const uint32* ReferencedMaterialIndex = ObjectToInstanceIndexMap.Find(Surface.Material->ReferenceID))
 				{
 					MaterialTemplate = ReferencedMaterials[*ReferencedMaterialIndex];
 				}
diff --git a/Engine/Plugins/Mutable/Source/CustomizableObject/Private/MuCO/CustomizableObjectSkeletalMesh.cpp b/Engine/Plugins/Mutable/Source/CustomizableObject/Private/MuCO/CustomizableObjectSkeletalMesh.cpp
--- a/Engine/Plugins/Mutable/Source/CustomizableObject/Private/MuCO/CustomizableObjectSkeletalMesh.cpp
+++ b/Engine/Plugins/Mutable/Source/CustomizableObject/Private/MuCO/CustomizableObjectSkeletalMesh.cpp
@@ -43,7 +43,7 @@ void UCustomizableObjectSkeletalMesh::InitMutableStreamingData(const TSharedRef<
 	for (int32 LODIndex = FirstLOD; LODIndex < LODCount; ++LODIndex)
 	{
 		const FInstanceUpdateData::FLOD& LOD = InOperationData->InstanceUpdateData.LODs[InstanceUpdateFirstLOD + LODIndex];
-		MeshIDs[LODIndex] = LOD.MeshID;
+		MeshIDs[LODIndex] = LOD.MeshId;
 
 		SurfaceIDs[LODIndex].SetNum(LOD.SurfaceCount);
 		check(LOD.Mesh && LOD.Mesh->GetSurfaceCount() == LOD.SurfaceCount);
diff --git a/Engine/Plugins/Mutable/Source/CustomizableObject/Private/MuCO/CustomizableObjectSystem.cpp b/Engine/Plugins/Mutable/Source/CustomizableObject/Private/MuCO/CustomizableObjectSystem.cpp
--- a/Engine/Plugins/Mutable/Source/CustomizableObject/Private/MuCO/CustomizableObjectSystem.cpp
+++ b/Engine/Plugins/Mutable/Source/CustomizableObject/Private/MuCO/CustomizableObjectSystem.cpp
@@ -2756,7 +2756,7 @@ namespace impl
 					
 					FInstanceUpdateData::FLOD& LOD = OperationData->InstanceUpdateData.LODs[Component.FirstLOD + LODIndex];
 					
-					LOD.MeshID = Instance->GetMeshId(InstanceComponentIndex, LODIndex);
+					LOD.MeshId = Instance->GetMeshId(InstanceComponentIndex, LODIndex);
 					
 					mu::EMeshContentFlags MeshContentFilter = mu::EMeshContentFlags::AllFlags; 
 					
@@ -2768,7 +2768,7 @@ namespace impl
 						}
 					}
 					
-					GetMeshesData->Add({ Component.FirstLOD + LODIndex, LOD.MeshID, MeshContentFilter });
+					GetMeshesData->Add({ Component.FirstLOD + LODIndex, LOD.MeshId, MeshContentFilter });
 				}
 			}
 
diff --git a/Engine/Plugins/Mutable/Source/MutableRuntime/Private/MuR/CodeRunner.cpp b/Engine/Plugins/Mutable/Source/MutableRuntime/Private/MuR/CodeRunner.cpp
--- a/Engine/Plugins/Mutable/Source/MutableRuntime/Private/MuR/CodeRunner.cpp
+++ b/Engine/Plugins/Mutable/Source/MutableRuntime/Private/MuR/CodeRunner.cpp
@@ -1175,7 +1175,7 @@ namespace mu
 			OP::ResourceConstantArgs Args = Program.GetOpArgs<OP::ResourceConstantArgs>(Item.At);
 
 			TSharedPtr<FMaterial> Source = MakeShared<FMaterial>();
-			Source->MaterialId = Args.value;
+			Source->ReferenceID = Args.value;
 
 			StoreMaterial(Item, Source);
 			break;
diff --git a/Engine/Plugins/Mutable/Source/MutableRuntime/Private/MuR/Material.cpp b/Engine/Plugins/Mutable/Source/MutableRuntime/Private/MuR/Material.cpp
--- a/Engine/Plugins/Mutable/Source/MutableRuntime/Private/MuR/Material.cpp
+++ b/Engine/Plugins/Mutable/Source/MutableRuntime/Private/MuR/Material.cpp
@@ -16,13 +16,13 @@ namespace mu
 		TSharedPtr<FMaterial> Result = MakeShared<FMaterial>();
 
 		Result->Material = Material;
-		Result->MaterialId = MaterialId;
+		Result->ReferenceID = ReferenceID;
 
 		return Result;
 	}
 
 	bool FMaterial::operator==(const FMaterial& Other) const
 	{
-		return MaterialId == Other.MaterialId && Material == Other.Material;
+		return ReferenceID == Other.ReferenceID && Material == Other.Material;
 	};
 }
diff --git a/Engine/Plugins/Mutable/Source/MutableRuntime/Public/MuR/Material.h b/Engine/Plugins/Mutable/Source/MutableRuntime/Public/MuR/Material.h
--- a/Engine/Plugins/Mutable/Source/MutableRuntime/Public/MuR/Material.h
+++ b/Engine/Plugins/Mutable/Source/MutableRuntime/Public/MuR/Material.h
@@ -35,11 +35,9 @@ namespace mu
 		//-----------------------------------------------------------------------------------------
 		UE_API bool operator==(const FMaterial& Other) const;
 
-	public:
-
 		TStrongObjectPtr<UMaterialInterface> Material;
 
-		int32 MaterialId = INDEX_NONE;
+		int32 ReferenceID = INDEX_NONE;
 	};
 }
 
