[Mutable] Fixed Default Values not exposed to Blueprint.

diff --git a/Engine/Plugins/Mutable/Source/CustomizableObject/Public/MuCO/CustomizableObject.h b/Engine/Plugins/Mutable/Source/CustomizableObject/Public/MuCO/CustomizableObject.h
--- a/Engine/Plugins/Mutable/Source/CustomizableObject/Public/MuCO/CustomizableObject.h
+++ b/Engine/Plugins/Mutable/Source/CustomizableObject/Public/MuCO/CustomizableObject.h
@@ -413,14 +413,17 @@ public:
 	
 	/** Get the default value of a parameter of type Texture.
 	  * @param InParameterName The name of the Texture parameter to get the default value of. */
+	UFUNCTION(BlueprintCallable, Category = CustomizableObject)
 	CUSTOMIZABLEOBJECT_API UTexture* GetTextureParameterDefaultValue(const FString& InParameterName) const;
 
 	/** Get the default value of a parameter of type Skeletal Mesh.
 	  * @param ParameterName The name of the Mesh parameter to get the default value of. */
+	UFUNCTION(BlueprintCallable, Category = CustomizableObject)
 	CUSTOMIZABLEOBJECT_API USkeletalMesh* GetSkeletalMeshParameterDefaultValue(const FString& ParameterName) const;
 	
 	/** Get the default value of a parameter of type Material.
 	* @param ParameterName The name of the Material parameter to get the default value of. */
+	UFUNCTION(BlueprintCallable, Category = CustomizableObject)
 	CUSTOMIZABLEOBJECT_API UMaterialInterface* GetMaterialParameterDefaultValue(const FString& ParameterName) const;
 
 	/** Return true if the parameter at the index provided is multidimensional.
