[Mutable] Remove GT syncronizations when converting Skeletal Mesh parameters.

- ConvertSkeletalMeshFromRuntimeData no longer syncs with the GT.
- Moved FUnrealMutableResourceProvider from System to the CO. Now is a SharedPtr. This simplifies the provider itself since now all assets come from a single CO.
- Renamed FMutableUpdateContext to FMutableUpdateImageContext.
- Renamed FMutableMeshOperationData to FMutableUpdateMeshContext.
- Removed encapsulation from FMutableUpdateImageContext. It is private and was not bringing anything to the table.
- Created a new ReuseInstance function. This allows to add members to the Instance instead of passing them each time we call a System function.
- Created FBoneNames which is modified when converting Skeletal Meshes at runtime. Thread safe.

diff --git a/Engine/Plugins/Mutable/Source/CustomizableObject/Internal/MuCO/BoneNames.h b/Engine/Plugins/Mutable/Source/CustomizableObject/Internal/MuCO/BoneNames.h
new file mode 100644
--- /dev/null
+++ b/Engine/Plugins/Mutable/Source/CustomizableObject/Internal/MuCO/BoneNames.h
@@ -0,0 +1,24 @@
+﻿// Copyright Epic Games, Inc. All Rights Reserved.
+
+#pragma once
+
+#include "MuR/Skeleton.h"
+
+class UModelResources;
+
+
+class FBoneNames
+{
+public:
+	CUSTOMIZABLEOBJECT_API FBoneNames(const UModelResources& ModelResources);
+
+	mu::FBoneName* Find(const FName& BoneName);
+	
+	mu::FBoneName FindOrAdd(const FName& BoneName);
+	
+private:
+	TMap<FString, mu::FBoneName> BoneNamesMap;
+
+	FCriticalSection CriticalSection;
+};
+
diff --git a/Engine/Plugins/Mutable/Source/CustomizableObject/Internal/MuCO/CustomizableObjectMipDataProvider.h b/Engine/Plugins/Mutable/Source/CustomizableObject/Internal/MuCO/CustomizableObjectMipDataProvider.h
--- a/Engine/Plugins/Mutable/Source/CustomizableObject/Internal/MuCO/CustomizableObjectMipDataProvider.h
+++ b/Engine/Plugins/Mutable/Source/CustomizableObject/Internal/MuCO/CustomizableObjectMipDataProvider.h
@@ -28,43 +28,14 @@ struct FModelStreamableBulkData;
 /** This struct stores the data relevant for the construction of a specific texture. 
 * This includes all the data required to rebuild the image (or any of its mips).
 */
-class FMutableUpdateContext
+class FMutableUpdateImageContext
 {
 public:
-	FMutableUpdateContext() = default;
-	
-	FMutableUpdateContext(const FString& InCustomizableObjectPathName,
-		const FString& InInstancePathName,
-		const TSharedPtr<mu::FSystem>& InSystem,
-		const TSharedPtr<mu::FModel>& InModel,
-		const TSharedPtr<mu::FMeshIdRegistry>& InMeshIdRegistry,
-		const TSharedPtr<mu::FImageIdRegistry>& InImageIdRegistry,
-		const TSharedPtr<mu::FMaterialIdRegistry>& InMaterialIdRegistry,
-		const TSharedPtr<FModelStreamableBulkData>& InModelStreamable,
-		const TSharedPtr<const mu::FParameters>& InParameters,
-		int32 InState);
-	
-	const FString& GetCustomizableObjectPathName() const;
-
-	const FString& GetInstancePathName() const;
-	
-	TSharedPtr<mu::FSystem> GetSystem() const;
-	
-	TSharedPtr<mu::FModel, ESPMode::ThreadSafe> GetModel() const;
-
-	TSharedPtr<FModelStreamableBulkData> GetModelStreamableBulkData();
-
-	TSharedPtr<const mu::FParameters> GetParameters() const;
-	
-	int32 GetState() const;
-
-	const TArray<TSharedPtr<const mu::FImage>>& GetImageParameterValues() const;
 
 	// Benchmarking Utility data (it may not always be present)
 	FString CapturedDescriptor;
 	bool bLevelBegunPlay = false;
 	
-private:
 	FString CustomizableObjectPathName;
 	FString InstancePathName;
 
@@ -73,12 +44,12 @@ private:
 	TSharedPtr<mu::FSystem> System;
 	TSharedPtr<mu::FModel> Model;
 
-public:
 	TSharedPtr<mu::FMeshIdRegistry> MeshIdRegistry;
 	TSharedPtr<mu::FImageIdRegistry> ImageIdRegistry;
 	TSharedPtr<mu::FMaterialIdRegistry> MaterialIdRegistry;
 
-private:
+	TSharedPtr<FUnrealMutableResourceProvider> ExternalResourceProvider;
+
 	TSharedPtr<FModelStreamableBulkData> ModelStreamableBulkData;
 	
 	TSharedPtr<const mu::FParameters> Parameters;
@@ -120,7 +91,7 @@ struct FMutableImageOperationData
 	int32 MipsToSkip = 0;
 	FMutableImageReference RequestedImage;
 
-	TSharedPtr<FMutableUpdateContext> UpdateContext;
+	TSharedPtr<FMutableUpdateImageContext> UpdateContext;
 
 	TSharedPtr<const mu::FImage> Result;
 
@@ -167,7 +138,7 @@ public:
 	UCustomizableObjectInstance* CustomizableObjectInstance = nullptr;
 
 	FMutableImageReference ImageRef;
-	TSharedPtr<FMutableUpdateContext> UpdateContext;
+	TSharedPtr<FMutableUpdateImageContext> UpdateContext;
 
 	bool bRequestAborted = false;
 
@@ -203,6 +174,6 @@ public:
 	TObjectPtr<UCustomizableObjectInstance> CustomizableObjectInstance = nullptr;
 
 	FMutableImageReference ImageRef;
-	TSharedPtr<FMutableUpdateContext> UpdateContext;
+	TSharedPtr<FMutableUpdateImageContext> UpdateContext;
 
 };
diff --git a/Engine/Plugins/Mutable/Source/CustomizableObject/Internal/MuCO/CustomizableObjectSystemPrivate.h b/Engine/Plugins/Mutable/Source/CustomizableObject/Internal/MuCO/CustomizableObjectSystemPrivate.h
--- a/Engine/Plugins/Mutable/Source/CustomizableObject/Internal/MuCO/CustomizableObjectSystemPrivate.h
+++ b/Engine/Plugins/Mutable/Source/CustomizableObject/Internal/MuCO/CustomizableObjectSystemPrivate.h
@@ -36,6 +36,8 @@
 
 #define UE_API CUSTOMIZABLEOBJECT_API
 
+class FBoneNames;
+class FUnrealMutableResourceProvider;
 class FMutableStreamRequest;
 class UEditorImageProvider;
 class UCustomizableObjectSystem;
@@ -625,6 +627,10 @@ public:
 	TSharedPtr<mu::FImageIdRegistry> ImageIdRegistry;
 	TSharedPtr<mu::FMaterialIdRegistry> MaterialIdRegistry;
 
+	TSharedPtr<FUnrealMutableResourceProvider> ExternalResourceProvider;
+
+	TSharedPtr<FBoneNames> BoneNames;
+	
 	// Number of possible components in the entire CO
 	// TODO: Redundant because it is in the CO
 	uint8 NumObjectComponents = 0;
@@ -816,11 +822,6 @@ public:
 	// This object is responsible for streaming data to the MutableSystem.
 	TSharedPtr<class FUnrealMutableModelBulkReader> Streamer;
 	
-	// This object is responsible for providing custom images and meshes to mutable (for image parameters, etc.)
-	// This object is called from the mutable thread, and it should only access data already safely submitted from
-	// the game thread and stored in FUnrealMutableImageProvider::GlobalExternalImages.
-	TSharedPtr<class FUnrealMutableResourceProvider> ResourceProvider;
-	
 	// The pending instance updates, discards or releases
 	FMutablePendingInstanceWork MutablePendingInstanceWork;
 	
@@ -889,9 +890,7 @@ public:
 	void UpdateMemoryLimit();
 
 	bool IsMutableAnimInfoDebuggingEnabled() const;
-
-	FUnrealMutableResourceProvider* GetResourceProviderChecked() const;
-
+	
 	/** Start the actual work of Update Skeletal Mesh process (Update Skeletal Mesh without the queue). */
 	void StartUpdateSkeletalMesh(const TSharedRef<FUpdateContextPrivate>& Context);
 
diff --git a/Engine/Plugins/Mutable/Source/CustomizableObject/Internal/MuCO/UnrealConversionUtils.h b/Engine/Plugins/Mutable/Source/CustomizableObject/Internal/MuCO/UnrealConversionUtils.h
--- a/Engine/Plugins/Mutable/Source/CustomizableObject/Internal/MuCO/UnrealConversionUtils.h
+++ b/Engine/Plugins/Mutable/Source/CustomizableObject/Internal/MuCO/UnrealConversionUtils.h
@@ -13,6 +13,7 @@
 #include "Containers/Map.h"
 #include "Tasks/Task.h"
 
+class FBoneNames;
 struct FClothBufferIndexMapping;
 struct FMeshToMeshVertData;
 struct FCustomizableObjectMeshToMeshVertData;
@@ -155,7 +156,7 @@ namespace UnrealConversionUtils
 
 	void MorphTargetVertexInfoBuffers(FSkeletalMeshLODRenderData& LODResource, const USkeletalMesh& Owner, const mu::FMesh& MutableMesh, const TMap<uint32, FMorphTargetMeshData>& MorphTargetMeshData, int32 LODIndex);
 
-	CUSTOMIZABLEOBJECT_API UE::Tasks::FTask ConvertSkeletalMeshFromRuntimeData(USkeletalMesh* SkeletalMesh, int32 LODIndex, int32 SectionIndex, UModelResources* ModelResources, mu::FMesh* Result, TFunction<void(bool bSuccess)>& CompletionCallback);
+	CUSTOMIZABLEOBJECT_API UE::Tasks::FTask ConvertSkeletalMeshFromRuntimeData(USkeletalMesh* SkeletalMesh, int32 LODIndex, int32 SectionIndex, TSharedPtr<mu::FMesh> Result, const TSharedRef<FBoneNames>& BoneNames);
 
 	void GetSectionClothData(const mu::FMesh& MutableMesh, int32 LODIndex, const TMap<uint32, FClothingMeshData>& ClothingMeshData, TArray<FSectionClothData>& SectionsClothData, int32& NumClothingDataNotFound);
 	
diff --git a/Engine/Plugins/Mutable/Source/CustomizableObject/Internal/MuCO/UnrealMutableImageProvider.h b/Engine/Plugins/Mutable/Source/CustomizableObject/Internal/MuCO/UnrealMutableImageProvider.h
new file mode 100644
--- /dev/null
+++ b/Engine/Plugins/Mutable/Source/CustomizableObject/Internal/MuCO/UnrealMutableImageProvider.h
@@ -0,0 +1,48 @@
+// Copyright Epic Games, Inc. All Rights Reserved.
+
+#pragma once
+
+#include "MuCO/BoneNames.h"
+#include "MuR/Image.h"
+#include "MuR/System.h"
+#include "Tasks/Task.h"
+
+class UModelResources;
+class UTexture;
+class USkeletalMesh;
+
+/** Implementation of a mutable core provider for image parameters that are application-specific. */
+class FUnrealMutableResourceProvider : public mu::FExternalResourceProvider
+{
+public:
+	CUSTOMIZABLEOBJECT_API FUnrealMutableResourceProvider(const TSharedRef<FBoneNames>& InBoneNames);
+	
+	// mu::ExternalResourceProvider interface
+	// Thread: worker
+	CUSTOMIZABLEOBJECT_API virtual TTuple<UE::Tasks::FTask, TFunction<void()>> GetImageAsync(UTexture* Texture, uint8 MipmapsToSkip, TFunction<void(TSharedPtr<mu::FImage>)>& ResultCallback) override;
+	CUSTOMIZABLEOBJECT_API virtual TTuple<UE::Tasks::FTask, TFunction<void()>> GetReferencedImageAsync( int32 Id, uint8 MipmapsToSkip, TFunction<void(TSharedPtr<mu::FImage>)>& ResultCallback) override;
+	CUSTOMIZABLEOBJECT_API virtual mu::FExtendedImageDesc GetImageDesc(UTexture* Texture) override;
+	CUSTOMIZABLEOBJECT_API virtual TTuple<UE::Tasks::FTask, TFunction<void()>> GetMeshAsync(USkeletalMesh* SkeletalMesh, int32 LODIndex, int32 SectionIndex, TFunction<void(TSharedPtr<mu::FMesh>)>& ResultCallback) override;
+
+#if WITH_EDITOR
+	void CacheRuntimeReferencedImages(const TArray<TSoftObjectPtr<UTexture2D>>& RuntimeReferencedTextures);
+#endif
+	
+private:
+	static inline const mu::FImageDesc DUMMY_IMAGE_DESC = 
+			mu::FImageDesc {mu::FImageSize(32, 32), mu::EImageFormat::RGBA_UByte, 1};
+
+	/** This will be called if an image Id has been requested by Mutable core but it has not been provided by any provider. */
+	static TSharedPtr<mu::FImage> CreateDummy();
+	static mu::FExtendedImageDesc CreateDummyDesc();
+
+#if WITH_EDITOR
+	TArray<TStrongObjectPtr<UTexture>> Images;
+#endif
+
+#if WITH_EDITOR
+	FCriticalSection RuntimeReferencedLock;
+#endif
+
+	TSharedRef<FBoneNames> BoneNames;
+};
diff --git a/Engine/Plugins/Mutable/Source/CustomizableObject/Private/BoneNames.cpp b/Engine/Plugins/Mutable/Source/CustomizableObject/Private/BoneNames.cpp
new file mode 100644
--- /dev/null
+++ b/Engine/Plugins/Mutable/Source/CustomizableObject/Private/BoneNames.cpp
@@ -0,0 +1,65 @@
+﻿// Copyright Epic Games, Inc. All Rights Reserved.
+
+#include "MuCO/BoneNames.h"
+
+#include "MuCO/CustomizableObjectPrivate.h"
+
+
+FBoneNames::FBoneNames(const UModelResources& ModelResources)
+{
+	BoneNamesMap.Reserve(ModelResources.BoneNamesMap.Num());
+
+	for (const TPair<FString, uint32>& Pair : ModelResources.BoneNamesMap)
+	{
+		BoneNamesMap.Add(Pair.Key, Pair.Value);
+	}
+}
+
+
+mu::FBoneName* FBoneNames::Find(const FName& BoneName)
+{
+	FScopeLock Lock(&CriticalSection);
+
+	return BoneNamesMap.Find(BoneName.ToString().ToLower());
+}
+
+
+mu::FBoneName FBoneNames::FindOrAdd(const FName& BoneName)
+{
+	FScopeLock Lock(&CriticalSection);
+
+	const FString BoneNameString = BoneName.ToString().ToLower();
+
+	if (mu::FBoneName* Result = BoneNamesMap.Find(BoneNameString))
+	{
+		return *Result;
+	}
+
+	// Go the slow way and add it to get a unique ID that the core can work with.
+	uint32 NewBoneId = CityHash32(reinterpret_cast<const char*>(*BoneNameString), BoneNameString.Len() * sizeof(FString::ElementType));
+
+	// See if the hash collides with an existing bone
+	bool bUnique = false;
+	while (!bUnique)
+	{
+		bUnique = true;
+		TMap<FString, mu::FBoneName>::TConstIterator MapIterator = BoneNamesMap.CreateConstIterator();
+		while (MapIterator)
+		{
+			if (MapIterator.Value() == NewBoneId)
+			{
+				bUnique = false;
+				break;
+			}
+			++MapIterator;
+		}
+
+		if (!bUnique)
+		{
+			NewBoneId++;
+		}
+	}
+
+	BoneNamesMap.Add(BoneNameString, NewBoneId);
+	return NewBoneId;
+}
diff --git a/Engine/Plugins/Mutable/Source/CustomizableObject/Private/MuCO/CustomizableObjectInstance.cpp b/Engine/Plugins/Mutable/Source/CustomizableObject/Private/MuCO/CustomizableObjectInstance.cpp
--- a/Engine/Plugins/Mutable/Source/CustomizableObject/Private/MuCO/CustomizableObjectInstance.cpp
+++ b/Engine/Plugins/Mutable/Source/CustomizableObject/Private/MuCO/CustomizableObjectInstance.cpp
@@ -35,6 +35,7 @@
 #include "MuCO/UnrealConversionUtils.h"
 #include "MuCO/UnrealPortabilityHelpers.h"
 #include "MuCO/LogBenchmarkUtil.h"
+#include "MuCO/BoneNames.h"
 
 #include "PhysicsEngine/PhysicsAsset.h"
 #include "Rendering/Texture2DResource.h"
@@ -54,6 +55,7 @@
 #include "MuCO/LoadUtils.h"
 #include "Serialization/BulkData.h"
 #include "HAL/PlatformFileManager.h"
+#include "MuCO/UnrealMutableImageProvider.h"
 
 #include "MuCO/Plugins/IMutableClothingModule.h"
 #include "Rendering/SkeletalMeshModel.h"
@@ -3840,16 +3842,13 @@ bool UCustomizableInstancePrivate::BuildSkeletonData(const TSharedRef<FUpdateCon
 		MUTABLE_CPUPROFILER_SCOPE(BuildSkeletonData_BuildBoneInfoMap);
 
 		BoneInfoMap.Reserve(RawRefBoneCount);
-
-		const UModelResources& ModelResources = CustomizableObject.GetPrivate()->GetModelResourcesChecked();
+		
 		for (int32 Index = 0; Index < RawRefBoneCount; ++Index)
 		{
 			const FName BoneName = RawRefBoneInfo[Index].Name;
-			const FString BoneNameString = BoneName.ToString().ToLower();
-			if(const uint32* Hash = ModelResources.BoneNamesMap.Find(BoneNameString))
+			if(const mu::FBoneName* Bone = OperationData->BoneNames->Find(BoneName))
 			{
-				const mu::FBoneName Bone(*Hash);
-				TPair<FName, uint16>& BoneInfo = BoneInfoMap.Add(Bone);
+				TPair<FName, uint16>& BoneInfo = BoneInfoMap.Add(*Bone);
 				BoneInfo.Key = BoneName;
 				BoneInfo.Value = Index;
 			}
@@ -5774,17 +5773,18 @@ void UCustomizableInstancePrivate::BuildMaterials(const TSharedRef<FUpdateContex
 	Exchange(OldGeneratedMaterials, GeneratedMaterials);
 	
 	// Prepare the data to store in order to regenerate resources for this instance (usually texture mips).
-	TSharedPtr<FMutableUpdateContext> UpdateContext = MakeShared<FMutableUpdateContext>(
-		CustomizableObject->GetPathName(),
-		Public->GetPathName(),
-		UCustomizableObjectSystem::GetInstance()->GetPrivate()->MutableSystem,
-		OperationData->Model,
-		CustomizableObject->GetPrivate()->MeshIdRegistry,
-	    CustomizableObject->GetPrivate()->ImageIdRegistry,
-	    CustomizableObject->GetPrivate()->MaterialIdRegistry,
-		CustomizableObject->GetPrivate()->GetModelStreamableBulkData(),
-		OperationData->Parameters,
-	    OperationData->GetCapturedDescriptor().GetState());
+	TSharedPtr<FMutableUpdateImageContext> UpdateContext = MakeShared<FMutableUpdateImageContext>();
+	UpdateContext->CustomizableObjectPathName = CustomizableObject->GetPathName();
+	UpdateContext->InstancePathName = Public->GetPathName();
+	UpdateContext->System = UCustomizableObjectSystem::GetInstance()->GetPrivate()->MutableSystem;
+	UpdateContext->Model = OperationData->Model;
+	UpdateContext->MeshIdRegistry = CustomizableObject->GetPrivate()->MeshIdRegistry;
+	UpdateContext->ImageIdRegistry = CustomizableObject->GetPrivate()->ImageIdRegistry;
+    UpdateContext->MaterialIdRegistry = CustomizableObject->GetPrivate()->MaterialIdRegistry,
+	UpdateContext->ExternalResourceProvider = OperationData->ExternalResourceProvider;
+	UpdateContext->ModelStreamableBulkData = CustomizableObject->GetPrivate()->GetModelStreamableBulkData();
+	UpdateContext->Parameters = OperationData->Parameters;
+	UpdateContext->State = OperationData->GetCapturedDescriptor().GetState();
 
 	// Cache the descriptor as a string if we want to later report it using our benchmark utility. 
 	if (FLogBenchmarkUtil::IsBenchmarkingReportingEnabled())
diff --git a/Engine/Plugins/Mutable/Source/CustomizableObject/Private/MuCO/CustomizableObjectMeshUpdate.cpp b/Engine/Plugins/Mutable/Source/CustomizableObject/Private/MuCO/CustomizableObjectMeshUpdate.cpp
--- a/Engine/Plugins/Mutable/Source/CustomizableObject/Private/MuCO/CustomizableObjectMeshUpdate.cpp
+++ b/Engine/Plugins/Mutable/Source/CustomizableObject/Private/MuCO/CustomizableObjectMeshUpdate.cpp
@@ -6,6 +6,7 @@ CustomizableObjectMeshUpdate.cpp: Helpers to stream in CustomizableObject skelet
 
 
 #include "MuCO/CustomizableObjectMeshUpdate.h"
+
 #include "Components/SkinnedMeshComponent.h"
 
 #include "MutableStreamRequest.h"
@@ -13,6 +14,7 @@ CustomizableObjectMeshUpdate.cpp: Helpers to stream in CustomizableObject skelet
 #include "MuCO/CustomizableObjectSystemPrivate.h"
 #include "MuCO/CustomizableObjectSkeletalMesh.h"
 #include "MuCO/UnrealConversionUtils.h"
+#include "MuCO/UnrealMutableImageProvider.h"
 
 #include "MuR/Model.h"
 #include "MuR/MeshBufferSet.h"
@@ -40,7 +42,7 @@ FCustomizableObjectMeshStreamIn::FCustomizableObjectMeshStreamIn(
 	const TSharedPtr<FModelStreamableBulkData>& ModelStreamableBulkData) :
 	FSkeletalMeshStreamIn(InMesh, CreateResourcesThread)
 {
-	OperationData = MakeShared<FMutableMeshOperationData>();
+	OperationData = MakeShared<FMutableUpdateMeshContext>();
 
 	// This must run in the mutable thread.
 	check(UCustomizableObjectSystem::IsCreated());
@@ -51,6 +53,7 @@ FCustomizableObjectMeshStreamIn::FCustomizableObjectMeshStreamIn(
 	OperationData->MeshIdRegistry = InMesh->MeshIdRegistry;
 	OperationData->ImageIdRegistry = InMesh->ImageIdRegistry;
 	OperationData->MaterialIdRegistry = InMesh->MaterialIdRegistry;
+	OperationData->ExternalResourceProvider = InMesh->ExternalResourceProvider;
 	OperationData->State = InMesh->State;
 
 	OperationData->ModelStreamableBulkData = ModelStreamableBulkData;
@@ -167,7 +170,7 @@ void FCustomizableObjectMeshStreamIn::DoCancelMeshUpdate(const FContext& Context
 
 namespace impl
 {
-	void Task_Mutable_UpdateMesh_End(const TSharedPtr<FMutableMeshOperationData> OperationData, TRefCountPtr<FCustomizableObjectMeshStreamIn>& Task, mu::FInstance::FID InstanceID)
+	void Task_Mutable_UpdateMesh_End(const TSharedPtr<FMutableUpdateMeshContext> OperationData, TRefCountPtr<FCustomizableObjectMeshStreamIn>& Task, mu::FInstance::FID InstanceID)
 	{
 		MUTABLE_CPUPROFILER_SCOPE(Task_Mutable_UpdateMesh_End);
 
@@ -186,7 +189,7 @@ namespace impl
 	}
 
 	void Task_Mutable_UpdateMesh_Loop(
-		const TSharedPtr<FMutableMeshOperationData> OperationData,
+		const TSharedPtr<FMutableUpdateMeshContext> OperationData,
 		TRefCountPtr<FCustomizableObjectMeshStreamIn>& Task,
 		mu::FInstance::FID InstanceID,
 		int32 LODIndex)
@@ -214,7 +217,7 @@ namespace impl
 			LowLevelTasks::ETaskPriority::Inherit));
 	}
 
-	void Task_Mutable_UpdateMesh(const TSharedPtr<FMutableMeshOperationData> OperationData, TRefCountPtr<FCustomizableObjectMeshStreamIn>& Task)
+	void Task_Mutable_UpdateMesh(const TSharedPtr<FMutableUpdateMeshContext> OperationData, TRefCountPtr<FCustomizableObjectMeshStreamIn>& Task)
 	{
 		MUTABLE_CPUPROFILER_SCOPE(Task_Mutable_UpdateMesh);
 
@@ -239,7 +242,7 @@ namespace impl
 #endif
 
 		// For now, we are forcing the recreation of mutable-side instances with every update.
-		mu::FInstance::FID InstanceID = System->NewInstance(Model);
+		mu::FInstance::FID InstanceID = System->NewInstance(Model, OperationData->ExternalResourceProvider);
 		UE_LOG(LogMutable, Verbose, TEXT("Creating Mutable instance with id [%d] for a mesh update"), InstanceID);
 
 		// LOD mask, set to all ones to build all LODs
diff --git a/Engine/Plugins/Mutable/Source/CustomizableObject/Private/MuCO/CustomizableObjectMeshUpdate.h b/Engine/Plugins/Mutable/Source/CustomizableObject/Private/MuCO/CustomizableObjectMeshUpdate.h
--- a/Engine/Plugins/Mutable/Source/CustomizableObject/Private/MuCO/CustomizableObjectMeshUpdate.h
+++ b/Engine/Plugins/Mutable/Source/CustomizableObject/Private/MuCO/CustomizableObjectMeshUpdate.h
@@ -22,13 +22,14 @@ namespace mu
 	class FParameters;
 }
 
-struct FMutableMeshOperationData
+struct FMutableUpdateMeshContext
 {
 	TSharedPtr<mu::FSystem> System;
-	TSharedPtr<mu::FModel, ESPMode::ThreadSafe> Model;
+	TSharedPtr<mu::FModel> Model;
 	TSharedPtr<mu::FMeshIdRegistry> MeshIdRegistry;
 	TSharedPtr<mu::FImageIdRegistry> ImageIdRegistry;
 	TSharedPtr<mu::FMaterialIdRegistry> MaterialIdRegistry;
+	TSharedPtr<FUnrealMutableResourceProvider> ExternalResourceProvider;
 	TSharedPtr<const mu::FParameters> Parameters;
 
 	TSharedPtr<FModelStreamableBulkData> ModelStreamableBulkData;
@@ -101,7 +102,7 @@ private:
 	void MarkRenderStateDirty(const FContext& Context);
 
 	/** Context of the MeshUpdate. */
-	TSharedPtr<FMutableMeshOperationData> OperationData;
+	TSharedPtr<FMutableUpdateMeshContext> OperationData;
 
 	/** MeshUpdate task Id to cancel the task if the stream in task is aborted. */
 	uint32 MutableTaskId = 0;
diff --git a/Engine/Plugins/Mutable/Source/CustomizableObject/Private/MuCO/CustomizableObjectMipDataProvider.cpp b/Engine/Plugins/Mutable/Source/CustomizableObject/Private/MuCO/CustomizableObjectMipDataProvider.cpp
--- a/Engine/Plugins/Mutable/Source/CustomizableObject/Private/MuCO/CustomizableObjectMipDataProvider.cpp
+++ b/Engine/Plugins/Mutable/Source/CustomizableObject/Private/MuCO/CustomizableObjectMipDataProvider.cpp
@@ -6,10 +6,10 @@
 #include "MuCO/CustomizableObjectInstance.h"
 #include "MuCO/CustomizableObjectSystemPrivate.h"
 #include "MuCO/MutableStreamRequest.h"
+#include "MuCO/UnrealMutableImageProvider.h"
 #include "MuR/Model.h"
 #include "MuR/Operations.h"
 #include "TextureResource.h"
-#include "UnrealMutableImageProvider.h"
 
 #include "Containers/Ticker.h"
 
@@ -37,78 +37,6 @@ static FAutoConsoleVariableRef CVarEnablePrefetchHighQualityMips(
 	TEXT("If true, prefetch the data of high-quality mips to ensure it is available during the Mip generation task."),
 	ECVF_Default);
 
-FMutableUpdateContext::FMutableUpdateContext(
-	const FString& InCustomizableObjectPathName,
-	const FString& InInstancePathName,
-	const TSharedPtr<mu::FSystem>& InSystem,
-	const TSharedPtr<mu::FModel>& InModel,
-	const TSharedPtr<mu::FMeshIdRegistry>& InMeshIdRegistry,
-	const TSharedPtr<mu::FImageIdRegistry>& InImageIdRegistry,
-	const TSharedPtr<mu::FMaterialIdRegistry>& InMaterialIdRegistry,
-	const TSharedPtr<FModelStreamableBulkData>& InModelStreamable,
-	const TSharedPtr<const mu::FParameters>& InParameters,
-	int32 InState) :
-	CustomizableObjectPathName(InCustomizableObjectPathName),
-	InstancePathName(InInstancePathName),
-	System(InSystem),
-	Model(InModel),
-	MeshIdRegistry(InMeshIdRegistry),
-	ImageIdRegistry(InImageIdRegistry),
-    MaterialIdRegistry(InMaterialIdRegistry),
-	ModelStreamableBulkData(InModelStreamable),
-	Parameters(InParameters),
-	State(InState)
-{
-}
-
-
-const FString& FMutableUpdateContext::GetCustomizableObjectPathName() const
-{
-	return CustomizableObjectPathName;
-}
-
-
-const FString& FMutableUpdateContext::GetInstancePathName() const
-{
-	return InstancePathName;
-}
-
-
-TSharedPtr<mu::FSystem> FMutableUpdateContext::GetSystem() const
-{
-	return System;
-}
-
-
-TSharedPtr<mu::FModel> FMutableUpdateContext::GetModel() const
-{
-	return Model;
-}
-
-
-TSharedPtr<FModelStreamableBulkData> FMutableUpdateContext::GetModelStreamableBulkData()
-{
-	return ModelStreamableBulkData;
-}
-
-
-TSharedPtr<const mu::FParameters> FMutableUpdateContext::GetParameters() const
-{
-	return Parameters;
-}
-
-
-int32 FMutableUpdateContext::GetState() const
-{
-	return State;
-}
-
-
-const TArray<TSharedPtr<const mu::FImage>>& FMutableUpdateContext::GetImageParameterValues() const
-{
-	return ImageParameterValues;
-}
-
 
 FMutableTextureMipDataProvider::FMutableTextureMipDataProvider(const UTexture* Texture, UCustomizableObjectInstance* InCustomizableObjectInstance, const FMutableImageReference& InImageRef)
 	: FTextureMipDataProvider(Texture, ETickState::Init, ETickThread::Async),
@@ -162,9 +90,9 @@ namespace Impl
 
 		// This runs in a worker thread.
 		check(OperationData.IsValid());
-		check(OperationData->UpdateContext->GetSystem());
-		check(OperationData->UpdateContext->GetModel());
-		check(OperationData->UpdateContext->GetParameters());
+		check(OperationData->UpdateContext->System);
+		check(OperationData->UpdateContext->Model);
+		check(OperationData->UpdateContext->Parameters);
 
 		if (!OperationData.IsValid())
 		{
@@ -173,8 +101,8 @@ namespace Impl
 
 		TRACE_BEGIN_REGION(UE_MUTABLE_MIPDATA_PROVIDER_UPDATE_IMAGE_REGION);
 		
-		TSharedPtr<mu::FSystem> System = OperationData->UpdateContext->GetSystem();
-		const TSharedPtr<mu::FModel, ESPMode::ThreadSafe> Model = OperationData->UpdateContext->GetModel();
+		TSharedPtr<mu::FSystem> System = OperationData->UpdateContext->System;
+		const TSharedPtr<mu::FModel, ESPMode::ThreadSafe> Model = OperationData->UpdateContext->Model;
 
 		auto EndUpdateImage = [](TSharedPtr<FMutableImageOperationData>& OperationData)
 		{
@@ -203,7 +131,7 @@ namespace Impl
 #endif
 
 		// For now, we are forcing the recreation of mutable-side instances with every update.
-		mu::FInstance::FID InstanceID = System->NewInstance(Model);
+		mu::FInstance::FID InstanceID = System->NewInstance(Model, OperationData->UpdateContext->ExternalResourceProvider);
 		UE_LOG(LogMutable, Verbose, TEXT("Creating Mutable instance with id [%d] for a single UpdateImage"), InstanceID)
 
 		TSharedPtr<const mu::FInstance> Instance;
@@ -214,11 +142,11 @@ namespace Impl
 			uint32 LODMask = 0xFFFFFFFF;
 
 			Instance = System->BeginUpdate_MutableThread(InstanceID,
-				OperationData->UpdateContext->GetParameters(),
+				OperationData->UpdateContext->Parameters,
 				OperationData->UpdateContext->MeshIdRegistry,
 				OperationData->UpdateContext->ImageIdRegistry,
 				OperationData->UpdateContext->MaterialIdRegistry,
-				OperationData->UpdateContext->GetState(),
+				OperationData->UpdateContext->State,
 				LODMask);
 			
 			check(Instance);
@@ -297,8 +225,8 @@ namespace Impl
 						// Report the peak memory used during the operation (operation + baseline)
 						const int64 RealMemoryPeak = PeakMemory + OperationData->ImageUpdateStartBytes;
 
-						const FString& CustomizableObjectPathName = OperationData->UpdateContext->GetCustomizableObjectPathName();
-						const FString& InstancePathName = OperationData->UpdateContext->GetInstancePathName();
+						const FString& CustomizableObjectPathName = OperationData->UpdateContext->CustomizableObjectPathName;
+						const FString& InstancePathName = OperationData->UpdateContext->InstancePathName;
 
 						const FString& Descriptor = OperationData->UpdateContext->CapturedDescriptor;
 						const bool bDidLevelBeginPlay = OperationData->UpdateContext->bLevelBegunPlay;
@@ -334,7 +262,7 @@ int32 FMutableTextureMipDataProvider::GetMips(const FTextureUpdateContext& Conte
 {
 	MUTABLE_CPUPROFILER_SCOPE(FMutableTextureMipDataProvider::GetMips);
 
-	FSoftObjectPath Path(UpdateContext->GetCustomizableObjectPathName());
+	FSoftObjectPath Path(UpdateContext->CustomizableObjectPathName);
 	TRACE_CPUPROFILER_EVENT_SCOPE_TEXT(*Path.GetAssetName())
 
 	if (!UCustomizableObjectSystem::IsActive())
@@ -428,7 +356,7 @@ int32 FMutableTextureMipDataProvider::GetMips(const FTextureUpdateContext& Conte
 		OperationData->Counter->Increment(); // MipUpdate task
 
 		MUTABLE_CPUPROFILER_SCOPE(ImagePrefetch);
-		const TSharedPtr<FModelStreamableBulkData>& ModelStreamableBulkData = OperationData->UpdateContext->GetModelStreamableBulkData();
+		const TSharedPtr<FModelStreamableBulkData>& ModelStreamableBulkData = OperationData->UpdateContext->ModelStreamableBulkData;
 		
 		PrefetchRequest = MakeUnique<FMutableStreamRequest>(ModelStreamableBulkData);
 
@@ -438,7 +366,7 @@ int32 FMutableTextureMipDataProvider::GetMips(const FTextureUpdateContext& Conte
 
 			uint32 MaxBlockSize = 0;
 
-			TSharedPtr<mu::FModel, ESPMode::ThreadSafe> Model = UpdateContext->GetModel();
+			TSharedPtr<mu::FModel, ESPMode::ThreadSafe> Model = UpdateContext->Model;
 			check(Model);
 
 			if (!OperationData->Levels.IsEmpty())
diff --git a/Engine/Plugins/Mutable/Source/CustomizableObject/Private/MuCO/CustomizableObjectSkeletalMesh.cpp b/Engine/Plugins/Mutable/Source/CustomizableObject/Private/MuCO/CustomizableObjectSkeletalMesh.cpp
--- a/Engine/Plugins/Mutable/Source/CustomizableObject/Private/MuCO/CustomizableObjectSkeletalMesh.cpp
+++ b/Engine/Plugins/Mutable/Source/CustomizableObject/Private/MuCO/CustomizableObjectSkeletalMesh.cpp
@@ -29,6 +29,7 @@ void UCustomizableObjectSkeletalMesh::InitMutableStreamingData(const TSharedRef<
 	Model = CustomizableObject.GetPrivate()->GetModel();
 	MeshIdRegistry = CustomizableObject.GetPrivate()->MeshIdRegistry;
 	ImageIdRegistry = CustomizableObject.GetPrivate()->ImageIdRegistry;
+	ExternalResourceProvider = InOperationData->ExternalResourceProvider;
 
 	Parameters = InOperationData->Parameters;
 	State = InOperationData->GetCapturedDescriptor().GetState();
diff --git a/Engine/Plugins/Mutable/Source/CustomizableObject/Private/MuCO/CustomizableObjectSystem.cpp b/Engine/Plugins/Mutable/Source/CustomizableObject/Private/MuCO/CustomizableObjectSystem.cpp
--- a/Engine/Plugins/Mutable/Source/CustomizableObject/Private/MuCO/CustomizableObjectSystem.cpp
+++ b/Engine/Plugins/Mutable/Source/CustomizableObject/Private/MuCO/CustomizableObjectSystem.cpp
@@ -852,12 +852,6 @@ void UCustomizableObjectSystem::InitSystem()
 	check(Private->Streamer != nullptr);
 	Private->MutableSystem->SetStreamingInterface(Private->Streamer);
 
-	// Set up the external image provider, for image parameters.
-	TSharedPtr<FUnrealMutableResourceProvider> Provider = MakeShared<FUnrealMutableResourceProvider>();
-	check(Provider != nullptr);
-	Private->ResourceProvider = Provider;
-	Private->MutableSystem->SetExternalResourceProvider(Provider);
-
 	GetPrivate()->DefaultInstanceLODManagement = NewObject<UCustomizableInstanceLODManagement>();
 	check(GetPrivate()->DefaultInstanceLODManagement != nullptr);
 	GetPrivate()->CurrentInstanceLODManagement = GetPrivate()->DefaultInstanceLODManagement;
@@ -1160,7 +1154,7 @@ void FinishUpdateGlobal(const TSharedRef<FUpdateContextPrivate>& Context)
 	const FName InstanceName = GetFNameSafe(Instance);
 	UE_LOG(LogMutable, Verbose, TEXT("Finished Update Skeletal Mesh Async. CustomizableObject=%s Instance=%s, Frame=%d  QueueTime=%f, UpdateTime=%f"), *ObjectName.ToString(), *InstanceName.ToString(), GFrameNumber, Context->QueueTime, Context->UpdateTime);
 	
-	if (SystemPrivate && FLogBenchmarkUtil::IsBenchmarkingReportingEnabled())
+	if (SystemPrivate && FLogBenchmarkUtil::IsBenchmarkingReportingEnabled())	
 	{
 		FFunctionGraphTask::CreateAndDispatchWhenReady( // Calling Benchmark in a task so we make sure we exited all scopes.
 		[Context]()
@@ -3214,11 +3208,7 @@ namespace impl
 		}
 		
 		UCustomizableObjectSystemPrivate* CustomizableObjectSystemPrivate = System->GetPrivate();
-
-		// Unregister the current CO from the external resource provider.
-		check(IsInGameThread());
-		CustomizableObjectSystemPrivate->ResourceProvider->SetCurrentObject(nullptr);
-
+		
 		// Next Task: Release Mutable. We need this regardless if we cancel or not
 		//-------------------------------------------------------------		
 		const  TSharedPtr<mu::FSystem> MutableSystem = CustomizableObjectSystemPrivate->MutableSystem;
@@ -3566,13 +3556,6 @@ namespace impl
 		Operation->ImageIdRegistry = ObjectPrivate->ImageIdRegistry;
 		Operation->MaterialIdRegistry = ObjectPrivate->MaterialIdRegistry;
 
-		// Register the current CO into the external resource provider for runtime meshes and images.
-		// It is needed to get access to the bone map when converting meshes at runtime, for example.
-		if (SystemPrivate->ResourceProvider)
-		{
-			SystemPrivate->ResourceProvider->SetCurrentObject(Object);
-		}
-
 		UE::Tasks::FTask CacheRuntimeTexturesEvent = UE::Tasks::MakeCompletedTask<void>();
 
 #if WITH_EDITOR
@@ -3607,7 +3590,7 @@ namespace impl
 				}
 			
 				const UModelResources& ModelResources = Object->GetPrivate()->GetModelResourcesChecked();
-				System->GetPrivate()->GetResourceProviderChecked()->CacheRuntimeReferencedImages(Operation->Model.ToSharedRef(), ModelResources.RuntimeReferencedTextures);
+				Operation->ExternalResourceProvider->CacheRuntimeReferencedImages(ModelResources.RuntimeReferencedTextures);
 				Event.Trigger();
 			}));
 		}
@@ -3696,19 +3679,24 @@ namespace impl
 			SystemPrivate->MutableTaskGraph.AllowLaunchingMutableTaskLowPriority(false, false);
 
 			const TSharedPtr<mu::FSystem> MutableSystem = SystemPrivate->MutableSystem;
+
+			const TSharedRef<FBoneNames> BoneNames = MakeShared<FBoneNames>(*Object->GetPrivate()->GetModelResources());
+			Operation->BoneNames = BoneNames;
+			
+			Operation->ExternalResourceProvider = MakeShared<FUnrealMutableResourceProvider>(BoneNames);
 			
 			if (Operation->bLiveUpdateMode)
 			{
 				if (Operation->InstanceID == 0)
 				{
 					// It's the first update since the instance was put in LiveUpdate Mode, this ID will be reused from now on
-					Operation->InstanceID = MutableSystem->NewInstance(Operation->Model);
+					Operation->InstanceID = MutableSystem->NewInstance(Operation->Model, Operation->ExternalResourceProvider);
 					UE_LOG(LogMutable, Verbose, TEXT("Creating Mutable instance with id [%d] for reuse "), Operation->InstanceID);
 				}
 				else
 				{
 					// The instance was already in LiveUpdate Mode, the ID is reused
-					check(Operation->InstanceID);
+					MutableSystem->ReuseInstance(Operation->InstanceID, Operation->ExternalResourceProvider);
 					UE_LOG(LogMutable, Verbose, TEXT("Reusing Mutable instance with id [%d] "), Operation->InstanceID);
 				}
 			}
@@ -3716,7 +3704,7 @@ namespace impl
 			{
 				// In non-LiveUpdate mode, we are forcing the recreation of mutable-side instances with every update.
 				check(Operation->InstanceID == 0);
-				Operation->InstanceID = MutableSystem->NewInstance(Operation->Model);
+				Operation->InstanceID = MutableSystem->NewInstance(Operation->Model, Operation->ExternalResourceProvider);
 				UE_LOG(LogMutable, Verbose, TEXT("Creating Mutable instance with id [%d] "), Operation->InstanceID);
 			}
 
@@ -4514,13 +4502,6 @@ bool UCustomizableObjectSystemPrivate::IsMutableAnimInfoDebuggingEnabled() const
 }
 
 
-FUnrealMutableResourceProvider* UCustomizableObjectSystemPrivate::GetResourceProviderChecked() const
-{
-	check(ResourceProvider)
-	return ResourceProvider.Get();
-}
-
-
 void UCustomizableObjectSystemPrivate::OnMutableEnabledChanged(IConsoleVariable* MutableEnabled)
 {
 	if (!UCustomizableObjectSystem::IsCreated())
diff --git a/Engine/Plugins/Mutable/Source/CustomizableObject/Private/MuCO/UnrealConversionUtils.cpp b/Engine/Plugins/Mutable/Source/CustomizableObject/Private/MuCO/UnrealConversionUtils.cpp
--- a/Engine/Plugins/Mutable/Source/CustomizableObject/Private/MuCO/UnrealConversionUtils.cpp
+++ b/Engine/Plugins/Mutable/Source/CustomizableObject/Private/MuCO/UnrealConversionUtils.cpp
@@ -17,6 +17,8 @@
 #include "Engine/SkeletalMesh.h"
 #include "UObject/StrongObjectPtr.h"
 #include "GPUSkinVertexFactory.h"
+#include "SkinnedAssetCompiler.h"
+#include "MuCO/BoneNames.h"
 #include "Rendering/SkeletalMeshRenderData.h"
 #include "Rendering/MorphTargetVertexCodec.h"
 
@@ -889,263 +891,234 @@ namespace UnrealConversionUtils
 	}
 
 	
-	UE::Tasks::FTask ConvertSkeletalMeshFromRuntimeData(USkeletalMesh* SkeletalMesh, int32 LODIndex, int32 SectionIndex, UModelResources* ModelResources, mu::FMesh* Result, TFunction<void(bool bSuccess)>& CompletionCallback)
+	UE::Tasks::FTask ConvertSkeletalMeshFromRuntimeData(USkeletalMesh* SkeletalMesh, int32 LODIndex, int32 SectionIndex, TSharedPtr<mu::FMesh> Result, const TSharedRef<FBoneNames>& BoneNames)
 	{
 		MUTABLE_CPUPROFILER_SCOPE(ConvertSkeletalMeshFromRuntimeData);
 
-		if (!SkeletalMesh || !Result)
+		if (!SkeletalMesh)
 		{
-			CompletionCallback(false);
-			ensure(false);
-			return UE::Tasks::MakeCompletedTask<void>();
+			return UE::Tasks::MakeCompletedTask<TSharedPtr<mu::FMesh>>();
 		}
-
-		FSkeletalMeshRenderData* Data = SkeletalMesh->GetResourceForRendering();
-
-		if (!Data->LODRenderData.IsValidIndex(LODIndex))
+		
+		struct FMeshConversionContext
 		{
-			// This may happen if we are requesting a LOD that is not present in the provided mesh.
-			// TODO: an empty mesh is returned, but should we use the last LOD instead?
-			UE_LOG(LogMutable, Warning, TEXT("Trying to convert mesh [%s] LOD %d but it only has %d LODs"), 
-				*SkeletalMesh->GetName(), LODIndex, Data->LODRenderData.Num() );
+			FSkeletalMeshLODRenderData StreamedLOD;
+			FSkeletalMeshLODRenderData* OriginalLOD = nullptr;
+			FSkeletalMeshLODRenderData* DataLOD = nullptr;
+			TUniquePtr<IBulkDataIORequest> Request;
 			
-			CompletionCallback(false);
+			FMeshConversionContext() : StreamedLOD(false) {}
+		};
 
-			return UE::Tasks::MakeCompletedTask<void>();
-		}
+		TSharedPtr<FMeshConversionContext> Context = MakeShared<FMeshConversionContext>();
 
-		FSkeletalMeshLODRenderData* LOD = &Data->LODRenderData[LODIndex];
 
-		if (!LOD->RenderSections.IsValidIndex(SectionIndex))
-		{
-			// This may happen in the provided mesh in the parameter value doesn't have the expected number of sections
-			UE_LOG(LogMutable, Warning, TEXT("Trying to convert mesh [%s] section %d in LOD %d but it only has %d sections"),
-				*SkeletalMesh->GetName(), SectionIndex, LODIndex, LOD->RenderSections.Num());
+		UE::Tasks::FTask CompilingTask = UE::Tasks::MakeCompletedTask<void>();
 
-			CompletionCallback(false);
+#if WITH_EDITOR
+		if (SkeletalMesh->IsCompiling())
+		{
+			UE::Tasks::FTaskEvent Event { UE_SOURCE_LOCATION };
+			CompilingTask = Event;
 			
-			return UE::Tasks::MakeCompletedTask<void>();
+			ExecuteOnGameThread(TEXT("GetMeshAsync"), [Event, SkeletalMesh = TStrongObjectPtr(SkeletalMesh)]() mutable
+			{
+				FSkinnedAssetCompilingManager::Get().FinishCompilation({ SkeletalMesh.Get() });
+				Event.Trigger();
+			});
 		}
+#endif
 
-		Result->SkeletalMeshes.AddUnique(TStrongObjectPtr<USkeletalMesh>(SkeletalMesh));
-
-		const FSkelMeshRenderSection& Section = LOD->RenderSections[SectionIndex];
+		UE::Tasks::FTask ConversionTask = UE::Tasks::Launch(TEXT("ConversionTask"),
+			[Context, SkeletalMesh = TStrongObjectPtr(SkeletalMesh), LODIndex, SectionIndex, BoneNames]()
+			{
+				FSkeletalMeshRenderData* Data = SkeletalMesh->GetResourceForRendering();
 
-		// Skeleton data
-		bool bBoneMapModified = false;
-		TArray<FBoneIndexType> BoneMap;
-		TArray<FBoneIndexType> RemappedBoneMapIndices;
+				if (!Data->LODRenderData.IsValidIndex(LODIndex))
+				{
+					// This may happen if we are requesting a LOD that is not present in the provided mesh.
+					// TODO: an empty mesh is returned, but should we use the last LOD instead?
+					UE_LOG(LogMutable, Warning, TEXT("Trying to convert mesh [%s] LOD %d but it only has %d LODs"), 
+						*SkeletalMesh->GetName(), LODIndex, Data->LODRenderData.Num() );
+					
+					return;
+				}
 
-		bool bIgnoreSkeleton = false;
-		if (!bIgnoreSkeleton)
-		{
-			USkeleton* Skeleton = SkeletalMesh->GetSkeleton();
-			if (!Skeleton)
-			{
-				CompletionCallback(false);
-				ensure(false);
-				return UE::Tasks::MakeCompletedTask<void>();
-			}
+				FSkeletalMeshLODRenderData* LOD = &Data->LODRenderData[LODIndex];
 
-			const TArray<uint16>& SourceRequiredBones = LOD->RequiredBones;
+				if (!LOD->RenderSections.IsValidIndex(SectionIndex))
+				{
+					// This may happen in the provided mesh in the parameter value doesn't have the expected number of sections
+					UE_LOG(LogMutable, Warning, TEXT("Trying to convert mesh [%s] section %d in LOD %d but it only has %d sections"),
+						*SkeletalMesh->GetName(), SectionIndex, LODIndex, LOD->RenderSections.Num());
 
-			// Build RequiredBones array
-			TArray<FBoneIndexType> RequiredBones;
-			RequiredBones.Reserve(SourceRequiredBones.Num());
+					return;
+				}
+					
+				Context->OriginalLOD = LOD;
+				Context->DataLOD = LOD;
 
-			for (const FBoneIndexType& RequiredBoneIndex : SourceRequiredBones)
-			{
-				RequiredBones.AddUnique(RequiredBoneIndex);
-			}
+				const bool bUseStreaming = !LOD->bStreamedDataInlined;
+				if (bUseStreaming)
+				{
+					MUTABLE_CPUPROFILER_SCOPE(MeshStreaming);
 
-			// Rebuild BoneMap
-			const TArray<uint16>& SourceBoneMap = Section.BoneMap;
-			const int32 NumBonesInBoneMap = SourceBoneMap.Num();
+					UE::Tasks::FTaskEvent LoadedEvent(TEXT("MutableMeshParamLoadCompleted"));
+					UE::Tasks::AddNested(LoadedEvent);
 
-			for (int32 BoneIndex = 0; BoneIndex < NumBonesInBoneMap; ++BoneIndex)
-			{
-				const FBoneIndexType FinalBoneIndex = SourceBoneMap[BoneIndex];
+					FBulkDataIORequestCallBack RequestCallback = [LoadedEvent, SkeletalMesh, Context, LODIndex](bool bWasCancelled, IBulkDataIORequest* IORequest) mutable
+						{
+							if (!bWasCancelled)
+							{
+								uint8* BulkData = IORequest->GetReadResults();
+								check(BulkData != nullptr);	
 
-				const int32 BoneMapIndex = BoneMap.AddUnique(FinalBoneIndex);
-				RemappedBoneMapIndices.Add(BoneMapIndex);
+								{
+									MUTABLE_CPUPROFILER_SCOPE(MeshStreamingSerialization);
 
-				bBoneMapModified = bBoneMapModified || SourceBoneMap[BoneIndex] != FinalBoneIndex;
-			}
+									FMemoryReaderView Ar(TArrayView<uint8>(BulkData, IORequest->GetSize()), true);
 
-			// Create the skeleton, poses, and BoneMap for this mesh
-			TSharedPtr<mu::FSkeleton> MutableSkeleton = MakeShared<mu::FSkeleton>();
-			Result->SetSkeleton(MutableSkeleton);
+									constexpr uint8 DummyStripFlags = 0;
+									const bool bForceKeepCPUResources = true;
+									const bool bNeedsCPUAccess = true;
+									Context->StreamedLOD.SerializeStreamedData(Ar, const_cast<USkeletalMesh*>(SkeletalMesh.Get()), LODIndex, DummyStripFlags, bNeedsCPUAccess, bForceKeepCPUResources);
 
-			const int32 NumRequiredBones = RequiredBones.Num();
-			Result->SetBonePoseCount(NumRequiredBones);
-			MutableSkeleton->SetBoneCount(NumRequiredBones);
+									Context->DataLOD = &Context->StreamedLOD;
+								}
+							}
+							else
+							{
+								// Can this happen even if we don't cancel ourselves?
+								check(false);
+							}
 
-			// MutableBoneMap will not keep an index to the Skeleton, but to the BoneName
-			TArray<mu::FBoneName> MutableBoneMap;
-			MutableBoneMap.SetNum(BoneMap.Num());
+							LoadedEvent.Trigger();
+						};
 
-			TArray<FMatrix> ComposedRefPoseMatrices;
-			ComposedRefPoseMatrices.SetNum(NumRequiredBones);
+					Context->Request.Reset( LOD->StreamingBulkData.CreateStreamingRequest( EAsyncIOPriorityAndFlags::AIOP_High, &RequestCallback, nullptr) );
+				}	
+			},
+			CompilingTask);
 
-			const TArray<FMeshBoneInfo>& RefBoneInfo = SkeletalMesh->GetRefSkeleton().GetRefBoneInfo();
-			for (int32 BoneIndex = 0; BoneIndex < NumRequiredBones; ++BoneIndex)
+		
+		UE::Tasks::FTask FinalConversionTask = UE::Tasks::Launch( TEXT("MutableRuntimeMeshConversionFinal"),
+			[Context, SkeletalMesh = TStrongObjectPtr(SkeletalMesh), LODIndex, SectionIndex, Result, BoneNames]()
 			{
-				const int32 RefSkeletonBoneIndex = RequiredBones[BoneIndex];
-
-				const FMeshBoneInfo& BoneInfo = RefBoneInfo[RefSkeletonBoneIndex];
-				const int32 ParentBoneIndex = RequiredBones.Find(BoneInfo.ParentIndex);
+				MUTABLE_CPUPROFILER_SCOPE(MeshConversion);
+					
+				if (!Context->OriginalLOD)
+				{
+					return;
+				}
+					
+				Result->SkeletalMeshes.AddUnique(SkeletalMesh);
 
-				// Set bone hierarchy
+				const FSkelMeshRenderSection& Section = Context->OriginalLOD->RenderSections[SectionIndex];
+					
+				// Skeleton data
+				bool bBoneMapModified = false;
+				TArray<FBoneIndexType> BoneMap;
+				TArray<FBoneIndexType> RemappedBoneMapIndices;
 
-				// Get the bone id
-				mu::FBoneName BoneName;
-				if (ModelResources)
+				bool bIgnoreSkeleton = false;
+				if (!bIgnoreSkeleton)
 				{
-					// See if the bone exists in the compiled CO data
-					const FString BoneNameString = BoneInfo.Name.ToString().ToLower();
-					uint32* BoneId = ModelResources->BoneNamesMap.Find(BoneNameString);
-
-					if (BoneId)
+					USkeleton* Skeleton = SkeletalMesh->GetSkeleton();
+					if (!Skeleton)
 					{
-						BoneName = *BoneId;
+						ensure(false);
+						return;
 					}
-					else
-					{
-						// Go the slow way and add it to get a unique ID that the core can work with.
-						uint32 NewBoneId = CityHash32(reinterpret_cast<const char*>(*BoneNameString), BoneNameString.Len() * sizeof(FString::ElementType));
 
-						// See if the hash collides with an existing bone
-						bool bUnique = false;
-						while (!bUnique)
-						{
-							bUnique = true;
-							TMap<FString, uint32>::TConstIterator MapIterator = ModelResources->BoneNamesMap.CreateConstIterator();
-							while (MapIterator)
-							{
-								if (MapIterator.Value() == NewBoneId)
-								{
-									bUnique = false;
-									break;
-								}
-								++MapIterator;
-							}
+					const TArray<uint16>& SourceRequiredBones = Context->OriginalLOD->RequiredBones;
 
-							if (!bUnique)
-							{
-								NewBoneId++;
-							}
-						}
+					// Build RequiredBones array
+					TArray<FBoneIndexType> RequiredBones;
+					RequiredBones.Reserve(SourceRequiredBones.Num());
 
-						ModelResources->BoneNamesMap.Add(BoneNameString, NewBoneId);
-						BoneName = NewBoneId;
+					for (const FBoneIndexType& RequiredBoneIndex : SourceRequiredBones)
+					{
+						RequiredBones.AddUnique(RequiredBoneIndex);
 					}
-				}
 
-				MutableSkeleton->SetBoneName(BoneIndex, BoneName);
-				MutableSkeleton->SetBoneParent(BoneIndex, ParentBoneIndex);
+					// Rebuild BoneMap
+					const TArray<uint16>& SourceBoneMap = Section.BoneMap;
+					const int32 NumBonesInBoneMap = SourceBoneMap.Num();
 
-				// Debug. Will not be serialized
-				MutableSkeleton->SetDebugName(BoneIndex, BoneInfo.Name);
-
-				// BoneMap: Convert RefSkeletonBoneIndex to BoneId
-				const int32 BoneMapIndex = BoneMap.Find(RefSkeletonBoneIndex);
-				if (BoneMapIndex != INDEX_NONE)
-				{
-					MutableBoneMap[BoneMapIndex] = BoneName;
-				}
+					for (int32 BoneIndex = 0; BoneIndex < NumBonesInBoneMap; ++BoneIndex)
+					{
+						const FBoneIndexType FinalBoneIndex = SourceBoneMap[BoneIndex];
 
-				if (ParentBoneIndex >= 0)
-				{
-					ComposedRefPoseMatrices[BoneIndex] = SkeletalMesh->GetRefPoseMatrix(RefSkeletonBoneIndex) * ComposedRefPoseMatrices[ParentBoneIndex];
-				}
-				else
-				{
-					ComposedRefPoseMatrices[BoneIndex] = SkeletalMesh->GetRefPoseMatrix(RefSkeletonBoneIndex);
-				}
+						const int32 BoneMapIndex = BoneMap.AddUnique(FinalBoneIndex);
+						RemappedBoneMapIndices.Add(BoneMapIndex);
 
-				// Set bone pose
-				FTransform3f BoneTransform;
-				BoneTransform.SetFromMatrix(FMatrix44f(ComposedRefPoseMatrices[BoneIndex]));
+						bBoneMapModified = bBoneMapModified || SourceBoneMap[BoneIndex] != FinalBoneIndex;
+					}
 
-				mu::EBoneUsageFlags BoneUsageFlags = mu::EBoneUsageFlags::None;
-				EnumAddFlags(BoneUsageFlags, BoneMapIndex != INDEX_NONE ? mu::EBoneUsageFlags::Skinning : mu::EBoneUsageFlags::None);
-				EnumAddFlags(BoneUsageFlags, ParentBoneIndex == INDEX_NONE ? mu::EBoneUsageFlags::Root : mu::EBoneUsageFlags::None);
+					// Create the skeleton, poses, and BoneMap for this mesh
+					TSharedPtr<mu::FSkeleton> MutableSkeleton = MakeShared<mu::FSkeleton>();
+					Result->SetSkeleton(MutableSkeleton);
 
-				Result->SetBonePose(BoneIndex, BoneName, BoneTransform, BoneUsageFlags);
-			}
+					const int32 NumRequiredBones = RequiredBones.Num();
+					Result->SetBonePoseCount(NumRequiredBones);
+					MutableSkeleton->SetBoneCount(NumRequiredBones);
 
-			Result->SetBoneMap(MutableBoneMap);
-		}
+					// MutableBoneMap will not keep an index to the Skeleton, but to the BoneName
+					TArray<mu::FBoneName> MutableBoneMap;
+					MutableBoneMap.SetNum(BoneMap.Num());
 
+					TArray<FMatrix> ComposedRefPoseMatrices;
+					ComposedRefPoseMatrices.SetNum(NumRequiredBones);
 
-		const bool bUseStreaming = !LOD->bStreamedDataInlined;
-		
-		UE::Tasks::FTaskEvent LoadedEvent(TEXT("MutableMeshParamLoadCompleted"));
+					const TArray<FMeshBoneInfo>& RefBoneInfo = SkeletalMesh->GetRefSkeleton().GetRefBoneInfo();
+					for (int32 BoneIndex = 0; BoneIndex < NumRequiredBones; ++BoneIndex)
+					{
+						const int32 RefSkeletonBoneIndex = RequiredBones[BoneIndex];
 
-		// This context stores the information that travels between tasks.
-		struct FMeshConversionContext
-		{
-			FSkeletalMeshLODRenderData StreamedLOD;
-			FSkeletalMeshLODRenderData* OriginalLOD = nullptr;
-			FSkeletalMeshLODRenderData* DataLOD = nullptr;
-			TUniquePtr<IBulkDataIORequest> Request;
+						const FMeshBoneInfo& BoneInfo = RefBoneInfo[RefSkeletonBoneIndex];
+						const int32 ParentBoneIndex = RequiredBones.Find(BoneInfo.ParentIndex);
 
-			FMeshConversionContext()
-				: StreamedLOD(false)
-			{
-			}
-		};
+						// Set bone hierarchy
 
-		TSharedPtr<FMeshConversionContext> Context = MakeShared<FMeshConversionContext>();
-		Context->OriginalLOD = LOD;
-		Context->DataLOD = LOD;
-		
-		if (!bUseStreaming)
-		{
-			LoadedEvent.Trigger();
-		}
-		else
-		{
-			MUTABLE_CPUPROFILER_SCOPE(MeshStreaming);
+						// Get the bone id
+						mu::FBoneName BoneName = BoneNames->FindOrAdd(BoneInfo.Name);
+						
+						MutableSkeleton->SetBoneName(BoneIndex, BoneName);
+						MutableSkeleton->SetBoneParent(BoneIndex, ParentBoneIndex);
 
-			FBulkDataIORequestCallBack RequestCallback = [LoadedEvent, SkeletalMesh, Context, LODIndex](bool bWasCancelled, IBulkDataIORequest* IORequest) mutable
-				{
-					if (!bWasCancelled)
-					{
-						uint8* BulkData = IORequest->GetReadResults();
-						check(BulkData != nullptr);
+						// Debug. Will not be serialized
+						MutableSkeleton->SetDebugName(BoneIndex, BoneInfo.Name);
 
+						// BoneMap: Convert RefSkeletonBoneIndex to BoneId
+						const int32 BoneMapIndex = BoneMap.Find(RefSkeletonBoneIndex);
+						if (BoneMapIndex != INDEX_NONE)
 						{
-							MUTABLE_CPUPROFILER_SCOPE(MeshStreamingSerialization);
-
-							FMemoryReaderView Ar(TArrayView<uint8>(BulkData, IORequest->GetSize()), true);
-
-							constexpr uint8 DummyStripFlags = 0;
-							const bool bForceKeepCPUResources = true;
-							const bool bNeedsCPUAccess = true;
-							Context->StreamedLOD.SerializeStreamedData(Ar, const_cast<USkeletalMesh*>(SkeletalMesh), LODIndex, DummyStripFlags, bNeedsCPUAccess, bForceKeepCPUResources);
+							MutableBoneMap[BoneMapIndex] = BoneName;
+						}
 
-							Context->DataLOD = &Context->StreamedLOD;
+						if (ParentBoneIndex >= 0)
+						{
+							ComposedRefPoseMatrices[BoneIndex] = SkeletalMesh->GetRefPoseMatrix(RefSkeletonBoneIndex) * ComposedRefPoseMatrices[ParentBoneIndex];
+						}
+						else
+						{
+							ComposedRefPoseMatrices[BoneIndex] = SkeletalMesh->GetRefPoseMatrix(RefSkeletonBoneIndex);
 						}
-					}
-					else
-					{
-						// Can this happen even if we don't cancel ourselves?
-						check(false);
-					}
 
-					LoadedEvent.Trigger();
-				};
+						// Set bone pose
+						FTransform3f BoneTransform;
+						BoneTransform.SetFromMatrix(FMatrix44f(ComposedRefPoseMatrices[BoneIndex]));
 
-			 Context->Request.Reset( LOD->StreamingBulkData.CreateStreamingRequest( EAsyncIOPriorityAndFlags::AIOP_High, &RequestCallback, nullptr) );
-		}
+						mu::EBoneUsageFlags BoneUsageFlags = mu::EBoneUsageFlags::None;
+						EnumAddFlags(BoneUsageFlags, BoneMapIndex != INDEX_NONE ? mu::EBoneUsageFlags::Skinning : mu::EBoneUsageFlags::None);
+						EnumAddFlags(BoneUsageFlags, ParentBoneIndex == INDEX_NONE ? mu::EBoneUsageFlags::Root : mu::EBoneUsageFlags::None);
 
-		UE::Tasks::FTask FinalConversionTask = UE::Tasks::Launch( TEXT("MutableRuntimeMeshConversionFinal"),
-			[Context, SkeletalMesh, LODIndex, SectionIndex, Result, CompletionCallback]()
-			{
-				MUTABLE_CPUPROFILER_SCOPE(MeshConversion);
+						Result->SetBonePose(BoneIndex, BoneName, BoneTransform, BoneUsageFlags);
+					}
 
-				const FSkelMeshRenderSection& Section = Context->OriginalLOD->RenderSections[SectionIndex];
+					Result->SetBoneMap(MutableBoneMap);
+				}
+					
 
 				// Mesh data
 				int32 FirstVertexIndex = Context->OriginalLOD->RenderSections[SectionIndex].BaseVertexIndex;
@@ -1698,13 +1671,8 @@ namespace UnrealConversionUtils
 				}
 
 				Result->EnsureSurfaceData();
-
-				CompletionCallback(true);
 			},
-
-			// Dependencies
-			LoadedEvent
-			);
+			ConversionTask);
 
 		return FinalConversionTask;
 	}
diff --git a/Engine/Plugins/Mutable/Source/CustomizableObject/Private/MuCO/UnrealMutableImageProvider.cpp b/Engine/Plugins/Mutable/Source/CustomizableObject/Private/MuCO/UnrealMutableImageProvider.cpp
--- a/Engine/Plugins/Mutable/Source/CustomizableObject/Private/MuCO/UnrealMutableImageProvider.cpp
+++ b/Engine/Plugins/Mutable/Source/CustomizableObject/Private/MuCO/UnrealMutableImageProvider.cpp
@@ -88,6 +88,11 @@ mu::EImageFormat GetMutablePixelFormat(EPixelFormat InTextureFormat)
 }
 
 
+FUnrealMutableResourceProvider::FUnrealMutableResourceProvider(const TSharedRef<FBoneNames>& InBoneNames) : BoneNames(InBoneNames)
+{
+}
+
+
 //-------------------------------------------------------------------------------------------------
 TTuple<UE::Tasks::FTask, TFunction<void()>> FUnrealMutableResourceProvider::GetImageAsync(UTexture* Texture, uint8 MipmapsToSkip, TFunction<void(TSharedPtr<mu::FImage>)>& ResultCallback)
 {
@@ -321,7 +326,7 @@ TTuple<UE::Tasks::FTask, TFunction<void()>> FUnrealMutableResourceProvider::GetI
 
 
 //-------------------------------------------------------------------------------------------------
-TTuple<UE::Tasks::FTask, TFunction<void()>> FUnrealMutableResourceProvider::GetReferencedImageAsync(const void* ModelPtr, int32 Id, uint8 MipmapsToSkip, TFunction<void(TSharedPtr<mu::FImage>)>& ResultCallback)
+TTuple<UE::Tasks::FTask, TFunction<void()>> FUnrealMutableResourceProvider::GetReferencedImageAsync(int32 Id, uint8 MipmapsToSkip, TFunction<void(TSharedPtr<mu::FImage>)>& ResultCallback)
 {
 	// Thread: worker
 	MUTABLE_CPUPROFILER_SCOPE(FUnrealMutableImageProvider::GetReferencedImageAsync);
@@ -333,15 +338,8 @@ TTuple<UE::Tasks::FTask, TFunction<void()>> FUnrealMutableResourceProvider::GetR
 
 #if WITH_EDITOR
 	FScopeLock Lock(&RuntimeReferencedLock);
-	if (!RuntimeReferencedImages.Contains(ModelPtr))
-	{
-		UE_LOG(LogMutable, Error, TEXT("Failed to load Referenced Image [%i]. Model not registered in the provider."), Id);
-		ResultCallback(CreateDummy());
-		return Invoke(TrivialReturn);
-	}
 
-	TArray<TStrongObjectPtr<UTexture>>& SourceTextures = RuntimeReferencedImages[ModelPtr].Images;
-	if (!SourceTextures.IsValidIndex(Id))
+	if (!Images.IsValidIndex(Id))
 	{
 		// This could happen in the editor, because some source textures may have changed while there was a background compilation.
 		// We just show a warning and move on. This cannot happen during cooks, so it is fine.
@@ -350,7 +348,7 @@ TTuple<UE::Tasks::FTask, TFunction<void()>> FUnrealMutableResourceProvider::GetR
 		return Invoke(TrivialReturn);
 	}
 
-	UTexture* Texture = SourceTextures[Id].Get();
+	UTexture* Texture = Images[Id].Get();
 	if (!Texture)
 	{
 		UE_LOG(LogMutable, Warning, TEXT("Invalid Referenced Image [%i]. Nullptr."), Id);
@@ -478,40 +476,7 @@ TTuple<UE::Tasks::FTask, TFunction<void()>> FUnrealMutableResourceProvider::GetM
 		return Invoke(TrivialReturn);
 	}
 	
-	UE::Tasks::FTaskEvent Completion(TEXT("MutableMeshParameterLoadInGameThread"));
-	
-	ExecuteOnGameThread(TEXT("GetMeshAsync"), [SkeletalMesh, LODIndex, SectionIndex, Result, Completion, Provider = SharedThis(this)]() mutable
-	{
-		check(IsInGameThread());
-
-		MUTABLE_CPUPROFILER_SCOPE(ActualLoad);
-
-		if (!SkeletalMesh)
-		{
-			UE_LOG(LogMutable, Warning, TEXT("No Skeletal Mesh was provided for the Skeletal Mesh parameter."));
-			Completion.Trigger();
-			return;
-		}
-		
-		UModelResources* ModelResources = nullptr;
-
-		TStrongObjectPtr<UCustomizableObject> CO = Provider->CurrentCustomizableObject.Pin();
-		if (CO)
-		{
-			ModelResources = CO->GetPrivate()->GetModelResources();
-		}
-
-		// It is valid not to have a CO or MutableStreamedResources. It is only used for skeleton data. 
-		// This may happen when updating mips, and for those operations we don't need skeleton data.
-		//ensure(ModelResources);
-
-		TFunction<void(bool)> Callback = [Completion](bool bSuccess) mutable // Mutable due Trigger not being const
-			{
-				Completion.Trigger();
-			};
-
-		UE::Tasks::FTask ConversionTask = UnrealConversionUtils::ConvertSkeletalMeshFromRuntimeData(SkeletalMesh, LODIndex, SectionIndex, ModelResources, Result.Get(), Callback);
-	});
+	UE::Tasks::FTask ConversionTask = UnrealConversionUtils::ConvertSkeletalMeshFromRuntimeData(SkeletalMesh, LODIndex, SectionIndex, Result, BoneNames);
 
 	return MakeTuple(
 		// Some post-game conversion stuff can happen here in a worker thread
@@ -520,7 +485,7 @@ TTuple<UE::Tasks::FTask, TFunction<void()>> FUnrealMutableResourceProvider::GetM
 			{
 				ResultCallback(Result);
 			},
-			Completion ),
+			ConversionTask),
 
 		// Cleanup code that will be called after the result is received in calling code.
 		[]()
@@ -531,7 +496,7 @@ TTuple<UE::Tasks::FTask, TFunction<void()>> FUnrealMutableResourceProvider::GetM
 
 
 #if WITH_EDITOR
-void FUnrealMutableResourceProvider::CacheRuntimeReferencedImages(const TSharedRef<const mu::FModel>& Model, const TArray<TSoftObjectPtr<UTexture2D>>& RuntimeReferencedTextures)
+void FUnrealMutableResourceProvider::CacheRuntimeReferencedImages(const TArray<TSoftObjectPtr<UTexture2D>>& RuntimeReferencedTextures)
 {
 	check(IsInGameThread());
 	
@@ -539,10 +504,7 @@ void FUnrealMutableResourceProvider::CacheRuntimeReferencedImages(const TSharedR
 	
 	FScopeLock Lock(&RuntimeReferencedLock);
 
-	FRuntimeReferencedImages& ModelImages = RuntimeReferencedImages.Add(&Model.Get());
-	ModelImages.Model = Model.ToWeakPtr();
-
-	ModelImages.Images.Reset();
+	Images.Reset();
 	for (const TSoftObjectPtr<UTexture2D>& RuntimeReferencedTexture : RuntimeReferencedTextures)
 	{
 		UTexture2D* Texture = RuntimeReferencedTexture.Get(); // Is already loaded.
@@ -558,7 +520,7 @@ void FUnrealMutableResourceProvider::CacheRuntimeReferencedImages(const TSharedR
 			}
 		}
 
-		ModelImages.Images.Add(TStrongObjectPtr(Texture)); // Perform a CopyTornOff. Once done, we no longer need the texture loaded.
+		Images.Add(TStrongObjectPtr(Texture)); // Perform a CopyTornOff. Once done, we no longer need the texture loaded.
 	}
 }
 #endif
@@ -604,35 +566,3 @@ mu::FExtendedImageDesc FUnrealMutableResourceProvider::CreateDummyDesc()
 	return mu::FExtendedImageDesc{ {DUMMY_IMAGE_DESC}, 0 };
 }
 
-
-TAutoConsoleVariable<bool> CVarMutableLockExternalImagesDuringGC(
-	TEXT("Mutable.LockExternalImagesDuringGC"),
-	true,
-	TEXT("If true, GlobalExternalImages where all texture parameters are stored will be locked from concurrent access during the AddReferencedObjects phase of GC."),
-	ECVF_Default);
-
-
-void FUnrealMutableResourceProvider::AddReferencedObjects(FReferenceCollector& Collector)
-{
-#if WITH_EDITOR
-	{
-		FScopeLock Lock(&RuntimeReferencedLock);
-
-		for (TMap<const void*, FRuntimeReferencedImages>::TIterator It = RuntimeReferencedImages.CreateIterator(); It; ++It)
-		{
-			if (!It.Value().Model.IsValid())
-			{
-				It.RemoveCurrent();
-			}
-		}
-	}
-#endif
-}
-
-
-void FUnrealMutableResourceProvider::SetCurrentObject(const TWeakObjectPtr<UCustomizableObject>& InObject)
-{
-	check(IsInGameThread());
-	CurrentCustomizableObject = InObject;
-}
-
diff --git a/Engine/Plugins/Mutable/Source/CustomizableObject/Private/MuCO/UnrealMutableImageProvider.h b/Engine/Plugins/Mutable/Source/CustomizableObject/Private/MuCO/UnrealMutableImageProvider.h
deleted file mode 100644
diff --git a/Engine/Plugins/Mutable/Source/CustomizableObject/Public/MuCO/CustomizableObjectSkeletalMesh.h b/Engine/Plugins/Mutable/Source/CustomizableObject/Public/MuCO/CustomizableObjectSkeletalMesh.h
--- a/Engine/Plugins/Mutable/Source/CustomizableObject/Public/MuCO/CustomizableObjectSkeletalMesh.h
+++ b/Engine/Plugins/Mutable/Source/CustomizableObject/Public/MuCO/CustomizableObjectSkeletalMesh.h
@@ -20,6 +20,7 @@ class UCustomizableInstancePrivate;
 class FUpdateContextPrivate;
 class FCustomizableObjectMeshStreamIn;
 struct FModelStreamableBulkData;
+class FUnrealMutableResourceProvider;
 
 
 /**
@@ -50,6 +51,8 @@ class UCustomizableObjectSkeletalMesh : public USkeletalMesh
 	TSharedPtr<mu::FImageIdRegistry> ImageIdRegistry;
 	TSharedPtr<mu::FMaterialIdRegistry> MaterialIdRegistry;
 
+	TSharedPtr<FUnrealMutableResourceProvider> ExternalResourceProvider;
+
 	/** Instance parameters at the time of creation. */
 	TSharedPtr<mu::FParameters> Parameters;
 	int32 State = -1;
diff --git a/Engine/Plugins/Mutable/Source/CustomizableObjectEditor/Private/MuCOE/SMutableCodeViewer.cpp b/Engine/Plugins/Mutable/Source/CustomizableObjectEditor/Private/MuCOE/SMutableCodeViewer.cpp
--- a/Engine/Plugins/Mutable/Source/CustomizableObjectEditor/Private/MuCOE/SMutableCodeViewer.cpp
+++ b/Engine/Plugins/Mutable/Source/CustomizableObjectEditor/Private/MuCOE/SMutableCodeViewer.cpp
@@ -25,7 +25,7 @@
 #include "MuCOE/Widgets/MutableExpanderArrow.h"
 #include "MuCO/CustomizableObject.h"
 #include "MuCO/CustomizableObjectPrivate.h"
-#include "MuCO/UnrealToMutableTextureConversionUtils.h"
+#include "MuCO/UnrealMutableImageProvider.h"
 #include "MuT/ErrorLog.h"
 #include "MuT/TypeInfo.h"
 #include "MuR/SystemPrivate.h"
@@ -295,9 +295,7 @@ void SMutableCodeViewer::ClearSelectedTreeRow() const
 
 
 void SMutableCodeViewer::Construct(const FArguments& InArgs,
-	const TSharedPtr<mu::FModel>& InMutableModel,
-	const TArray<TSoftObjectPtr<const UTexture>>& InReferencedTextures,
-	const TArray<TSoftObjectPtr<const UStreamableRenderAsset>>& InReferencedMeshes)
+	const TSharedPtr<mu::FModel>& InMutableModel)
 {
 	// Min width allowed for the column. Needed to avoid having issues with the constants space being to small
 	// and then getting too tall on the y axis crashing the UI drawer.
@@ -307,8 +305,7 @@ void SMutableCodeViewer::Construct(const FArguments& InArgs,
 	MeshIdRegistry = InArgs._MeshIdRegistry;
 	ImageIdRegistry = InArgs._ImageIdRegistry;
 	MaterialIdRegistry = InArgs._MaterialIdRegistry;
-	ReferencedTextures = InReferencedTextures;
-	ReferencedMeshes = InReferencedMeshes;
+	ExternalResourceProvider = InArgs._ExternalResourceProvider;
 	PreviewParameters = mu::FModel::NewParameters(MutableModel);
 
 	RootNodes.Empty();
@@ -2423,9 +2420,6 @@ void SMutableCodeViewer::TreeExpandInstance()
 	TreeExpandElements(RootNodes,false, mu::EDataType::Instance);
 }
 
-#pragma endregion
-
-#pragma region Caching of operations related to constant resource
 
 void SMutableCodeViewer::CacheRootNodeAddresses()
 {
@@ -2801,142 +2795,6 @@ bool SMutableCodeViewer::IsConstantResourceUsedByOperation(const int32 IndexOnCo
 }
 
 
-
-#pragma endregion 
-
-
-namespace
-{
-	/** Test implementation to provide image parameters. It will generate some images of a fixed size and format. */
-	class TestResourceProvider : public mu::FExternalResourceProvider
-	{
-		static inline const mu::FImageDesc IMAGE_DESC = 
-			mu::FImageDesc(mu::FImageSize(1024, 1024), mu::EImageFormat::RGBA_UByte, 1);
-
-	public:
-
-		/** */
-		TArray<TSoftObjectPtr<const UTexture>> ReferencedTextures;
-		TArray<TSoftObjectPtr<const UStreamableRenderAsset>> ReferencedMeshes;
-
-	public:
-
-
-		virtual TTuple<UE::Tasks::FTask, TFunction<void()>> GetImageAsync(UTexture* Texture, uint8 MipmapsToSkip, TFunction<void(TSharedPtr<mu::FImage>)>& ResultCallback) override
-		{
-			MUTABLE_CPUPROFILER_SCOPE(TestImageProvider_GetImage);
-
-			int32 Size = IMAGE_DESC.m_size[0];
-			Size = FMath::Max(4, Size / (1 << MipmapsToSkip));
-
-			TSharedPtr<mu::FImage> Image = MakeShared<mu::FImage>(
-				Size, Size, IMAGE_DESC.m_lods,
-				IMAGE_DESC.m_format,
-				mu::EInitializationType::NotInitialized);
-
-			// Generate an alpha-tested circle with an horizontal gradient color.
-			uint8* Data = Image->GetLODData(0);
-			int32 CircleRadius = (Size * 2) / 5;
-			int32 CircleRadius2 = CircleRadius * CircleRadius;
-			int32 Color[3] = {255, 128, 0};
-
-			int32 LogSize = FMath::CeilLogTwo(Size);
-
-			int32 HalfSize = Size >> 1;
-			for (int32 RadY = -HalfSize; RadY < HalfSize; ++RadY)
-			{
-				int32 RadY2 = RadY * RadY;
-				for (int32 x = 0; x < Size; ++x)
-				{
-					int32 RadX = (x - HalfSize);
-					int32 R2 = RadX * RadX + RadY2;
-					int32 Opacity = FMath::Clamp(((CircleRadius2 - R2) * 512) / CircleRadius2 - 64, 0, 255);
-					Data[0] = uint8((Color[0] * x) >> LogSize);
-					Data[1] = uint8((Color[1] * x) >> LogSize);
-					Data[2] = uint8((Color[2] * x) >> LogSize);
-					Data[3] = uint8(Opacity);
-					Data += 4;
-				}
-			}
-
-			ResultCallback(Image);
-
-			return MakeTuple(UE::Tasks::MakeCompletedTask<void>(), []() -> void {});
-		}
-
-		virtual mu::FExtendedImageDesc GetImageDesc(UTexture* Texture) override
-		{
-			return mu::FExtendedImageDesc{ IMAGE_DESC };
-		}
-
-
-		virtual TTuple<UE::Tasks::FTask, TFunction<void()>> GetReferencedImageAsync(const void* ModelPtr, int32 Id, uint8 MipmapsToSkip, TFunction<void(TSharedPtr<mu::FImage>)>& ResultCallback)
-		{
-			check(ReferencedTextures.IsValidIndex(Id));
-
-			const UTexture* Texture = ReferencedTextures[Id].Get();
-			if (!Texture)
-			{
-				UE_LOG(LogMutable, Warning, TEXT("Failed to load Referenced Image [%i]. Nullptr."), Id);
-				return MakeTuple(UE::Tasks::MakeCompletedTask<void>(), []() -> void {});
-			}
-			
-			const UTexture2D* Texture2D = Cast<UTexture2D>(Texture);
-			if (!Texture2D)
-			{
-				UE_LOG(LogMutable, Warning, TEXT("Invalid Referenced Image [%i, %s]. Is not a UTexture2D."), Id, *Texture->GetName());
-				return MakeTuple(UE::Tasks::MakeCompletedTask<void>(), []() -> void {});
-			}
-			
-			// In the editor the src data can be directly accessed
-			int32 MipIndex = (MipmapsToSkip < Texture2D->GetPlatformData()->Mips.Num()) ? MipmapsToSkip : Texture2D->GetPlatformData()->Mips.Num() - 1;
-			check(MipIndex >= 0);
-
-			TSharedPtr<mu::FImage> ResultImage = MakeShared<mu::FImage>();
-
-			FMutableSourceTextureData Tex(*Texture2D);
-			EUnrealToMutableConversionError Error = ConvertTextureUnrealSourceToMutable(ResultImage.Get(), Tex, MipmapsToSkip);
-
-			if (Error != EUnrealToMutableConversionError::Success)
-			{
-				// This could happen in the editor, because some source textures may have changed while there was a background compilation.
-				// We just show a warning and move on. This cannot happen during cooks, so it is fine.
-				UE_LOG(LogMutable, Warning, TEXT("Failed to load some source texture data for Referenced Image [%i, %s]. Some textures may be corrupted."), Id, *Texture->GetName());
-			}
-
-			ResultCallback(ResultImage);
-
-			return MakeTuple(UE::Tasks::MakeCompletedTask<void>(), []() -> void {});
-		}
-
-
-		virtual TTuple<UE::Tasks::FTask, TFunction<void()>> GetMeshAsync(USkeletalMesh* SkeletalMesh, int32 LODIndex, int32 SectionIndex, TFunction<void(TSharedPtr<mu::FMesh>)>& ResultCallback) override
-		{
-			// Thread: worker
-			MUTABLE_CPUPROFILER_SCOPE(FUnrealMutableImageProvider::GetMeshAsync);
-
-			check(SkeletalMesh);
-
-			TSharedPtr<mu::FMesh> Result = MakeShared<mu::FMesh>();
-
-			TFunction<void(bool)> DummyCallback = [](bool) {};
-			UE::Tasks::FTask ConversionTask = UnrealConversionUtils::ConvertSkeletalMeshFromRuntimeData(SkeletalMesh, LODIndex, SectionIndex, nullptr, Result.Get(), DummyCallback);
-
-			return MakeTuple(
-				UE::Tasks::Launch( TEXT("FinalizeGetMesh"),
-					[Result, ResultCallback]()
-					{
-						ResultCallback(Result);
-					}, 
-					ConversionTask)
-				, 
-				[]() -> void {}
-			);
-		}
-	};
-}
-
-
 void SMutableCodeViewer::Tick(const FGeometry& AllottedGeometry, const double InCurrentTime, const float InDeltaTime)
 {
 	SCompoundWidget::Tick(AllottedGeometry, InCurrentTime, InDeltaTime);
@@ -2966,12 +2824,7 @@ void SMutableCodeViewer::Tick(const FGeometry& AllottedGeometry, const double In
 
 	mu::FSettings Settings;
 	const TSharedPtr<mu::FSystem> System = MakeShared<mu::FSystem>(Settings);
-
-	TSharedPtr<TestResourceProvider> ExternalResourceProvider = MakeShared<TestResourceProvider>();
-	ExternalResourceProvider->ReferencedTextures = ReferencedTextures;
-	ExternalResourceProvider->ReferencedMeshes = ReferencedMeshes;
-	System->SetExternalResourceProvider(ExternalResourceProvider);
-
+	
 	System->GetPrivate()->BeginBuild(MutableModel, MeshIdRegistry, ImageIdRegistry, MaterialIdRegistry);
 
 	switch (OperationDataType)
@@ -2979,7 +2832,7 @@ void SMutableCodeViewer::Tick(const FGeometry& AllottedGeometry, const double In
 	case mu::EDataType::Layout:
 	{
 		check(PreviewLayoutViewer);
-		TSharedPtr<const mu::FLayout> MutableLayout = System->GetPrivate()->BuildLayout(MutableModel, PreviewParameters.Get(), SelectedOperationAddress);
+		TSharedPtr<const mu::FLayout> MutableLayout = System->GetPrivate()->BuildLayout(ExternalResourceProvider, MutableModel, PreviewParameters.Get(), SelectedOperationAddress);
 		PreviewLayoutViewer->SetLayout(MutableLayout);
 		break;
 	}
@@ -2987,7 +2840,7 @@ void SMutableCodeViewer::Tick(const FGeometry& AllottedGeometry, const double In
 	case mu::EDataType::Image:
 	{
 		check(PreviewImageViewer);
-		TSharedPtr<const mu::FImage> MutableImage = System->GetPrivate()->BuildImage(MutableModel, PreviewParameters.Get(), SelectedOperationAddress, MipsToSkip, 0);
+		TSharedPtr<const mu::FImage> MutableImage = System->GetPrivate()->BuildImage(ExternalResourceProvider, MutableModel, PreviewParameters.Get(), SelectedOperationAddress, MipsToSkip, 0);
 		PreviewImageViewer->SetImage(MutableImage, 0);
 		break;
 	}
@@ -2995,7 +2848,7 @@ void SMutableCodeViewer::Tick(const FGeometry& AllottedGeometry, const double In
 	case mu::EDataType::Mesh:
 	{
 		check(PreviewMeshViewer);
-		TSharedPtr<const mu::FMesh> MutableMesh = System->GetPrivate()->BuildMesh(MutableModel, PreviewParameters.Get(), SelectedOperationAddress, mu::EMeshContentFlags::AllFlags);
+		TSharedPtr<const mu::FMesh> MutableMesh = System->GetPrivate()->BuildMesh(ExternalResourceProvider, MutableModel, PreviewParameters.Get(), SelectedOperationAddress, mu::EMeshContentFlags::AllFlags);
 		PreviewMeshViewer->SetMesh(MutableMesh);
 		break;
 	}
@@ -3003,7 +2856,7 @@ void SMutableCodeViewer::Tick(const FGeometry& AllottedGeometry, const double In
 	case mu::EDataType::Instance:
 	{
 		check(PreviewInstanceViewer);
-		TSharedPtr<const mu::FInstance> MutableInstance = System->GetPrivate()->BuildInstance(MutableModel, PreviewParameters.Get(), SelectedOperationAddress);
+		TSharedPtr<const mu::FInstance> MutableInstance = System->GetPrivate()->BuildInstance(ExternalResourceProvider, MutableModel, PreviewParameters.Get(), SelectedOperationAddress);
 		PreviewInstanceViewer->SetInstance(MutableInstance, MutableModel, PreviewParameters, *System);
 		break;
 	}
@@ -3011,7 +2864,7 @@ void SMutableCodeViewer::Tick(const FGeometry& AllottedGeometry, const double In
 	case mu::EDataType::Bool:
 	{
 		check(PreviewBoolViewer);
-		const bool MutableBool = System->GetPrivate()->BuildBool(MutableModel, PreviewParameters.Get(), SelectedOperationAddress);
+		const bool MutableBool = System->GetPrivate()->BuildBool(ExternalResourceProvider, MutableModel, PreviewParameters.Get(), SelectedOperationAddress);
 		PreviewBoolViewer->SetBool(MutableBool);
 		break;
 	}
@@ -3019,7 +2872,7 @@ void SMutableCodeViewer::Tick(const FGeometry& AllottedGeometry, const double In
 	case mu::EDataType::Int:
 	{
 		check(PreviewIntViewer);
-		const int32 MutableInt = System->GetPrivate()->BuildInt(MutableModel, PreviewParameters.Get(), SelectedOperationAddress);
+		const int32 MutableInt = System->GetPrivate()->BuildInt(ExternalResourceProvider, MutableModel, PreviewParameters.Get(), SelectedOperationAddress);
 		PreviewIntViewer->SetInt(MutableInt);
 		break;
 	}
@@ -3027,7 +2880,7 @@ void SMutableCodeViewer::Tick(const FGeometry& AllottedGeometry, const double In
 	case mu::EDataType::Scalar:
 	{
 		check(PreviewScalarViewer);
-		const float MutableScalar = System->GetPrivate()->BuildScalar(MutableModel, PreviewParameters.Get(), SelectedOperationAddress);
+		const float MutableScalar = System->GetPrivate()->BuildScalar(ExternalResourceProvider, MutableModel, PreviewParameters.Get(), SelectedOperationAddress);
 		PreviewScalarViewer->SetScalar(MutableScalar);
 		break;
 	}
@@ -3035,7 +2888,7 @@ void SMutableCodeViewer::Tick(const FGeometry& AllottedGeometry, const double In
 	case mu::EDataType::String:
 	{
 		check(PreviewStringViewer);
-		TSharedPtr<const mu::String> MutableString = System->GetPrivate()->BuildString(MutableModel, PreviewParameters.Get(), SelectedOperationAddress);
+		TSharedPtr<const mu::String> MutableString = System->GetPrivate()->BuildString(ExternalResourceProvider, MutableModel, PreviewParameters.Get(), SelectedOperationAddress);
 		const FText MutableText = FText::FromString(FString(MutableString->GetValue()));
 		PreviewStringViewer->SetString(MutableText);
 		break;
@@ -3044,7 +2897,7 @@ void SMutableCodeViewer::Tick(const FGeometry& AllottedGeometry, const double In
 	case mu::EDataType::Color:
 	{
 		check(PreviewColorViewer);
-		FVector4f Color = System->GetPrivate()->BuildColour(MutableModel, PreviewParameters.Get(), SelectedOperationAddress);
+		FVector4f Color = System->GetPrivate()->BuildColour(ExternalResourceProvider, MutableModel, PreviewParameters.Get(), SelectedOperationAddress);
 		PreviewColorViewer->SetColor(Color);
 		break;
 	}
@@ -3052,7 +2905,7 @@ void SMutableCodeViewer::Tick(const FGeometry& AllottedGeometry, const double In
 	case mu::EDataType::Projector:
 	{
 		check (PreviewProjectorViewer);
-		mu::FProjector Projector = System->GetPrivate()->BuildProjector(MutableModel, PreviewParameters.Get(), SelectedOperationAddress);
+		mu::FProjector Projector = System->GetPrivate()->BuildProjector(ExternalResourceProvider, MutableModel, PreviewParameters.Get(), SelectedOperationAddress);
 		PreviewProjectorViewer->SetProjector(Projector);
 	}
 	
@@ -3111,7 +2964,7 @@ void SMutableCodeViewer::PrepareInstanceViewer()
 {
 	if (!PreviewInstanceViewer)
 	{
-		PreviewInstanceViewer = SNew(SMutableInstanceViewer);
+		PreviewInstanceViewer = SNew(SMutableInstanceViewer, ExternalResourceProvider);
 	}
 
 	PreviewBorder->SetContent(PreviewInstanceViewer.ToSharedRef());
diff --git a/Engine/Plugins/Mutable/Source/CustomizableObjectEditor/Private/MuCOE/SMutableCodeViewer.h b/Engine/Plugins/Mutable/Source/CustomizableObjectEditor/Private/MuCOE/SMutableCodeViewer.h
--- a/Engine/Plugins/Mutable/Source/CustomizableObjectEditor/Private/MuCOE/SMutableCodeViewer.h
+++ b/Engine/Plugins/Mutable/Source/CustomizableObjectEditor/Private/MuCOE/SMutableCodeViewer.h
@@ -11,6 +11,7 @@
 #include "MuR/Mesh.h"
 #include "UObject/SoftObjectPtr.h"
 
+class FUnrealMutableResourceProvider;
 class STableViewBase;
 namespace ESelectInfo { enum Type : int; }
 template <typename ItemType> class STreeView;
@@ -58,13 +59,12 @@ public:
 	SLATE_ARGUMENT(TSharedPtr<mu::FMeshIdRegistry>, MeshIdRegistry)
 	SLATE_ARGUMENT(TSharedPtr<mu::FImageIdRegistry>, ImageIdRegistry)
 	SLATE_ARGUMENT(TSharedPtr<mu::FMaterialIdRegistry>, MaterialIdRegistry)
+	SLATE_ARGUMENT(TSharedPtr<FUnrealMutableResourceProvider>, ExternalResourceProvider)
 		
 	SLATE_END_ARGS()
 
 	void Construct(const FArguments& InArgs,
-		const TSharedPtr<mu::FModel>& InMutableModel,
-		const TArray<TSoftObjectPtr<const UTexture>>& ReferencedTextures,
-		const TArray<TSoftObjectPtr<const UStreamableRenderAsset>>& ReferencedMeshes);
+		const TSharedPtr<mu::FModel>& InMutableModel);
 
 	// SWidget interface
 	virtual void Tick(const FGeometry& AllottedGeometry, const double InCurrentTime, const float InDeltaTime) override;
@@ -83,6 +83,8 @@ private:
 	TSharedPtr<mu::FMeshIdRegistry> MeshIdRegistry;
 	TSharedPtr<mu::FImageIdRegistry> ImageIdRegistry;
 	TSharedPtr<mu::FMaterialIdRegistry> MaterialIdRegistry;
+
+	TSharedPtr<FUnrealMutableResourceProvider> ExternalResourceProvider;
 	
 	/** Array of external referenced textures in MutableModel, indexed by id. */
 	TArray<TSoftObjectPtr<const UTexture>> ReferencedTextures;
diff --git a/Engine/Plugins/Mutable/Source/CustomizableObjectEditor/Private/MuCOE/SMutableInstanceViewer.cpp b/Engine/Plugins/Mutable/Source/CustomizableObjectEditor/Private/MuCOE/SMutableInstanceViewer.cpp
--- a/Engine/Plugins/Mutable/Source/CustomizableObjectEditor/Private/MuCOE/SMutableInstanceViewer.cpp
+++ b/Engine/Plugins/Mutable/Source/CustomizableObjectEditor/Private/MuCOE/SMutableInstanceViewer.cpp
@@ -4,6 +4,7 @@
 
 #include "Engine/SkeletalMesh.h"
 #include "Framework/Views/TableViewMetadata.h"
+#include "MuCO/UnrealMutableImageProvider.h"
 #include "MuCOE/SMutableMeshViewport.h"
 #include "MuCOE/SMutableSkeletonViewer.h"
 #include "MuT/TypeInfo.h"
@@ -252,12 +253,14 @@ private:
 };
 
 
-void SMutableInstanceViewer::Construct(const FArguments& InArgs)
+void SMutableInstanceViewer::Construct(const FArguments& InArgs, const TSharedPtr<FUnrealMutableResourceProvider>& InExternalResourceProvider)
 {
 	// Splitter values
 	constexpr float TablesSplitterValue = 0.5f;
 	constexpr float ViewportSplitterValue = 0.5f;
 
+	ExternalResourceProvider = InExternalResourceProvider;
+
 	ChildSlot
 	[
 		SNew(SSplitter)
@@ -314,6 +317,7 @@ void SMutableInstanceViewer::SetInstance(const TSharedPtr<const mu::FInstance>&
 					{
 						mu::OP::ADDRESS MeshRootAddress = MeshId.GetKey()->Address;
 						TSharedPtr<const mu::FMesh> MutableMesh = System.GetPrivate()->BuildMesh(
+								ExternalResourceProvider,
 								Model, 
 								Parameters.Get(), 
 								MeshRootAddress, 
diff --git a/Engine/Plugins/Mutable/Source/CustomizableObjectEditor/Private/MuCOE/SMutableInstanceViewer.h b/Engine/Plugins/Mutable/Source/CustomizableObjectEditor/Private/MuCOE/SMutableInstanceViewer.h
--- a/Engine/Plugins/Mutable/Source/CustomizableObjectEditor/Private/MuCOE/SMutableInstanceViewer.h
+++ b/Engine/Plugins/Mutable/Source/CustomizableObjectEditor/Private/MuCOE/SMutableInstanceViewer.h
@@ -7,6 +7,7 @@
 #include "UObject/GCObject.h"
 #include "Widgets/SCompoundWidget.h"
 
+class FUnrealMutableResourceProvider;
 template <typename ItemType> class SListView;
 
 class ITableRow;
@@ -56,18 +57,18 @@ public:
 	SLATE_END_ARGS()
 
 	/** Builds the widget */
-	void Construct(const FArguments& InArgs);
+	void Construct(const FArguments& InArgs, const TSharedPtr<FUnrealMutableResourceProvider>& InExternalResourceProvider);
 
 	/** Set the Mutable Instance to be used for this widget */
 	void SetInstance(const TSharedPtr<const mu::FInstance>&, const TSharedPtr<const mu::FModel>&, const TSharedPtr<const mu::FParameters>&, const mu::FSystem&);
-
 	
 	/** */
 	TSharedRef<SWidget> GenerateLODsListView(const TSharedPtr<TArray<TSharedPtr<FMutableInstanceViewerLODElement>>>& InBufferChannelElements);
 	TSharedRef<SWidget> GenerateSurfaceListView(const TSharedPtr<TArray<TSharedPtr<FMutableInstanceViewerSurfaceElement>>>& InBufferChannelElements);
 
 private:
-
+	TSharedPtr<FUnrealMutableResourceProvider> ExternalResourceProvider;
+	
 	/** Data backend for the widget. It represents the Instance that is being "displayed" */
 	TSharedPtr<const mu::FInstance> MutableInstance = nullptr;
 
diff --git a/Engine/Plugins/Mutable/Source/CustomizableObjectEditor/Private/MuCOE/SMutableObjectViewer.cpp b/Engine/Plugins/Mutable/Source/CustomizableObjectEditor/Private/MuCOE/SMutableObjectViewer.cpp
--- a/Engine/Plugins/Mutable/Source/CustomizableObjectEditor/Private/MuCOE/SMutableObjectViewer.cpp
+++ b/Engine/Plugins/Mutable/Source/CustomizableObjectEditor/Private/MuCOE/SMutableObjectViewer.cpp
@@ -7,6 +7,7 @@
 #include "Interfaces/ITargetPlatform.h"
 #include "Interfaces/ITargetPlatformManagerModule.h"
 #include "MuCO/CustomizableObject.h"
+#include "MuCO/UnrealMutableImageProvider.h"
 #include "MuCOE/CustomizableObjectCompileRunnable.h"
 #include "MuCOE/CustomizableObjectEditorStyle.h"
 #include "MuCOE/SMutableCodeViewer.h"
@@ -194,11 +195,17 @@ void SMutableObjectViewer::CompileMutableCodePressed()
 
 	FString DataTag = FString::Printf( TEXT("Mutable Code : %s opt %d"), *CompileOptions.TargetPlatform->PlatformName(), CompileOptions.OptimizationLevel );
 
-	TSharedRef<SMutableCodeViewer> CodeViewer = SNew(SMutableCodeViewer, CompileTask->Model, RuntimeTextures, RuntimeMeshes)
+	const TSharedRef<FBoneNames> BoneNames = MakeShared<FBoneNames>(*CustomizableObject->GetPrivate()->GetModelResources());
+	
+	TSharedPtr<FUnrealMutableResourceProvider> ExternalResourceProvider = MakeShared<FUnrealMutableResourceProvider>(BoneNames);
+
+	TSharedRef<SMutableCodeViewer> CodeViewer = SNew(SMutableCodeViewer, CompileTask->Model)
 		.DataTag(DataTag)
 		.MeshIdRegistry(CustomizableObject->GetPrivate()->MeshIdRegistry)
 		.ImageIdRegistry(CustomizableObject->GetPrivate()->ImageIdRegistry)
-		.MaterialIdRegistry(CustomizableObject->GetPrivate()->MaterialIdRegistry);
+		.MaterialIdRegistry(CustomizableObject->GetPrivate()->MaterialIdRegistry)
+		.ExternalResourceProvider(ExternalResourceProvider)
+		.DataTag(DataTag);
 
 	DebuggerContentsBox->ClearContent();
 	DebuggerContentsBox->SetContent(CodeViewer);
diff --git a/Engine/Plugins/Mutable/Source/MutableRuntime/Internal/MuR/CodeVisitor.h b/Engine/Plugins/Mutable/Source/MutableRuntime/Internal/MuR/CodeVisitor.h
--- a/Engine/Plugins/Mutable/Source/MutableRuntime/Internal/MuR/CodeVisitor.h
+++ b/Engine/Plugins/Mutable/Source/MutableRuntime/Internal/MuR/CodeVisitor.h
@@ -274,7 +274,7 @@ namespace mu
 
                 if (Args.condition)
                 {
-					bValue = System->BuildBool(Model, Params, Args.condition);
+					bValue = System->BuildBool(nullptr, Model, Params, Args.condition);
                 }
 
                 if (bValue)
@@ -317,7 +317,7 @@ namespace mu
 
                     PARENT::RecurseWithCurrentState( VarAddress );
 
-                    int32 Var = System->BuildInt( Model, Params, VarAddress );
+                    int32 Var = System->BuildInt(nullptr, Model, Params, VarAddress );
 
 					OP::ADDRESS ValueAt = DefAddress;
 
diff --git a/Engine/Plugins/Mutable/Source/MutableRuntime/Internal/MuR/SystemPrivate.h b/Engine/Plugins/Mutable/Source/MutableRuntime/Internal/MuR/SystemPrivate.h
--- a/Engine/Plugins/Mutable/Source/MutableRuntime/Internal/MuR/SystemPrivate.h
+++ b/Engine/Plugins/Mutable/Source/MutableRuntime/Internal/MuR/SystemPrivate.h
@@ -1201,6 +1201,8 @@ namespace mu
 		/** Cached data for the generation of this instance. */
 		TSharedPtr<FProgramCache> Cache;
 
+		TSharedPtr<FExternalResourceProvider> ExternalResourceProvider;
+		
 		~FLiveInstance()
 		{
 			// Manually done to trace mem deallocations
@@ -2014,25 +2016,23 @@ namespace mu
     	
 		MUTABLERUNTIME_API void EndBuild();
 
-		MUTABLERUNTIME_API bool BuildBool(const TSharedPtr<const FModel>&, const FParameters*, OP::ADDRESS) ;
-		MUTABLERUNTIME_API int32 BuildInt(const TSharedPtr<const FModel>&, const FParameters*, OP::ADDRESS) ;
-		MUTABLERUNTIME_API float BuildScalar(const TSharedPtr<const FModel>&, const FParameters*, OP::ADDRESS) ;
-		MUTABLERUNTIME_API FVector4f BuildColour(const TSharedPtr<const FModel>&, const FParameters*, OP::ADDRESS) ;
-		MUTABLERUNTIME_API TSharedPtr<const String> BuildString(const TSharedPtr<const FModel>&, const FParameters*, OP::ADDRESS) ;
-		MUTABLERUNTIME_API TSharedPtr<const FImage> BuildImage(const TSharedPtr<const FModel>&, const FParameters*, OP::ADDRESS, int32 MipsToSkip, int32 LOD);
-		MUTABLERUNTIME_API TSharedPtr<const FMesh> BuildMesh(const TSharedPtr<const FModel>&, const FParameters*, OP::ADDRESS, EMeshContentFlags MeshContentFilter);
-		MUTABLERUNTIME_API TSharedPtr<const FInstance> BuildInstance(const TSharedPtr<const FModel>&, const FParameters*, OP::ADDRESS);
-		MUTABLERUNTIME_API TSharedPtr<const FLayout> BuildLayout(const TSharedPtr<const FModel>&, const FParameters*, OP::ADDRESS) ;
-    	MUTABLERUNTIME_API FProjector BuildProjector(const TSharedPtr<const FModel>&, const FParameters*, OP::ADDRESS) ;
+		MUTABLERUNTIME_API bool BuildBool(const TSharedPtr<FExternalResourceProvider>& ExternalResourceProvider, const TSharedPtr<const FModel>&, const FParameters*, OP::ADDRESS) ;
+		MUTABLERUNTIME_API int32 BuildInt(const TSharedPtr<FExternalResourceProvider>& ExternalResourceProvider, const TSharedPtr<const FModel>&, const FParameters*, OP::ADDRESS) ;
+		MUTABLERUNTIME_API float BuildScalar(const TSharedPtr<FExternalResourceProvider>& ExternalResourceProvider, const TSharedPtr<const FModel>&, const FParameters*, OP::ADDRESS) ;
+		MUTABLERUNTIME_API FVector4f BuildColour(const TSharedPtr<FExternalResourceProvider>& ExternalResourceProvider, const TSharedPtr<const FModel>&, const FParameters*, OP::ADDRESS) ;
+		MUTABLERUNTIME_API TSharedPtr<const String> BuildString(const TSharedPtr<FExternalResourceProvider>& ExternalResourceProvider, const TSharedPtr<const FModel>&, const FParameters*, OP::ADDRESS) ;
+		MUTABLERUNTIME_API TSharedPtr<const FImage> BuildImage(const TSharedPtr<FExternalResourceProvider>& ExternalResourceProvider, const TSharedPtr<const FModel>&, const FParameters*, OP::ADDRESS, int32 MipsToSkip, int32 LOD);
+		MUTABLERUNTIME_API TSharedPtr<const FMesh> BuildMesh(const TSharedPtr<FExternalResourceProvider>& ExternalResourceProvider, const TSharedPtr<const FModel>&, const FParameters*, OP::ADDRESS, EMeshContentFlags MeshContentFilter);
+		MUTABLERUNTIME_API TSharedPtr<const FInstance> BuildInstance(const TSharedPtr<FExternalResourceProvider>& ExternalResourceProvider, const TSharedPtr<const FModel>&, const FParameters*, OP::ADDRESS);
+		MUTABLERUNTIME_API TSharedPtr<const FLayout> BuildLayout(const TSharedPtr<FExternalResourceProvider>& ExternalResourceProvider, const TSharedPtr<const FModel>&, const FParameters*, OP::ADDRESS) ;
+    	MUTABLERUNTIME_API FProjector BuildProjector(const TSharedPtr<FExternalResourceProvider>& ExternalResourceProvider, const TSharedPtr<const FModel>&, const FParameters*, OP::ADDRESS) ;
     	
         //!
         FSettings Settings;
 
         //! Data streaming interface, if any.
 		TSharedPtr<FModelReader> StreamInterface;
-
-		TSharedPtr<FExternalResourceProvider> ExternalResourceProvider;
-
+    	
 		/** Counter used to generate unique IDs for every new instance created in the system. */
         FInstance::FID LastInstanceID = 0;
 
@@ -2054,7 +2054,14 @@ namespace mu
         bool CheckUpdatedParameters( const FLiveInstance*, const TSharedPtr<const FParameters>&, uint64& OutUpdatedParameters );
 
 
-		void RunCode(const TSharedPtr<const FModel>&, const FParameters*, OP::ADDRESS, uint32 LODs = FSystem::AllLODs, uint8 ExecutionOptions = 0, int32 LOD = 0);
+		void RunCode(
+			const TSharedPtr<FExternalResourceProvider>& ExternalResourceProvider,
+			const TSharedPtr<const FModel>&,
+			const FParameters*,
+			OP::ADDRESS,
+			uint32 LODs = FSystem::AllLODs,
+			uint8 ExecutionOptions = 0,
+			int32 LOD = 0);
 
 		//!
 		void PrepareCache(const FModel*, int32 State);
diff --git a/Engine/Plugins/Mutable/Source/MutableRuntime/Private/MuR/CodeRunner.cpp b/Engine/Plugins/Mutable/Source/MutableRuntime/Private/MuR/CodeRunner.cpp
--- a/Engine/Plugins/Mutable/Source/MutableRuntime/Private/MuR/CodeRunner.cpp
+++ b/Engine/Plugins/Mutable/Source/MutableRuntime/Private/MuR/CodeRunner.cpp
@@ -143,25 +143,33 @@ namespace mu
 
 	TSharedRef<CodeRunner> CodeRunner::Create(
 		const FSettings& InSettings,
-		class FSystem::Private* InSystem,
+		FSystem::Private* InSystem,
 		EExecutionStrategy InExecutionStrategy,
 		const TSharedPtr<const FModel>& InModel,
 		const FParameters* InParams,
 		OP::ADDRESS At,
-		uint32 InLODMask, uint8 ExecutionOptions, int32 InImageLOD, FScheduledOp::EType Type)
+		uint32 InLODMask,
+		uint8 ExecutionOptions,
+		int32 InImageLOD,
+		FScheduledOp::EType Type,
+		const TSharedPtr<FExternalResourceProvider>& ExternalResourceProvider)
 	{
 		return MakeShared<CodeRunner>(FPrivateToken {},
-				InSettings, InSystem, InExecutionStrategy, InModel, InParams, At, InLODMask, ExecutionOptions, InImageLOD, Type);
+				InSettings, InSystem, InExecutionStrategy, InModel, InParams, At, InLODMask, ExecutionOptions, InImageLOD, Type, ExternalResourceProvider);
 	}
 
     CodeRunner::CodeRunner(FPrivateToken PrivateToken, 
 		const FSettings& InSettings,
-		class FSystem::Private* InSystem,
+		FSystem::Private* InSystem,
 		EExecutionStrategy InExecutionStrategy,
 		const TSharedPtr<const FModel>& InModel,
 		const FParameters* InParams,
 		OP::ADDRESS At,
-		uint32 InLodMask, uint8 ExecutionOptions, int32 InImageLOD, FScheduledOp::EType Type )
+		uint32 InLodMask,
+		uint8 ExecutionOptions,
+		int32 InImageLOD,
+		FScheduledOp::EType Type,
+		const TSharedPtr<FExternalResourceProvider>& ExternalResourceProvider)
 		: Settings(InSettings)
 		, RunnerCompletionEvent(TEXT("CodeRunnerCompletioneEventInit"))
 		, ExecutionStrategy(InExecutionStrategy)
@@ -169,6 +177,7 @@ namespace mu
 		, Model(InModel)
 		, Params(InParams)
 		, LODMask(InLodMask)
+		, ExternalResourceProvider(ExternalResourceProvider)
 	{
 		MUTABLE_CPUPROFILER_SCOPE(CodeRunner_Create)
 		const FProgram& Program = Model->GetPrivate()->Program;
@@ -211,20 +220,18 @@ namespace mu
 	TTuple<UE::Tasks::FTask, TFunction<void()>> CodeRunner::LoadExternalImageAsync(FExternalResourceId Id, uint8 MipmapsToSkip, TFunction<void(TSharedPtr<FImage>)>& ResultCallback)
     {
 		MUTABLE_CPUPROFILER_SCOPE(LoadExternalImageAsync);
-
-		check(System);
-
-		if (System->ExternalResourceProvider)
+		
+		if (ExternalResourceProvider)
 		{
 			if (Id.ReferenceResourceId < 0)
 			{
 				// It's a parameter image
-				return System->ExternalResourceProvider->GetImageAsync(Id.ImageParameter, MipmapsToSkip, ResultCallback);
+				return ExternalResourceProvider->GetImageAsync(Id.ImageParameter, MipmapsToSkip, ResultCallback);
 			}
 			else
 			{
 				// It's an image reference
-				return System->ExternalResourceProvider->GetReferencedImageAsync(Model.Get(), Id.ReferenceResourceId, MipmapsToSkip, ResultCallback);
+				return ExternalResourceProvider->GetReferencedImageAsync( Id.ReferenceResourceId, MipmapsToSkip, ResultCallback);
 			}
 		}
 		else
@@ -241,14 +248,12 @@ namespace mu
 	{
 		MUTABLE_CPUPROFILER_SCOPE(LoadExternalImageAsync);
 
-		check(System);
-
-		if (System->ExternalResourceProvider)
+		if (ExternalResourceProvider)
 		{
 			if (Id.ReferenceResourceId < 0)
 			{
 				// It's a parameter mesh
-				return System->ExternalResourceProvider->GetMeshAsync(Id.MeshParameter, LODIndex, SectionIndex, ResultCallback);
+				return ExternalResourceProvider->GetMeshAsync(Id.MeshParameter, LODIndex, SectionIndex, ResultCallback);
 			}
 			else
 			{
@@ -272,9 +277,9 @@ namespace mu
 
 		check(System);
 
-		if (System->ExternalResourceProvider)
+		if (ExternalResourceProvider)
 		{
-			return System->ExternalResourceProvider->GetImageDesc(Image);
+			return ExternalResourceProvider->GetImageDesc(Image);
 		}
 		else
 		{
diff --git a/Engine/Plugins/Mutable/Source/MutableRuntime/Private/MuR/CodeRunner.h b/Engine/Plugins/Mutable/Source/MutableRuntime/Private/MuR/CodeRunner.h
--- a/Engine/Plugins/Mutable/Source/MutableRuntime/Private/MuR/CodeRunner.h
+++ b/Engine/Plugins/Mutable/Source/MutableRuntime/Private/MuR/CodeRunner.h
@@ -44,20 +44,31 @@ namespace  mu
 	public:
 		static TSharedRef<CodeRunner> Create(
 				const FSettings&, 
-				class FSystem::Private*, 
+				FSystem::Private*, 
 				EExecutionStrategy,
 				const TSharedPtr<const FModel>&, 
 				const FParameters* Params,
-				OP::ADDRESS At, uint32 LodMask, uint8 ExecutionOptions, int32 InImageLOD, FScheduledOp::EType);
-		// Private constructor to prevent stack allocation. In general we can not call AsShared() if the lifetime is
+				OP::ADDRESS At,
+				uint32 LodMask,
+				uint8 ExecutionOptions,
+				int32 InImageLOD,
+				FScheduledOp::EType Type,
+				const TSharedPtr<FExternalResourceProvider>& ExternalResourceProvider);
+
+    	// Private constructor to prevent stack allocation. In general we can not call AsShared() if the lifetime is
 		// bounded.
 		explicit CodeRunner(FPrivateToken, 
 			const FSettings&, 
-			class FSystem::Private*, 
+			FSystem::Private*, 
 			EExecutionStrategy,
 			const TSharedPtr<const FModel>&, 
 			const FParameters* Params,
-			OP::ADDRESS At, uint32 LodMask, uint8 ExecutionOptions, int32 InImageLOD, FScheduledOp::EType);
+			OP::ADDRESS At,
+			uint32 LodMask,
+			uint8 ExecutionOptions,
+			int32 InImageLOD,
+			FScheduledOp::EType Type,
+			const TSharedPtr<FExternalResourceProvider>& ExternalResourceProvider);
 
     protected:
 		struct FProfileContext
@@ -615,10 +626,7 @@ namespace  mu
 		{
 			System->WorkingMemoryManager.Release(Resource);
 		}
-
-	// TODO: protect this.
-	public:
-		
+    	
 		/** This flag is turned on when a streaming error or similar happens. Results are not usable.
 		* This should only happen in-editor.
 		*/
@@ -629,6 +637,8 @@ namespace  mu
 		const FParameters* Params = nullptr;
 		uint32 LODMask = 0;
 
+    	TSharedPtr<FExternalResourceProvider> ExternalResourceProvider;
+    
     private:
     	struct FRomLoadOps
     	{
diff --git a/Engine/Plugins/Mutable/Source/MutableRuntime/Private/MuR/System.cpp b/Engine/Plugins/Mutable/Source/MutableRuntime/Private/MuR/System.cpp
--- a/Engine/Plugins/Mutable/Source/MutableRuntime/Private/MuR/System.cpp
+++ b/Engine/Plugins/Mutable/Source/MutableRuntime/Private/MuR/System.cpp
@@ -134,8 +134,7 @@ namespace mu
 
 		// Make it explicit to try to capture metrics
         StreamInterface = nullptr;
-		ExternalResourceProvider = nullptr;
-    }
+	}
 
 
     //---------------------------------------------------------------------------------------------
@@ -187,15 +186,6 @@ namespace mu
 		m_pD->UpdateStats();
 	}
 
-	
-    //---------------------------------------------------------------------------------------------
-    void FSystem::SetExternalResourceProvider(const TSharedPtr<FExternalResourceProvider>& InInterface )
-    {
-		LLM_SCOPE_BYNAME(TEXT("MutableRuntime"));
-
-        m_pD->ExternalResourceProvider = InInterface;
-    }
-
 
 	//---------------------------------------------------------------------------------------------
 	void FSystem::SetImagePixelConversionOverride(const FImageOperator::FImagePixelFormatFunc& FormatFunc)
@@ -207,7 +197,7 @@ namespace mu
 
 
     //---------------------------------------------------------------------------------------------
-    FInstance::FID FSystem::NewInstance( const TSharedPtr<const FModel>& InModel )
+    FInstance::FID FSystem::NewInstance( const TSharedPtr<const FModel>& InModel, const TSharedPtr<FExternalResourceProvider>& ExternalResourceProvider)
     {
 		LLM_SCOPE_BYNAME(TEXT("MutableRuntime"));
 		MUTABLE_CPUPROFILER_SCOPE(NewInstance);
@@ -218,6 +208,7 @@ namespace mu
 		InstanceData.Model = InModel;
 		InstanceData.State = -1;
 		InstanceData.Cache = MakeShared<FProgramCache>();
+    	InstanceData.ExternalResourceProvider = ExternalResourceProvider;
 		m_pD->WorkingMemoryManager.LiveInstances.Add(InstanceData);
 
 		TRACE_COUNTER_SET(MutableRuntime_LiveInstances, m_pD->WorkingMemoryManager.LiveInstances.Num());
@@ -226,6 +217,15 @@ namespace mu
 	}
 
 
+    void FSystem::ReuseInstance(FInstance::FID InstanceId, const TSharedPtr<FExternalResourceProvider>& ExternalResourceProvider)
+    {
+    	FLiveInstance* LiveInstance = m_pD->FindLiveInstance(InstanceId);
+    	check(LiveInstance);
+
+    	LiveInstance->ExternalResourceProvider = ExternalResourceProvider;
+    }
+
+
     //---------------------------------------------------------------------------------------------
     TSharedPtr<const FInstance> FSystem::BeginUpdate_GameThread(FInstance::FID InInstanceID,
                                      const TSharedPtr<const FParameters>& InParams,
@@ -343,7 +343,7 @@ namespace mu
 		// Ensure the model cache has been created
 		m_pD->WorkingMemoryManager.FindOrAddModelCache(LiveInstance->Model);
 
-		m_pD->RunCode(LiveInstance->Model, InParams.Get(), RootAt, InLodMask);
+		m_pD->RunCode(LiveInstance->ExternalResourceProvider, LiveInstance->Model, InParams.Get(), RootAt, InLodMask);
 
 		TSharedPtr<const FInstance> Result = LiveInstance->Cache->GetInstance(FCacheAddress(RootAt, 0, 0));
 
@@ -390,7 +390,7 @@ namespace mu
 		m_pD->WorkingMemoryManager.CurrentInstanceCache = LiveInstance->Cache;
 
 		OP::ADDRESS RootAddress = ImageId.GetKey()->Address;
-		pResult = m_pD->BuildImage(LiveInstance->Model, LiveInstance->OldParameters.Get(), RootAddress, MipsToSkip, InImageLOD);
+		pResult = m_pD->BuildImage(LiveInstance->ExternalResourceProvider, LiveInstance->Model, LiveInstance->OldParameters.Get(), RootAddress, MipsToSkip, InImageLOD);
 
 		// We always need to return something valid.
 		if (!pResult)
@@ -432,7 +432,15 @@ namespace mu
 		}
 		
 		TSharedRef<CodeRunner> Runner = CodeRunner::Create(
-				m_pD->Settings, m_pD, EExecutionStrategy::MinimizeMemory, LiveInstance->Model, LiveInstance->OldParameters.Get(), RootAddress, FSystem::AllLODs, MipsToSkip, InImageLOD, FScheduledOp::EType::Full);
+				m_pD->Settings,
+				m_pD,
+				EExecutionStrategy::MinimizeMemory,
+				LiveInstance->Model,
+				LiveInstance->OldParameters.Get(),
+				RootAddress, FSystem::AllLODs,
+				MipsToSkip, InImageLOD,
+				FScheduledOp::EType::Full,
+				LiveInstance->ExternalResourceProvider);
 	
 		
 		constexpr bool bForceInlineExecution = false;
@@ -510,7 +518,17 @@ namespace mu
 					
 			int8 executionOptions = 0;
 			TSharedRef<CodeRunner> Runner = CodeRunner::Create(
-					m_pD->Settings, m_pD, EExecutionStrategy::MinimizeMemory, LiveInstance->Model, LiveInstance->OldParameters.Get(), RootAddress, FSystem::AllLODs, executionOptions, 0, FScheduledOp::EType::ImageDesc);
+					m_pD->Settings,
+					m_pD,
+					EExecutionStrategy::MinimizeMemory,
+					LiveInstance->Model,
+					LiveInstance->OldParameters.Get(),
+					RootAddress,
+					FSystem::AllLODs,
+					executionOptions,
+					0,
+					FScheduledOp::EType::ImageDesc,
+					LiveInstance->ExternalResourceProvider);
 			
 			constexpr bool bForceInlineExecution = true;
 			UE::Tasks::FTask CompletionEvent = Runner->StartRun(bForceInlineExecution);
@@ -567,7 +585,17 @@ namespace mu
 		// GetImageDesc may call normal execution paths where meshes are computed.
 		int8 ExecutionOptions = 0;
 		TSharedRef<CodeRunner> Runner = CodeRunner::Create(
-				m_pD->Settings, m_pD, EExecutionStrategy::MinimizeMemory, LiveInstance->Model, LiveInstance->OldParameters.Get(), RootAddress, FSystem::AllLODs, ExecutionOptions, 0, FScheduledOp::EType::ImageDesc);
+				m_pD->Settings,
+				m_pD,
+				EExecutionStrategy::MinimizeMemory,
+				LiveInstance->Model,
+				LiveInstance->OldParameters.Get(),
+				RootAddress,
+				FSystem::AllLODs,
+				ExecutionOptions,
+				0,
+				FScheduledOp::EType::ImageDesc,
+				LiveInstance->ExternalResourceProvider);
 		
 		constexpr bool bForceInlineExecution = false;
 		UE::Tasks::FTask RunnerCompletionEvent = Runner->StartRun(bForceInlineExecution);
@@ -607,6 +635,7 @@ namespace mu
 
 		OP::ADDRESS RootAddress = MeshId.GetKey()->Address;
 		Result = m_pD->BuildMesh(
+				LiveInstance->ExternalResourceProvider,
 				LiveInstance->Model, 
 				LiveInstance->OldParameters.Get(), 
 				RootAddress, 
@@ -661,7 +690,8 @@ namespace mu
 				FSystem::AllLODs, 
 				ExecutionOptions, 
 				0, 
-				FScheduledOp::EType::Full);
+				FScheduledOp::EType::Full,
+				LiveInstance->ExternalResourceProvider);
 		
 		constexpr bool bForceInlineExecution = false;
 		UE::Tasks::FTask RunnerCompletionEvent = Runner->StartRun(bForceInlineExecution);
@@ -725,7 +755,8 @@ namespace mu
 				FSystem::AllLODs, 
 				0, 
 				0, 
-				FScheduledOp::EType::Full);
+				FScheduledOp::EType::Full,
+				LiveInstance->ExternalResourceProvider);
 		
 		constexpr bool bForceInlineExecution = false;
 		UE::Tasks::FTask RunnerCompletionEvent = Runner->StartRun(bForceInlineExecution);
@@ -1007,11 +1038,27 @@ namespace mu
 
 
 	//---------------------------------------------------------------------------------------------
-	void FSystem::Private::RunCode(const TSharedPtr<const FModel>& InModel,
-		const FParameters* InParameters, OP::ADDRESS InCodeRoot, uint32 InLODs, uint8 ExecutionOptions, int32 InImageLOD)
+	void FSystem::Private::RunCode(
+		const TSharedPtr<FExternalResourceProvider>& ExternalResourceProvider,
+		const TSharedPtr<const FModel>& InModel,
+		const FParameters* InParameters,
+		OP::ADDRESS InCodeRoot,
+		uint32 InLODs,
+		uint8 ExecutionOptions,
+		int32 InImageLOD)
 	{
-		TSharedRef<CodeRunner> Runner = CodeRunner::Create(Settings, this, EExecutionStrategy::MinimizeMemory, InModel, InParameters, InCodeRoot, InLODs,
-			ExecutionOptions, InImageLOD, FScheduledOp::EType::Full);
+		TSharedRef<CodeRunner> Runner = CodeRunner::Create(
+			Settings,
+			this,
+			EExecutionStrategy::MinimizeMemory,
+			InModel,
+			InParameters,
+			InCodeRoot,
+			InLODs,
+			ExecutionOptions,
+			InImageLOD,
+			FScheduledOp::EType::Full,
+			ExternalResourceProvider);
 		
 		constexpr bool bForceInlineExecutution = true;
 		UE::Tasks::FTask RunnerCompletionEvent = Runner->StartRun(bForceInlineExecutution);
@@ -1022,11 +1069,11 @@ namespace mu
 
 
 	//---------------------------------------------------------------------------------------------
-	bool FSystem::Private::BuildBool(const TSharedPtr<const FModel>& pModel, const FParameters* Params, OP::ADDRESS at)
+	bool FSystem::Private::BuildBool(const TSharedPtr<FExternalResourceProvider>& ExternalResourceProvider, const TSharedPtr<const FModel>& pModel, const FParameters* Params, OP::ADDRESS at)
 	{
 		WorkingMemoryManager.BeginRunnerThread();
 
-		RunCode(pModel, Params, at);
+		RunCode(ExternalResourceProvider, pModel, Params, at);
 
 		bool bResult = false;
 		if (!bUnrecoverableError)
@@ -1041,11 +1088,11 @@ namespace mu
 
 
 	//---------------------------------------------------------------------------------------------
-	float FSystem::Private::BuildScalar(const TSharedPtr<const FModel>& pModel, const FParameters* Params, OP::ADDRESS at)
+	float FSystem::Private::BuildScalar(const TSharedPtr<FExternalResourceProvider>& ExternalResourceProvider, const TSharedPtr<const FModel>& pModel, const FParameters* Params, OP::ADDRESS at)
 	{
 		WorkingMemoryManager.BeginRunnerThread();
 
-		RunCode(pModel, Params, at);
+		RunCode(ExternalResourceProvider, pModel, Params, at);
 
 		float Result = 0.0f;		
 		if (!bUnrecoverableError)
@@ -1060,11 +1107,11 @@ namespace mu
 
 
 	//---------------------------------------------------------------------------------------------
-	int32 FSystem::Private::BuildInt(const TSharedPtr<const FModel>& pModel, const FParameters* Params, OP::ADDRESS at)
+	int32 FSystem::Private::BuildInt(const TSharedPtr<FExternalResourceProvider>& ExternalResourceProvider, const TSharedPtr<const FModel>& pModel, const FParameters* Params, OP::ADDRESS at)
 	{
 		WorkingMemoryManager.BeginRunnerThread();
 
-		RunCode(pModel, Params, at);
+		RunCode(ExternalResourceProvider, pModel, Params, at);
 
 		int32 Result = 0;
 		if (!bUnrecoverableError)
@@ -1079,7 +1126,7 @@ namespace mu
 
 
 	//---------------------------------------------------------------------------------------------
-	FVector4f FSystem::Private::BuildColour(const TSharedPtr<const FModel>& pModel, const FParameters* Params, OP::ADDRESS at)
+	FVector4f FSystem::Private::BuildColour(const TSharedPtr<FExternalResourceProvider>& ExternalResourceProvider, const TSharedPtr<const FModel>& pModel, const FParameters* Params, OP::ADDRESS at)
 	{
 		WorkingMemoryManager.BeginRunnerThread();
 
@@ -1088,7 +1135,7 @@ namespace mu
 		mu::EOpType OpType = pModel->GetPrivate()->Program.GetOpType(at);
 		if (GetOpDataType(OpType) == EDataType::Color)
 		{
-			RunCode(pModel, Params, at);
+			RunCode(ExternalResourceProvider, pModel, Params, at);
 			if (!bUnrecoverableError)
 			{
 				Result = WorkingMemoryManager.CurrentInstanceCache->GetColour(FCacheAddress(at, 0, 0));
@@ -1102,11 +1149,11 @@ namespace mu
 
 	
 	//---------------------------------------------------------------------------------------------
-	FProjector FSystem::Private::BuildProjector(const TSharedPtr<const FModel>& pModel, const FParameters* Params, OP::ADDRESS at)
+	FProjector FSystem::Private::BuildProjector(const TSharedPtr<FExternalResourceProvider>& ExternalResourceProvider, const TSharedPtr<const FModel>& pModel, const FParameters* Params, OP::ADDRESS at)
 	{
 		WorkingMemoryManager.BeginRunnerThread();
 
-    	RunCode(pModel, Params, at);
+    	RunCode(ExternalResourceProvider, pModel, Params, at);
 
 		FProjector Result;
 		if (!bUnrecoverableError)
@@ -1120,7 +1167,7 @@ namespace mu
 	}
 
 	
-	TSharedPtr<const FImage> FSystem::Private::BuildImage(const TSharedPtr<const FModel>& pModel, const FParameters* Params, OP::ADDRESS at, int32 MipsToSkip, int32 InImageLOD)
+	TSharedPtr<const FImage> FSystem::Private::BuildImage(const TSharedPtr<FExternalResourceProvider>& ExternalResourceProvider, const TSharedPtr<const FModel>& pModel, const FParameters* Params, OP::ADDRESS at, int32 MipsToSkip, int32 InImageLOD)
 	{
 		WorkingMemoryManager.BeginRunnerThread();
 
@@ -1129,7 +1176,7 @@ namespace mu
 		mu::EOpType OpType = pModel->GetPrivate()->Program.GetOpType(at);
 		if (GetOpDataType(OpType) == EDataType::Image)
 		{
-			RunCode(pModel, Params, at, FSystem::AllLODs, uint8(MipsToSkip), InImageLOD);
+			RunCode(ExternalResourceProvider, pModel, Params, at, FSystem::AllLODs, uint8(MipsToSkip), InImageLOD);
 			if (!bUnrecoverableError)
 			{
 				Result = WorkingMemoryManager.LoadImage(FCacheAddress(at, 0, MipsToSkip), true);
@@ -1142,7 +1189,7 @@ namespace mu
 	}
 
 
-	TSharedPtr<const FMesh> FSystem::Private::BuildMesh(const TSharedPtr<const FModel>& Model, const FParameters* Params, OP::ADDRESS RootAddress, EMeshContentFlags MeshContentFilter)
+	TSharedPtr<const FMesh> FSystem::Private::BuildMesh(const TSharedPtr<FExternalResourceProvider>& ExternalResourceProvider, const TSharedPtr<const FModel>& Model, const FParameters* Params, OP::ADDRESS RootAddress, EMeshContentFlags MeshContentFilter)
 	{
 		WorkingMemoryManager.BeginRunnerThread();
 
@@ -1153,7 +1200,7 @@ namespace mu
 		if (GetOpDataType(OpType) == EDataType::Mesh)
 		{
 			uint8 ExecutionOptions = static_cast<uint8>(MeshContentFilter);
-			RunCode(Model, Params, RootAddress, FSystem::AllLODs, ExecutionOptions);
+			RunCode(ExternalResourceProvider, Model, Params, RootAddress, FSystem::AllLODs, ExecutionOptions);
 			if (!bUnrecoverableError)
 			{
 				Result = WorkingMemoryManager.LoadMesh(FCacheAddress(RootAddress, 0, ExecutionOptions), true);
@@ -1166,7 +1213,7 @@ namespace mu
 	}
 
 
-	TSharedPtr<const FInstance> FSystem::Private::BuildInstance(const TSharedPtr<const FModel>& pModel, const FParameters* Params, OP::ADDRESS at)
+	TSharedPtr<const FInstance> FSystem::Private::BuildInstance(const TSharedPtr<FExternalResourceProvider>& ExternalResourceProvider, const TSharedPtr<const FModel>& pModel, const FParameters* Params, OP::ADDRESS at)
 	{
 		WorkingMemoryManager.BeginRunnerThread();
 
@@ -1175,7 +1222,7 @@ namespace mu
 		mu::EOpType OpType = pModel->GetPrivate()->Program.GetOpType(at);
 		if (GetOpDataType(OpType) == EDataType::Instance)
 		{
-			RunCode(pModel, Params, at);
+			RunCode(ExternalResourceProvider, pModel, Params, at);
 			if (!bUnrecoverableError)
 			{
 				Result = WorkingMemoryManager.CurrentInstanceCache->GetInstance(FCacheAddress(at, 0, 0));
@@ -1189,7 +1236,7 @@ namespace mu
 
 
 	//---------------------------------------------------------------------------------------------
-	TSharedPtr<const FLayout> FSystem::Private::BuildLayout(const TSharedPtr<const FModel>& pModel, const FParameters* Params, OP::ADDRESS at)
+	TSharedPtr<const FLayout> FSystem::Private::BuildLayout(const TSharedPtr<FExternalResourceProvider>& ExternalResourceProvider, const TSharedPtr<const FModel>& pModel, const FParameters* Params, OP::ADDRESS at)
 	{
 		WorkingMemoryManager.BeginRunnerThread();
 
@@ -1200,7 +1247,7 @@ namespace mu
 			mu::EOpType OpType = pModel->GetPrivate()->Program.GetOpType(at);
 			if (GetOpDataType(OpType) == EDataType::Layout)
 			{
-				RunCode(pModel, Params, at);
+				RunCode(ExternalResourceProvider, pModel, Params, at);
 				if (!bUnrecoverableError)
 				{
 					Result = WorkingMemoryManager.CurrentInstanceCache->GetLayout(FCacheAddress(at, 0, 0));
@@ -1215,7 +1262,7 @@ namespace mu
 
 
 	//---------------------------------------------------------------------------------------------
-	TSharedPtr<const String> FSystem::Private::BuildString(const TSharedPtr<const FModel>& pModel, const FParameters* Params, OP::ADDRESS at)
+	TSharedPtr<const String> FSystem::Private::BuildString(const TSharedPtr<FExternalResourceProvider>& ExternalResourceProvider, const TSharedPtr<const FModel>& pModel, const FParameters* Params, OP::ADDRESS at)
 	{
 		WorkingMemoryManager.BeginRunnerThread();
 
@@ -1226,7 +1273,7 @@ namespace mu
 			mu::EOpType OpType = pModel->GetPrivate()->Program.GetOpType(at);
 			if (GetOpDataType(OpType) == EDataType::String)
 			{
-				RunCode(pModel, Params, at);
+				RunCode(ExternalResourceProvider, pModel, Params, at);
 				if (!bUnrecoverableError)
 				{
 					Result = WorkingMemoryManager.CurrentInstanceCache->GetString(FCacheAddress(at, 0, 0));
diff --git a/Engine/Plugins/Mutable/Source/MutableRuntime/Public/MuR/System.h b/Engine/Plugins/Mutable/Source/MutableRuntime/Public/MuR/System.h
--- a/Engine/Plugins/Mutable/Source/MutableRuntime/Public/MuR/System.h
+++ b/Engine/Plugins/Mutable/Source/MutableRuntime/Public/MuR/System.h
@@ -77,15 +77,14 @@ namespace mu
     class FExternalResourceProvider
     {
     public:
-
         //! Ensure virtual destruction
         virtual ~FExternalResourceProvider() = default;
 
         /** Returns the completion event and a cleanup function that must be called once event is completed. */
 		virtual TTuple<UE::Tasks::FTask, TFunction<void()>> GetImageAsync(UTexture* Texture, uint8 MipmapsToSkip, TFunction<void(TSharedPtr<FImage>)>& ResultCallback) = 0;
-		virtual TTuple<UE::Tasks::FTask, TFunction<void()>> GetReferencedImageAsync(const void* ModelPtr, int32 Id, uint8 MipmapsToSkip, TFunction<void(TSharedPtr<FImage>)>& ResultCallback) { check(false); return {}; }
+		virtual TTuple<UE::Tasks::FTask, TFunction<void()>> GetReferencedImageAsync(int32 Id, uint8 MipmapsToSkip, TFunction<void(TSharedPtr<FImage>)>& ResultCallback) { check(false); return {}; }
 
-		virtual mu::FExtendedImageDesc GetImageDesc(UTexture* Texture) = 0;
+		virtual FExtendedImageDesc GetImageDesc(UTexture* Texture) = 0;
 
 		/** Returns the completion event and a cleanup function that must be called once event is completed. */
 		virtual TTuple<UE::Tasks::FTask, TFunction<void()>> GetMeshAsync(USkeletalMesh* SkeleltalMesh, int32 InLODIndex, int32 InSectionIndex, TFunction<void(TSharedPtr<FMesh>)>& ResultCallback) = 0;
@@ -120,9 +119,6 @@ namespace mu
 		*/
 		UE_API void ClearWorkingMemory();
 		
-        /** Set a new provider for external image data. This is only necessary if image parameters are used in the models. */
-        UE_API void SetExternalResourceProvider(const TSharedPtr<FExternalResourceProvider>&);
-
 		/** Set a function that will be used to convert image pixel formats instead of the internal conversion. 
 		* \warning The provided function can be called from any thread, and also concurrently.
 		* If this function fails (returns false in the first parameter) the internal function is attempted next.
@@ -134,10 +130,11 @@ namespace mu
         //! calls to BeginUpdate/EndUpdate.
         //! A call to NewInstance must be paired to a call to ReleasesInstance when the instance is
         //! no longer needed.
-        //! \param pModel Model to build an instance of
         //! \return An identifier that is always bigger than 0.
-        UE_API FInstance::FID NewInstance(const TSharedPtr<const FModel>& Model);
+        UE_API FInstance::FID NewInstance(const TSharedPtr<const FModel>& Model, const TSharedPtr<FExternalResourceProvider>& ExternalResourceProvider);
 
+		UE_API void ReuseInstance(FInstance::FID InstanceId, const TSharedPtr<FExternalResourceProvider>& ExternalResourceProvider);
+		
         //! \brief Update an instance with a new parameter set and/or state.
         //!
         //! \warning a call to BeginUpdate must be paired with a call to EndUpdate once the returned
diff --git a/Engine/Plugins/Mutable/Source/MutableTools/Private/MuT/CodeOptimiser.cpp b/Engine/Plugins/Mutable/Source/MutableTools/Private/MuT/CodeOptimiser.cpp
--- a/Engine/Plugins/Mutable/Source/MutableTools/Private/MuT/CodeOptimiser.cpp
+++ b/Engine/Plugins/Mutable/Source/MutableTools/Private/MuT/CodeOptimiser.cpp
@@ -423,7 +423,7 @@ namespace mu
 			{
 				MUTABLE_CPUPROFILER_SCOPE(ConstantMesh);
 
-				TSharedPtr<const FMesh> MeshBuild = System->GetPrivate()->BuildMesh(Model, LocalParams.Get(), at, EMeshContentFlags::AllFlags);
+				TSharedPtr<const FMesh> MeshBuild = System->GetPrivate()->BuildMesh(nullptr, Model, LocalParams.Get(), at, EMeshContentFlags::AllFlags);
 
 				if (MeshBuild)
 				{
@@ -440,7 +440,7 @@ namespace mu
 			{
 				MUTABLE_CPUPROFILER_SCOPE(ConstantImage);
 
-				TSharedPtr<const FImage> pImage = System->GetPrivate()->BuildImage( Model, LocalParams.Get(), at, 0, 0 );
+				TSharedPtr<const FImage> pImage = System->GetPrivate()->BuildImage(nullptr, Model, LocalParams.Get(), at, 0, 0);
 
 				if (pImage)
 				{
@@ -457,7 +457,7 @@ namespace mu
 			{
 				MUTABLE_CPUPROFILER_SCOPE(ConstantLayout);
 
-				TSharedPtr<const FLayout> pLayout = System->GetPrivate()->BuildLayout( Model, LocalParams.Get(), at );
+				TSharedPtr<const FLayout> pLayout = System->GetPrivate()->BuildLayout(nullptr, Model, LocalParams.Get(), at);
 
 				if (pLayout)
 				{
@@ -473,7 +473,7 @@ namespace mu
 			{
 				MUTABLE_CPUPROFILER_SCOPE(ConstantBool);
 
-				bool value = System->GetPrivate()->BuildBool( Model, LocalParams.Get(), at );
+				bool value = System->GetPrivate()->BuildBool(nullptr, Model, LocalParams.Get(), at);
 				Result = new ASTOpConstantBool(value);
 				break;
 			}
@@ -483,7 +483,7 @@ namespace mu
 				MUTABLE_CPUPROFILER_SCOPE(ConstantBool);
 
 				FVector4f ResultColor(0, 0, 0, 0);
-				ResultColor = System->GetPrivate()->BuildColour( Model, LocalParams.Get(), at );
+				ResultColor = System->GetPrivate()->BuildColour(nullptr, Model, LocalParams.Get(), at);
 
 				{
 					mu::Ptr<ASTOpConstantColor> ConstantOp = new ASTOpConstantColor();
