[Mutable] Fix Material Id crash with Mesh Streaming.

diff --git a/Engine/Plugins/Mutable/Source/CustomizableObject/Private/MuCO/CustomizableObjectSkeletalMesh.cpp b/Engine/Plugins/Mutable/Source/CustomizableObject/Private/MuCO/CustomizableObjectSkeletalMesh.cpp
--- a/Engine/Plugins/Mutable/Source/CustomizableObject/Private/MuCO/CustomizableObjectSkeletalMesh.cpp
+++ b/Engine/Plugins/Mutable/Source/CustomizableObject/Private/MuCO/CustomizableObjectSkeletalMesh.cpp
@@ -29,6 +29,7 @@ void UCustomizableObjectSkeletalMesh::InitMutableStreamingData(const TSharedRef<
 	Model = CustomizableObject.GetPrivate()->GetModel();
 	MeshIdRegistry = CustomizableObject.GetPrivate()->MeshIdRegistry;
 	ImageIdRegistry = CustomizableObject.GetPrivate()->ImageIdRegistry;
+	MaterialIdRegistry = CustomizableObject.GetPrivate()->MaterialIdRegistry;
 	ExternalResourceProvider = InOperationData->ExternalResourceProvider;
 
 	Parameters = InOperationData->Parameters;
