[mutable] Fixed two very sneaky bugs that had gone unnoticed for nearly a year involving dynamic components that are deleted after an update. They retained their LastMeshIDs forever and could not ever be visible again.
#tests Editor
#rnx

diff --git a/Engine/Plugins/Mutable/Source/CustomizableObject/Internal/MuCO/CustomizableObjectInstancePrivate.h b/Engine/Plugins/Mutable/Source/CustomizableObject/Internal/MuCO/CustomizableObjectInstancePrivate.h
--- a/Engine/Plugins/Mutable/Source/CustomizableObject/Internal/MuCO/CustomizableObjectInstancePrivate.h
+++ b/Engine/Plugins/Mutable/Source/CustomizableObject/Internal/MuCO/CustomizableObjectInstancePrivate.h
@@ -278,6 +278,7 @@ public:
 	UE_API void InitCustomizableObjectData(const UCustomizableObject* InCustomizableObject);
 
 	UE_API FCustomizableInstanceComponentData* GetComponentData(const FName& ComponentName);
+	void ClearInstanceComponentData(const FName& ComponentName);
 
 	ECOInstanceFlags GetCOInstanceFlags() const { return InstanceFlagsPrivate; }
 	void SetCOInstanceFlags(ECOInstanceFlags FlagsToSet) { InstanceFlagsPrivate = (ECOInstanceFlags)(InstanceFlagsPrivate | FlagsToSet); }
diff --git a/Engine/Plugins/Mutable/Source/CustomizableObject/Private/MuCO/CustomizableObjectInstance.cpp b/Engine/Plugins/Mutable/Source/CustomizableObject/Private/MuCO/CustomizableObjectInstance.cpp
--- a/Engine/Plugins/Mutable/Source/CustomizableObject/Private/MuCO/CustomizableObjectInstance.cpp
+++ b/Engine/Plugins/Mutable/Source/CustomizableObject/Private/MuCO/CustomizableObjectInstance.cpp
@@ -315,6 +315,36 @@ const FCustomizableInstanceComponentData* UCustomizableInstancePrivate::GetCompo
 }
 
 
+void UCustomizableInstancePrivate::ClearInstanceComponentData(const FName& ComponentName)
+{
+	FCustomizableInstanceComponentData* ComponentData = GetComponentData(ComponentName);
+
+	if (ComponentData)
+	{
+		ComponentData->LastMeshIdPerLOD.Reset(MAX_MESH_LOD_COUNT);
+		ComponentData->LastMeshIdPerLOD.AddDefaulted(MAX_MESH_LOD_COUNT);
+
+		ComponentData->AnimSlotToBP.Empty();
+		ComponentData->AssetUserDataArray.Empty();
+		ComponentData->ClothingPhysicsAssetsToStream.Empty();
+
+#if WITH_EDITORONLY_DATA
+		ComponentData->MeshPartPaths.Empty();
+#endif
+		ComponentData->OverlayMaterial = nullptr;
+		ComponentData->OverrideMaterials.Empty();
+		ComponentData->PhysicsAssets.AdditionalPhysicsAssets.Empty();
+		ComponentData->PhysicsAssets.AdditionalPhysicsAssetsToLoad.Empty();
+		ComponentData->PhysicsAssets.PhysicsAssetsToMerge.Empty();
+		ComponentData->PhysicsAssets.PhysicsAssetToLoad.Empty();
+		ComponentData->Skeletons.Skeleton = nullptr;
+		ComponentData->Skeletons.SkeletonIds.Empty();
+		ComponentData->Skeletons.SkeletonsToMerge.Empty();
+		ComponentData->StreamedResourceIndex.Empty();
+	}
+}
+
+
 UCustomizableObjectInstance::UCustomizableObjectInstance()
 {
 	SetFlags(RF_Transactional);
@@ -2178,6 +2208,35 @@ bool UCustomizableInstancePrivate::UpdateSkeletalMesh_PostBeginUpdate0(UCustomiz
 		}
 	}
 
+	// Detect the components that have disappeared since the last update and blank their properties so they aren't reused the
+	// next time they reappear
+	{
+		for (const auto& MeshIt : SkeletalMeshes)
+		{
+			FName DeletedComponentCandidateName = MeshIt.Key;
+			bool bIsDeleted = true;
+
+			for (const auto& InstanceUpdateComponent : OperationData->InstanceUpdateData.Components)
+			{
+				FName InstanceUpdateComponentName = ModelResources.ComponentNamesPerObjectComponent[InstanceUpdateComponent.Id.GetValue()];
+
+				if (DeletedComponentCandidateName == InstanceUpdateComponentName)
+				{
+					bIsDeleted = false;
+
+					break;
+				}
+			}
+
+			if (bIsDeleted)
+			{
+				ClearInstanceComponentData(DeletedComponentCandidateName);
+
+				bUpdateMeshes = true;
+			}
+		}
+	}
+
 	// None of the current meshes requires a mesh update. Continue to BuildMaterials
 	if (!bUpdateMeshes)
 	{
@@ -2211,6 +2270,11 @@ bool UCustomizableInstancePrivate::UpdateSkeletalMesh_PostBeginUpdate0(UCustomiz
 			return false;
 		}
 
+		TArray<mu::FMeshId>* MeshDescriptors = OperationData->GetMeshDescriptors(ObjectComponentIndex);
+		check(MeshDescriptors);
+		check(MeshDescriptors->Num() == MAX_MESH_LOD_COUNT);
+		ComponentData->LastMeshIdPerLOD = *MeshDescriptors;
+
 		// If the component doesn't need an update copy the previously generated mesh.
 		if (!OperationData->MeshChangedPerInstanceComponent[InstanceComponentIndex.GetValue()])
 		{
@@ -2221,11 +2285,6 @@ bool UCustomizableInstancePrivate::UpdateSkeletalMesh_PostBeginUpdate0(UCustomiz
 
 			continue;
 		}
-		
-		TArray<mu::FMeshId>* MeshDescriptors = OperationData->GetMeshDescriptors(ObjectComponentIndex);
-		check(MeshDescriptors);
-		check(MeshDescriptors->Num() == MAX_MESH_LOD_COUNT);
-		ComponentData->LastMeshIdPerLOD = *MeshDescriptors;
 
 		uint32 MeshHash = INDEX_NONE;
 		for (int32 LODIndex = OperationData->FirstLODAvailable[ComponentName]; LODIndex < Component->LODCount && LODIndex < MAX_MESH_LOD_COUNT; ++LODIndex)
