[Mutable] Add MeshConversionContext to allow meshes to be converted asynchronously.
GenerationContext now has a SharedPtr to the FMutableCompilationContext instead of a reference.
Remove unused Layout index from GetPinSection().
Remove SkeletalMesh sync load when generating thumbnails.

#rnx

diff --git a/Engine/Plugins/Mutable/Source/CustomizableObject/Internal/MuCO/UnrealToMutableTextureConversionUtils.h b/Engine/Plugins/Mutable/Source/CustomizableObject/Internal/MuCO/UnrealToMutableTextureConversionUtils.h
--- a/Engine/Plugins/Mutable/Source/CustomizableObject/Internal/MuCO/UnrealToMutableTextureConversionUtils.h
+++ b/Engine/Plugins/Mutable/Source/CustomizableObject/Internal/MuCO/UnrealToMutableTextureConversionUtils.h
@@ -13,9 +13,12 @@
 class UTexture2D;
 class UAnimInstance;
 class USkeletalMesh;
+struct FMutableCompilationContext;
+
 namespace mu
 {
 	class FImage;
+	class FMesh;
 }
 
 struct FMutableSourceTextureData
@@ -58,8 +61,14 @@ enum class EMutableMeshConversionFlags : uint32
 	// Ignore Morphs
 	IgnoreMorphs = 1 << 2,
 
+	// Ignore Texture coordinates
+	IgnoreTexCoords = 1 << 3,
+
+	// Ignore AssetUserData
+	IgnoreAUD = 1 << 4,
+
 	// Prevent this mesh generation from adding per mesh metadata. 
-	DoNotCreateMeshMetadata = 1 << 3
+	DoNotCreateMeshMetadata = 1 << 5
 };
 
 ENUM_CLASS_FLAGS(EMutableMeshConversionFlags)
@@ -80,7 +89,7 @@ struct FMutableSourceMeshData
 	/** Assets involved in the conversion. */
 	TSoftObjectPtr<UStreamableRenderAsset> Mesh;
 	TSoftClassPtr<UAnimInstance> AnimInstance;
-	TSoftObjectPtr<const USkeletalMesh> TableReferenceSkeletalMesh;
+	TSoftObjectPtr<UStreamableRenderAsset> TableReferenceSkeletalMesh;
 	FName Component = NAME_None;
 
 	/** */
@@ -128,7 +137,27 @@ enum class EUnrealToMutableConversionError
     Unknown
 };
 
-CUSTOMIZABLEOBJECT_API EUnrealToMutableConversionError ConvertTextureUnrealSourceToMutable(mu::FImage* OutResult, FMutableSourceTextureData&, uint8 MipmapsToSkip);
+UE_API EUnrealToMutableConversionError ConvertTextureUnrealSourceToMutable(mu::FImage* OutResult, FMutableSourceTextureData&, uint8 MipmapsToSkip);
+
+/** Mesh conversion context. */
+class FMeshConversionContext
+{
+public:
+	UE_API FMeshConversionContext() = default;
+
+public:
+
+	// Source Data
+	FMutableSourceMeshData Source;
+	FString MorphName;
+
+	TSharedPtr<FMutableCompilationContext> CompilationContext;
+
+	// Keep strong references to meshes and other UObjects involved in the conversion
+	TArray<TStrongObjectPtr<UObject>> ReferencedObjects;
+
+	TSharedPtr<mu::FMesh> Result;
+};
 
 #endif
 
diff --git a/Engine/Plugins/Mutable/Source/CustomizableObjectEditor/Private/MuCOE/CustomizableObjectCompileRunnable.cpp b/Engine/Plugins/Mutable/Source/CustomizableObjectEditor/Private/MuCOE/CustomizableObjectCompileRunnable.cpp
--- a/Engine/Plugins/Mutable/Source/CustomizableObjectEditor/Private/MuCOE/CustomizableObjectCompileRunnable.cpp
+++ b/Engine/Plugins/Mutable/Source/CustomizableObjectEditor/Private/MuCOE/CustomizableObjectCompileRunnable.cpp
@@ -20,9 +20,12 @@
 #include "Containers/Ticker.h"
 #include "Serialization/ObjectAndNameAsStringProxyArchive.h"
 #include "Trace/Trace.inl"
+#include "UObject/StrongObjectPtr.h"
 #include "UObject/SoftObjectPtr.h"
 #include "UObject/ICookInfo.h"
 #include "Engine/AssetUserData.h"
+#include "Engine/SkeletalMesh.h"
+#include "Animation/AnimInstance.h"
 
 #include "DerivedDataCacheInterface.h"
 #include "DerivedDataCache.h"
@@ -212,7 +215,17 @@ uint32 FCustomizableObjectCompileRunnable::Run()
 					// Ensure we don't pull unwanted data into the package when cooking.
 					FCookLoadScope CookLoadScope(ECookLoadType::EditorOnly);
 
-					*ResolvedMesh = ConvertSkeletalMeshToMutable(MeshData, *Compiler->CompilationContext, MorphName);
+					TSharedRef<FMeshConversionContext> MeshConversionContext = MakeShared<FMeshConversionContext>();
+					MeshConversionContext->Source = MeshData;
+					MeshConversionContext->MorphName = MorphName;
+					MeshConversionContext->CompilationContext = Compiler->CompilationContext;
+
+					// Todo: Load assets async
+					MeshConversionContext->ReferencedObjects.Emplace(TStrongObjectPtr<UObject>(MutablePrivate::LoadObject(MeshData.Mesh)));
+					MeshConversionContext->ReferencedObjects.Emplace(TStrongObjectPtr<UObject>(MutablePrivate::LoadObject(MeshData.TableReferenceSkeletalMesh)));
+					MeshConversionContext->ReferencedObjects.Emplace(TStrongObjectPtr<UObject>(MutablePrivate::LoadClass(MeshData.AnimInstance)));
+
+					*ResolvedMesh = ConvertSkeletalMeshToMutable(MeshConversionContext);
 
 					CompletionEvent.Trigger();
 				};
diff --git a/Engine/Plugins/Mutable/Source/CustomizableObjectEditor/Private/MuCOE/CustomizableObjectCompiler.cpp b/Engine/Plugins/Mutable/Source/CustomizableObjectEditor/Private/MuCOE/CustomizableObjectCompiler.cpp
--- a/Engine/Plugins/Mutable/Source/CustomizableObjectEditor/Private/MuCOE/CustomizableObjectCompiler.cpp
+++ b/Engine/Plugins/Mutable/Source/CustomizableObjectEditor/Private/MuCOE/CustomizableObjectCompiler.cpp
@@ -676,9 +676,9 @@ mu::Ptr<mu::NodeObject> GenerateMutableRoot(
 		GenerationContext.NoNameNodeObjectArray.AddUnique(LocalRootNodeObject);
 	}
 
-	if ((Object->GetPrivate()->GetMeshCompileType() == EMutableCompileMeshType::Full) || GenerationContext.CompilationContext.Options.bIsCooking)
+	if ((Object->GetPrivate()->GetMeshCompileType() == EMutableCompileMeshType::Full) || GenerationContext.CompilationContext->Options.bIsCooking)
 	{
-		if (LocalRootNodeObject->ParentObject!=nullptr && GenerationContext.CompilationContext.Options.bIsCooking)
+		if (LocalRootNodeObject->ParentObject!=nullptr && GenerationContext.CompilationContext->Options.bIsCooking)
 		{
 			// This happens while packaging.
 			return nullptr;
@@ -759,9 +759,9 @@ mu::Ptr<mu::NodeObject> GenerateMutableRoot(
 		}
 	}, &GenerationContext.GroupIdToExternalNodeMap, &GenerationContext.MacroNodesStack);
 	
-    GenerationContext.CompilationContext.RealTimeMorphTargetsOverrides = RootNodeObject->RealTimeMorphSelectionOverrides;
+    GenerationContext.CompilationContext->RealTimeMorphTargetsOverrides = RootNodeObject->RealTimeMorphSelectionOverrides;
 
-	if (!GenerationContext.CompilationContext.Options.ParamNamesToSelectedOptions.IsEmpty())
+	if (!GenerationContext.CompilationContext->Options.ParamNamesToSelectedOptions.IsEmpty())
 	{
 		if (const UModelResources* ModelResources = Object->GetPrivate()->GetModelResources())
 		{
@@ -905,7 +905,7 @@ void FCustomizableObjectCompiler::CompileInternal(bool bAsync)
 	}
 
 	CompilationContext = MakeShared<FMutableCompilationContext>(CurrentObject, SharedThis(this), CurrentOptions);
-	FMutableGraphGenerationContext GenerationContext(*CompilationContext);
+	FMutableGraphGenerationContext GenerationContext(CompilationContext);
 
 	// Perform a first participating objects pass
 	TMap<FName, FGuid> ParticipatingObjects = ICustomizableObjectEditorModule::GetChecked().GetParticipatingObjects(CurrentObject, &CurrentOptions);
@@ -1014,7 +1014,7 @@ void FCustomizableObjectCompiler::CompileInternal(bool bAsync)
 										ImageProp.AddressX, ImageProp.AddressY });
 		}
 
-		const TArray<FMutableComponentInfo>& ComponentInfos = GenerationContext.CompilationContext.ComponentInfos;
+		const TArray<FMutableComponentInfo>& ComponentInfos = GenerationContext.CompilationContext->ComponentInfos;
 		for (const FMutableComponentInfo& ComponentInfo : ComponentInfos)
 		{
 			if (!Cast<UCustomizableObjectNodeComponentMesh>(ComponentInfo.Node))
@@ -1632,7 +1632,7 @@ mu::Ptr<mu::Node> FCustomizableObjectCompiler::Export(UCustomizableObject* Objec
 	CompilerOptions.bAnimBpPhysicsManipulationEnabled = Object->GetPrivate()->IsEnabledAnimBpPhysicsAssetsManipulation();
 	
 	CompilationContext = MakeShared<FMutableCompilationContext>(Object, SharedThis(this), CompilerOptions);
-	FMutableGraphGenerationContext GenerationContext(*CompilationContext);
+	FMutableGraphGenerationContext GenerationContext(CompilationContext);
 	
 	// Generate the mutable node expression
 	mu::Ptr<mu::NodeObject> MutableRoot = GenerateMutableRoot(Object, GenerationContext);
diff --git a/Engine/Plugins/Mutable/Source/CustomizableObjectEditor/Private/MuCOE/CustomizableObjectEditor.cpp b/Engine/Plugins/Mutable/Source/CustomizableObjectEditor/Private/MuCOE/CustomizableObjectEditor.cpp
--- a/Engine/Plugins/Mutable/Source/CustomizableObjectEditor/Private/MuCOE/CustomizableObjectEditor.cpp
+++ b/Engine/Plugins/Mutable/Source/CustomizableObjectEditor/Private/MuCOE/CustomizableObjectEditor.cpp
@@ -852,18 +852,16 @@ void FCustomizableObjectEditor::ShowGizmoClipMesh(UCustomizableObjectNode& Node,
 	{
 		if (const UEdGraphNode* ConnectedNode = ConnectedPin->GetOwningNode())
 		{
-			int32 DummyIndex = 0;
-
 			if (const UCustomizableObjectNodeStaticMesh* StaticMeshNode = Cast<UCustomizableObjectNodeStaticMesh>(ConnectedNode))
 			{
 				ClipMesh = MutablePrivate::LoadObject(StaticMeshNode->GetMesh());
-				StaticMeshNode->GetPinSection(*ConnectedPin, LODIndex, SectionIndex, DummyIndex);
+				StaticMeshNode->GetPinSection(*ConnectedPin, LODIndex, SectionIndex);
 				MaterialSlotIndex = SectionIndex;
 			}
 			else if (const UCustomizableObjectNodeSkeletalMesh* SkeletalMeshNode = Cast<UCustomizableObjectNodeSkeletalMesh>(ConnectedNode))
 			{
 				ClipMesh = MutablePrivate::LoadObject(SkeletalMeshNode->GetMesh());
-				SkeletalMeshNode->GetPinSection(*ConnectedPin, LODIndex, SectionIndex, DummyIndex);
+				SkeletalMeshNode->GetPinSection(*ConnectedPin, LODIndex, SectionIndex);
 				MaterialSlotIndex = SkeletalMeshNode->GetSkeletalMaterialIndexFor(*ConnectedPin);
 			}
 			else if (const UCustomizableObjectNodeTable* TableNode = Cast<UCustomizableObjectNodeTable>(ConnectedNode))
diff --git a/Engine/Plugins/Mutable/Source/CustomizableObjectEditor/Private/MuCOE/CustomizableObjectEditorModule.cpp b/Engine/Plugins/Mutable/Source/CustomizableObjectEditor/Private/MuCOE/CustomizableObjectEditorModule.cpp
--- a/Engine/Plugins/Mutable/Source/CustomizableObjectEditor/Private/MuCOE/CustomizableObjectEditorModule.cpp
+++ b/Engine/Plugins/Mutable/Source/CustomizableObjectEditor/Private/MuCOE/CustomizableObjectEditorModule.cpp
@@ -879,7 +879,7 @@ TMap<FName, FGuid> FCustomizableObjectEditorModule::GetParticipatingObjects(cons
 	MUTABLE_CPUPROFILER_SCOPE(FCustomizableObjectEditorModule::GetParticipatingObjects);
 
 	FCompilationOptions Options = InOptions ? *InOptions : Object->GetPrivate()->GetCompileOptions();
-	FMutableCompilationContext CompilationContext(Object, nullptr, Options);
+	TSharedPtr<FMutableCompilationContext> CompilationContext = MakeShared<FMutableCompilationContext>(Object, nullptr, Options);
 	FMutableGraphGenerationContext Context(CompilationContext);
 
 	IAssetRegistry& AssetRegistry = FModuleManager::LoadModuleChecked<FAssetRegistryModule>("AssetRegistry").Get();
diff --git a/Engine/Plugins/Mutable/Source/CustomizableObjectEditor/Private/MuCOE/CustomizableObjectLayout.cpp b/Engine/Plugins/Mutable/Source/CustomizableObjectEditor/Private/MuCOE/CustomizableObjectLayout.cpp
--- a/Engine/Plugins/Mutable/Source/CustomizableObjectEditor/Private/MuCOE/CustomizableObjectLayout.cpp
+++ b/Engine/Plugins/Mutable/Source/CustomizableObjectEditor/Private/MuCOE/CustomizableObjectLayout.cpp
@@ -134,7 +134,7 @@ void UCustomizableObjectLayout::GenerateAutomaticBlocksFromUVs()
 
 	FCompilationOptions Options = Object->GetPrivate()->GetCompileOptions();
 
-	FMutableCompilationContext CompilationContext(Object, Compiler, Options);
+	TSharedPtr<FMutableCompilationContext> CompilationContext = MakeShared<FMutableCompilationContext>(Object, Compiler, Options);
 	FMutableGraphGenerationContext GenerationContext(CompilationContext);
 
 	mu::Ptr<mu::NodeLayout> LayoutNode = CreateMutableLayoutNode(GenerationContext, this, false);
diff --git a/Engine/Plugins/Mutable/Source/CustomizableObjectEditor/Private/MuCOE/ExtensionDataCompilerInterface.cpp b/Engine/Plugins/Mutable/Source/CustomizableObjectEditor/Private/MuCOE/ExtensionDataCompilerInterface.cpp
--- a/Engine/Plugins/Mutable/Source/CustomizableObjectEditor/Private/MuCOE/ExtensionDataCompilerInterface.cpp
+++ b/Engine/Plugins/Mutable/Source/CustomizableObjectEditor/Private/MuCOE/ExtensionDataCompilerInterface.cpp
@@ -39,8 +39,8 @@ TSharedPtr<const mu::FExtensionData> FExtensionDataCompilerInterface::MakeAlways
 
 const UObject* FExtensionDataCompilerInterface::GetOuterForAlwaysLoadedObjects()
 {
-	check(GenerationContext.CompilationContext.Object);
-	return GenerationContext.CompilationContext.Object.Get();
+	check(GenerationContext.CompilationContext->Object);
+	return GenerationContext.CompilationContext->Object.Get();
 }
 
 void FExtensionDataCompilerInterface::AddGeneratedNode(const UCustomizableObjectNode* InNode)
diff --git a/Engine/Plugins/Mutable/Source/CustomizableObjectEditor/Private/MuCOE/GenerateMutableSource/GenerateMutableSource.cpp b/Engine/Plugins/Mutable/Source/CustomizableObjectEditor/Private/MuCOE/GenerateMutableSource/GenerateMutableSource.cpp
--- a/Engine/Plugins/Mutable/Source/CustomizableObjectEditor/Private/MuCOE/GenerateMutableSource/GenerateMutableSource.cpp
+++ b/Engine/Plugins/Mutable/Source/CustomizableObjectEditor/Private/MuCOE/GenerateMutableSource/GenerateMutableSource.cpp
@@ -84,7 +84,7 @@ bool FGraphCycle::FoundCycle() const
 
 	if (const UCustomizableObject** Result = Context.VisitedPins.Find(Key))
 	{
-		Context.CompilationContext.Log(LOCTEXT("CycleFoundNode", "Cycle detected."), &Node, EMessageSeverity::Error, true);
+		Context.CompilationContext->Log(LOCTEXT("CycleFoundNode", "Cycle detected."), &Node, EMessageSeverity::Error, true);
 		Context.CustomizableObjectWithCycle = *Result;
 		return true;
 	}
@@ -114,7 +114,7 @@ void CheckNumOutputs(const UEdGraphPin& Pin, const FMutableGraphGenerationContex
 
 			if (numOutLinks > 1)
 			{
-				GenerationContext.CompilationContext.Log(LOCTEXT("MultipleOutgoing", "The node has several outgoing connections, but it should be limited to 1."), CastChecked<UCustomizableObjectNode>(Pin.GetOwningNode()));
+				GenerationContext.CompilationContext->Log(LOCTEXT("MultipleOutgoing", "The node has several outgoing connections, but it should be limited to 1."), CastChecked<UCustomizableObjectNode>(Pin.GetOwningNode()));
 			}
 		}
 	}
@@ -258,7 +258,7 @@ bool FMutableCompilationContext::FindBone(const FName& InBoneName, mu::FBoneName
 }
 
 
-FMutableGraphGenerationContext::FMutableGraphGenerationContext(FMutableCompilationContext& InCompilationContext)
+FMutableGraphGenerationContext::FMutableGraphGenerationContext(const TSharedPtr<FMutableCompilationContext>& InCompilationContext)
 	: CompilationContext(InCompilationContext)
 	, ExtensionDataCompilerInterface(*this)
 {
@@ -272,19 +272,19 @@ FMutableGraphGenerationContext::FMutableGraphGenerationContext(FMutableCompilati
 
 FString FMutableGraphGenerationContext::GetObjectName()
 {
-	return CompilationContext.GetObjectName();
+	return CompilationContext->GetObjectName();
 }
 
 
 void FMutableGraphGenerationContext::Log(const FText& Message, const TArray<const UObject*>& Context, const EMessageSeverity::Type MessageSeverity, const bool bAddBaseObjectInfo, const ELoggerSpamBin SpamBin) const
 {
-	CompilationContext.Log(Message, Context, MessageSeverity, bAddBaseObjectInfo, SpamBin);
+	CompilationContext->Log(Message, Context, MessageSeverity, bAddBaseObjectInfo, SpamBin);
 }
 
 
 void FMutableGraphGenerationContext::Log(const FText& Message, const UObject* Context, const EMessageSeverity::Type MessageSeverity, const bool bAddBaseObjectInfo, const ELoggerSpamBin SpamBin) const
 {
-	CompilationContext.Log(Message, Context, MessageSeverity, bAddBaseObjectInfo, SpamBin);
+	CompilationContext->Log(Message, Context, MessageSeverity, bAddBaseObjectInfo, SpamBin);
 }
 
 
@@ -371,7 +371,7 @@ const FGuid FMutableGraphGenerationContext::GetNodeIdUnchecked(const UCustomizab
 
 FMutableComponentInfo* FMutableGraphGenerationContext::GetCurrentComponentInfo()
 {
-	return CompilationContext.GetComponentInfo(CurrentMeshComponent);
+	return CompilationContext->GetComponentInfo(CurrentMeshComponent);
 }
 
 
@@ -703,7 +703,7 @@ mu::Ptr<mu::NodeMeshApplyPose> CreateNodeMeshApplyPose(FMutableGraphGenerationCo
 	for (int32 i = 0; i < ArrayBoneName.Num(); ++i)
 	{
 		const FName BoneName = ArrayBoneName[i];
-		const mu::FBoneName& BoneId = GenerationContext.CompilationContext.GetBoneUnique(BoneName);
+		const mu::FBoneName& BoneId = GenerationContext.CompilationContext->GetBoneUnique(BoneName);
 
 		MutableSkeleton->SetDebugName(i, BoneName);
 		MutableSkeleton->SetBoneName(i, { BoneId });
@@ -895,7 +895,7 @@ mu::Ptr<mu::NodeObject> GenerateMutableSource(const UEdGraphPin * Pin, FMutableG
 			int32 MinQualityLevelLOD = INT32_MAX;
 			int32 MaxQualityLevelLOD = 0;
 
-			for (const FMutableComponentInfo& ComponentInfo : GenerationContext.CompilationContext.ComponentInfos)
+			for (const FMutableComponentInfo& ComponentInfo : GenerationContext.CompilationContext->ComponentInfos)
 			{
 				if (!Cast<UCustomizableObjectNodeComponentMesh>(ComponentInfo.Node))
 				{
@@ -903,7 +903,7 @@ mu::Ptr<mu::NodeObject> GenerateMutableSource(const UEdGraphPin * Pin, FMutableG
 				}
 				
 				const FMutableLODSettings& LODSettings = ComponentInfo.Node->LODSettings;
-				const ITargetPlatform* TargetPlatform = GenerationContext.CompilationContext.Options.TargetPlatform;
+				const ITargetPlatform* TargetPlatform = GenerationContext.CompilationContext->Options.TargetPlatform;
 				const FSupportedQualityLevelArray SupportedQualityLevels = LODSettings.MinQualityLevelLOD.GetSupportedQualityLevels(*TargetPlatform->GetPlatformInfo().IniPlatformName.ToString());
 				
 				for (const int32& QualitySetting : SupportedQualityLevels)
@@ -953,7 +953,7 @@ mu::Ptr<mu::NodeObject> GenerateMutableSource(const UEdGraphPin * Pin, FMutableG
 			TArray<UEdGraphPin*> ConnectedModifierPins = FollowInputPinArray(*ModifierPin);
 			for (UEdGraphPin* const ConnectedModifier : ConnectedModifierPins)
 			{
-				for (const FMutableComponentInfo& Pair : GenerationContext.CompilationContext.ComponentInfos)
+				for (const FMutableComponentInfo& Pair : GenerationContext.CompilationContext->ComponentInfos)
 				{
 					GenerationContext.CurrentMeshComponent = Pair.ComponentName;
 					GenerationContext.CurrentAutoLODStrategy = Pair.Node->AutoLODStrategy;
@@ -1114,7 +1114,7 @@ mu::Ptr<mu::NodeObject> GenerateMutableSource(const UEdGraphPin * Pin, FMutableG
 
 			UCustomizableObjectNodeObject* CustomizableObjectNodeObject = Cast<UCustomizableObjectNodeObject>(MacroContextPin->GetOwningNode());
 
-			const FString* SelectedOptionName = GenerationContext.CompilationContext.Options.ParamNamesToSelectedOptions.Find(TypedNodeGroup->GetGroupName(&GenerationContext.MacroNodesStack)); // If the param is in the map restrict to only the selected option
+			const FString* SelectedOptionName = GenerationContext.CompilationContext->Options.ParamNamesToSelectedOptions.Find(TypedNodeGroup->GetGroupName(&GenerationContext.MacroNodesStack)); // If the param is in the map restrict to only the selected option
 			mu::Ptr<mu::NodeObject> ChildNode;
 
 			if (bConnectAtLeastTheLastChild || !SelectedOptionName || (CustomizableObjectNodeObject && *SelectedOptionName == CustomizableObjectNodeObject->GetObjectName(&GenerationContext.MacroNodesStack)) )
@@ -1240,7 +1240,7 @@ mu::Ptr<mu::NodeObject> GenerateMutableSource(const UEdGraphPin * Pin, FMutableG
 
 			UCustomizableObjectNodeObject* CustomizableObjectNodeObject = Cast<UCustomizableObjectNodeObject>(ExternalChildNode->OutputPin()->GetOwningNode());
 
-			const FString* SelectedOptionName = GenerationContext.CompilationContext.Options.ParamNamesToSelectedOptions.Find(TypedNodeGroup->GetGroupName(&GenerationContext.MacroNodesStack)); // If the param is in the map restrict to only the selected option
+			const FString* SelectedOptionName = GenerationContext.CompilationContext->Options.ParamNamesToSelectedOptions.Find(TypedNodeGroup->GetGroupName(&GenerationContext.MacroNodesStack)); // If the param is in the map restrict to only the selected option
 			mu::Ptr<mu::NodeObject> ChildNode;
 
 			if (bConnectAtLeastTheLastChild || !SelectedOptionName || (CustomizableObjectNodeObject  && *SelectedOptionName == CustomizableObjectNodeObject->GetObjectName(&GenerationContext.MacroNodesStack)) )
@@ -1382,14 +1382,14 @@ FString GenerateGameplayTag(const FString& GameplayTag)
 
 void PopulateReferenceSkeletalMeshesData(FMutableGraphGenerationContext& GenerationContext)
 {
-	const FString PlatformName = GenerationContext.CompilationContext.Options.TargetPlatform->IniPlatformName();
+	const FString PlatformName = GenerationContext.CompilationContext->Options.TargetPlatform->IniPlatformName();
 	
-	const uint32 ComponentCount = GenerationContext.CompilationContext.ComponentInfos.Num();
+	const uint32 ComponentCount = GenerationContext.CompilationContext->ComponentInfos.Num();
 
 	GenerationContext.ReferenceSkeletalMeshesData.AddDefaulted(ComponentCount);
 	for (uint32 ComponentIndex = 0; ComponentIndex < ComponentCount; ++ComponentIndex)
 	{
-		const FMutableComponentInfo& ComponentInfo = GenerationContext.CompilationContext.ComponentInfos[ComponentIndex];
+		const FMutableComponentInfo& ComponentInfo = GenerationContext.CompilationContext->ComponentInfos[ComponentIndex];
 		const FName& ComponentName = ComponentInfo.ComponentName;
 		
 		USkeletalMesh* RefSkeletalMesh = ComponentInfo.RefSkeletalMesh.Get();
@@ -1398,8 +1398,8 @@ void PopulateReferenceSkeletalMeshesData(FMutableGraphGenerationContext& Generat
 		FMutableRefSkeletalMeshData& Data = GenerationContext.ReferenceSkeletalMeshesData[ComponentIndex];
 
 		// Set the RefSkeletalMesh
-		if (!GenerationContext.CompilationContext.Options.TargetPlatform->IsClientOnly()
-			|| GenerationContext.CompilationContext.Object->bEnableUseRefSkeletalMeshAsPlaceholder)
+		if (!GenerationContext.CompilationContext->Options.TargetPlatform->IsClientOnly()
+			|| GenerationContext.CompilationContext->Object->bEnableUseRefSkeletalMeshAsPlaceholder)
 		{
 			Data.SkeletalMesh = RefSkeletalMesh;
 		}
@@ -1541,9 +1541,9 @@ uint32 ComputeLODBiasForTexture(const FMutableGraphGenerationContext& Generation
 	constexpr int32 MaxAllowedLODBias = 6;
 
 	// Force a large LODBias for debug
-	if (GenerationContext.CompilationContext.Options.bForceLargeLODBias)
+	if (GenerationContext.CompilationContext->Options.bForceLargeLODBias)
 	{
-		return FMath::Min(GenerationContext.CompilationContext.Options.DebugBias, MaxAllowedLODBias);
+		return FMath::Min(GenerationContext.CompilationContext->Options.DebugBias, MaxAllowedLODBias);
 	}
 
 	// Max size and number of mips from Texture. 
@@ -1562,7 +1562,7 @@ uint32 ComputeLODBiasForTexture(const FMutableGraphGenerationContext& Generation
 		return 0;
 	}
 
-	const UTextureLODSettings& LODSettings = GenerationContext.CompilationContext.Options.TargetPlatform->GetTextureLODSettings();
+	const UTextureLODSettings& LODSettings = GenerationContext.CompilationContext->Options.TargetPlatform->GetTextureLODSettings();
 
 	// Get the MaxTextureSize for the TargetPlatform.
 	const int32 MaxTextureSize = GetMaxTextureSize(ReferenceTexture ? *ReferenceTexture : Texture, LODSettings);
@@ -1678,7 +1678,7 @@ TSharedPtr<mu::FImage> GenerateImageConstant(UTexture* Texture, FMutableGraphGen
 	if (!bIsReference)
 	{
 		bForceLoad = true;
-		if (GenerationContext.CompilationContext.Options.OptimizationLevel == 0)
+		if (GenerationContext.CompilationContext->Options.OptimizationLevel == 0)
 		{
 			bIsCompileTime = false;
 		}
@@ -1744,7 +1744,7 @@ TSharedPtr<mu::FMesh> GenerateMeshConstant(
 		bForceLoad = true;
 		// TODO: For now, even in none, we turn them into compile-time references, since runtime meshes are not
 		// supported yet.
-		//if (GenerationContext.CompilationContext.Options.OptimizationLevel == 0)
+		//if (GenerationContext.CompilationContext->Options.OptimizationLevel == 0)
 		//{
 		//	bIsCompileTime = false;
 		//}
diff --git a/Engine/Plugins/Mutable/Source/CustomizableObjectEditor/Private/MuCOE/GenerateMutableSource/GenerateMutableSource.h b/Engine/Plugins/Mutable/Source/CustomizableObjectEditor/Private/MuCOE/GenerateMutableSource/GenerateMutableSource.h
--- a/Engine/Plugins/Mutable/Source/CustomizableObjectEditor/Private/MuCOE/GenerateMutableSource/GenerateMutableSource.h
+++ b/Engine/Plugins/Mutable/Source/CustomizableObjectEditor/Private/MuCOE/GenerateMutableSource/GenerateMutableSource.h
@@ -569,9 +569,9 @@ uint32 GetTypeHash(const FSharedSurfaceKey& Key);
 */
 struct FMutableGraphGenerationContext
 {
-	FMutableGraphGenerationContext(FMutableCompilationContext& InCompilationContext);
+	FMutableGraphGenerationContext(const TSharedPtr<FMutableCompilationContext>& InCompilationContext);
 	
-	FMutableCompilationContext& CompilationContext;
+	TSharedPtr<FMutableCompilationContext> CompilationContext;
 
 	/** Full hierarchy root. */
 	UCustomizableObjectNodeObject* Root = nullptr;
@@ -1001,4 +1001,4 @@ uint32 AddUniqueSurfaceMetadata(const FMutableSourceSurfaceMetadata& MeshData, F
 
 /** Convert a mesh to Mutable format.
 */
-TSharedPtr<mu::FMesh> ConvertSkeletalMeshToMutable(const FMutableSourceMeshData& Source, FMutableCompilationContext& Context, const FString& MorphName);
+TSharedPtr<mu::FMesh> ConvertSkeletalMeshToMutable(TSharedRef<FMeshConversionContext> MeshConversionContext);
diff --git a/Engine/Plugins/Mutable/Source/CustomizableObjectEditor/Private/MuCOE/GenerateMutableSource/GenerateMutableSourceComponent.cpp b/Engine/Plugins/Mutable/Source/CustomizableObjectEditor/Private/MuCOE/GenerateMutableSource/GenerateMutableSourceComponent.cpp
--- a/Engine/Plugins/Mutable/Source/CustomizableObjectEditor/Private/MuCOE/GenerateMutableSource/GenerateMutableSourceComponent.cpp
+++ b/Engine/Plugins/Mutable/Source/CustomizableObjectEditor/Private/MuCOE/GenerateMutableSource/GenerateMutableSourceComponent.cpp
@@ -136,7 +136,7 @@ mu::Ptr<mu::NodeComponent> GenerateMutableSourceComponent(const UEdGraphPin* Pin
 	
 	if (const UCustomizableObjectNodeComponentMesh* TypedComponentMesh = Cast<UCustomizableObjectNodeComponentMesh>(Node))
 	{
-		if (!GenerationContext.CompilationContext.ComponentInfos.ContainsByPredicate([&](const FMutableComponentInfo& ComponentInfo)
+		if (!GenerationContext.CompilationContext->ComponentInfos.ContainsByPredicate([&](const FMutableComponentInfo& ComponentInfo)
 		{
 			return ComponentInfo.Node == TypedComponentMesh;
 		}))
@@ -179,7 +179,7 @@ mu::Ptr<mu::NodeComponent> GenerateMutableSourceComponent(const UEdGraphPin* Pin
 	{
 		const FName& ParentComponentName = TypedComponentMeshExtend->GetParentComponentName(&GenerationContext.MacroNodesStack);
 
-		if (FMutableComponentInfo* FindResult = GenerationContext.CompilationContext.ComponentInfos.FindByPredicate([&](const FMutableComponentInfo& Element)
+		if (FMutableComponentInfo* FindResult = GenerationContext.CompilationContext->ComponentInfos.FindByPredicate([&](const FMutableComponentInfo& Element)
 		{
 			return Element.ComponentName == ParentComponentName;
 		}))
@@ -482,7 +482,7 @@ void FirstPass(UCustomizableObjectNodeComponentMesh& Node, FMutableGraphGenerati
 		return;
 	}
 	
-	if (FMutableComponentInfo* Result = GenerationContext.CompilationContext.ComponentInfos.FindByPredicate([&](const FMutableComponentInfo& Element)
+	if (FMutableComponentInfo* Result = GenerationContext.CompilationContext->ComponentInfos.FindByPredicate([&](const FMutableComponentInfo& Element)
 	{
 		return Element.ComponentName == ComponentName;
 	}))
@@ -550,7 +550,7 @@ void FirstPass(UCustomizableObjectNodeComponentMesh& Node, FMutableGraphGenerati
 		// Array of Quality values for each of the scalability settings. The valid values go from [0,4] each one representing one Scalability setting:
 		//		0 = Low
 		//		4 = Cine = Cinematic
-		FSupportedQualityLevelArray SupportedQualityLevels = LODSettings.MinQualityLevelLOD.GetSupportedQualityLevels(*GenerationContext.CompilationContext.Options.TargetPlatform->GetPlatformInfo().IniPlatformName.ToString());
+		FSupportedQualityLevelArray SupportedQualityLevels = LODSettings.MinQualityLevelLOD.GetSupportedQualityLevels(*GenerationContext.CompilationContext->Options.TargetPlatform->GetPlatformInfo().IniPlatformName.ToString());
 			
 		// If no scalability settings are found, use the MaxLOD as the MinLOD to be used.
 		int32 MinValue = NumLODs - 1;
@@ -572,7 +572,7 @@ void FirstPass(UCustomizableObjectNodeComponentMesh& Node, FMutableGraphGenerati
 	}
 	else
 	{
-		FirstLODAvailable = LODSettings.MinLOD.GetValueForPlatform(*GenerationContext.CompilationContext.Options.TargetPlatform->IniPlatformName());
+		FirstLODAvailable = LODSettings.MinLOD.GetValueForPlatform(*GenerationContext.CompilationContext->Options.TargetPlatform->IniPlatformName());
 	}
 
 	FirstLODAvailable = FMath::Clamp(FirstLODAvailable, 0, NumLODs - 1);
@@ -582,20 +582,20 @@ void FirstPass(UCustomizableObjectNodeComponentMesh& Node, FMutableGraphGenerati
 	// Find the streaming settings for the target platform
 	if (LODSettings.bOverrideLODStreamingSettings)
 	{
-		GenerationContext.bEnableLODStreaming = LODSettings.bEnableLODStreaming.GetValueForPlatform(*GenerationContext.CompilationContext.Options.TargetPlatform->IniPlatformName());
-		NumMaxLODsToStream = LODSettings.NumMaxStreamedLODs.GetValueForPlatform(*GenerationContext.CompilationContext.Options.TargetPlatform->IniPlatformName());
+		GenerationContext.bEnableLODStreaming = LODSettings.bEnableLODStreaming.GetValueForPlatform(*GenerationContext.CompilationContext->Options.TargetPlatform->IniPlatformName());
+		NumMaxLODsToStream = LODSettings.NumMaxStreamedLODs.GetValueForPlatform(*GenerationContext.CompilationContext->Options.TargetPlatform->IniPlatformName());
 	}
 	else
 	{
-		for (int32 MeshIndex = 0; MeshIndex < GenerationContext.CompilationContext.ComponentInfos.Num(); ++MeshIndex)
+		for (int32 MeshIndex = 0; MeshIndex < GenerationContext.CompilationContext->ComponentInfos.Num(); ++MeshIndex)
 		{
-			RefSkeletalMesh = GenerationContext.CompilationContext.ComponentInfos[MeshIndex].RefSkeletalMesh.Get();
+			RefSkeletalMesh = GenerationContext.CompilationContext->ComponentInfos[MeshIndex].RefSkeletalMesh.Get();
 			check(RefSkeletalMesh);
 
 			GenerationContext.bEnableLODStreaming = GenerationContext.bEnableLODStreaming &&
-				RefSkeletalMesh->GetEnableLODStreaming(GenerationContext.CompilationContext.Options.TargetPlatform);
+				RefSkeletalMesh->GetEnableLODStreaming(GenerationContext.CompilationContext->Options.TargetPlatform);
 
-			NumMaxLODsToStream = FMath::Min(NumMaxLODsToStream, static_cast<uint8>(RefSkeletalMesh->GetMaxNumStreamedLODs(GenerationContext.CompilationContext.Options.TargetPlatform)));
+			NumMaxLODsToStream = FMath::Min(NumMaxLODsToStream, static_cast<uint8>(RefSkeletalMesh->GetMaxNumStreamedLODs(GenerationContext.CompilationContext->Options.TargetPlatform)));
 		}
 	}
 
@@ -611,10 +611,10 @@ void FirstPass(UCustomizableObjectNodeComponentMesh& Node, FMutableGraphGenerati
 	ComponentInfo.AccumulateBonesToRemovePerLOD(Node.LODReductionSettings, Node.NumLODs);
 	ComponentInfo.MacroContext = GenerationContext.MacroNodesStack;
 
-	GenerationContext.CompilationContext.ComponentInfos.Add(ComponentInfo);
+	GenerationContext.CompilationContext->ComponentInfos.Add(ComponentInfo);
 
 	// Make sure the Skeleton from the reference mesh is added to the list of referenced Skeletons.
-	GenerationContext.CompilationContext.ReferencedSkeletons.Add(RefSkeleton);
+	GenerationContext.CompilationContext->ReferencedSkeletons.Add(RefSkeleton);
 
 	GenerationContext.ComponentNames.Add(ComponentName);
 }
diff --git a/Engine/Plugins/Mutable/Source/CustomizableObjectEditor/Private/MuCOE/GenerateMutableSource/GenerateMutableSourceGroupProjector.cpp b/Engine/Plugins/Mutable/Source/CustomizableObjectEditor/Private/MuCOE/GenerateMutableSource/GenerateMutableSourceGroupProjector.cpp
--- a/Engine/Plugins/Mutable/Source/CustomizableObjectEditor/Private/MuCOE/GenerateMutableSource/GenerateMutableSourceGroupProjector.cpp
+++ b/Engine/Plugins/Mutable/Source/CustomizableObjectEditor/Private/MuCOE/GenerateMutableSource/GenerateMutableSourceGroupProjector.cpp
@@ -310,7 +310,7 @@ TOptional<FGroupProjectorTempData> GenerateMutableGroupProjector(UCustomizableOb
 
 	FGroupProjectorTempData GroupProjectorTempData;
 	
-	const FCompilationOptions& CompilationOptions = GenerationContext.CompilationContext.Options;
+	const FCompilationOptions& CompilationOptions = GenerationContext.CompilationContext->Options;
 
 	// The static cast works here because it's already known to be a mu::NodeProjectorParameter* because of the UE5 Cast in the previous line
 	GroupProjectorTempData.NodeProjectorParameterPtr = static_cast<mu::NodeProjectorParameter*>(GenerateMutableSourceProjector(&ProjParamNode->OutputPin(), GenerationContext).get());
@@ -509,7 +509,7 @@ TOptional<FGroupProjectorTempData> GenerateMutableGroupProjector(UCustomizableOb
 				// Calculate the number of mips to tag as high res for this image.
 				if (ProjParamNode->ReferenceTexture)
 				{
-					const UTextureLODSettings& TextureLODSettings = GenerationContext.CompilationContext.Options.TargetPlatform->GetTextureLODSettings();
+					const UTextureLODSettings& TextureLODSettings = GenerationContext.CompilationContext->Options.TargetPlatform->GetTextureLODSettings();
 					const FTextureLODGroup& LODGroupInfo = TextureLODSettings.GetTextureLODGroup(ProjParamNode->ReferenceTexture->LODGroup);
 
 					ImageNode->SourceDataDescriptor.OptionalMaxLODSize = LODGroupInfo.OptionalMaxLODSize;
diff --git a/Engine/Plugins/Mutable/Source/CustomizableObjectEditor/Private/MuCOE/GenerateMutableSource/GenerateMutableSourceImage.cpp b/Engine/Plugins/Mutable/Source/CustomizableObjectEditor/Private/MuCOE/GenerateMutableSource/GenerateMutableSourceImage.cpp
--- a/Engine/Plugins/Mutable/Source/CustomizableObjectEditor/Private/MuCOE/GenerateMutableSource/GenerateMutableSourceImage.cpp
+++ b/Engine/Plugins/Mutable/Source/CustomizableObjectEditor/Private/MuCOE/GenerateMutableSource/GenerateMutableSourceImage.cpp
@@ -107,7 +107,7 @@ mu::Ptr<mu::NodeImage> GenerateMutableSourceImage(const UEdGraphPin* Pin, FMutab
 
 	mu::Ptr<mu::NodeImage> Result;
 
-	const FCompilationOptions& CompilationOptions = GenerationContext.CompilationContext.Options;
+	const FCompilationOptions& CompilationOptions = GenerationContext.CompilationContext->Options;
 
 	if (const UCustomizableObjectNodeTexture* TypedNodeTex = Cast<UCustomizableObjectNodeTexture>(Node))
 	{
@@ -131,7 +131,7 @@ mu::Ptr<mu::NodeImage> GenerateMutableSourceImage(const UEdGraphPin* Pin, FMutab
 				mu::Ptr<mu::NodeImageConstant> ConstantImageNode = new mu::NodeImageConstant;
 				ConstantImageNode->SetValue( ImageConstant );
 				
-				const UTextureLODSettings& TextureLODSettings = GenerationContext.CompilationContext.Options.TargetPlatform->GetTextureLODSettings();
+				const UTextureLODSettings& TextureLODSettings = GenerationContext.CompilationContext->Options.TargetPlatform->GetTextureLODSettings();
 				const FTextureLODGroup& LODGroupInfo = TextureLODSettings.GetTextureLODGroup(BaseTexture->LODGroup);
 
 				ConstantImageNode->SourceDataDescriptor.OptionalMaxLODSize = LODGroupInfo.OptionalMaxLODSize;
@@ -234,7 +234,7 @@ mu::Ptr<mu::NodeImage> GenerateMutableSourceImage(const UEdGraphPin* Pin, FMutab
 			const uint32 MipsToSkip = ComputeLODBiasForTexture(GenerationContext, *Texture, nullptr, ReferenceTextureSize);
 			Result = ResizeTextureByNumMips(ImageNode, MipsToSkip);
 
-			const UTextureLODSettings& TextureLODSettings = GenerationContext.CompilationContext.Options.TargetPlatform->GetTextureLODSettings();
+			const UTextureLODSettings& TextureLODSettings = GenerationContext.CompilationContext->Options.TargetPlatform->GetTextureLODSettings();
 			const FTextureLODGroup& LODGroupInfo = TextureLODSettings.GetTextureLODGroup(Texture->LODGroup);
 
 			ImageNode->SourceDataDescriptor.OptionalMaxLODSize = LODGroupInfo.OptionalMaxLODSize;
@@ -1137,7 +1137,7 @@ mu::Ptr<mu::NodeImage> GenerateMutableSourceImage(const UEdGraphPin* Pin, FMutab
 								ImageTableNode->MaxTextureSize = MaxTextureSize;
 								ImageTableNode->ReferenceImageDesc = ImageDesc;
 				
-								const UTextureLODSettings& TextureLODSettings = GenerationContext.CompilationContext.Options.TargetPlatform->GetTextureLODSettings();
+								const UTextureLODSettings& TextureLODSettings = GenerationContext.CompilationContext->Options.TargetPlatform->GetTextureLODSettings();
 								const FTextureLODGroup& LODGroupInfo = TextureLODSettings.GetTextureLODGroup(DefaultTexture2D->LODGroup);
 
 								ImageTableNode->SourceDataDescriptor.OptionalMaxLODSize = LODGroupInfo.OptionalMaxLODSize;
diff --git a/Engine/Plugins/Mutable/Source/CustomizableObjectEditor/Private/MuCOE/GenerateMutableSource/GenerateMutableSourceMaterial.cpp b/Engine/Plugins/Mutable/Source/CustomizableObjectEditor/Private/MuCOE/GenerateMutableSource/GenerateMutableSourceMaterial.cpp
--- a/Engine/Plugins/Mutable/Source/CustomizableObjectEditor/Private/MuCOE/GenerateMutableSource/GenerateMutableSourceMaterial.cpp
+++ b/Engine/Plugins/Mutable/Source/CustomizableObjectEditor/Private/MuCOE/GenerateMutableSource/GenerateMutableSourceMaterial.cpp
@@ -228,7 +228,7 @@ mu::Ptr<mu::NodeMaterial> GenerateMutableSourceMaterial(const UEdGraphPin* Pin,
 					check(GenerationContext.ReferencedMaterials.IsValidIndex(GenerationContext.CurrentReferencedMaterialIndex));
 					UMaterialInterface* BaseMaterial = GenerationContext.ReferencedMaterials[GenerationContext.CurrentReferencedMaterialIndex];
 
-					const bool bTableMaterialCheckDisabled = GenerationContext.CompilationContext.Object->GetPrivate()->IsTableMaterialsParentCheckDisabled();
+					const bool bTableMaterialCheckDisabled = GenerationContext.CompilationContext->Object->GetPrivate()->IsTableMaterialsParentCheckDisabled();
 					const bool bMaterialParentMismatch = !bTableMaterialCheckDisabled && BaseMaterial
 						&& TableMaterial->GetMaterial() != BaseMaterial->GetMaterial();
 
diff --git a/Engine/Plugins/Mutable/Source/CustomizableObjectEditor/Private/MuCOE/GenerateMutableSource/GenerateMutableSourceMesh.cpp b/Engine/Plugins/Mutable/Source/CustomizableObjectEditor/Private/MuCOE/GenerateMutableSource/GenerateMutableSourceMesh.cpp
--- a/Engine/Plugins/Mutable/Source/CustomizableObjectEditor/Private/MuCOE/GenerateMutableSource/GenerateMutableSourceMesh.cpp
+++ b/Engine/Plugins/Mutable/Source/CustomizableObjectEditor/Private/MuCOE/GenerateMutableSource/GenerateMutableSourceMesh.cpp
@@ -589,6 +589,7 @@ TArray<TTuple<UPhysicsAsset*, int32>> GetPhysicsAssetsFromAnimInstance(const TSo
 		return Result;
 	}
 
+	// TODO: Get Physics Asset using AssetRegistry
 	UClass* AnimInstanceClass = MutablePrivate::LoadClass(AnimInstance);
 	UAnimBlueprintGeneratedClass* AnimClass = Cast<UAnimBlueprintGeneratedClass>(AnimInstanceClass);
 
@@ -861,16 +862,16 @@ namespace MutablePrivate
 }
 
 
-TSharedPtr<mu::FMesh> ConvertSkeletalMeshToMutable(
-	const FMutableSourceMeshData& Source,
-	FMutableCompilationContext& Context,
-	const FString& MorphName
-)
+TSharedPtr<mu::FMesh> ConvertSkeletalMeshToMutable(TSharedRef<FMeshConversionContext> MeshConversionContext)
 {
+	const FMutableSourceMeshData& Source = MeshConversionContext->Source;
+	FMutableCompilationContext& Context = *MeshConversionContext->CompilationContext;
+	const FString& MorphName = MeshConversionContext->MorphName;
+
 	const FString MeshName = Source.Mesh.GetLongPackageName().ToLower();
 	const uint32 MeshId = CityHash32(reinterpret_cast<const char*>(*MeshName), MeshName.Len() * sizeof(FString::ElementType));
 
-	const USkeletalMesh* SkeletalMesh = Cast<USkeletalMesh>( MutablePrivate::LoadObject(Source.Mesh) );
+	const USkeletalMesh* SkeletalMesh = Cast<USkeletalMesh>(Source.Mesh.Get());
 	check(SkeletalMesh);
 	const FSkeletalMeshModel* ImportedModel = SkeletalMesh->GetImportedModel();
 	if (!ImportedModel)
@@ -1469,11 +1470,12 @@ TSharedPtr<mu::FMesh> ConvertSkeletalMeshToMutable(
 				}
 				else
 				{
+					USkeletalMesh* TableReferenceSkeletalMesh = Cast<USkeletalMesh>(Source.TableReferenceSkeletalMesh.Get());
 					for (const UMorphTarget* MorphTarget : SkeletalMeshMorphTargets)
 					{
-						check(MorphTarget);
+						check(MorphTarget && TableReferenceSkeletalMesh);
 
-						if (MutablePrivate::LoadObject(Source.TableReferenceSkeletalMesh)->FindMorphTarget(MorphTarget->GetFName()))
+						if (TableReferenceSkeletalMesh->FindMorphTarget(MorphTarget->GetFName()))
 						{
 							MorphTargetsNames.Add(MorphTarget->GetFName());
 						}
@@ -2601,7 +2603,7 @@ TSharedPtr<mu::FMesh> ConvertSkeletalMeshToMutable(
 		}
 
 		// Add clothing asset participating objects
-		if (GenerationContext.CompilationContext.Options.bClothingEnabled)
+		if (GenerationContext.CompilationContext->Options.bClothingEnabled)
 		{
 			for (const FAssetData& Dep : Assets)
 			{
@@ -2636,8 +2638,17 @@ TSharedPtr<mu::FMesh> ConvertSkeletalMeshToMutable(
 	else
 	{
 		// At some point this shouldn't happen anymore because all mesh conversion could be moved to the core compilation stage.
-		FString NoMorph;
-		MutableMesh = ConvertSkeletalMeshToMutable(Source, GenerationContext.CompilationContext, NoMorph);
+		TSharedRef<FMeshConversionContext> MeshConversionContext = MakeShared<FMeshConversionContext>();
+		MeshConversionContext->Source = Source;
+		//MeshConversionContext->MorphName = FString(); // NoMorph
+		MeshConversionContext->CompilationContext = GenerationContext.CompilationContext;
+
+		// Load assets sync
+		MeshConversionContext->ReferencedObjects.Emplace(TStrongObjectPtr<UObject>(MutablePrivate::LoadObject(Source.Mesh)));
+		MeshConversionContext->ReferencedObjects.Emplace(TStrongObjectPtr<UObject>(MutablePrivate::LoadObject(Source.TableReferenceSkeletalMesh)));
+		MeshConversionContext->ReferencedObjects.Emplace(TStrongObjectPtr<UObject>(MutablePrivate::LoadClass(Source.AnimInstance)));
+
+		MutableMesh = ConvertSkeletalMeshToMutable(MeshConversionContext);
 	}
 
 	return MutableMesh;
@@ -2935,8 +2946,7 @@ TSharedPtr<mu::FMesh> BuildMorphedMutableMesh(const UEdGraphPin* BaseSourcePin,
 
 	else if (const ICustomizableObjectNodeMeshInterface* TypedNodeSkeletalMesh = Cast<ICustomizableObjectNodeMeshInterface>(Node))
 	{
-		int32 LayoutIndex;
-		TypedNodeSkeletalMesh->GetPinSection(*BaseSourcePin, LODIndexConnected, SectionIndexConnected, LayoutIndex);
+		TypedNodeSkeletalMesh->GetPinSection(*BaseSourcePin, LODIndexConnected, SectionIndexConnected);
 		SkeletalMesh = Cast<USkeletalMesh>(MutablePrivate::LoadObject(TypedNodeSkeletalMesh->GetMesh()));
 	}
 	
@@ -3269,7 +3279,7 @@ bool GetAndValidateReshapePhysicsToDeform(
 	{
 		TArray<PhysicsInfoType> PhysicsAssetsInfo;
 
-		const bool bAnimBpOverridePhysicsManipulationEnabled = GenerationContext.CompilationContext.Options.bAnimBpPhysicsManipulationEnabled;
+		const bool bAnimBpOverridePhysicsManipulationEnabled = GenerationContext.CompilationContext->Options.bAnimBpPhysicsManipulationEnabled;
 		for (const MeshInfoType& Mesh : SkeletalMeshesInfo)
 		{
 			const USkeletalMesh* SkeletalMesh = Mesh.Get<USkeletalMesh*>();
@@ -3624,7 +3634,7 @@ mu::NodeMeshPtr GenerateMorphMesh(const UEdGraphPin* Pin,
 
 						for (const FName BoneName : BonesToDeform)
 						{
-							Result->BonesToDeform.Add(GenerationContext.CompilationContext.GetBoneUnique(BoneName));
+							Result->BonesToDeform.Add(GenerationContext.CompilationContext->GetBoneUnique(BoneName));
 						}
 					}
 
@@ -3640,7 +3650,7 @@ mu::NodeMeshPtr GenerateMorphMesh(const UEdGraphPin* Pin,
 	
 						for (const FName& PhysicsBoneName : PhysicsToDeform)
 						{
-							Result->PhysicsToDeform.Add(GenerationContext.CompilationContext.GetBoneUnique(PhysicsBoneName));
+							Result->PhysicsToDeform.Add(GenerationContext.CompilationContext->GetBoneUnique(PhysicsBoneName));
 						}
 					}
 						
@@ -3808,8 +3818,7 @@ mu::Ptr<mu::NodeMesh> GenerateMutableSourceMesh(const UEdGraphPin* Pin,
 				int32 SectionIndex = -1;
 
 				{
-					int32 LayoutIndex;
-					TypedNodeSkel->GetPinSection(*Pin, LODIndexConnected, SectionIndexConnected, LayoutIndex);
+					TypedNodeSkel->GetPinSection(*Pin, LODIndexConnected, SectionIndexConnected);
 				}
 
 
@@ -3840,7 +3849,7 @@ mu::Ptr<mu::NodeMesh> GenerateMutableSourceMesh(const UEdGraphPin* Pin,
 
 				TArray<FCustomizableObjectStreameableResourceId> StreamedResources;
 
-				if (GenerationContext.CompilationContext.Object->GetPrivate()->IsAssetUserDataMergeEnabled())
+				if (GenerationContext.CompilationContext->Object->GetPrivate()->IsAssetUserDataMergeEnabled())
 				{
 					const TArray<UAssetUserData*>* AssetUserDataArray = SkeletalMesh->GetAssetUserDataArray();
 
@@ -4328,7 +4337,7 @@ mu::Ptr<mu::NodeMesh> GenerateMutableSourceMesh(const UEdGraphPin* Pin,
 				
 				for (const FName& BoneName : BonesToDeform)
 				{
-					MeshNode->BonesToDeform.Add(GenerationContext.CompilationContext.GetBoneUnique(BoneName));
+					MeshNode->BonesToDeform.Add(GenerationContext.CompilationContext->GetBoneUnique(BoneName));
 				}
 			}
 
@@ -4343,7 +4352,7 @@ mu::Ptr<mu::NodeMesh> GenerateMutableSourceMesh(const UEdGraphPin* Pin,
 
 				for (const FName& PhysicsBoneName : PhysicsToDeform)
 				{
-					MeshNode->PhysicsToDeform.Add(GenerationContext.CompilationContext.GetBoneUnique(PhysicsBoneName));
+					MeshNode->PhysicsToDeform.Add(GenerationContext.CompilationContext->GetBoneUnique(PhysicsBoneName));
 				}	
 			}
 			
@@ -4614,8 +4623,7 @@ mu::Ptr<mu::NodeMesh> GenerateMutableSourceMesh(const UEdGraphPin* Pin,
 		int32 SectionIndex = 0;
 		
 		{
-			int32 LayoutIndex = 0;
-			TypedNodeParam->GetPinSection(*Pin, LODIndex, SectionIndex, LayoutIndex);
+			TypedNodeParam->GetPinSection(*Pin, LODIndex, SectionIndex);
 			ParameterNode->LODIndex = LODIndex;
 			ParameterNode->SectionIndex = SectionIndex;
 			ParameterNode->bAutomaticLODs = GenerationContext.CurrentAutoLODStrategy == ECustomizableObjectAutomaticLODStrategy::AutomaticFromMesh;
diff --git a/Engine/Plugins/Mutable/Source/CustomizableObjectEditor/Private/MuCOE/GenerateMutableSource/GenerateMutableSourceModifier.cpp b/Engine/Plugins/Mutable/Source/CustomizableObjectEditor/Private/MuCOE/GenerateMutableSource/GenerateMutableSourceModifier.cpp
--- a/Engine/Plugins/Mutable/Source/CustomizableObjectEditor/Private/MuCOE/GenerateMutableSource/GenerateMutableSourceModifier.cpp
+++ b/Engine/Plugins/Mutable/Source/CustomizableObjectEditor/Private/MuCOE/GenerateMutableSource/GenerateMutableSourceModifier.cpp
@@ -93,7 +93,7 @@ mu::Ptr<mu::NodeModifier> GenerateMutableSourceModifier(const UEdGraphPin * Pin,
 		ClipNode->SetParams(TypedNodeClip->B, TypedNodeClip->Exponent);
 		ClipNode->SetMorphEllipse(TypedNodeClip->Radius, TypedNodeClip->Radius2, TypedNodeClip->RotationAngle);
 
-		ClipNode->SetVertexSelectionBone(GenerationContext.CompilationContext.GetBoneUnique(TypedNodeClip->BoneName), TypedNodeClip->MaxEffectRadius);
+		ClipNode->SetVertexSelectionBone(GenerationContext.CompilationContext->GetBoneUnique(TypedNodeClip->BoneName), TypedNodeClip->MaxEffectRadius);
 
 		ClipNode->MultipleTagsPolicy = TypedNodeClip->MultipleTagPolicy;
 		ClipNode->RequiredTags = TypedNodeClip->GetNodeRequiredTags(&GenerationContext.MacroNodesStack);
@@ -281,13 +281,13 @@ mu::Ptr<mu::NodeModifier> GenerateMutableSourceModifier(const UEdGraphPin * Pin,
 					//NOTE: This is the same is done in GenerateMutableSourceSurface. 
 					if (const UEdGraphPin* SkeletalMeshPin = FindMeshBaseSource(*ConnectedPin, false, &GenerationContext.MacroNodesStack))
 					{
-						int32 MetadataLODIndex, MetadataSectionIndex, MetadataLayoutIndex;
-						MetadataLODIndex = MetadataSectionIndex = MetadataLayoutIndex = INDEX_NONE;
+						int32 MetadataLODIndex, MetadataSectionIndex;
+						MetadataLODIndex = MetadataSectionIndex = INDEX_NONE;
 
 						if (const UCustomizableObjectNodeSkeletalMesh* SkeletalMeshNode = Cast<UCustomizableObjectNodeSkeletalMesh>(SkeletalMeshPin->GetOwningNode()))
 						{
 							MeshData.Metadata.Mesh = SkeletalMeshNode->GetMesh().ToSoftObjectPath();
-							SkeletalMeshNode->GetPinSection(*SkeletalMeshPin, MetadataLODIndex, MetadataSectionIndex, MetadataLayoutIndex);
+							SkeletalMeshNode->GetPinSection(*SkeletalMeshPin, MetadataLODIndex, MetadataSectionIndex);
 						}
 						else if (const UCustomizableObjectNodeTable* TableNode = Cast<UCustomizableObjectNodeTable>(SkeletalMeshPin->GetOwningNode()))
 						{
diff --git a/Engine/Plugins/Mutable/Source/CustomizableObjectEditor/Private/MuCOE/GenerateMutableSource/GenerateMutableSourceSurface.cpp b/Engine/Plugins/Mutable/Source/CustomizableObjectEditor/Private/MuCOE/GenerateMutableSource/GenerateMutableSourceSurface.cpp
--- a/Engine/Plugins/Mutable/Source/CustomizableObjectEditor/Private/MuCOE/GenerateMutableSource/GenerateMutableSourceSurface.cpp
+++ b/Engine/Plugins/Mutable/Source/CustomizableObjectEditor/Private/MuCOE/GenerateMutableSource/GenerateMutableSourceSurface.cpp
@@ -148,13 +148,13 @@ mu::Ptr<mu::NodeSurface> GenerateMutableSourceSurface(const UEdGraphPin * Pin, F
 			//NOTE: This is the same is done in GenerateMutableSourceSurface. 
 			if (const UEdGraphPin* SkeletalMeshPin = FindMeshBaseSource(*ConnectedMaterialPin, false, &GenerationContext.MacroNodesStack))
 			{
-				int32 MetadataLODIndex, MetadataSectionIndex, MetadataLayoutIndex;
-				MetadataLODIndex = MetadataSectionIndex = MetadataLayoutIndex = INDEX_NONE;
+				int32 MetadataLODIndex, MetadataSectionIndex;
+				MetadataLODIndex = MetadataSectionIndex = INDEX_NONE;
 
 				if (const UCustomizableObjectNodeSkeletalMesh* SkeletalMeshNode = Cast<UCustomizableObjectNodeSkeletalMesh>(SkeletalMeshPin->GetOwningNode()))
 				{
 					MeshData.Metadata.Mesh = SkeletalMeshNode->GetMesh().ToSoftObjectPath();
-					SkeletalMeshNode->GetPinSection(*SkeletalMeshPin, MetadataLODIndex, MetadataSectionIndex, MetadataLayoutIndex);
+					SkeletalMeshNode->GetPinSection(*SkeletalMeshPin, MetadataLODIndex, MetadataSectionIndex);
 				}
 				else if (const UCustomizableObjectNodeTable* TableNode = Cast<UCustomizableObjectNodeTable>(SkeletalMeshPin->GetOwningNode()))
 				{
@@ -364,7 +364,7 @@ mu::Ptr<mu::NodeSurface> GenerateMutableSourceSurface(const UEdGraphPin * Pin, F
 
 
 								// MaxTextureSize setting. Based on the ReferenceTexture and Platform settings.
-								const UTextureLODSettings& TextureLODSettings = GenerationContext.CompilationContext.Options.TargetPlatform->GetTextureLODSettings();
+								const UTextureLODSettings& TextureLODSettings = GenerationContext.CompilationContext->Options.TargetPlatform->GetTextureLODSettings();
 								Props.MaxTextureSize = GetMaxTextureSize(*ReferenceTexture, TextureLODSettings);
 
 								// ReferenceTexture source size. Textures contributing to this Image should be equal to or smaller than TextureSize. 
@@ -424,12 +424,12 @@ mu::Ptr<mu::NodeSurface> GenerateMutableSourceSurface(const UEdGraphPin * Pin, F
 										const uint32 MipsToSkip = ComputeLODBiasForTexture(GenerationContext, *Texture2D, nullptr, Props.TextureSize);
 										mu::Ptr<mu::NodeImage> Result =  ResizeTextureByNumMips(ConstImageNode, MipsToSkip);
 	
-										const UTextureLODSettings& TextureLODSettings = GenerationContext.CompilationContext.Options.TargetPlatform->GetTextureLODSettings();
+										const UTextureLODSettings& TextureLODSettings = GenerationContext.CompilationContext->Options.TargetPlatform->GetTextureLODSettings();
 										const FTextureLODGroup& LODGroupInfo = TextureLODSettings.GetTextureLODGroup(Texture2D->LODGroup);
 
 										ConstImageNode->SourceDataDescriptor.OptionalMaxLODSize = LODGroupInfo.OptionalMaxLODSize;
 										ConstImageNode->SourceDataDescriptor.OptionalLODBias = LODGroupInfo.OptionalLODBias;
-										ConstImageNode->SourceDataDescriptor.NumNonOptionalLODs = GenerationContext.CompilationContext.Options.MinDiskMips;
+										ConstImageNode->SourceDataDescriptor.NumNonOptionalLODs = GenerationContext.CompilationContext->Options.MinDiskMips;
 
 										const FString TextureName = GetNameSafe(Texture2D).ToLower();
 										ConstImageNode->SourceDataDescriptor.SourceId = CityHash32(reinterpret_cast<const char*>(*TextureName), TextureName.Len() * sizeof(FString::ElementType));
@@ -528,14 +528,14 @@ mu::Ptr<mu::NodeSurface> GenerateMutableSourceSurface(const UEdGraphPin * Pin, F
 									CompositedImage->Normal = NormalCompositeMipmapImage;
 
 									CompositeNormalImage->SourceDataDescriptor.OptionalMaxLODSize = 0;
-									if (GenerationContext.CompilationContext.Options.TargetPlatform)
+									if (GenerationContext.CompilationContext->Options.TargetPlatform)
 									{
-										const UTextureLODSettings& TextureLODSettings = GenerationContext.CompilationContext.Options.TargetPlatform->GetTextureLODSettings();
+										const UTextureLODSettings& TextureLODSettings = GenerationContext.CompilationContext->Options.TargetPlatform->GetTextureLODSettings();
 										const FTextureLODGroup& LODGroupInfo = TextureLODSettings.GetTextureLODGroup(ReferenceCompositeNormalTexture->LODGroup);
 
 										CompositeNormalImage->SourceDataDescriptor.OptionalMaxLODSize = LODGroupInfo.OptionalMaxLODSize;
 										CompositeNormalImage->SourceDataDescriptor.OptionalLODBias = LODGroupInfo.OptionalLODBias;
-										CompositeNormalImage->SourceDataDescriptor.NumNonOptionalLODs = GenerationContext.CompilationContext.Options.MinDiskMips;
+										CompositeNormalImage->SourceDataDescriptor.NumNonOptionalLODs = GenerationContext.CompilationContext->Options.MinDiskMips;
 									}
 
 									const FString TextureName = GetNameSafe(ReferenceCompositeNormalTexture).ToLower();
@@ -553,17 +553,17 @@ mu::Ptr<mu::NodeSurface> GenerateMutableSourceSurface(const UEdGraphPin * Pin, F
 							LastImage = FormatImage;
 
 							TArray<TArray<FTextureBuildSettings>> BuildSettingsPerFormatPerLayer;
-							if (GenerationContext.CompilationContext.Options.TargetPlatform)
+							if (GenerationContext.CompilationContext->Options.TargetPlatform)
 							{
-								ReferenceTexture->GetTargetPlatformBuildSettings(GenerationContext.CompilationContext.Options.TargetPlatform, BuildSettingsPerFormatPerLayer);
+								ReferenceTexture->GetTargetPlatformBuildSettings(GenerationContext.CompilationContext->Options.TargetPlatform, BuildSettingsPerFormatPerLayer);
 
-								bool bIsServerOnly = GenerationContext.CompilationContext.Options.TargetPlatform->IsServerOnly();
+								bool bIsServerOnly = GenerationContext.CompilationContext->Options.TargetPlatform->IsServerOnly();
 								// Suppress the message for server only platforms. Images are discarded is ServerOnly. 
 								if (!bIsServerOnly && BuildSettingsPerFormatPerLayer.IsEmpty())
 								{
 									const FString ReplacedImageFormatMsg = FString::Printf(TEXT("In object [%s] for platform [%s] the unsupported image format of texture [%s] is used, RGBA_UByte will be used instead."),
 										*GenerationContext.GetObjectName(),
-										*GenerationContext.CompilationContext.Options.TargetPlatform->PlatformName(),
+										*GenerationContext.CompilationContext->Options.TargetPlatform->PlatformName(),
 										*ReferenceTexture->GetName());
 									const FText ReplacedImageFormatText = FText::FromString(ReplacedImageFormatMsg);
 									GenerationContext.Log(ReplacedImageFormatText, Node, EMessageSeverity::Info);
@@ -573,7 +573,7 @@ mu::Ptr<mu::NodeSurface> GenerateMutableSourceSurface(const UEdGraphPin * Pin, F
 								{
 									const FString ReplacedImageFormatMsg = FString::Printf(TEXT("In object [%s] for platform [%s] the image format of texture [%s] has multiple target formats. Only one will be used.."),
 										*GenerationContext.GetObjectName(),
-										*GenerationContext.CompilationContext.Options.TargetPlatform->PlatformName(),
+										*GenerationContext.CompilationContext->Options.TargetPlatform->PlatformName(),
 										*ReferenceTexture->GetName());
 									const FText ReplacedImageFormatText = FText::FromString(ReplacedImageFormatMsg);
 									GenerationContext.Log(ReplacedImageFormatText, Node, EMessageSeverity::Info);
@@ -585,7 +585,7 @@ mu::Ptr<mu::NodeSurface> GenerateMutableSourceSurface(const UEdGraphPin * Pin, F
 							{
 								const TArray<FTextureBuildSettings>& BuildSettingsPerLayer = BuildSettingsPerFormatPerLayer[0];
 
-								if (GenerationContext.CompilationContext.Options.TextureCompression!=ECustomizableObjectTextureCompression::None)
+								if (GenerationContext.CompilationContext->Options.TextureCompression!=ECustomizableObjectTextureCompression::None)
 								{
 									static ITextureFormatManagerModule* TextureFormatManager = nullptr;
 									if (!TextureFormatManager)
@@ -603,9 +603,9 @@ mu::Ptr<mu::NodeSurface> GenerateMutableSourceSurface(const UEdGraphPin * Pin, F
 									mu::EImageFormat MutableFormatIfAlpha = QualityAndPerformanceFix(UnrealToMutablePixelFormat(UnrealTargetPlatformFormatAlpha,true));
 
 									// Temp hack to enable RG->LA 
-									if (GenerationContext.CompilationContext.Options.TargetPlatform)
+									if (GenerationContext.CompilationContext->Options.TargetPlatform)
 									{
-										bool bUseLA = GenerationContext.CompilationContext.Options.TargetPlatform->SupportsFeature(ETargetPlatformFeatures::NormalmapLAEncodingMode);
+										bool bUseLA = GenerationContext.CompilationContext->Options.TargetPlatform->SupportsFeature(ETargetPlatformFeatures::NormalmapLAEncodingMode);
 										if (bUseLA)
 										{
 											// We'll have to trust the reference texture because the actual internal settings are opaque.
diff --git a/Engine/Plugins/Mutable/Source/CustomizableObjectEditor/Private/MuCOE/GenerateMutableSource/GenerateMutableSourceTable.cpp b/Engine/Plugins/Mutable/Source/CustomizableObjectEditor/Private/MuCOE/GenerateMutableSource/GenerateMutableSourceTable.cpp
--- a/Engine/Plugins/Mutable/Source/CustomizableObjectEditor/Private/MuCOE/GenerateMutableSource/GenerateMutableSourceTable.cpp
+++ b/Engine/Plugins/Mutable/Source/CustomizableObjectEditor/Private/MuCOE/GenerateMutableSource/GenerateMutableSourceTable.cpp
@@ -36,7 +36,7 @@ bool FillTableColumn(const UCustomizableObjectNodeTable* TableNode, mu::Ptr<mu::
 	UDataTable* DataTablePtr = GetDataTable(TableNode, GenerationContext);
 	FString ColumnPropertyName = ColumnProperty->GetAuthoredName();
 
-	const bool bIgnoreMissingData = GenerationContext.CompilationContext.Object->GetPrivate()->bDisableTableMissingDataWarning;
+	const bool bIgnoreMissingData = GenerationContext.CompilationContext->Object->GetPrivate()->bDisableTableMissingDataWarning;
 
 	// Getting property type
 	if (const FSoftObjectProperty* SoftObjectProperty = CastField<FSoftObjectProperty>(ColumnProperty))
@@ -244,7 +244,7 @@ bool FillTableColumn(const UCustomizableObjectNodeTable* TableNode, mu::Ptr<mu::
 			
 			TArray<FCustomizableObjectStreameableResourceId> StreamedResources;
 
-			if (GenerationContext.CompilationContext.Object->GetPrivate()->IsAssetUserDataMergeEnabled())
+			if (GenerationContext.CompilationContext->Object->GetPrivate()->IsAssetUserDataMergeEnabled())
 			{
 				const TArray<UAssetUserData*>* AssetUserDataArray = SkeletalMesh->GetAssetUserDataArray();
 
@@ -441,7 +441,7 @@ bool FillTableColumn(const UCustomizableObjectNodeTable* TableNode, mu::Ptr<mu::
 				return false;
 			}
 
-			const bool bTableMaterialCheckDisabled = GenerationContext.CompilationContext.Object->GetPrivate()->IsTableMaterialsParentCheckDisabled();
+			const bool bTableMaterialCheckDisabled = GenerationContext.CompilationContext->Object->GetPrivate()->IsTableMaterialsParentCheckDisabled();
 			const bool bMaterialParentMismatch = !bTableMaterialCheckDisabled && MaterialInstance
 												 && ReferenceMaterial->GetMaterial() != MaterialInstance->GetMaterial();
 
@@ -584,7 +584,7 @@ bool FillTableColumn(const UCustomizableObjectNodeTable* TableNode, mu::Ptr<mu::
 				for (int32 i = 0; i < ArrayBoneName.Num(); ++i)
 				{
 					const FName BoneName = ArrayBoneName[i];
-					const mu::FBoneName& BoneId = GenerationContext.CompilationContext.GetBoneUnique(BoneName);
+					const mu::FBoneName& BoneId = GenerationContext.CompilationContext->GetBoneUnique(BoneName);
 
 					MutableSkeleton->SetDebugName(i, BoneName);
 					MutableSkeleton->SetBoneName(i, { BoneId });
@@ -682,7 +682,7 @@ FName GetAnotherOption(FName SelectedOptionName, const TArray<FName>& RowNames)
 
 void RestrictRowNamesToSelectedOption(TArray<FName>& InOutRowNames, const UCustomizableObjectNodeTable& TableNode, FMutableGraphGenerationContext& GenerationContext)
 {
-	if (!GenerationContext.CompilationContext.Options.ParamNamesToSelectedOptions.IsEmpty())
+	if (!GenerationContext.CompilationContext->Options.ParamNamesToSelectedOptions.IsEmpty())
 	{
 		FMutableParamNameSet* ParamNameSet = GenerationContext.TableToParamNames.Find(TableNode.Table->GetPathName());
 
@@ -693,7 +693,7 @@ void RestrictRowNamesToSelectedOption(TArray<FName>& InOutRowNames, const UCusto
 			for (const FString& ParamName : ParamNameSet->ParamNames)
 			{
 				// If the param is in the map restrict to only the selected option
-				const FString* SelectedOptionString = GenerationContext.CompilationContext.Options.ParamNamesToSelectedOptions.Find(ParamName);
+				const FString* SelectedOptionString = GenerationContext.CompilationContext->Options.ParamNamesToSelectedOptions.Find(ParamName);
 
 				if (SelectedOptionString)
 				{
@@ -1039,7 +1039,7 @@ mu::Ptr<mu::FTable> GenerateMutableSourceTable(const UDataTable* DataTable, cons
 {
 	check(DataTable && TableNode);
 
-	if (GenerationContext.CompilationContext.Options.ParamNamesToSelectedOptions.IsEmpty())
+	if (GenerationContext.CompilationContext->Options.ParamNamesToSelectedOptions.IsEmpty())
 	{
 		FMutableParamNameSet* ParamNameSet = GenerationContext.TableToParamNames.Find(DataTable->GetPathName());
 
diff --git a/Engine/Plugins/Mutable/Source/CustomizableObjectEditor/Private/MuCOE/Nodes/CustomizableObjectNodeMeshInterface.h b/Engine/Plugins/Mutable/Source/CustomizableObjectEditor/Private/MuCOE/Nodes/CustomizableObjectNodeMeshInterface.h
--- a/Engine/Plugins/Mutable/Source/CustomizableObjectEditor/Private/MuCOE/Nodes/CustomizableObjectNodeMeshInterface.h
+++ b/Engine/Plugins/Mutable/Source/CustomizableObjectEditor/Private/MuCOE/Nodes/CustomizableObjectNodeMeshInterface.h
@@ -43,8 +43,7 @@ public:
 	 * 
 	 * @param Pin In. Node's owning pin.
 	 * @param LODIndex Out.
-	 * @param SectionIndex Out.
-	 * @param LayoutIndex Out. */
-	virtual void GetPinSection(const UEdGraphPin& Pin, int32& OutLODIndex, int32& OutSectionIndex, int32& OutLayoutIndex) const PURE_VIRTUAL(UCustomizableObjectNodeMesh::GetPinSection, );
+	 * @param SectionIndex Out. */
+	virtual void GetPinSection(const UEdGraphPin& Pin, int32& OutLODIndex, int32& OutSectionIndex) const PURE_VIRTUAL(UCustomizableObjectNodeMesh::GetPinSection, );
 };
 
diff --git a/Engine/Plugins/Mutable/Source/CustomizableObjectEditor/Private/MuCOE/Nodes/CustomizableObjectNodeSkeletalMesh.cpp b/Engine/Plugins/Mutable/Source/CustomizableObjectEditor/Private/MuCOE/Nodes/CustomizableObjectNodeSkeletalMesh.cpp
--- a/Engine/Plugins/Mutable/Source/CustomizableObjectEditor/Private/MuCOE/Nodes/CustomizableObjectNodeSkeletalMesh.cpp
+++ b/Engine/Plugins/Mutable/Source/CustomizableObjectEditor/Private/MuCOE/Nodes/CustomizableObjectNodeSkeletalMesh.cpp
@@ -4,6 +4,7 @@
 
 #include "Engine/SkeletalMesh.h"
 #include "AssetThumbnail.h"
+#include "AssetRegistry/AssetRegistryModule.h"
 #include "ISinglePropertyView.h"
 #include "Modules/ModuleManager.h"
 #include "MuCO/CustomizableObjectCustomVersion.h"
@@ -35,11 +36,11 @@ static const TCHAR* SKELETAL_MESH_PIN_NAME = TEXT("Skeletal Mesh");
 
 void UCustomizableObjectNodeSkeletalMesh::LoadObjects()
 {
-	if (SkeletalMesh)
+	if (USkeletalMesh* LoadedMesh = MutablePrivate::LoadObject(SkeletalMesh))
 	{
-		ConditionalPostLoadReference(*SkeletalMesh);
+		ConditionalPostLoadReference(*LoadedMesh);
 
-		for (const FSkeletalMaterial& SkeletalMaterial : SkeletalMesh->GetMaterials())
+		for (const FSkeletalMaterial& SkeletalMaterial : LoadedMesh->GetMaterials())
 		{
 			if (UMaterialInterface* Material = SkeletalMaterial.MaterialInterface)
 			{
@@ -257,28 +258,18 @@ UEdGraphPin* UCustomizableObjectNodeSkeletalMesh::GetMeshPin(const int32 LODInde
 }
 
 
-void UCustomizableObjectNodeSkeletalMesh::GetPinSection(const UEdGraphPin& Pin, int32& OutLODIndex, int32& OutSectionIndex, int32& OutLayoutIndex) const
+void UCustomizableObjectNodeSkeletalMesh::GetPinSection(const UEdGraphPin& Pin, int32& OutLODIndex, int32& OutSectionIndex) const
 {
 	if (const UCustomizableObjectNodeSkeletalMeshPinDataSection* PinData = Cast<UCustomizableObjectNodeSkeletalMeshPinDataSection>(GetPinData(Pin)))
 	{
 		OutLODIndex = PinData->GetLODIndex();
 		OutSectionIndex = PinData->GetSectionIndex();
-
-		if (const UCustomizableObjectNodeSkeletalMeshPinDataLayout* LayoutPinData = Cast<UCustomizableObjectNodeSkeletalMeshPinDataLayout>(GetPinData(Pin)))
-		{
-			OutLayoutIndex = LayoutPinData->GetUVIndex();
-		}
-		else
-		{
-			OutLayoutIndex = -1;
-		}
-		
-		return;
 	}
-	
-	OutLODIndex = -1;
-	OutSectionIndex = -1;
-	OutLayoutIndex = -1;
+	else
+	{
+		OutLODIndex = -1;
+		OutSectionIndex = -1;
+	}
 }
 
 
@@ -297,8 +288,7 @@ FSkeletalMaterial* UCustomizableObjectNodeSkeletalMesh::GetSkeletalMaterialFor(c
 {
 	int32 LODIndex;
 	int32 SectionIndex;
-	int32 LayoutIndex;
-	GetPinSection(Pin, LODIndex, SectionIndex, LayoutIndex);
+	GetPinSection(Pin, LODIndex, SectionIndex);
 
 	return GetSkeletalMaterialFor(LODIndex, SectionIndex);
 }
@@ -308,8 +298,7 @@ int32 UCustomizableObjectNodeSkeletalMesh::GetSkeletalMaterialIndexFor(const UEd
 {
 	int32 LODIndex;
 	int32 SectionIndex;
-	int32 LayoutIndex;
-	GetPinSection(Pin, LODIndex, SectionIndex, LayoutIndex);
+	GetPinSection(Pin, LODIndex, SectionIndex);
 
 	return GetSkeletalMaterialIndexFor(LODIndex, SectionIndex);
 }
@@ -318,8 +307,7 @@ const FSkelMeshSection* UCustomizableObjectNodeSkeletalMesh::GetSkeletalMeshSect
 {
 	int32 LODIndex;
 	int32 SectionIndex;
-	int32 LayoutIndex;
-	GetPinSection(Pin, LODIndex, SectionIndex, LayoutIndex);
+	GetPinSection(Pin, LODIndex, SectionIndex);
 	
 	return GetSkeletalMeshSectionFor(LODIndex, SectionIndex);
 }
@@ -708,46 +696,6 @@ TSharedPtr<SGraphNode> UCustomizableObjectNodeSkeletalMesh::CreateVisualWidget()
 }
 
 
-bool UCustomizableObjectNodeSkeletalMesh::CheckIsValidLayout(const UEdGraphPin* InPin, int32& LayoutIndex, FString& MaterialName)
-{
-	const UEdGraphPin* ConnectedPin = FollowOutputPin(*InPin);
-	if (!ConnectedPin)
-	{
-		return true;
-	}
-
-	int32 LODIndex;
-	int32 SectionIndex;
-	GetPinSection(*ConnectedPin, LODIndex, SectionIndex, LayoutIndex);
-
-	if (const UMaterialInterface* MaterialInterface = GetMaterialInterfaceFor(LODIndex, SectionIndex))
-	{
-		MaterialName = MaterialInterface->GetName();
-	}
-	
-	if (LayoutIndex == 0)
-	{
-		return true;
-	}
-	
-	TBitArray VisitedLayouts(false, LayoutIndex);
-	
-	for (const UEdGraphPin* Pin : GetAllNonOrphanPins())
-	{
-		if (const UCustomizableObjectNodeSkeletalMeshPinDataLayout* PinData = Cast<UCustomizableObjectNodeSkeletalMeshPinDataLayout>(GetPinData(*Pin)))
-		{
-			if (PinData->GetLODIndex() == LODIndex &&
-				PinData->GetSectionIndex() == SectionIndex &&
-				PinData->GetUVIndex() < LayoutIndex)
-			{
-				VisitedLayouts[PinData->GetUVIndex()] = true;
-			}
-		}
-	}
-
-	return !VisitedLayouts.Contains(false);
-}
-
 UMaterialInterface* UCustomizableObjectNodeSkeletalMesh::GetMaterialInterfaceFor(const int32 LODIndex, const int32 SectionIndex) const
 {
 	if (FSkeletalMaterial* SkeletalMaterial = GetSkeletalMaterialFor(LODIndex, SectionIndex))
@@ -849,7 +797,12 @@ void SGraphNodeSkeletalMesh::Construct(const FArguments& InArgs, UEdGraphNode* I
 	// Thumbnail
 	AssetThumbnailPool = MakeShareable(new FAssetThumbnailPool(32));
 	USkeletalMesh* SkeletalMesh = MutablePrivate::LoadObject(NodeSkeletalMesh->SkeletalMesh);
-	AssetThumbnail = MakeShareable(new FAssetThumbnail(SkeletalMesh, ThumbnailSize, ThumbnailSize, AssetThumbnailPool));
+
+	const FAssetRegistryModule& AssetRegistryModule = FModuleManager::LoadModuleChecked<FAssetRegistryModule>(AssetRegistryConstants::ModuleName);
+
+	FAssetData AssetData;
+	AssetRegistryModule.TryGetAssetByObjectPath(NodeSkeletalMesh->SkeletalMesh.ToSoftObjectPath(), AssetData);
+	AssetThumbnail = MakeShareable(new FAssetThumbnail(AssetData, ThumbnailSize, ThumbnailSize, AssetThumbnailPool));
 
 	// Selector
 	FPropertyEditorModule& PropPlugin = FModuleManager::LoadModuleChecked<FPropertyEditorModule>("PropertyEditor");
diff --git a/Engine/Plugins/Mutable/Source/CustomizableObjectEditor/Private/MuCOE/Nodes/CustomizableObjectNodeSkeletalMesh.h b/Engine/Plugins/Mutable/Source/CustomizableObjectEditor/Private/MuCOE/Nodes/CustomizableObjectNodeSkeletalMesh.h
--- a/Engine/Plugins/Mutable/Source/CustomizableObjectEditor/Private/MuCOE/Nodes/CustomizableObjectNodeSkeletalMesh.h
+++ b/Engine/Plugins/Mutable/Source/CustomizableObjectEditor/Private/MuCOE/Nodes/CustomizableObjectNodeSkeletalMesh.h
@@ -234,7 +234,7 @@ public:
 	UE_API virtual TArray<UCustomizableObjectLayout*> GetLayouts(const UEdGraphPin& MeshPin) const override;
 	UE_API virtual TSoftObjectPtr<UStreamableRenderAsset> GetMesh() const override;
 	UE_API virtual UEdGraphPin* GetMeshPin(int32 LOD, int32 SectionIndex) const override;
-	UE_API virtual void GetPinSection(const UEdGraphPin& Pin, int32& OutLODIndex, int32& OutSectionIndex, int32& OutLayoutIndex) const override;
+	UE_API virtual void GetPinSection(const UEdGraphPin& Pin, int32& OutLODIndex, int32& OutSectionIndex) const override;
 
 	// Own interface
 	
@@ -255,9 +255,6 @@ public:
 	// Creates the SGraph Node widget for the thumbnail
 	UE_API virtual TSharedPtr<SGraphNode> CreateVisualWidget() override;
 
-	/** Check if lower UV channels pins of the same LOD and Section are also linked to a Layout Node. */
-	UE_API bool CheckIsValidLayout(const UEdGraphPin* Pin, int32& LayoutIndex, FString& MaterialName);
-
 	// Determines if the Node is collapsed or not
 	bool bCollapsed = true;
 
diff --git a/Engine/Plugins/Mutable/Source/CustomizableObjectEditor/Private/MuCOE/Nodes/CustomizableObjectNodeSkeletalMeshParameter.cpp b/Engine/Plugins/Mutable/Source/CustomizableObjectEditor/Private/MuCOE/Nodes/CustomizableObjectNodeSkeletalMeshParameter.cpp
--- a/Engine/Plugins/Mutable/Source/CustomizableObjectEditor/Private/MuCOE/Nodes/CustomizableObjectNodeSkeletalMeshParameter.cpp
+++ b/Engine/Plugins/Mutable/Source/CustomizableObjectEditor/Private/MuCOE/Nodes/CustomizableObjectNodeSkeletalMeshParameter.cpp
@@ -164,11 +164,10 @@ UEdGraphPin* UCustomizableObjectNodeSkeletalMeshParameter::GetMeshPin(const int3
 }
 
 
-void UCustomizableObjectNodeSkeletalMeshParameter::GetPinSection(const UEdGraphPin& Pin, int32& OutLODIndex, int32& OutSectionIndex, int32& OutLayoutIndex) const
+void UCustomizableObjectNodeSkeletalMeshParameter::GetPinSection(const UEdGraphPin& Pin, int32& OutLODIndex, int32& OutSectionIndex) const
 {
 	OutLODIndex = -1;
 	OutSectionIndex = -1;
-	OutLayoutIndex = -1;
 
 	if (const UCustomizableObjectNodeSkeletalMeshParameterPinDataSection* PinData = Cast<UCustomizableObjectNodeSkeletalMeshParameterPinDataSection>(GetPinData(Pin)))
 	{
diff --git a/Engine/Plugins/Mutable/Source/CustomizableObjectEditor/Private/MuCOE/Nodes/CustomizableObjectNodeSkeletalMeshParameter.h b/Engine/Plugins/Mutable/Source/CustomizableObjectEditor/Private/MuCOE/Nodes/CustomizableObjectNodeSkeletalMeshParameter.h
--- a/Engine/Plugins/Mutable/Source/CustomizableObjectEditor/Private/MuCOE/Nodes/CustomizableObjectNodeSkeletalMeshParameter.h
+++ b/Engine/Plugins/Mutable/Source/CustomizableObjectEditor/Private/MuCOE/Nodes/CustomizableObjectNodeSkeletalMeshParameter.h
@@ -82,7 +82,7 @@ public:
 	UE_API virtual TArray<UCustomizableObjectLayout*> GetLayouts(const UEdGraphPin& MeshPin) const override;
 	UE_API virtual TSoftObjectPtr<UStreamableRenderAsset> GetMesh() const override;
 	UE_API virtual UEdGraphPin* GetMeshPin(int32 LOD, int32 SectionIndex) const override;
-	UE_API virtual void GetPinSection(const UEdGraphPin& Pin, int32& OutLODIndex, int32& OutSectionIndex, int32& OutLayoutIndex) const override;
+	UE_API virtual void GetPinSection(const UEdGraphPin& Pin, int32& OutLODIndex, int32& OutSectionIndex) const override;
 
 	// Own interface
 	UE_API void GetPinSection(const UEdGraphPin& Pin, int32& OutSectionIndex) const;
diff --git a/Engine/Plugins/Mutable/Source/CustomizableObjectEditor/Private/MuCOE/Nodes/CustomizableObjectNodeStaticMesh.cpp b/Engine/Plugins/Mutable/Source/CustomizableObjectEditor/Private/MuCOE/Nodes/CustomizableObjectNodeStaticMesh.cpp
--- a/Engine/Plugins/Mutable/Source/CustomizableObjectEditor/Private/MuCOE/Nodes/CustomizableObjectNodeStaticMesh.cpp
+++ b/Engine/Plugins/Mutable/Source/CustomizableObjectEditor/Private/MuCOE/Nodes/CustomizableObjectNodeStaticMesh.cpp
@@ -380,10 +380,8 @@ FText UCustomizableObjectNodeStaticMesh::GetTooltipText() const
 }
 
 
-void UCustomizableObjectNodeStaticMesh::GetPinSection(const UEdGraphPin& Pin, int32& OutLODIndex, int32& OutSectionIndex, int32& OutLayoutIndex) const
+void UCustomizableObjectNodeStaticMesh::GetPinSection(const UEdGraphPin& Pin, int32& OutLODIndex, int32& OutSectionIndex) const
 {
-	OutLayoutIndex = 0; // Currently layouts different from 0 not supported for static mesh.
-	
 	for (OutLODIndex = 0; OutLODIndex < LODs.Num(); ++OutLODIndex)
 	{
 		const FCustomizableObjectNodeStaticMeshLOD& LOD = LODs[OutLODIndex];
@@ -426,8 +424,7 @@ UMaterialInterface* UCustomizableObjectNodeStaticMesh::GetMaterialFor(const UEdG
 		check(Pin);
 		int32 LODIndex;
 		int32 SectionIndex;
-		int32 LayoutIndex;
-		GetPinSection(*Pin, LODIndex, SectionIndex, LayoutIndex);
+		GetPinSection(*Pin, LODIndex, SectionIndex);
 
 		const int32 MaterialIndex = StaticMesh->GetRenderData()->LODResources[LODIndex].Sections[SectionIndex].MaterialIndex;
 
diff --git a/Engine/Plugins/Mutable/Source/CustomizableObjectEditor/Private/MuCOE/Nodes/CustomizableObjectNodeStaticMesh.h b/Engine/Plugins/Mutable/Source/CustomizableObjectEditor/Private/MuCOE/Nodes/CustomizableObjectNodeStaticMesh.h
--- a/Engine/Plugins/Mutable/Source/CustomizableObjectEditor/Private/MuCOE/Nodes/CustomizableObjectNodeStaticMesh.h
+++ b/Engine/Plugins/Mutable/Source/CustomizableObjectEditor/Private/MuCOE/Nodes/CustomizableObjectNodeStaticMesh.h
@@ -154,7 +154,7 @@ public:
 	UE_API virtual TArray<UCustomizableObjectLayout*> GetLayouts(const UEdGraphPin& OutPin) const override;
 	UE_API virtual TSoftObjectPtr<UStreamableRenderAsset> GetMesh() const override;
 	UE_API virtual UEdGraphPin* GetMeshPin(int32 LOD, int32 SectionIndex) const override;
-	UE_API virtual void GetPinSection(const UEdGraphPin& Pin, int32& OutLODIndex, int32& OutSectionIndex, int32& OutLayoutIndex) const override;
+	UE_API virtual void GetPinSection(const UEdGraphPin& Pin, int32& OutLODIndex, int32& OutSectionIndex) const override;
 
 	/** Returns the material associated to the given output pin. */
 	UE_API UMaterialInterface* GetMaterialFor(const UEdGraphPin* Pin) const;
