[mutable] Fixed crash when trying to retrieve a texture from the TextureCache (UE level)
- The crash was caused by expecting to have a texture cached when performing a bake (the cache was "disabled" in that case)
- Enabled the use of the TextureCache for the textures generated for baking.
- Added bool to the TextureCache key to separate bake-oriented and non-bake-oriented textures as we do not want to mix them.

#rnx

diff --git a/Engine/Plugins/Mutable/Source/CustomizableObject/Internal/MuCO/CustomizableObjectPrivate.h b/Engine/Plugins/Mutable/Source/CustomizableObject/Internal/MuCO/CustomizableObjectPrivate.h
--- a/Engine/Plugins/Mutable/Source/CustomizableObject/Internal/MuCO/CustomizableObjectPrivate.h
+++ b/Engine/Plugins/Mutable/Source/CustomizableObject/Internal/MuCO/CustomizableObjectPrivate.h
@@ -281,6 +281,11 @@ public:
 
 		int32 SkippedMips = 0;
 
+#if WITH_EDITORONLY_DATA
+		// TODO: Remove this bool and make the actual generation of the textures a step prior to the build of the materials
+		bool bIsBake = false;
+#endif
+		
 		bool operator==(const FId&) const = default;
 
 		friend uint32 GetTypeHash(const FId& Key);
diff --git a/Engine/Plugins/Mutable/Source/CustomizableObject/Private/MuCO/CustomizableObject.cpp b/Engine/Plugins/Mutable/Source/CustomizableObject/Private/MuCO/CustomizableObject.cpp
--- a/Engine/Plugins/Mutable/Source/CustomizableObject/Private/MuCO/CustomizableObject.cpp
+++ b/Engine/Plugins/Mutable/Source/CustomizableObject/Private/MuCO/CustomizableObject.cpp
@@ -3459,7 +3459,11 @@ FMutableParameterData::FMutableParameterData(const FMutableParamUIMetadata& InPa
 
 uint32 GetTypeHash(const FTextureCache::FId& Key)
 {
+#if WITH_EDITOR
+	return HashCombine(GetTypeHash(Key.Resource), Key.SkippedMips, Key.bIsBake);
+#else
 	return HashCombine(GetTypeHash(Key.Resource), Key.SkippedMips);
+#endif
 }
 
 
diff --git a/Engine/Plugins/Mutable/Source/CustomizableObject/Private/MuCO/CustomizableObjectInstance.cpp b/Engine/Plugins/Mutable/Source/CustomizableObject/Private/MuCO/CustomizableObjectInstance.cpp
--- a/Engine/Plugins/Mutable/Source/CustomizableObject/Private/MuCO/CustomizableObjectInstance.cpp
+++ b/Engine/Plugins/Mutable/Source/CustomizableObject/Private/MuCO/CustomizableObjectInstance.cpp
@@ -6159,11 +6159,17 @@ void UCustomizableInstancePrivate::BuildMaterials(const TSharedRef<FUpdateContex
 						UTexture* PassThroughTexture = nullptr; // Texture not generated by mutable
 
 						const FTextureReuseCacheKey TextureReuseCacheRef = bReuseTextures ? FTextureReuseCacheKey{Image.BaseLOD, ObjectComponentIndex.GetValue(), Surface.SurfaceId, ImageIndex} : FTextureReuseCacheKey{};
+
+#if WITH_EDITOR
+						const FTextureCache::FId TextureCacheKey = FTextureCache::FId(Image.ImageID, OperationData->MipsToSkip, OperationData->bBake);
+#else
+						const FTextureCache::FId TextureCacheKey = FTextureCache::FId(Image.ImageID, OperationData->MipsToSkip);
+#endif
 						
 						// If the mutable image is null, it must be in the cache
 						if (!MutableImage)
 						{
-							MutableTexture = TextureCache.Get(FTextureCache::FId(Image.ImageID, OperationData->MipsToSkip));
+							MutableTexture = TextureCache.Get(TextureCacheKey);
 							check(MutableTexture);
 						}
 
@@ -6356,12 +6362,8 @@ void UCustomizableInstancePrivate::BuildMaterials(const TSharedRef<FUpdateContex
 											MutableTexture->UpdateResource();
 										}
 									}
-
-									// During the bake do not add the texture to the cache as we do not want it to be used for other instance updates
-									if (!OperationData->bBake)
-									{
-										TextureCache.Add(FTextureCache::FId(Image.ImageID, OperationData->MipsToSkip), MutableTexture);						
-									}
+									
+									TextureCache.Add(TextureCacheKey, MutableTexture);						
 								}
 								else
 								{
diff --git a/Engine/Plugins/Mutable/Source/CustomizableObject/Private/MuCO/CustomizableObjectSystem.cpp b/Engine/Plugins/Mutable/Source/CustomizableObject/Private/MuCO/CustomizableObjectSystem.cpp
--- a/Engine/Plugins/Mutable/Source/CustomizableObject/Private/MuCO/CustomizableObjectSystem.cpp
+++ b/Engine/Plugins/Mutable/Source/CustomizableObject/Private/MuCO/CustomizableObjectSystem.cpp
@@ -3323,15 +3323,18 @@ namespace impl
 		//
 		// If the program aborted, AreAnyLoadsPending may return true, as the program doesn't cancel
 		// its own loads on exit (maybe it should?)
-
-		// Do not update the CachedTextures during a bake as we do not want the instance update to use cached textures
-		if (!OperationData->bBake)
+		
 		{
 			FTextureCache& TextureCache = CustomizableObject->GetPrivate()->TextureCache;
 		
 			for (const FInstanceUpdateData::FImage& Image : OperationData->InstanceUpdateData.Images)
 			{
-				if (const UTexture2D* Texture = TextureCache.Get(FTextureCache::FId(Image.ImageID, OperationData->MipsToSkip)))
+#if	WITH_EDITOR
+				const FTextureCache::FId TextureCacheKey = FTextureCache::FId(Image.ImageID, OperationData->MipsToSkip, OperationData->bBake);
+#else
+				const FTextureCache::FId TextureCacheKey = FTextureCache::FId(Image.ImageID, OperationData->MipsToSkip);
+#endif
+				if (const UTexture2D* Texture = TextureCache.Get(TextureCacheKey))
 				{
 					OperationData->Objects.Add(TStrongObjectPtr(Texture));
 					OperationData->CachedTextures.Add(Image.ImageID);
