[Mutable] Fix multiple Dataless bugs.

- Crash with Default Material parameter.
- Crash with Dataless Morphs having names but no data.
- Fix Texture Parameter not croping images with layout operations.

diff --git a/Engine/Plugins/Mutable/Source/CustomizableObject/Private/MuCO/CustomizableObjectInstanceDescriptor.cpp b/Engine/Plugins/Mutable/Source/CustomizableObject/Private/MuCO/CustomizableObjectInstanceDescriptor.cpp
--- a/Engine/Plugins/Mutable/Source/CustomizableObject/Private/MuCO/CustomizableObjectInstanceDescriptor.cpp
+++ b/Engine/Plugins/Mutable/Source/CustomizableObject/Private/MuCO/CustomizableObjectInstanceDescriptor.cpp
@@ -2892,7 +2892,20 @@ void FCustomizableObjectInstanceDescriptor::SetDefaultValue(int32 ParamIndex)
 		}
 		break;
 	}
- 
+
+	case EMutableParameterType::Material:
+	{
+		UMaterialInterface* DefaultValue = CustomizableObject->GetMaterialParameterDefaultValue(ParamName);
+
+		if (ensure(MaterialParameters.IsValidIndex(TypedIndex)))
+		{
+			FCustomizableObjectMaterialParameterValue& MaterialParameter = MaterialParameters[TypedIndex];
+			MaterialParameter.ParameterValue = DefaultValue;
+			MaterialParameter.ParameterRangeValues.Empty();
+		}
+		break;
+	}
+		
 	default:
 		// Parameter type replication not implemented.
 		unimplemented();
diff --git a/Engine/Plugins/Mutable/Source/CustomizableObject/Private/MuCO/CustomizableObjectSystem.cpp b/Engine/Plugins/Mutable/Source/CustomizableObject/Private/MuCO/CustomizableObjectSystem.cpp
--- a/Engine/Plugins/Mutable/Source/CustomizableObject/Private/MuCO/CustomizableObjectSystem.cpp
+++ b/Engine/Plugins/Mutable/Source/CustomizableObject/Private/MuCO/CustomizableObjectSystem.cpp
@@ -2252,6 +2252,11 @@ namespace impl
 					const mu::FMeshSurface& Surface = LOD.Mesh->Surfaces[SectionIndex];
 					for (const mu::FSurfaceSubMesh& SubMesh : Surface.SubMeshes)
 					{
+						if (SubMesh.ExternalId == 0)
+						{
+							continue;	
+						}
+						
 						const uint32 MorphMetadataId = ModelResources->MeshMetadata[SubMesh.ExternalId].MorphMetadataId;
 
 						FRealTimeMorphStreamable* Result = ModelStreamableBulkData->RealTimeMorphStreamables.Find(MorphMetadataId);
diff --git a/Engine/Plugins/Mutable/Source/CustomizableObject/Private/MuCO/UnrealConversionUtils.cpp b/Engine/Plugins/Mutable/Source/CustomizableObject/Private/MuCO/UnrealConversionUtils.cpp
--- a/Engine/Plugins/Mutable/Source/CustomizableObject/Private/MuCO/UnrealConversionUtils.cpp
+++ b/Engine/Plugins/Mutable/Source/CustomizableObject/Private/MuCO/UnrealConversionUtils.cpp
@@ -1354,6 +1354,11 @@ namespace UnrealConversionUtils
 				
 					for (const FMorphTargetLODModel* MorphModel : MorphTargetLODs)
 					{
+						if (!MorphModel)
+						{
+							continue;
+						}
+						
 						uint32 BatchStartOffset = BatchHeaders.Num();
 
 						float MaximumValues[4] = { -FLT_MAX, -FLT_MAX, -FLT_MAX, -FLT_MAX };
@@ -1409,6 +1414,7 @@ namespace UnrealConversionUtils
 						? TConstArrayView<uint32>(RuntimeCompressedMorphData)
 						: TConstArrayView<uint32>(MorphTargetBuffers.GetData(), MorphTargetBuffers.GetMorphDataSizeInBytes()/sizeof(uint32));
 
+				if (SourceMorphDataView.GetData())
 				{	
 					MUTABLE_CPUPROFILER_SCOPE(ExtractSectionMorphs)
 					mu::FMeshMorph& Morph = Result->Morph;
@@ -1557,7 +1563,7 @@ namespace UnrealConversionUtils
 						}
 
 						TArrayView<uint32> MorphBatchHeadersView(
-								Result->MorphDataBuffer.GetData() + BatchOffset*NumBatchHeaderDwords,
+								Result->MorphDataBuffer.GetData() + BatchOffset * NumBatchHeaderDwords,
 								MorphNumBatches * NumBatchHeaderDwords);
 
 						// Batches still point to the source data.
diff --git a/Engine/Plugins/Mutable/Source/MutableRuntime/Private/MuR/Model.cpp b/Engine/Plugins/Mutable/Source/MutableRuntime/Private/MuR/Model.cpp
--- a/Engine/Plugins/Mutable/Source/MutableRuntime/Private/MuR/Model.cpp
+++ b/Engine/Plugins/Mutable/Source/MutableRuntime/Private/MuR/Model.cpp
@@ -497,10 +497,11 @@ namespace mu
     		for (const TPair<FName, TObjectPtr<UTexture>>& Pair : *ImageDefaults)
     		{
     			int32 Index = pRes->GetPrivate()->Find(Pair.Key.ToString());
-				check(Index != INDEX_NONE);
-    			check(pRes->GetType(Index) == EParameterType::Image);
-    			
-    			pRes->SetImageValue(Index, Pair.Value);
+    			if (Index != INDEX_NONE)
+    			{
+    				check(pRes->GetType(Index) == EParameterType::Image);
+    				pRes->SetImageValue(Index, Pair.Value);
+    			}
     		}
     	}
     	
@@ -509,10 +510,11 @@ namespace mu
     		for (const TPair<FName, TObjectPtr<USkeletalMesh>>& Pair : *MeshDefaults)
     		{
     			int32 Index = pRes->GetPrivate()->Find(Pair.Key.ToString());
-    			check(Index != INDEX_NONE);
-    			check(pRes->GetType(Index) == EParameterType::Mesh);
-    			
-    			pRes->SetMeshValue(Index, Pair.Value);
+    			if (Index != INDEX_NONE)
+    			{
+    				check(pRes->GetType(Index) == EParameterType::Mesh);
+    				pRes->SetMeshValue(Index, Pair.Value);
+    			}
     		}
     	}
 
@@ -521,10 +523,11 @@ namespace mu
 			for (const TPair<FName, TObjectPtr<UMaterialInterface>>& Pair : *MaterialDefaults)
 			{
 				int32 Index = pRes->GetPrivate()->Find(Pair.Key.ToString());
-				check(Index != INDEX_NONE);
-				check(pRes->GetType(Index) == EParameterType::Material);
-
-				pRes->SetMaterialValue(Index, Pair.Value);
+				if (Index != INDEX_NONE)
+				{
+					check(pRes->GetType(Index) == EParameterType::Material);
+					pRes->SetMaterialValue(Index, Pair.Value);
+				}
 			}
 		}
 
diff --git a/Engine/Plugins/Mutable/Source/MutableTools/Private/MuT/CodeGenerator.h b/Engine/Plugins/Mutable/Source/MutableTools/Private/MuT/CodeGenerator.h
--- a/Engine/Plugins/Mutable/Source/MutableTools/Private/MuT/CodeGenerator.h
+++ b/Engine/Plugins/Mutable/Source/MutableTools/Private/MuT/CodeGenerator.h
@@ -494,6 +494,7 @@ namespace mu
 		struct FImageGenerationResult
 		{
 			Ptr<ASTOp> op;
+			bool Parameter = false;
 		};
 
 		/** */
@@ -552,6 +553,8 @@ namespace mu
 		void GenerateImage_NormalComposite(const FImageGenerationOptions&, FImageGenerationResult&, const NodeImageNormalComposite*);
 		void GenerateImage_Transform(const FImageGenerationOptions&, FImageGenerationResult&, const NodeImageTransform*);
 
+    	void GenerateCrop(const FImageGenerationOptions& Options, FImageGenerationResult& Result, const NodeImage&);
+
 		/** */
 		TSharedPtr<FImage> GenerateMissingImage(EImageFormat);
 
diff --git a/Engine/Plugins/Mutable/Source/MutableTools/Private/MuT/CodeGenerator_Image.cpp b/Engine/Plugins/Mutable/Source/MutableTools/Private/MuT/CodeGenerator_Image.cpp
--- a/Engine/Plugins/Mutable/Source/MutableTools/Private/MuT/CodeGenerator_Image.cpp
+++ b/Engine/Plugins/Mutable/Source/MutableTools/Private/MuT/CodeGenerator_Image.cpp
@@ -145,58 +145,15 @@
 	}
 
 
-   void CodeGenerator::GenerateImage_Constant(const FImageGenerationOptions& Options, FImageGenerationResult& Result, const NodeImageConstant* InNode)
-    {
-		const NodeImageConstant& node = *InNode;
-		
-        // TODO: check duplicates
-        TSharedPtr<const FImage> pImage;
-		if (node.Proxy)
-		{
-			pImage = node.Proxy->Get();
-		}
-
-        if (!pImage)
-        {
-            // This data is required
-            pImage = GenerateMissingImage(EImageFormat::RGB_UByte );
-
-            // Log an error message
-            ErrorLog->Add( "Constant image not set.", ELMT_WARNING, InNode->GetMessageContext());
-        }
-
-
-		if (pImage->IsReference())
-		{
-			Ptr<ASTOpReferenceResource> ReferenceOp = new ASTOpReferenceResource();
-			ReferenceOp->Type = EOpType::IM_REFERENCE;
-			ReferenceOp->ID = pImage->GetReferencedTexture();
-			ReferenceOp->bForceLoad = pImage->IsForceLoad();
-			ReferenceOp->SourceDataDescriptor = InNode->SourceDataDescriptor;
-
-			// Don't store the format. Format can vary between loaded constant image and reference and cause
-			// code optimization bugs.
-			// As it is now, reference will always have alpha channel but constant resolution can remove the 
-			// channel if not used.
-			// TODO: review this, probably the reference descriptor generation needs to check for alpha channels 
-			// as well.
-			ReferenceOp->ImageDesc = FImageDesc(pImage->GetSize(), EImageFormat::None, pImage->GetLODCount());
-			Result.op = ReferenceOp;
-		}
-		else
-		{
-			Ptr<ASTOpConstantResource> op = new ASTOpConstantResource();
-			op->Type = EOpType::IM_CONSTANT;
-			op->SetValue(pImage, CompilerOptions->OptimisationOptions.DiskCacheContext);
-			op->SourceDataDescriptor = InNode->SourceDataDescriptor;
-			Result.op = op;
-		}
-
+	void CodeGenerator::GenerateCrop(const FImageGenerationOptions& Options, FImageGenerationResult& Result, const NodeImage& InNode)
+	{
 		if (Options.ImageLayoutStrategy!= CompilerOptions::TextureLayoutStrategy::None && Options.LayoutToApply)
 		{
 			// We want to generate only a block from the image.
 
-			FIntVector2 SourceImageSize(pImage->GetSizeX(), pImage->GetSizeY());
+			const FImageDesc ImageDesc = Result.op->GetImageDesc();
+			
+			FIntVector2 SourceImageSize(ImageDesc.m_size[0], ImageDesc.m_size[1]);
 
 			int32 BlockIndex = Options.LayoutToApply->FindBlock(Options.LayoutBlockId);
 			check( BlockIndex>=0 );
@@ -218,7 +175,7 @@
 			RectInPixels.size[1] = (RectInCells.size[1] * SourceImageSize[1]) / grid[1];
 
 			// Do we need to crop?
-			if (RectInPixels.min[0]!=0 || RectInPixels.min[1]!=0 || pImage->GetSizeX()!=RectInPixels.size[0] || pImage->GetSizeY()!=RectInPixels.size[1])
+			if (RectInPixels.min[0] != 0 || RectInPixels.min[1] != 0 || SourceImageSize[0] != RectInPixels.size[0] || SourceImageSize[1] != RectInPixels.size[1])
 			{
 				// See if the rect belongs to a single texture tile
 				FIntVector2 TileMin(RectInPixels.min[0] / SourceImageSize[0], RectInPixels.min[1] / SourceImageSize[1]);
@@ -231,7 +188,7 @@
 					ensure(false);
 
 					// Log an error message
-					ErrorLog->Add("A layout block goes across different texture tiles. This is not supported yet.", ELMT_ERROR, InNode->GetMessageContext());
+					ErrorLog->Add("A layout block goes across different texture tiles. This is not supported yet.", ELMT_ERROR, InNode.GetMessageContext());
 				}
 				else
 				{
@@ -247,6 +204,57 @@
 				}
 			}
 		}
+	}
+	
+
+   void CodeGenerator::GenerateImage_Constant(const FImageGenerationOptions& Options, FImageGenerationResult& Result, const NodeImageConstant* InNode)
+    {
+		const NodeImageConstant& node = *InNode;
+		
+        // TODO: check duplicates
+        TSharedPtr<const FImage> pImage;
+		if (node.Proxy)
+		{
+			pImage = node.Proxy->Get();
+		}
+
+        if (!pImage)
+        {
+            // This data is required
+            pImage = GenerateMissingImage(EImageFormat::RGB_UByte );
+
+            // Log an error message
+            ErrorLog->Add( "Constant image not set.", ELMT_WARNING, InNode->GetMessageContext());
+        }
+
+
+		if (pImage->IsReference())
+		{
+			Ptr<ASTOpReferenceResource> ReferenceOp = new ASTOpReferenceResource();
+			ReferenceOp->Type = EOpType::IM_REFERENCE;
+			ReferenceOp->ID = pImage->GetReferencedTexture();
+			ReferenceOp->bForceLoad = pImage->IsForceLoad();
+			ReferenceOp->SourceDataDescriptor = InNode->SourceDataDescriptor;
+
+			// Don't store the format. Format can vary between loaded constant image and reference and cause
+			// code optimization bugs.
+			// As it is now, reference will always have alpha channel but constant resolution can remove the 
+			// channel if not used.
+			// TODO: review this, probably the reference descriptor generation needs to check for alpha channels 
+			// as well.
+			ReferenceOp->ImageDesc = FImageDesc(pImage->GetSize(), EImageFormat::None, pImage->GetLODCount());
+			Result.op = ReferenceOp;
+		}
+		else
+		{
+			Ptr<ASTOpConstantResource> op = new ASTOpConstantResource();
+			op->Type = EOpType::IM_CONSTANT;
+			op->SetValue(pImage, CompilerOptions->OptimisationOptions.DiskCacheContext);
+			op->SourceDataDescriptor = InNode->SourceDataDescriptor;
+			Result.op = op;
+		}
+
+		GenerateCrop(Options, Result, *InNode);
     }
 
 
@@ -289,6 +297,7 @@
 		}
 
 		Result.op = op;
+		Result.Parameter = true;
     }
 
 
@@ -902,6 +911,7 @@
 			FImageGenerationResult BaseResult;
 			GenerateImage(Options, BaseResult, Node->Source);
 			fop->Source = BaseResult.op;
+			Result.Parameter = BaseResult.Parameter;
 		}
 
         Result.op = fop;
@@ -1079,11 +1089,10 @@
 	void CodeGenerator::GenerateImage_Resize(const FImageGenerationOptions& Options, FImageGenerationResult& Result, const NodeImageResize* Node)
 	{
 		MUTABLE_CPUPROFILER_SCOPE(NodeImageResize);
-
-        Ptr<ASTOp> at;
-
+		
         // Source image
         Ptr<ASTOp> base;
+		bool bParameter = false;
         if (Node->Base )
         {
 			FImageGenerationOptions NewOptions = Options;
@@ -1096,7 +1105,8 @@
 
 			FImageGenerationResult BaseResult;
 			GenerateImage(NewOptions, BaseResult, Node->Base);
-			base = BaseResult.op;
+        	base = BaseResult.op;
+        	bParameter = BaseResult.Parameter;
         }
         else
         {
@@ -1111,13 +1121,13 @@
 			NewOp->Factor[0] = Node->SizeX;
 			NewOp->Factor[1] = Node->SizeY;
 			NewOp->Source = base;
-            at = NewOp;
+	        Result.op = NewOp;
         }
         else
         {
 			// Apply the layout block to the rect size
 			UE::Math::TIntVector2<int32> FinalImageSize = { int32(Node->SizeX), int32(Node->SizeY) };
-			if (Options.LayoutToApply)
+			if (Options.LayoutToApply && !bParameter)
 			{
 				int32 BlockIndex = Options.LayoutToApply->FindBlock(Options.LayoutBlockId);
 				check(BlockIndex >= 0);
@@ -1140,10 +1150,13 @@
 			ResizeOp->Size[0] = (uint16)FinalImageSize[0];
 			ResizeOp->Size[1] = (uint16)FinalImageSize[1];
 			ResizeOp->Source = base;
-            at = ResizeOp;
+        	Result.op = ResizeOp;
+
+        	if (bParameter)
+        	{
+        		GenerateCrop(Options, Result, *Node);
+        	}
         }
-
-        Result.op = at;
     }
 
 
