[Mutable] Fix issues with GenerateMutableSourceImage cache.

diff --git a/Engine/Plugins/Mutable/Source/CustomizableObject/Internal/MuCO/CustomizableObjectPrivate.h b/Engine/Plugins/Mutable/Source/CustomizableObject/Internal/MuCO/CustomizableObjectPrivate.h
--- a/Engine/Plugins/Mutable/Source/CustomizableObject/Internal/MuCO/CustomizableObjectPrivate.h
+++ b/Engine/Plugins/Mutable/Source/CustomizableObject/Internal/MuCO/CustomizableObjectPrivate.h
@@ -1687,6 +1687,8 @@ enum class ECustomizableObjectVersions : int32
 
 	AddMaterialArgs,
 
+	CacheReferenceTextureSize,
+
 	// -----<new versions can be added above this line>--------
 	LastCustomizableObjectVersion
 };
diff --git a/Engine/Plugins/Mutable/Source/CustomizableObjectEditor/Private/MuCOE/GenerateMutableSource/GenerateMutableSource.cpp b/Engine/Plugins/Mutable/Source/CustomizableObjectEditor/Private/MuCOE/GenerateMutableSource/GenerateMutableSource.cpp
--- a/Engine/Plugins/Mutable/Source/CustomizableObjectEditor/Private/MuCOE/GenerateMutableSource/GenerateMutableSource.cpp
+++ b/Engine/Plugins/Mutable/Source/CustomizableObjectEditor/Private/MuCOE/GenerateMutableSource/GenerateMutableSource.cpp
@@ -408,6 +408,7 @@ uint32 GetTypeHash(const FGeneratedKey& Key)
 	//Hash = HashCombine(Hash, GetTypeHash(Key.MeshMorphStack)); // Does not support array
 	Hash = HashCombine(Hash, GetTypeHash(Key.bOnlyConnectedLOD));
 	Hash = HashCombine(Hash, GetTypeHash(Key.CurrentMeshComponent));
+	Hash = HashCombine(Hash, GetTypeHash(Key.ReferenceTextureSize));
 
 	for (int32 MacroIndex = 0; MacroIndex < Key.MacroContext.Num(); ++MacroIndex)
 	{
diff --git a/Engine/Plugins/Mutable/Source/CustomizableObjectEditor/Private/MuCOE/GenerateMutableSource/GenerateMutableSource.h b/Engine/Plugins/Mutable/Source/CustomizableObjectEditor/Private/MuCOE/GenerateMutableSource/GenerateMutableSource.h
--- a/Engine/Plugins/Mutable/Source/CustomizableObjectEditor/Private/MuCOE/GenerateMutableSource/GenerateMutableSource.h
+++ b/Engine/Plugins/Mutable/Source/CustomizableObjectEditor/Private/MuCOE/GenerateMutableSource/GenerateMutableSource.h
@@ -164,6 +164,8 @@ public:
 
 	/** Pointer to control if this is a node inside a Mutable Macro. */
 	TArray<const UCustomizableObjectNodeMacroInstance*> MacroContext;
+
+	int32 ReferenceTextureSize = -1;
 };
 
 
diff --git a/Engine/Plugins/Mutable/Source/CustomizableObjectEditor/Private/MuCOE/GenerateMutableSource/GenerateMutableSourceImage.cpp b/Engine/Plugins/Mutable/Source/CustomizableObjectEditor/Private/MuCOE/GenerateMutableSource/GenerateMutableSourceImage.cpp
--- a/Engine/Plugins/Mutable/Source/CustomizableObjectEditor/Private/MuCOE/GenerateMutableSource/GenerateMutableSourceImage.cpp
+++ b/Engine/Plugins/Mutable/Source/CustomizableObjectEditor/Private/MuCOE/GenerateMutableSource/GenerateMutableSourceImage.cpp
@@ -90,7 +90,9 @@ mu::Ptr<mu::NodeImage> GenerateMutableSourceImage(const UEdGraphPin* Pin, FMutab
 
 	UCustomizableObjectNode* Node = CastChecked<UCustomizableObjectNode>(Pin->GetOwningNode());
 
-	const FGeneratedKey Key(reinterpret_cast<void*>(&GenerateMutableSourceImage), *Pin, *Node, GenerationContext, true);
+	FGeneratedKey Key(&GenerateMutableSourceImage, *Pin, *Node, GenerationContext, true);
+	Key.ReferenceTextureSize = ReferenceTextureSize;
+	
 	if (const FGeneratedData* Generated = GenerationContext.Generated.Find(Key))
 	{
 		return static_cast<mu::NodeImage*>(Generated->Node.get());
