[Mutable] Fix no matching constructor for FGeneratedKey.

Static cast required on some compilers.

diff --git a/Engine/Plugins/Mutable/Source/CustomizableObjectEditor/Private/MuCOE/GenerateMutableSource/GenerateMutableSourceImage.cpp b/Engine/Plugins/Mutable/Source/CustomizableObjectEditor/Private/MuCOE/GenerateMutableSource/GenerateMutableSourceImage.cpp
--- a/Engine/Plugins/Mutable/Source/CustomizableObjectEditor/Private/MuCOE/GenerateMutableSource/GenerateMutableSourceImage.cpp
+++ b/Engine/Plugins/Mutable/Source/CustomizableObjectEditor/Private/MuCOE/GenerateMutableSource/GenerateMutableSourceImage.cpp
@@ -90,7 +90,7 @@ mu::Ptr<mu::NodeImage> GenerateMutableSourceImage(const UEdGraphPin* Pin, FMutab
 
 	UCustomizableObjectNode* Node = CastChecked<UCustomizableObjectNode>(Pin->GetOwningNode());
 
-	FGeneratedKey Key(&GenerateMutableSourceImage, *Pin, *Node, GenerationContext, true);
+	FGeneratedKey Key(static_cast<void*>(&GenerateMutableSourceImage), *Pin, *Node, GenerationContext, true); // Static cast to void required on some compilers.
 	Key.ReferenceTextureSize = ReferenceTextureSize;
 	
 	if (const FGeneratedData* Generated = GenerationContext.Generated.Find(Key))
