[Mutable] Decoupled COI from Update.

Removed functions from being part of the COI class. They will be moved to the System in a future submit.
Removed many unused arguments, locals.

diff --git a/Engine/Plugins/Mutable/Source/CustomizableObject/Internal/MuCO/CustomizableObjectInstancePrivate.h b/Engine/Plugins/Mutable/Source/CustomizableObject/Internal/MuCO/CustomizableObjectInstancePrivate.h
--- a/Engine/Plugins/Mutable/Source/CustomizableObject/Internal/MuCO/CustomizableObjectInstancePrivate.h
+++ b/Engine/Plugins/Mutable/Source/CustomizableObject/Internal/MuCO/CustomizableObjectInstancePrivate.h
@@ -285,24 +285,18 @@ public:
 	void ClearCOInstanceFlags(ECOInstanceFlags FlagsToClear) { InstanceFlagsPrivate = (ECOInstanceFlags)(InstanceFlagsPrivate & ~FlagsToClear); }
 	bool HasCOInstanceFlags(ECOInstanceFlags FlagsToCheck) const { return (InstanceFlagsPrivate & FlagsToCheck) != 0; }
 
-	UE_API void BuildMaterials(const TSharedRef<FUpdateContextPrivate>& OperationData, UCustomizableObjectInstance* Public);
+	UE_API void BuildMaterials(const TSharedRef<FUpdateContextPrivate>& Context, UCustomizableObjectInstance* Public);
 
 	UE_API void ReuseTexture(UTexture2D* Texture, TSharedRef<FTexturePlatformData, ESPMode::ThreadSafe>& PlatformData);
 
-	/** Returns the task that will be called when all assets and data are loaded, may be already completed if no assets or data needs loading.
-	 * If no StreamableManager is provided, it will load assets synchronously. */
-	UE_API UE::Tasks::FTask LoadAdditionalAssetsAndData(const TSharedRef<FUpdateContextPrivate>& OperationData);
-
-	UE_API void AdditionalAssetsAsyncLoaded(TSharedRef<FUpdateContextPrivate> OperationData, UE::Tasks::FTaskEvent Event);
-
 	UE_API void TickUpdateCloseCustomizableObjects(UCustomizableObjectInstance& Public, FMutableInstanceUpdateMap& InOutRequestedUpdates);
 	UE_API void UpdateInstanceIfNotGenerated(UCustomizableObjectInstance& Public, FMutableInstanceUpdateMap& InOutRequestedUpdates);
 
 	// Returns true if success (?)
-	UE_API bool UpdateSkeletalMesh_PostBeginUpdate0(UCustomizableObjectInstance* Public, const TSharedRef<FUpdateContextPrivate>& OperationData);
+	UE_API bool UpdateSkeletalMesh_PostBeginUpdate0(UCustomizableObjectInstance* Instance, const TSharedRef<FUpdateContextPrivate>& Context);
 	
 	// Copy data generated in the mutable thread over to the instance and initializes additional data required during the update
-	UE_API void PrepareForUpdate(const TSharedRef<FUpdateContextPrivate>& OperationData);
+	UE_API void PrepareForUpdate();
 
 	// The following method is basically copied from PostEditChangeProperty and/or SkeletalMesh.cpp to be able to replicate PostEditChangeProperty without the editor
 	UE_API void PostEditChangePropertyWithoutEditor();
@@ -320,37 +314,16 @@ public:
 #if WITH_EDITORONLY_DATA
 	UE_API void RegenerateImportedModels();
 #endif
-
-private:
-
-	UE_API void InitSkeletalMeshData(const TSharedRef<FUpdateContextPrivate>& OperationData, USkeletalMesh* SkeletalMesh, const FMutableRefSkeletalMeshData& RefSkeletalMeshData, const UCustomizableObject& CustomizableObject, FCustomizableObjectComponentIndex ObjectComponentIndex);
-
-	UE_API bool BuildSkeletonData(const TSharedRef<FUpdateContextPrivate>& OperationData, USkeletalMesh& SkeletalMesh, const FMutableRefSkeletalMeshData& RefSkeletalMeshData, UCustomizableObject& CustomizableObject, FCustomizableObjectInstanceComponentIndex InstanceComponentIndex);
-	UE_API void BuildMeshSockets(const TSharedRef<FUpdateContextPrivate>& OperationData, USkeletalMesh* SkeletalMesh, const UModelResources& ModelResources, const FMutableRefSkeletalMeshData& RefSkeletalMeshData, TSharedPtr<const mu::FMesh> MutableMesh);
-	UE_API void BuildOrCopyElementData(const TSharedRef<FUpdateContextPrivate>& OperationData, USkeletalMesh* SkeletalMesh, UCustomizableObject& CustomizableObjectInstance, FCustomizableObjectInstanceComponentIndex InstanceComponentIndex);
-	UE_API void BuildOrCopyMorphTargetsData(const TSharedRef<FUpdateContextPrivate>& OperationData, USkeletalMesh* SkeletalMesh, const USkeletalMesh* SrcSkeletalMesh, UCustomizableObject& CustomizableObject, FCustomizableObjectInstanceComponentIndex InstanceComponentIndex);
-	UE_API bool BuildOrCopyRenderData(const TSharedRef<FUpdateContextPrivate>& OperationData, USkeletalMesh* SkeletalMesh, const USkeletalMesh* SrcSkeletalMesh, UCustomizableObjectInstance* CustomizableObjectInstance, FCustomizableObjectInstanceComponentIndex InstanceComponentIndex);
-
-	static UE_API void BuildOrCopyClothingData(const TSharedRef<FUpdateContextPrivate>& OperationData, USkeletalMesh* SkeletalMesh, const UModelResources& CustomizableObjectInstance, FCustomizableObjectInstanceComponentIndex InstanceComponentIndex, const TArray<TObjectPtr<UPhysicsAsset>>& ClothingPhysicsAssets);
-
+	
 	// 
 	UE_API FCustomizableInstanceComponentData* GetComponentData(FCustomizableObjectComponentIndex ObjectComponentIndex);
 	UE_API const FCustomizableInstanceComponentData* GetComponentData(FCustomizableObjectComponentIndex ObjectComponentIndex) const;
 	
-	//
-	UE_API USkeleton* MergeSkeletons(UCustomizableObject& CustomizableObject, const FMutableRefSkeletalMeshData& RefSkeletalMeshData, FCustomizableObjectComponentIndex ObjectComponentIndex, bool& bOutCreatedNewSkeleton);
-
-	//
-	UE_API UPhysicsAsset* GetOrBuildMainPhysicsAsset(const TSharedRef<FUpdateContextPrivate>& OperationData, TObjectPtr<class UPhysicsAsset> TamplateAsset, const mu::FPhysicsBody* PhysicsBody, bool bDisableCollisionBetweenAssets, FCustomizableObjectInstanceComponentIndex InstanceComponentIndex);
-	
 	// Create a transient texture and add it to the TextureTrackerArray
 	UE_API UTexture2D* CreateTexture(const FString& TextureName);
 
 	UE_API void InvalidateGeneratedData();
 
-	UE_API bool DoComponentsNeedUpdate(UCustomizableObjectInstance* CustomizableObjectInstance, const TSharedRef<FUpdateContextPrivate>& OperationData, bool& bOutEmptyMesh);
-
-public:
 	UE_API bool LoadParametersFromProfile(int32 ProfileIndex);
 	
 	UE_API bool SaveParametersToProfile(int32 ProfileIndex);
@@ -423,16 +396,6 @@ public:
 	TArray<TSoftObjectPtr<const UTexture>> PassThroughTexturesToLoad;
 	TArray<TSoftObjectPtr<const UStreamableRenderAsset>> PassThroughMeshesToLoad;
 
-	// Used during an update to prevent the pass-through textures loaded by LoadAdditionalAssetsAsync() from being unloaded by GC
-	// between AdditionalAssetsAsyncLoaded() and their setting into the generated materials in BuildMaterials()
-	UPROPERTY(Transient)
-	TArray<TObjectPtr<const UTexture>> LoadedPassThroughTexturesPendingSetMaterial;
-
-	// Used during an update to prevent the pass-through meshes loaded by LoadAdditionalAssetsAsync() from being unloaded by GC
-	// between AdditionalAssetsAsyncLoaded() and their setting into the generated materials in BuildMaterials()
-	UPROPERTY(Transient)
-	TArray<TObjectPtr<const UStreamableRenderAsset>> LoadedPassThroughMeshesPendingSetMaterial;
-
 private:
 	ECOInstanceFlags InstanceFlagsPrivate = ECOInstanceFlags::ECONone;
 
@@ -479,4 +442,30 @@ public:
 #endif // WITH_EDITOR
 };
 
+
+UE_API UPhysicsAsset* GetOrBuildMainPhysicsAsset(const TSharedRef<FUpdateContextPrivate>& Context, TObjectPtr<UPhysicsAsset> TamplateAsset, const mu::FPhysicsBody* PhysicsBody, bool bDisableCollisionBetweenAssets, FCustomizableObjectInstanceComponentIndex InstanceComponentIndex);
+
+UE_API USkeleton* MergeSkeletons(UCustomizableObject& CustomizableObject, FCustomizableInstanceComponentData& ObjectComponentIndex, bool& bOutCreatedNewSkeleton);
+
+UE_API void InitSkeletalMeshData(const TSharedRef<FUpdateContextPrivate>& Context, USkeletalMesh* SkeletalMesh, const FMutableRefSkeletalMeshData& RefSkeletalMeshData, const UCustomizableObject& CustomizableObject, FCustomizableObjectComponentIndex ObjectComponentIndex);
+
+UE_API void BuildOrCopyClothingData(const TSharedRef<FUpdateContextPrivate>& Context, USkeletalMesh* SkeletalMesh, const UModelResources& CustomizableObjectInstance, FCustomizableObjectInstanceComponentIndex InstanceComponentIndex, const TArray<TObjectPtr<UPhysicsAsset>>& ClothingPhysicsAssets);
+
+UE_API bool BuildSkeletonData(const TSharedRef<FUpdateContextPrivate>& Context, USkeletalMesh& SkeletalMesh, FCustomizableObjectInstanceComponentIndex InstanceComponentIndex);
+
+UE_API void BuildMeshSockets(USkeletalMesh* SkeletalMesh, const UModelResources& ModelResources, const FMutableRefSkeletalMeshData& RefSkeletalMeshData, TSharedPtr<const mu::FMesh> MutableMesh);
+
+UE_API bool BuildOrCopyRenderData(const TSharedRef<FUpdateContextPrivate>& Context, USkeletalMesh* SkeletalMesh, UCustomizableObjectInstance* CustomizableObjectInstance, FCustomizableObjectInstanceComponentIndex InstanceComponentIndex);
+
+UE_API void BuildOrCopyElementData(const TSharedRef<FUpdateContextPrivate>& Context, USkeletalMesh* SkeletalMesh, FCustomizableObjectInstanceComponentIndex InstanceComponentIndex);
+
+UE_API void BuildOrCopyMorphTargetsData(const TSharedRef<FUpdateContextPrivate>& Context, USkeletalMesh* SkeletalMesh, FCustomizableObjectInstanceComponentIndex InstanceComponentIndex);
+	
+UE_API bool DoComponentsNeedUpdate(const TSharedRef<FUpdateContextPrivate>& Context, bool& bOutEmptyMesh);
+
+UE_API UE::Tasks::FTask LoadAdditionalAssetsAndData(const TSharedRef<FUpdateContextPrivate>& Context);
+
+UE_API void AdditionalAssetsAsyncLoaded(TSharedRef<FUpdateContextPrivate> Context, UE::Tasks::FTaskEvent Event);
+
+
 #undef UE_API
diff --git a/Engine/Plugins/Mutable/Source/CustomizableObject/Internal/MuCO/CustomizableObjectSystemPrivate.h b/Engine/Plugins/Mutable/Source/CustomizableObject/Internal/MuCO/CustomizableObjectSystemPrivate.h
--- a/Engine/Plugins/Mutable/Source/CustomizableObject/Internal/MuCO/CustomizableObjectSystemPrivate.h
+++ b/Engine/Plugins/Mutable/Source/CustomizableObject/Internal/MuCO/CustomizableObjectSystemPrivate.h
@@ -707,6 +707,14 @@ public:
 	/** Hard references to objects. Avoids GC to collect them. */
 	TArray<TStrongObjectPtr<const UObject>> Objects;
 
+	// Used during an update to prevent the pass-through textures loaded by LoadAdditionalAssetsAsync() from being unloaded by GC
+	// between AdditionalAssetsAsyncLoaded() and their setting into the generated materials in BuildMaterials()
+	TArray<TStrongObjectPtr<const UTexture>> LoadedPassThroughTexturesPendingSetMaterial;
+
+	// Used during an update to prevent the pass-through meshes loaded by LoadAdditionalAssetsAsync() from being unloaded by GC
+	// between AdditionalAssetsAsyncLoaded() and their setting into the generated materials in BuildMaterials()
+	TArray<TStrongObjectPtr<const UStreamableRenderAsset>> LoadedPassThroughMeshesPendingSetMaterial;
+	
 	TArray<mu::FImageId> CachedTextures;
 };
 
diff --git a/Engine/Plugins/Mutable/Source/CustomizableObject/Private/MuCO/CustomizableObjectInstance.cpp b/Engine/Plugins/Mutable/Source/CustomizableObject/Private/MuCO/CustomizableObjectInstance.cpp
--- a/Engine/Plugins/Mutable/Source/CustomizableObject/Private/MuCO/CustomizableObjectInstance.cpp
+++ b/Engine/Plugins/Mutable/Source/CustomizableObject/Private/MuCO/CustomizableObjectInstance.cpp
@@ -4,7 +4,6 @@
 
 #include "Algo/Find.h"
 #include "Algo/MaxElement.h"
-#include "Animation/AnimClassInterface.h"
 #include "Animation/AnimBlueprintGeneratedClass.h"
 #include "Animation/AnimInstance.h"
 #include "Animation/Skeleton.h"
@@ -19,7 +18,6 @@
 #include "MutableStreamRequest.h"
 #include "Materials/Material.h"
 #include "Materials/MaterialInstanceDynamic.h"
-#include "Misc/Optional.h"
 #include "Modules/ModuleManager.h"
 #include "Tasks/Task.h"
 
@@ -47,15 +45,12 @@
 #include "PhysicsEngine/AggregateGeom.h"
 #include "PhysicsEngine/PhysicsConstraintTemplate.h"
 #include "PhysicsEngine/SkeletalBodySetup.h"
-#include "Hash/CityHash.h"
 #include "MuCO/CustomizableObjectCustomVersion.h"
 #include "MuCO/CustomizableObjectResourceData.h"
 #include "MuCO/CustomizableObjectStreamedResourceData.h"
 #include "MuCO/CustomizableObjectResourceDataTypes.h"
 #include "MuCO/LoadUtils.h"
 #include "Serialization/BulkData.h"
-#include "HAL/PlatformFileManager.h"
-#include "MuCO/UnrealMutableImageProvider.h"
 
 #include "MuCO/Plugins/IMutableClothingModule.h"
 #include "Rendering/SkeletalMeshModel.h"
@@ -110,8 +105,8 @@ struct FMutableMaterialPlaceholder
 		FName ParamName;
 		EPlaceHolderParamType Type;
 		int32 LayerIndex; // Set to -1 for non-multilayer params
-		float Scalar;
-		FLinearColor Vector;
+		float Scalar = 0.0f;
+		FLinearColor Vector = FLinearColor::Black;
 		FGeneratedTexture Texture;
 
 		FMutableMaterialPlaceHolderParam(const FName& InParamName, const int32 InLayerIndex, const FLinearColor& InVector)
@@ -370,7 +365,7 @@ void UCustomizableInstancePrivate::SetDescriptor(const FCustomizableObjectInstan
 }
 
 
-void UCustomizableInstancePrivate::PrepareForUpdate(const TSharedRef<FUpdateContextPrivate>& OperationData)
+void UCustomizableInstancePrivate::PrepareForUpdate()
 {
 	// Clear the ComponentData from previous updates
 	for (FCustomizableInstanceComponentData& ComponentData : ComponentsData)
@@ -1179,7 +1174,6 @@ bool AreSkeletonsCompatible(const TArray<TObjectPtr<USkeleton>>& InSkeletons)
 
 		const FReferenceSkeleton& ReferenceSkeleton = Skeleton->GetReferenceSkeleton();
 		const TArray<FMeshBoneInfo>& Bones = ReferenceSkeleton.GetRawRefBoneInfo();
-		const TArray<FTransform>& BonePoses = ReferenceSkeleton.GetRawRefBonePose();
 
 		const int32 NumBones = Bones.Num();
 		for (int32 BoneIndex = 0; BoneIndex < NumBones; ++BoneIndex)
@@ -1252,15 +1246,12 @@ bool AreSkeletonsCompatible(const TArray<TObjectPtr<USkeleton>>& InSkeletons)
 #endif
 
 
-USkeleton* UCustomizableInstancePrivate::MergeSkeletons(UCustomizableObject& CustomizableObject, const FMutableRefSkeletalMeshData& RefSkeletalMeshData, FCustomizableObjectComponentIndex ObjectComponentIndex, bool& bOutCreatedNewSkeleton)
+USkeleton* MergeSkeletons(UCustomizableObject& CustomizableObject, FCustomizableInstanceComponentData& ComponentData, bool& bOutCreatedNewSkeleton)
 {
 	MUTABLE_CPUPROFILER_SCOPE(BuildSkeletonData_MergeSkeletons);
 	bOutCreatedNewSkeleton = false;
-
-	FCustomizableInstanceComponentData* ComponentData = GetComponentData(ObjectComponentIndex);
-	check(ComponentData);
-
-	FReferencedSkeletons& ReferencedSkeletons = ComponentData->Skeletons;
+	
+	FReferencedSkeletons& ReferencedSkeletons = ComponentData.Skeletons;
 	
 	// Merged skeleton found in the cache
 	if (ReferencedSkeletons.Skeleton)
@@ -1293,9 +1284,7 @@ USkeleton* UCustomizableInstancePrivate::MergeSkeletons(UCustomizableObject& Cus
 	USkeleton* FinalSkeleton = USkeletalMergingLibrary::MergeSkeletons(Params);
 	if (!FinalSkeleton)
 	{
-		FString Msg = FString::Printf(TEXT("MergeSkeletons failed for Customizable Object [%s], Instance [%s]. Skeletons involved: "),
-			*CustomizableObject.GetName(),
-			*GetOuter()->GetName());
+		FString Msg = FString::Printf(TEXT("MergeSkeletons failed for Customizable Object [%s]. Skeletons involved: "), *CustomizableObject.GetName());
 		
 		const int32 SkeletonCount = Params.SkeletonsToMerge.Num();
 		for (int32 SkeletonIndex = 0; SkeletonIndex < SkeletonCount; ++SkeletonIndex)
@@ -1508,7 +1497,7 @@ namespace
 	};
 
 	TObjectPtr<UPhysicsAsset> MakePhysicsAssetFromTemplateAndMutableBody(
-		const TSharedRef<FUpdateContextPrivate>& OperationData,
+		const TSharedRef<FUpdateContextPrivate>& Context,
 		TObjectPtr<UPhysicsAsset> TemplateAsset,
 		const mu::FPhysicsBody* MutablePhysics,
 		FCustomizableObjectInstanceComponentIndex InstanceComponentIndex)
@@ -1526,7 +1515,7 @@ namespace
 
 		Result->bNotForDedicatedServer = TemplateAsset->bNotForDedicatedServer;
 
-		const TMap<mu::FBoneName, TPair<FName, uint16>>& BoneInfoMap = OperationData->InstanceUpdateData.SkeletonsPerInstanceComponent[InstanceComponentIndex.GetValue()].BoneInfoMap;
+		const TMap<mu::FBoneName, TPair<FName, uint16>>& BoneInfoMap = Context->InstanceUpdateData.SkeletonsPerInstanceComponent[InstanceComponentIndex.GetValue()].BoneInfoMap;
 		TMap<FName, int32> BonesInUse;
 
 		const int32 MutablePhysicsBodyCount = MutablePhysics->GetBodyCount();
@@ -1650,32 +1639,32 @@ namespace
 	}
 }
 
-UPhysicsAsset* UCustomizableInstancePrivate::GetOrBuildMainPhysicsAsset(
-	const TSharedRef<FUpdateContextPrivate>& OperationData,
+UPhysicsAsset* GetOrBuildMainPhysicsAsset(
+	const TSharedRef<FUpdateContextPrivate>& Context,
 	TObjectPtr<UPhysicsAsset> TemplateAsset,
 	const mu::FPhysicsBody* MutablePhysics,
 	bool bDisableCollisionsBetweenDifferentAssets,
 	FCustomizableObjectInstanceComponentIndex InstanceComponentIndex)
 {
-
 	MUTABLE_CPUPROFILER_SCOPE(MergePhysicsAssets);
 
 	check(MutablePhysics);
 
+	UCustomizableObjectInstance* Instance = Context->Instance.Get();
+	UCustomizableInstancePrivate* InstancePrivate = Instance->GetPrivate();
+	
 	UPhysicsAsset* Result = nullptr;
 
-	const FInstanceUpdateData::FComponent* Component = OperationData->GetComponentUpdateData(InstanceComponentIndex);
+	const FInstanceUpdateData::FComponent* Component = Context->GetComponentUpdateData(InstanceComponentIndex);
 	if (!Component)
 	{
 		return nullptr;
 	}
 
 	FCustomizableObjectComponentIndex ObjectComponentIndex = Component->Id;
-	FCustomizableInstanceComponentData* ComponentData = GetComponentData(ObjectComponentIndex);
+	FCustomizableInstanceComponentData* ComponentData = InstancePrivate->GetComponentData(ObjectComponentIndex);
 	check(ComponentData);
-
-	TArray<TObjectPtr<UPhysicsAsset>>& PhysicsAssets = ComponentData->PhysicsAssets.PhysicsAssetsToMerge;
-
+	
 	TArray<TObjectPtr<UPhysicsAsset>> ValidAssets;
 
 	const int32 NumPhysicsAssets = ComponentData->PhysicsAssets.PhysicsAssetsToMerge.Num();
@@ -1715,7 +1704,7 @@ UPhysicsAsset* UCustomizableInstancePrivate::GetOrBuildMainPhysicsAsset(
 
 	Result->bNotForDedicatedServer = TemplateAsset->bNotForDedicatedServer;
 
-	const TMap<mu::FBoneName, TPair<FName, uint16>>& BoneInfoMap = OperationData->InstanceUpdateData.SkeletonsPerInstanceComponent[InstanceComponentIndex.GetValue()].BoneInfoMap;
+	const TMap<mu::FBoneName, TPair<FName, uint16>>& BoneInfoMap = Context->InstanceUpdateData.SkeletonsPerInstanceComponent[InstanceComponentIndex.GetValue()].BoneInfoMap;
 	TMap<FName, int32> BonesInUse;
 
 	const int32 MutablePhysicsBodyCount = MutablePhysics->GetBodyCount();
@@ -1736,9 +1725,6 @@ UPhysicsAsset* UCustomizableInstancePrivate::GetOrBuildMainPhysicsAsset(
 	// Map from BodyName/BoneName to set and index in set.
 	TMap<FName, CollisionSetEntryType> BodySetupSetMap;
 	
-	// {SetIndex0, SetIndex1, BodyIndex}
-	using MultiSetCollisionEnableType = TTuple<int32, int32, int32>;
-
 	// Only for elements that belong to two or more differnet sets. 
 	// Contains in which set the elements belong.
 	using MultiSetArrayType = TArray<int32, TInlineAllocator<4>>;
@@ -1975,11 +1961,14 @@ void SetMeshUVChannelDensity(FMeshUVChannelInfo& UVChannelInfo, float Density =
 }
 
 
-bool UCustomizableInstancePrivate::DoComponentsNeedUpdate(UCustomizableObjectInstance* Public, const TSharedRef<FUpdateContextPrivate>& OperationData, bool& bHasInvalidMesh)
+bool DoComponentsNeedUpdate(const TSharedRef<FUpdateContextPrivate>& Context, bool& bHasInvalidMesh)
 {
 	MUTABLE_CPUPROFILER_SCOPE(UCustomizableInstancePrivate::DoComponentsNeedUpdate);
 
-	UCustomizableObject* CustomizableObject = Public->GetCustomizableObject();
+	UCustomizableObjectInstance* Instance = Context->Instance.Get();
+	UCustomizableInstancePrivate* InstancePrivate = Instance->GetPrivate();
+	
+	UCustomizableObject* CustomizableObject = Instance->GetCustomizableObject();
 	
 	if (!CustomizableObject)
 	{
@@ -1989,28 +1978,28 @@ bool UCustomizableInstancePrivate::DoComponentsNeedUpdate(UCustomizableObjectIns
 	bHasInvalidMesh = false;
 	bool bNeedsUpdate = false;
 	
-	const int32 NumInstanceComponents = OperationData->InstanceUpdateData.Components.Num();
-	bNeedsUpdate |= NumInstanceComponents != SkeletalMeshes.Num();
+	const int32 NumInstanceComponents = Context->InstanceUpdateData.Components.Num();
+	bNeedsUpdate |= NumInstanceComponents != InstancePrivate->SkeletalMeshes.Num();
 
 	// Find which components need an update
-	OperationData->MeshChangedPerInstanceComponent.Init(false, NumInstanceComponents);
+	Context->MeshChangedPerInstanceComponent.Init(false, NumInstanceComponents);
 
 	for (int32 InstanceComponentIndex = 0; InstanceComponentIndex < NumInstanceComponents; ++InstanceComponentIndex)
 	{
-		const FInstanceUpdateData::FComponent& Component = OperationData->InstanceUpdateData.Components[InstanceComponentIndex];
+		const FInstanceUpdateData::FComponent& Component = Context->InstanceUpdateData.Components[InstanceComponentIndex];
 		const FCustomizableObjectComponentIndex ObjectComponentIndex = Component.Id;
-		const FName ComponentName = OperationData->ComponentNames[ObjectComponentIndex.GetValue()];
+		const FName ComponentName = Context->ComponentNames[ObjectComponentIndex.GetValue()];
 		
 		bool bComponentWithMesh = false;
 
-		const int32 FirstLOD = OperationData->bStreamMeshLODs ?
-			OperationData->FirstLODAvailable[ComponentName] :
-			OperationData->GetFirstRequestedLOD()[ComponentName];
+		const int32 FirstLOD = Context->bStreamMeshLODs ?
+			Context->FirstLODAvailable[ComponentName] :
+			Context->GetFirstRequestedLOD()[ComponentName];
 
 		// Look for invalid meshes
 		for (int32 LODIndex = FirstLOD; LODIndex < Component.LODCount; ++LODIndex)
 		{
-			const FInstanceUpdateData::FLOD& LOD = OperationData->InstanceUpdateData.LODs[Component.FirstLOD + LODIndex];
+			const FInstanceUpdateData::FLOD& LOD = Context->InstanceUpdateData.LODs[Component.FirstLOD + LODIndex];
 
 			if (!LOD.Mesh)
 			{
@@ -2026,7 +2015,7 @@ bool UCustomizableInstancePrivate::DoComponentsNeedUpdate(UCustomizableObjectIns
 			if (!LOD.Mesh->IsReference() && LOD.Mesh->GetVertexCount() == 0)
 			{
 				UE_LOG(LogMutable, Error, TEXT("Failed to generate SkeletalMesh for CO Instance [%s]. CO [%s] has invalid geometry for LOD [%d] Component [%d]."),
-					*Public->GetName(), *CustomizableObject->GetName(),
+					*Instance->GetName(), *CustomizableObject->GetName(),
 					LODIndex, InstanceComponentIndex);
 				bHasInvalidMesh = true; // Unreal does not support empty LODs.
 			}
@@ -2038,16 +2027,16 @@ bool UCustomizableInstancePrivate::DoComponentsNeedUpdate(UCustomizableObjectIns
 		bool bMeshChanged = true;
 
 		// Compare Mesh IDs 
-		const TArray<mu::FMeshId>* MeshIDs = OperationData->GetMeshDescriptors(ObjectComponentIndex);
+		const TArray<mu::FMeshId>* MeshIDs = Context->GetMeshDescriptors(ObjectComponentIndex);
 		if (ensure(MeshIDs))
 		{
-			if (const FCustomizableInstanceComponentData* ComponentData = GetComponentData(ObjectComponentIndex)) // Could be nullptr if the component has not been generated.
+			if (const FCustomizableInstanceComponentData* ComponentData = InstancePrivate->GetComponentData(ObjectComponentIndex)) // Could be nullptr if the component has not been generated.
 			{
 				bMeshChanged = *MeshIDs != ComponentData->LastMeshIdPerLOD;
 			}
 		}
 
-		OperationData->MeshChangedPerInstanceComponent[InstanceComponentIndex] = bMeshChanged;
+		Context->MeshChangedPerInstanceComponent[InstanceComponentIndex] = bMeshChanged;
 		bNeedsUpdate |= bMeshChanged;
 	}
 
@@ -2055,22 +2044,22 @@ bool UCustomizableInstancePrivate::DoComponentsNeedUpdate(UCustomizableObjectIns
 }
 
 
-bool UCustomizableInstancePrivate::UpdateSkeletalMesh_PostBeginUpdate0(UCustomizableObjectInstance* Public, const TSharedRef<FUpdateContextPrivate>& OperationData)
+bool UCustomizableInstancePrivate::UpdateSkeletalMesh_PostBeginUpdate0(UCustomizableObjectInstance* Instance, const TSharedRef<FUpdateContextPrivate>& Context)
 {
 	MUTABLE_CPUPROFILER_SCOPE(UCustomizableInstancePrivate::UpdateSkeletalMesh_PostBeginUpdate0)
 
 	bool bHasInvalidMesh = false;
 
-	bool bUpdateMeshes = DoComponentsNeedUpdate(Public, OperationData, bHasInvalidMesh);
+	bool bUpdateMeshes = DoComponentsNeedUpdate(Context, bHasInvalidMesh);
 
-	UCustomizableObject* CustomizableObject = Public->GetCustomizableObject();
+	UCustomizableObject* CustomizableObject = Instance->GetCustomizableObject();
 
 	if (!CustomizableObject)
 	{
-		UE_LOG(LogMutable, Warning, TEXT("Failed to generate SkeletalMesh for CO Instance %s. It does not have a CO."), *Public->GetName());
+		UE_LOG(LogMutable, Warning, TEXT("Failed to generate SkeletalMesh for CO Instance %s. It does not have a CO."), *Instance->GetName());
 		
 		InvalidateGeneratedData();
-		OperationData->UpdateResult = EUpdateResult::Error;
+		Context->UpdateResult = EUpdateResult::Error;
 
 		return false;
 	}
@@ -2078,10 +2067,10 @@ bool UCustomizableInstancePrivate::UpdateSkeletalMesh_PostBeginUpdate0(UCustomiz
 	// We can not handle empty meshes, clear any generated mesh and return
 	if (bHasInvalidMesh)
 	{
-		UE_LOG(LogMutable, Warning, TEXT("Failed to generate SkeletalMesh for CO Instance %s. CO [%s]"), *Public->GetName(), *GetNameSafe(CustomizableObject));
+		UE_LOG(LogMutable, Warning, TEXT("Failed to generate SkeletalMesh for CO Instance %s. CO [%s]"), *Instance->GetName(), *GetNameSafe(CustomizableObject));
 
 		InvalidateGeneratedData();
-		OperationData->UpdateResult = EUpdateResult::Error;
+		Context->UpdateResult = EUpdateResult::Error;
 
 		return false;
 	}
@@ -2103,7 +2092,7 @@ bool UCustomizableInstancePrivate::UpdateSkeletalMesh_PostBeginUpdate0(UCustomiz
 		//
 		// This ensures that extension functions such as OnSkeletalMeshCreated are called for each
 		// extension, even if they didn't produce any extension data.
-		Public->GetPrivate()->ExtensionInstanceData.Empty(AllExtensions.Num());
+		Instance->GetPrivate()->ExtensionInstanceData.Empty(AllExtensions.Num());
 		for (const UCustomizableObjectExtension* Extension : AllExtensions)
 		{
 			ExtensionToExtensionData.Add(Extension);
@@ -2111,7 +2100,7 @@ bool UCustomizableInstancePrivate::UpdateSkeletalMesh_PostBeginUpdate0(UCustomiz
 
 		const TArrayView<const FRegisteredObjectNodeInputPin> ExtensionPins = ICustomizableObjectModule::Get().GetAdditionalObjectNodePins();
 
-		for (FInstanceUpdateData::FNamedExtensionData& ExtensionOutput : OperationData->InstanceUpdateData.ExtendedInputPins)
+		for (FInstanceUpdateData::FNamedExtensionData& ExtensionOutput : Context->InstanceUpdateData.ExtendedInputPins)
 		{
 			const FRegisteredObjectNodeInputPin* FoundPin =
 				Algo::FindBy(ExtensionPins, ExtensionOutput.Name, &FRegisteredObjectNodeInputPin::GlobalPinName);
@@ -2202,7 +2191,7 @@ bool UCustomizableInstancePrivate::UpdateSkeletalMesh_PostBeginUpdate0(UCustomiz
 		FInstancedStruct NewExtensionInstanceData = Pair.Key->GenerateExtensionInstanceData(Pair.Value);
 		if (NewExtensionInstanceData.IsValid())
 		{
-			FExtensionInstanceData& NewData = Public->GetPrivate()->ExtensionInstanceData.AddDefaulted_GetRef();
+			FExtensionInstanceData& NewData = Instance->GetPrivate()->ExtensionInstanceData.AddDefaulted_GetRef();
 			NewData.Extension = Pair.Key;
 			NewData.Data = MoveTemp(NewExtensionInstanceData);
 		}
@@ -2216,7 +2205,7 @@ bool UCustomizableInstancePrivate::UpdateSkeletalMesh_PostBeginUpdate0(UCustomiz
 			FName DeletedComponentCandidateName = MeshIt.Key;
 			bool bIsDeleted = true;
 
-			for (const auto& InstanceUpdateComponent : OperationData->InstanceUpdateData.Components)
+			for (const auto& InstanceUpdateComponent : Context->InstanceUpdateData.Components)
 			{
 				FName InstanceUpdateComponentName = ModelResources.ComponentNamesPerObjectComponent[InstanceUpdateComponent.Id.GetValue()];
 
@@ -2245,11 +2234,11 @@ bool UCustomizableInstancePrivate::UpdateSkeletalMesh_PostBeginUpdate0(UCustomiz
 
 	SkeletalMeshes.Reset();
 	
-	const int32 NumInstanceComponents = OperationData->InstanceUpdateData.Components.Num();
+	const int32 NumInstanceComponents = Context->InstanceUpdateData.Components.Num();
 	for (int32 InstanceComponentIndexValue = 0; InstanceComponentIndexValue < NumInstanceComponents; ++InstanceComponentIndexValue)
 	{
 		FCustomizableObjectInstanceComponentIndex InstanceComponentIndex(InstanceComponentIndexValue);
-		const FInstanceUpdateData::FComponent* Component = OperationData->GetComponentUpdateData( InstanceComponentIndex );
+		const FInstanceUpdateData::FComponent* Component = Context->GetComponentUpdateData( InstanceComponentIndex );
 		if (!Component)
 		{
 			continue;
@@ -2270,13 +2259,13 @@ bool UCustomizableInstancePrivate::UpdateSkeletalMesh_PostBeginUpdate0(UCustomiz
 			return false;
 		}
 
-		TArray<mu::FMeshId>* MeshDescriptors = OperationData->GetMeshDescriptors(ObjectComponentIndex);
+		TArray<mu::FMeshId>* MeshDescriptors = Context->GetMeshDescriptors(ObjectComponentIndex);
 		check(MeshDescriptors);
 		check(MeshDescriptors->Num() == MAX_MESH_LOD_COUNT);
 		ComponentData->LastMeshIdPerLOD = *MeshDescriptors;
 
 		// If the component doesn't need an update copy the previously generated mesh.
-		if (!OperationData->MeshChangedPerInstanceComponent[InstanceComponentIndex.GetValue()])
+		if (!Context->MeshChangedPerInstanceComponent[InstanceComponentIndex.GetValue()])
 		{
 			if (TObjectPtr<USkeletalMesh>* Result = OldSkeletalMeshes.Find(ComponentName))
 			{
@@ -2287,13 +2276,13 @@ bool UCustomizableInstancePrivate::UpdateSkeletalMesh_PostBeginUpdate0(UCustomiz
 		}
 
 		uint32 MeshHash = INDEX_NONE;
-		for (int32 LODIndex = OperationData->FirstLODAvailable[ComponentName]; LODIndex < Component->LODCount && LODIndex < MAX_MESH_LOD_COUNT; ++LODIndex)
+		for (int32 LODIndex = Context->FirstLODAvailable[ComponentName]; LODIndex < Component->LODCount && LODIndex < MAX_MESH_LOD_COUNT; ++LODIndex)
 		{
 			const mu::FGeneratedMeshKey* MeshKey = (*MeshDescriptors)[LODIndex].GetKey();
 			check(MeshKey);
-			const uint32 MeshIDHash = GetTypeHash(*MeshKey, OperationData->bBake);
+			const uint32 MeshIDHash = GetTypeHash(*MeshKey, Context->bBake);
 
-			if (!OperationData->bBake)
+			if (!Context->bBake)
 			{
 				MeshHash = HashCombineFast(MeshHash, MeshIDHash);
 			}
@@ -2303,7 +2292,7 @@ bool UCustomizableInstancePrivate::UpdateSkeletalMesh_PostBeginUpdate0(UCustomiz
 			}
 		}
 
-		if (OperationData->bUseMeshCache)
+		if (Context->bUseMeshCache)
 		{
 			if (USkeletalMesh* CachedMesh = CustomizableObject->GetPrivate()->MeshCache.Get(*MeshDescriptors))
 			{
@@ -2316,7 +2305,7 @@ bool UCustomizableInstancePrivate::UpdateSkeletalMesh_PostBeginUpdate0(UCustomiz
 		TSharedPtr<const mu::FMesh> ComponentMesh;
 		for (int32 LODIndex = 0; LODIndex < Component->LODCount && !ComponentMesh; ++LODIndex)
 		{
-			ComponentMesh = OperationData->InstanceUpdateData.LODs[Component->FirstLOD + LODIndex].Mesh;
+			ComponentMesh = Context->InstanceUpdateData.LODs[Component->FirstLOD + LODIndex].Mesh;
 		}
 		
 		if (!ComponentMesh)
@@ -2366,10 +2355,10 @@ bool UCustomizableInstancePrivate::UpdateSkeletalMesh_PostBeginUpdate0(UCustomiz
 
 		USkeletalMesh* SkeletalMesh = nullptr;
 		{
-			FString SkeletalMeshName = FString::Printf (TEXT("SK_%s_%s_%lu"), *Public->GetCustomizableObject()->GetName(), *ComponentName.ToString(), MeshHash);
+			FString SkeletalMeshName = FString::Printf (TEXT("SK_%s_%s_%lu"), *Instance->GetCustomizableObject()->GetName(), *ComponentName.ToString(), MeshHash);
 
 #if WITH_EDITOR
-			if (OperationData->bBake)
+			if (Context->bBake)
 			{
 				SkeletalMeshName = FBakingConfiguration::BakedResourcePrefix + SkeletalMeshName;
 			}
@@ -2378,7 +2367,7 @@ bool UCustomizableInstancePrivate::UpdateSkeletalMesh_PostBeginUpdate0(UCustomiz
 			// Make name unique to avoid collisions with other objects.
 			MakeMutableGeneratedObjectNameUnique(SkeletalMeshName, MeshHash, USkeletalMesh::StaticClass());
 			
-			if (!OperationData->bBake)
+			if (!Context->bBake)
 			{
 				SkeletalMesh = NewObject<UCustomizableObjectSkeletalMesh>(GetTransientPackage(), FName(*SkeletalMeshName), RF_Transient);
 			}
@@ -2396,10 +2385,10 @@ bool UCustomizableInstancePrivate::UpdateSkeletalMesh_PostBeginUpdate0(UCustomiz
 		const FMutableRefSkeletalMeshData& RefSkeletalMeshData = ModelResources.ReferenceSkeletalMeshesData[ObjectComponentIndex.GetValue()];
 
 		// Set up the default information any mesh from this component will have (LODArrayInfos, RenderData, Mesh settings, etc). 
-		InitSkeletalMeshData(OperationData, SkeletalMesh, RefSkeletalMeshData, *CustomizableObject, ObjectComponentIndex);
+		InitSkeletalMeshData(Context, SkeletalMesh, RefSkeletalMeshData, *CustomizableObject, ObjectComponentIndex);
 		
 		// Construct a new skeleton, fix up ActiveBones and Bonemap arrays and recompute the RefInvMatrices
-		const bool bBuildSkeletonDataSuccess = BuildSkeletonData(OperationData, *SkeletalMesh, RefSkeletalMeshData, *CustomizableObject, FCustomizableObjectInstanceComponentIndex(InstanceComponentIndex) );
+		const bool bBuildSkeletonDataSuccess = BuildSkeletonData(Context, *SkeletalMesh, FCustomizableObjectInstanceComponentIndex(InstanceComponentIndex) );
 		if (!bBuildSkeletonDataSuccess)
 		{
 			InvalidateGeneratedData();
@@ -2411,7 +2400,7 @@ bool UCustomizableInstancePrivate::UpdateSkeletalMesh_PostBeginUpdate0(UCustomiz
 		{
 			constexpr bool bDisallowCollisionBetweenAssets = true;
 			UPhysicsAsset* PhysicsAssetResult = GetOrBuildMainPhysicsAsset(
-				OperationData, RefSkeletalMeshData.PhysicsAsset, MutablePhysics.Get(), bDisallowCollisionBetweenAssets, FCustomizableObjectInstanceComponentIndex(InstanceComponentIndex) );
+				Context, RefSkeletalMeshData.PhysicsAsset, MutablePhysics.Get(), bDisallowCollisionBetweenAssets, FCustomizableObjectInstanceComponentIndex(InstanceComponentIndex) );
 
 			SkeletalMesh->SetPhysicsAsset(PhysicsAssetResult);
 
@@ -2458,31 +2447,27 @@ bool UCustomizableInstancePrivate::UpdateSkeletalMesh_PostBeginUpdate0(UCustomiz
 
 			Entry.PropertyIndex = Info.PropertyIndex;
 			Entry.PhysicsAsset = MakePhysicsAssetFromTemplateAndMutableBody(
-				OperationData,PhysicsAssetTemplate, AdditionalPhysicsBody.Get(), InstanceComponentIndex);
+				Context,PhysicsAssetTemplate, AdditionalPhysicsBody.Get(), InstanceComponentIndex);
 		}
 
 		// Add sockets from the SkeletalMesh of reference and from the MutableMesh
-		BuildMeshSockets(OperationData, SkeletalMesh, ModelResources, RefSkeletalMeshData, ComponentMesh);
+		BuildMeshSockets(SkeletalMesh, ModelResources, RefSkeletalMeshData, ComponentMesh);
 		
 		for (const TPair<const UCustomizableObjectExtension*, TArray<FInputPinDataContainer>>& Pair : ExtensionToExtensionData)
 		{
 			Pair.Key->OnSkeletalMeshCreated(Pair.Value, ComponentName, SkeletalMesh);
 		}
 		
-		// Mesh to copy data from if possible.
-		TObjectPtr<USkeletalMesh>* OldSkeletalMeshPtr = OldSkeletalMeshes.Find(ComponentName);
-		USkeletalMesh* OldSkeletalMesh = OldSkeletalMeshPtr ? OldSkeletalMeshPtr->Get() : nullptr;
-
-		BuildOrCopyElementData(OperationData, SkeletalMesh, *CustomizableObject, InstanceComponentIndex);
-		bool const bCopyRenderDataSuccess = BuildOrCopyRenderData(OperationData, SkeletalMesh, OldSkeletalMesh, Public, InstanceComponentIndex);
+		BuildOrCopyElementData(Context, SkeletalMesh, InstanceComponentIndex);
+		bool const bCopyRenderDataSuccess = BuildOrCopyRenderData(Context, SkeletalMesh, Instance, InstanceComponentIndex);
 		if (!bCopyRenderDataSuccess)
 		{
 			InvalidateGeneratedData();
 			return false;
 		}
 
-		BuildOrCopyMorphTargetsData(OperationData, SkeletalMesh, OldSkeletalMesh, *CustomizableObject, InstanceComponentIndex);
-		BuildOrCopyClothingData(OperationData, SkeletalMesh, ModelResources, InstanceComponentIndex, ClothingPhysicsAssets);
+		BuildOrCopyMorphTargetsData(Context, SkeletalMesh, InstanceComponentIndex);
+		BuildOrCopyClothingData(Context, SkeletalMesh, ModelResources, InstanceComponentIndex, ClothingPhysicsAssets);
 
 		FSkeletalMeshRenderData* RenderData = SkeletalMesh->GetResourceForRendering();
 		ensure(RenderData && RenderData->LODRenderData.Num() > 0);
@@ -2496,9 +2481,9 @@ bool UCustomizableInstancePrivate::UpdateSkeletalMesh_PostBeginUpdate0(UCustomiz
 			}
 		}
 		
-		if (OperationData->bUseMeshCache)
+		if (Context->bUseMeshCache)
 		{
-			const TArray<mu::FMeshId>* MeshId = OperationData->GetMeshDescriptors(ObjectComponentIndex);
+			const TArray<mu::FMeshId>* MeshId = Context->GetMeshDescriptors(ObjectComponentIndex);
 			if (MeshId)
 			{
 				CustomizableObject->GetPrivate()->MeshCache.Add(*MeshId, CastChecked<UCustomizableObjectSkeletalMesh>(SkeletalMesh));
@@ -2507,7 +2492,7 @@ bool UCustomizableInstancePrivate::UpdateSkeletalMesh_PostBeginUpdate0(UCustomiz
 
 		if (UCustomizableObjectSkeletalMesh* StreamableMesh = Cast<UCustomizableObjectSkeletalMesh>(SkeletalMesh))
 		{
-			StreamableMesh->InitMutableStreamingData(OperationData, ComponentName, Component->FirstLOD, Component->LODCount);
+			StreamableMesh->InitMutableStreamingData(Context, ComponentName, Component->FirstLOD, Component->LODCount);
 		}
 	}
 
@@ -3260,17 +3245,6 @@ void UCustomizableInstancePrivate::SetReferenceSkeletalMesh() const
 }
 
 
-namespace
-{
-	inline float unpack_uint8(uint8 i)
-	{
-		float res = i;
-		res -= 127.5f;
-		res /= 127.5f;
-		return res;
-	}
-}
-
 bool MutableTextureUsesOfflineProcessedData()
 {
 #if PLATFORM_DESKTOP || PLATFORM_ANDROID || PLATFORM_IOS
@@ -3675,7 +3649,6 @@ void ConvertImage(UTexture2D* Texture, TSharedPtr<const mu::FImage> MutableImage
 		}
 
 		MutableImage = Converted;
-		MutableFormat = mu::EImageFormat::RGBA_UByte;
 	}
 	else if (MutableFormat == mu::EImageFormat::BGRA_UByte)
 	{
@@ -3698,7 +3671,6 @@ void ConvertImage(UTexture2D* Texture, TSharedPtr<const mu::FImage> MutableImage
 		}
 
 		MutableImage = Converted;
-		MutableFormat = mu::EImageFormat::RGBA_UByte;
 	}
 
 	if (OnlyLOD >= 0)
@@ -3706,7 +3678,7 @@ void ConvertImage(UTexture2D* Texture, TSharedPtr<const mu::FImage> MutableImage
 		OnlyLOD = FMath::Min( OnlyLOD, MutableImage->GetLODCount()-1 );
 	}
 
-	Texture->SetPlatformData(MutableCreateImagePlatformData(MutableImage,OnlyLOD,0,0) );
+	Texture->SetPlatformData(MutableCreateImagePlatformData(MutableImage, OnlyLOD,0,0) );
 }
 
 static int32 EnableRayTracingFix = 0;
@@ -3718,21 +3690,24 @@ FAutoConsoleVariableRef CVarMutableEnableRayTracingFix(
 	TEXT("If 2, Enable fix for all generated meshes. Meshes will have ray tracing disabled.")
 	);
 
-void UCustomizableInstancePrivate::InitSkeletalMeshData(const TSharedRef<FUpdateContextPrivate>& OperationData, USkeletalMesh* SkeletalMesh, const FMutableRefSkeletalMeshData& RefSkeletalMeshData, const UCustomizableObject& CustomizableObject, FCustomizableObjectComponentIndex ObjectComponentIndex)
+void InitSkeletalMeshData(const TSharedRef<FUpdateContextPrivate>& Context, USkeletalMesh* SkeletalMesh, const FMutableRefSkeletalMeshData& RefSkeletalMeshData, const UCustomizableObject& CustomizableObject, FCustomizableObjectComponentIndex ObjectComponentIndex)
 {
 	MUTABLE_CPUPROFILER_SCOPE(UCustomizableInstancePrivate::InitSkeletalMesh);
 
 	check(SkeletalMesh);
 
-	const FName ComponentName = OperationData->ComponentNames[ObjectComponentIndex.GetValue()];
+	UCustomizableObjectInstance* Instance = Context->Instance.Get();
+	UCustomizableInstancePrivate* InstancePrivate = Instance->GetPrivate();
+	
+	const FName ComponentName = Context->ComponentNames[ObjectComponentIndex.GetValue()];
 
-	SkeletalMesh->NeverStream = !OperationData->bStreamMeshLODs;
+	SkeletalMesh->NeverStream = !Context->bStreamMeshLODs;
 
 	SkeletalMesh->SetImportedBounds(RefSkeletalMeshData.Bounds);
 	SkeletalMesh->SetPostProcessAnimBlueprint(RefSkeletalMeshData.PostProcessAnimInst.Get());
 	SkeletalMesh->SetShadowPhysicsAsset(RefSkeletalMeshData.ShadowPhysicsAsset.Get());
 	
-	const bool bEnableRayTracingFix = EnableRayTracingFix == 2 || (EnableRayTracingFix == 1 && OperationData->bStreamMeshLODs);
+	const bool bEnableRayTracingFix = EnableRayTracingFix == 2 || (EnableRayTracingFix == 1 && Context->bStreamMeshLODs);
 	if (bEnableRayTracingFix)
 	{
 		SkeletalMesh->SetSupportRayTracing(false);
@@ -3746,7 +3721,7 @@ void UCustomizableInstancePrivate::InitSkeletalMeshData(const TSharedRef<FUpdate
 
 	// Asset User Data
 	{
-		const FCustomizableInstanceComponentData* ComponentData = GetComponentData(ObjectComponentIndex);
+		const FCustomizableInstanceComponentData* ComponentData = InstancePrivate->GetComponentData(ObjectComponentIndex);
 		check(ComponentData);
 		for (TObjectPtr<UAssetUserData> AssetUserData : ComponentData->AssetUserDataArray)
 		{
@@ -3754,11 +3729,11 @@ void UCustomizableInstancePrivate::InitSkeletalMeshData(const TSharedRef<FUpdate
 		}
 
 		//Custom Asset User Data
-		if (OperationData->Instance->GetAnimationGameplayTags().Num() ||
+		if (Context->Instance->GetAnimationGameplayTags().Num() ||
 			ComponentData->AnimSlotToBP.Num())
 		{
 			UCustomizableObjectInstanceUserData* InstanceData = NewObject<UCustomizableObjectInstanceUserData>(SkeletalMesh, NAME_None, RF_Public | RF_Transactional);
-			InstanceData->AnimationGameplayTag = OperationData->Instance->GetAnimationGameplayTags();
+			InstanceData->AnimationGameplayTag = Context->Instance->GetAnimationGameplayTags();
 			
 			for (const TTuple<FName, TSoftClassPtr<UAnimInstance>>& AnimSlot : ComponentData->AnimSlotToBP)
 			{
@@ -3779,12 +3754,12 @@ void UCustomizableInstancePrivate::InitSkeletalMeshData(const TSharedRef<FUpdate
 		SkeletalMesh->AllocateResourceForRendering();
 		
 		FSkeletalMeshRenderData* RenderData = SkeletalMesh->GetResourceForRendering();
-		int32 NumLODsAvailablePerComponent = OperationData->NumLODsAvailable[ComponentName];
-		RenderData->NumInlinedLODs = NumLODsAvailablePerComponent - OperationData->FirstResidentLOD[ComponentName];
-		RenderData->NumNonOptionalLODs = NumLODsAvailablePerComponent - OperationData->FirstLODAvailable[ComponentName];
-		RenderData->CurrentFirstLODIdx = OperationData->FirstResidentLOD[ComponentName];
+		int32 NumLODsAvailablePerComponent = Context->NumLODsAvailable[ComponentName];
+		RenderData->NumInlinedLODs = NumLODsAvailablePerComponent - Context->FirstResidentLOD[ComponentName];
+		RenderData->NumNonOptionalLODs = NumLODsAvailablePerComponent - Context->FirstLODAvailable[ComponentName];
+		RenderData->CurrentFirstLODIdx = Context->FirstResidentLOD[ComponentName];
 		RenderData->PendingFirstLODIdx = RenderData->CurrentFirstLODIdx;
-		RenderData->LODBiasModifier = OperationData->FirstLODAvailable[ComponentName];		
+		RenderData->LODBiasModifier = Context->FirstLODAvailable[ComponentName];		
 		
 		if (bEnableRayTracingFix)
 		{
@@ -3796,8 +3771,8 @@ void UCustomizableInstancePrivate::InitSkeletalMeshData(const TSharedRef<FUpdate
 			RenderData->LODRenderData.Add(new FSkeletalMeshLODRenderData());
 			
 			FSkeletalMeshLODRenderData& LODRenderData = RenderData->LODRenderData[LODIndex];
-			LODRenderData.bIsLODOptional = LODIndex < OperationData->FirstLODAvailable[ComponentName];
-			LODRenderData.bStreamedDataInlined = LODIndex >= OperationData->FirstResidentLOD[ComponentName];
+			LODRenderData.bIsLODOptional = LODIndex < Context->FirstLODAvailable[ComponentName];
+			LODRenderData.bStreamedDataInlined = LODIndex >= Context->FirstResidentLOD[ComponentName];
 
 			const FMutableRefLODData& LODData = RefSkeletalMeshData.LODData[LODIndex];
 			FSkeletalMeshLODInfo& LODInfo = SkeletalMesh->AddLODInfo();
@@ -3856,7 +3831,7 @@ void UCustomizableInstancePrivate::InitSkeletalMeshData(const TSharedRef<FUpdate
 
 	// Set Min LOD (Override the Reference Skeletal Mesh LOD Settings)
 	const UModelResources& ModelResources = CustomizableObject.GetPrivate()->GetModelResourcesChecked();
-	SkeletalMesh->SetMinLod(FMath::Max(ModelResources.MinLODPerComponent.FindChecked(ComponentName).GetDefault(), static_cast<int32>(OperationData->FirstLODAvailable[ComponentName])));
+	SkeletalMesh->SetMinLod(FMath::Max(ModelResources.MinLODPerComponent.FindChecked(ComponentName).GetDefault(), static_cast<int32>(Context->FirstLODAvailable[ComponentName])));
 	SkeletalMesh->SetQualityLevelMinLod(ModelResources.MinQualityLevelLODPerComponent.FindChecked(ComponentName));
 
 	// Set up unreal's default material, will be replaced when building materials
@@ -3873,14 +3848,22 @@ void UCustomizableInstancePrivate::InitSkeletalMeshData(const TSharedRef<FUpdate
 }
 
 
-bool UCustomizableInstancePrivate::BuildSkeletonData(const TSharedRef<FUpdateContextPrivate>& OperationData, USkeletalMesh& SkeletalMesh, const FMutableRefSkeletalMeshData& RefSkeletalMeshData, UCustomizableObject& CustomizableObject, FCustomizableObjectInstanceComponentIndex InstanceComponentIndex)
+bool BuildSkeletonData(const TSharedRef<FUpdateContextPrivate>& Context, USkeletalMesh& SkeletalMesh, FCustomizableObjectInstanceComponentIndex InstanceComponentIndex)
 {
 	MUTABLE_CPUPROFILER_SCOPE(UCustomizableInstancePrivate::BuildSkeletonData);
 
-	FCustomizableObjectComponentIndex ObjectComponentIndex = OperationData->GetObjectComponentIndex(InstanceComponentIndex);
+	UCustomizableObjectInstance* Instance = Context->Instance.Get();
+	UCustomizableInstancePrivate* InstancePrivate = Instance->GetPrivate();
+
+	UCustomizableObject* CustomizableObject = Instance->GetCustomizableObject();
+	
+	FCustomizableObjectComponentIndex ObjectComponentIndex = Context->GetObjectComponentIndex(InstanceComponentIndex);
 
+	FCustomizableInstanceComponentData* ComponentData = InstancePrivate->GetComponentData(ObjectComponentIndex);
+	check(ComponentData);
+	
 	bool bCreatedNewSkeleton = false;
-	const TObjectPtr<USkeleton> Skeleton = MergeSkeletons(CustomizableObject, RefSkeletalMeshData, ObjectComponentIndex, bCreatedNewSkeleton);
+	const TObjectPtr<USkeleton> Skeleton = MergeSkeletons(*CustomizableObject, *ComponentData, bCreatedNewSkeleton);
 	if (!Skeleton)
 	{
 		return false;
@@ -3894,8 +3877,8 @@ bool UCustomizableInstancePrivate::BuildSkeletonData(const TSharedRef<FUpdateCon
 	const TArray<FMeshBoneInfo>& RawRefBoneInfo = ReferenceSkeleton.GetRawRefBoneInfo();
 	const int32 RawRefBoneCount = ReferenceSkeleton.GetRawBoneNum();
 
-	const TArray<FInstanceUpdateData::FBone>& BonePose = OperationData->InstanceUpdateData.SkeletonsPerInstanceComponent[InstanceComponentIndex.GetValue()].BonePose;
-	TMap<mu::FBoneName, TPair<FName,uint16>>& BoneInfoMap = OperationData->InstanceUpdateData.SkeletonsPerInstanceComponent[InstanceComponentIndex.GetValue()].BoneInfoMap;
+	const TArray<FInstanceUpdateData::FBone>& BonePose = Context->InstanceUpdateData.SkeletonsPerInstanceComponent[InstanceComponentIndex.GetValue()].BonePose;
+	TMap<mu::FBoneName, TPair<FName,uint16>>& BoneInfoMap = Context->InstanceUpdateData.SkeletonsPerInstanceComponent[InstanceComponentIndex.GetValue()].BoneInfoMap;
 
 	{
 		MUTABLE_CPUPROFILER_SCOPE(BuildSkeletonData_BuildBoneInfoMap);
@@ -3905,7 +3888,7 @@ bool UCustomizableInstancePrivate::BuildSkeletonData(const TSharedRef<FUpdateCon
 		for (int32 Index = 0; Index < RawRefBoneCount; ++Index)
 		{
 			const FName BoneName = RawRefBoneInfo[Index].Name;
-			if(const mu::FBoneName* Bone = OperationData->BoneNames->Find(BoneName))
+			if(const mu::FBoneName* Bone = Context->BoneNames->Find(BoneName))
 			{
 				TPair<FName, uint16>& BoneInfo = BoneInfoMap.Add(*Bone);
 				BoneInfo.Key = BoneName;
@@ -3970,15 +3953,12 @@ bool UCustomizableInstancePrivate::BuildSkeletonData(const TSharedRef<FUpdateCon
 
 	USkeleton* GeneratedSkeleton = SkeletalMesh.GetSkeleton();
 
-	if(GeneratedSkeleton && bCreatedNewSkeleton)
+	if (GeneratedSkeleton && bCreatedNewSkeleton)
 	{
 		// If the skeleton is new, it means it has just been merged and the retargeting modes need merging too as the
 		// MergeSkeletons function doesn't do it. Only do it for newly generated ones, not for cached or non-transient ones.
 		GeneratedSkeleton->RecreateBoneTree(&SkeletalMesh);
-
-		FCustomizableInstanceComponentData* ComponentData = GetComponentData(ObjectComponentIndex);
-		check(ComponentData);
-
+		
 		TArray<TObjectPtr<USkeleton>>& SkeletonsToMerge = ComponentData->Skeletons.SkeletonsToMerge;
 		check(SkeletonsToMerge.Num() > 1);
 
@@ -4020,7 +4000,7 @@ bool UCustomizableInstancePrivate::BuildSkeletonData(const TSharedRef<FUpdateCon
 }
 
 
-void UCustomizableInstancePrivate::BuildMeshSockets(const TSharedRef<FUpdateContextPrivate>& OperationData, USkeletalMesh* SkeletalMesh, const UModelResources& ModelResources, const FMutableRefSkeletalMeshData& RefSkeletalMeshData, TSharedPtr<const mu::FMesh> MutableMesh)
+void BuildMeshSockets(USkeletalMesh* SkeletalMesh, const UModelResources& ModelResources, const FMutableRefSkeletalMeshData& RefSkeletalMeshData, TSharedPtr<const mu::FMesh> MutableMesh)
 {
 	// Build mesh sockets.
 	MUTABLE_CPUPROFILER_SCOPE(UCustomizableInstancePrivate::BuildMeshSockets);
@@ -4118,22 +4098,21 @@ void UCustomizableInstancePrivate::BuildMeshSockets(const TSharedRef<FUpdateCont
 }
 
 
-void UCustomizableInstancePrivate::BuildOrCopyElementData(const TSharedRef<FUpdateContextPrivate>& OperationData, USkeletalMesh* SkeletalMesh, UCustomizableObject& CustomizableObject, 
-	FCustomizableObjectInstanceComponentIndex InstanceComponentIndex)
+void BuildOrCopyElementData(const TSharedRef<FUpdateContextPrivate>& Context, USkeletalMesh* SkeletalMesh, FCustomizableObjectInstanceComponentIndex InstanceComponentIndex)
 {
 	MUTABLE_CPUPROFILER_SCOPE(UCustomizableInstancePrivate::BuildOrCopyElementData);
 
-	const FInstanceUpdateData::FComponent* Component = OperationData->GetComponentUpdateData(InstanceComponentIndex);
+	const FInstanceUpdateData::FComponent* Component = Context->GetComponentUpdateData(InstanceComponentIndex);
 	if (!Component)
 	{
 		return;
 	}
 
-	const FName ComponentName = OperationData->ComponentNames[Component->Id.GetValue()];
+	const FName ComponentName = Context->ComponentNames[Component->Id.GetValue()];
 
-	for (int32 LODIndex = OperationData->FirstLODAvailable[ComponentName]; LODIndex < Component->LODCount; ++LODIndex)
+	for (int32 LODIndex = Context->FirstLODAvailable[ComponentName]; LODIndex < Component->LODCount; ++LODIndex)
 	{
-		const FInstanceUpdateData::FLOD& LOD = OperationData->InstanceUpdateData.LODs[Component->FirstLOD+LODIndex];
+		const FInstanceUpdateData::FLOD& LOD = Context->InstanceUpdateData.LODs[Component->FirstLOD+LODIndex];
 
 		if (!LOD.SurfaceCount)
 		{
@@ -4147,7 +4126,7 @@ void UCustomizableInstancePrivate::BuildOrCopyElementData(const TSharedRef<FUpda
 	}
 }
 
-void UCustomizableInstancePrivate::BuildOrCopyMorphTargetsData(const TSharedRef<FUpdateContextPrivate>& OperationData, USkeletalMesh* SkeletalMesh, const USkeletalMesh* LastUpdateSkeletalMesh, UCustomizableObject& CustomizableObject, FCustomizableObjectInstanceComponentIndex InstanceComponentIndex)
+void BuildOrCopyMorphTargetsData(const TSharedRef<FUpdateContextPrivate>& Context, USkeletalMesh* SkeletalMesh, FCustomizableObjectInstanceComponentIndex InstanceComponentIndex)
 {
 	MUTABLE_CPUPROFILER_SCOPE(UCustomizableInstancePrivate::BuildOrCopyMorphTargetsData);
 
@@ -4162,10 +4141,10 @@ void UCustomizableInstancePrivate::BuildOrCopyMorphTargetsData(const TSharedRef<
 		return;
 	}
 	
-	FCustomizableObjectComponentIndex ObjectComponentIndex = OperationData->GetObjectComponentIndex(InstanceComponentIndex);
-	const FName ComponentName = OperationData->ComponentNames[ObjectComponentIndex.GetValue()];
+	FCustomizableObjectComponentIndex ObjectComponentIndex = Context->GetObjectComponentIndex(InstanceComponentIndex);
+	const FName ComponentName = Context->ComponentNames[ObjectComponentIndex.GetValue()];
 
-	FSkeletalMeshMorphTargets* ComponentMorphTargets = OperationData->InstanceUpdateData.RealTimeMorphTargets.Find(ComponentName);
+	FSkeletalMeshMorphTargets* ComponentMorphTargets = Context->InstanceUpdateData.RealTimeMorphTargets.Find(ComponentName);
 	if (!ComponentMorphTargets)
 	{
 		return;
@@ -4192,13 +4171,13 @@ void UCustomizableInstancePrivate::BuildOrCopyMorphTargetsData(const TSharedRef<
 
 		TArray<FMorphTargetLODModel>& MorphLODModels = NewMorphTarget->GetMorphLODModels();
 
-		if (OperationData->bStreamMeshLODs)
+		if (Context->bStreamMeshLODs)
 		{
 			MorphLODModels.SetNum(ComponentMorphTargets->RealTimeMorphsLODData[MorphTargetIndex].Num());
 			
 			// Streamed LODs
-			const int32 FirstLODAvailable = OperationData->FirstLODAvailable[ComponentName];
-			for (int32 LODIndex = FirstLODAvailable; LODIndex < OperationData->FirstResidentLOD[ComponentName]; ++LODIndex)
+			const int32 FirstLODAvailable = Context->FirstLODAvailable[ComponentName];
+			for (int32 LODIndex = FirstLODAvailable; LODIndex < Context->FirstResidentLOD[ComponentName]; ++LODIndex)
 			{
 				// Copy data required for streaming
 				MorphLODModels[LODIndex].NumVertices = 1; // Trick the engine
@@ -4206,7 +4185,7 @@ void UCustomizableInstancePrivate::BuildOrCopyMorphTargetsData(const TSharedRef<
 			}
 			
 			// Residents LODs
-			for (int32 LODIndex = OperationData->GetFirstRequestedLOD()[ComponentName]; LODIndex < OperationData->NumLODsAvailable[ComponentName]; ++LODIndex)
+			for (int32 LODIndex = Context->GetFirstRequestedLOD()[ComponentName]; LODIndex < Context->NumLODsAvailable[ComponentName]; ++LODIndex)
 			{
 				MorphLODModels[LODIndex] = ComponentMorphTargets->RealTimeMorphsLODData[MorphTargetIndex][LODIndex];
 			}
@@ -4220,11 +4199,11 @@ void UCustomizableInstancePrivate::BuildOrCopyMorphTargetsData(const TSharedRef<
 	}
 
 	// Mutable hacky LOD Streaming
-	if (!OperationData->bStreamMeshLODs)
+	if (!Context->bStreamMeshLODs)
 	{
 		// Copy MorphTargets from the FirstGeneratedLOD to the LODs below
-		const int32 FirstRequestedLOD = OperationData->GetFirstRequestedLOD()[ComponentName];
-		for (int32 LODIndex = OperationData->FirstLODAvailable[ComponentName]; LODIndex < FirstRequestedLOD; ++LODIndex)
+		const int32 FirstRequestedLOD = Context->GetFirstRequestedLOD()[ComponentName];
+		for (int32 LODIndex = Context->FirstLODAvailable[ComponentName]; LODIndex < FirstRequestedLOD; ++LODIndex)
 		{
 			MUTABLE_CPUPROFILER_SCOPE(CopyMorphTargetsData);
 
@@ -4235,7 +4214,7 @@ void UCustomizableInstancePrivate::BuildOrCopyMorphTargetsData(const TSharedRef<
 		}
 	}
 
-	const bool bInKeepEmptyMorphTargets = OperationData->bStreamMeshLODs;
+	const bool bInKeepEmptyMorphTargets = Context->bStreamMeshLODs;
 	SkeletalMesh->InitMorphTargets(bInKeepEmptyMorphTargets); // True to avoid removing streamed Morph Targets.
 }
 
@@ -4253,8 +4232,8 @@ namespace
 
 }
 
-void UCustomizableInstancePrivate::BuildOrCopyClothingData(
-		const TSharedRef<FUpdateContextPrivate>& OperationData, 
+void BuildOrCopyClothingData(
+		const TSharedRef<FUpdateContextPrivate>& Context, 
 		USkeletalMesh* SkeletalMesh, 
 		const UModelResources& ModelResources,
 		FCustomizableObjectInstanceComponentIndex InstanceComponentIndex,
@@ -4283,18 +4262,18 @@ void UCustomizableInstancePrivate::BuildOrCopyClothingData(
 	const TArray<FCustomizableObjectClothingAssetData>& ClothingAssetsData = ModelResources.ClothingAssetsData;
 	const TArray<FCustomizableObjectClothConfigData>& ClothSharedConfigsData = ModelResources.ClothSharedConfigsData;
 
-	if (!(ClothingAssetsData.Num() && OperationData->InstanceUpdateData.ClothingMeshData.Num()))
+	if (!(ClothingAssetsData.Num() && Context->InstanceUpdateData.ClothingMeshData.Num()))
 	{
 		return;
 	}
 
-	const FInstanceUpdateData::FComponent* Component = OperationData->GetComponentUpdateData(InstanceComponentIndex);
+	const FInstanceUpdateData::FComponent* Component = Context->GetComponentUpdateData(InstanceComponentIndex);
 	if (!Component)
 	{
 		return;
 	}
 
-	const FName& ComponentName = OperationData->ComponentNames[Component->Id.GetValue()];
+	const FName& ComponentName = Context->ComponentNames[Component->Id.GetValue()];
 
 	TArray<FSectionClothMetadata> SectionClothMetadata; // Sections must be sorted ascending
 	SectionClothMetadata.Reserve(32);
@@ -4309,9 +4288,9 @@ void UCustomizableInstancePrivate::BuildOrCopyClothingData(
 	{
 		int32 NumClothingDataNotFound = 0;
 			
-		for (int32 LODIndex = OperationData->FirstLODAvailable[ComponentName]; LODIndex < Component->LODCount; ++LODIndex)
+		for (int32 LODIndex = Context->FirstLODAvailable[ComponentName]; LODIndex < Component->LODCount; ++LODIndex)
 		{
-			const FInstanceUpdateData::FLOD& LOD = OperationData->InstanceUpdateData.LODs[Component->FirstLOD+LODIndex];
+			const FInstanceUpdateData::FLOD& LOD = Context->InstanceUpdateData.LODs[Component->FirstLOD+LODIndex];
 
 			if (TSharedPtr<const mu::FMesh> MutableMesh = LOD.Mesh)
 			{
@@ -4336,7 +4315,7 @@ void UCustomizableInstancePrivate::BuildOrCopyClothingData(
 						continue;
 					}
 					
-					FClothingMeshData* SectionClothingData = OperationData->InstanceUpdateData.ClothingMeshData.Find(ClothResourceId);
+					FClothingMeshData* SectionClothingData = Context->InstanceUpdateData.ClothingMeshData.Find(ClothResourceId);
 
 					if (!SectionClothingData)
 					{
@@ -4418,13 +4397,13 @@ void UCustomizableInstancePrivate::BuildOrCopyClothingData(
 
 		int32 NumClothingDataNotFound = 0;
 
-		for (int32 LODIndex = OperationData->FirstResidentLOD[ComponentName]; LODIndex < Component->LODCount; ++LODIndex)
+		for (int32 LODIndex = Context->FirstResidentLOD[ComponentName]; LODIndex < Component->LODCount; ++LODIndex)
 		{
-			const FInstanceUpdateData::FLOD& LOD = OperationData->InstanceUpdateData.LODs[Component->FirstLOD+LODIndex];
+			const FInstanceUpdateData::FLOD& LOD = Context->InstanceUpdateData.LODs[Component->FirstLOD+LODIndex];
 
 			if (TSharedPtr<const mu::FMesh> MutableMesh = LOD.Mesh)
 			{
-				GetSectionClothData(*MutableMesh, LODIndex, OperationData->InstanceUpdateData.ClothingMeshData, SectionsClothData, NumClothingDataNotFound);
+				GetSectionClothData(*MutableMesh, LODIndex, Context->InstanceUpdateData.ClothingMeshData, SectionsClothData, NumClothingDataNotFound);
 			}
 		}
 
@@ -4543,7 +4522,7 @@ void UCustomizableInstancePrivate::BuildOrCopyClothingData(
 
 	const bool bAllowClothingPhysicsEdits = !bDisableClothingPhysicsEditsPropagation &&
 		ModelResources.bAllowClothingPhysicsEditsPropagation &&
-		!OperationData->bStreamMeshLODs;
+		!Context->bStreamMeshLODs;
 	
 	if (bAllowClothingPhysicsEdits)
 	{
@@ -4598,9 +4577,9 @@ void UCustomizableInstancePrivate::BuildOrCopyClothingData(
 	ResidentLODClothIndexMapping.SetNum(Component->LODCount);
 
 	// Zero all LODs (even those which do not use cloth).
-	for (int32 LODIndex = OperationData->FirstLODAvailable[ComponentName]; LODIndex < Component->LODCount; ++LODIndex)
+	for (int32 LODIndex = Context->FirstLODAvailable[ComponentName]; LODIndex < Component->LODCount; ++LODIndex)
 	{
-		const FInstanceUpdateData::FLOD& LOD = OperationData->InstanceUpdateData.LODs[Component->FirstLOD + LODIndex];
+		const FInstanceUpdateData::FLOD& LOD = Context->InstanceUpdateData.LODs[Component->FirstLOD + LODIndex];
 		
 		TSharedPtr<const mu::FMesh> MutableMesh = LOD.Mesh;
 		if (!MutableMesh)
@@ -4622,7 +4601,7 @@ void UCustomizableInstancePrivate::BuildOrCopyClothingData(
 		MUTABLE_CPUPROFILER_SCOPE(InitClothRenderData)
 
 		// Streamed
-		for (int32 LODIndex = OperationData->FirstLODAvailable[ComponentName]; LODIndex < OperationData->FirstResidentLOD[ComponentName]; ++LODIndex)
+		for (int32 LODIndex = Context->FirstLODAvailable[ComponentName]; LODIndex < Context->FirstResidentLOD[ComponentName]; ++LODIndex)
 		{
 			FSkeletalMeshLODRenderData& LODModel = RenderResource->LODRenderData[LODIndex];
 
@@ -4658,7 +4637,7 @@ void UCustomizableInstancePrivate::BuildOrCopyClothingData(
 
 
 		// Resident
-		for (int32 LODIndex = OperationData->FirstResidentLOD[ComponentName]; LODIndex < OperationData->NumLODsAvailable[ComponentName]; ++LODIndex)
+		for (int32 LODIndex = Context->FirstResidentLOD[ComponentName]; LODIndex < Context->NumLODsAvailable[ComponentName]; ++LODIndex)
 		{
 			FSkeletalMeshLODRenderData& LODModel = RenderResource->LODRenderData[LODIndex];
 	
@@ -4691,7 +4670,7 @@ void UCustomizableInstancePrivate::BuildOrCopyClothingData(
 }
 
 
-bool UCustomizableInstancePrivate::BuildOrCopyRenderData(const TSharedRef<FUpdateContextPrivate>& OperationData, USkeletalMesh* SkeletalMesh, const USkeletalMesh* LastUpdateSkeletalMesh, 
+bool BuildOrCopyRenderData(const TSharedRef<FUpdateContextPrivate>& Context, USkeletalMesh* SkeletalMesh, 
 	UCustomizableObjectInstance* Public, FCustomizableObjectInstanceComponentIndex InstanceComponentIndex)
 {
 	MUTABLE_CPUPROFILER_SCOPE(UCustomizableInstancePrivate::BuildOrCopyRenderData);
@@ -4704,7 +4683,7 @@ bool UCustomizableInstancePrivate::BuildOrCopyRenderData(const TSharedRef<FUpdat
 	// It must be not null as it's checked in the calling function
 	check(CustomizableObject);
 
-	const FInstanceUpdateData::FComponent* Component = OperationData->GetComponentUpdateData(InstanceComponentIndex);
+	const FInstanceUpdateData::FComponent* Component = Context->GetComponentUpdateData(InstanceComponentIndex);
 	if (!Component)
 	{
 		return false;
@@ -4713,15 +4692,15 @@ bool UCustomizableInstancePrivate::BuildOrCopyRenderData(const TSharedRef<FUpdat
 	const UModelResources& ModelResources = CustomizableObject->GetPrivate()->GetModelResourcesChecked();
 	const FName ComponentName = ModelResources.ComponentNamesPerObjectComponent[Component->Id.GetValue()];
 
-	const int32 FirstLOD = OperationData->bStreamMeshLODs ?
-		OperationData->FirstLODAvailable[ComponentName] :
-		OperationData->GetFirstRequestedLOD()[ComponentName];
+	const int32 FirstLOD = Context->bStreamMeshLODs ?
+		Context->FirstLODAvailable[ComponentName] :
+		Context->GetFirstRequestedLOD()[ComponentName];
 	
 	for (int32 LODIndex = FirstLOD; LODIndex < Component->LODCount; ++LODIndex)
 	{
 		MUTABLE_CPUPROFILER_SCOPE(BuildRenderData);
 		
-		const FInstanceUpdateData::FLOD& LOD = OperationData->InstanceUpdateData.LODs[Component->FirstLOD+LODIndex];
+		const FInstanceUpdateData::FLOD& LOD = Context->InstanceUpdateData.LODs[Component->FirstLOD+LODIndex];
 
 		// There could be components without a mesh in LODs
 		if (!LOD.Mesh || LOD.SurfaceCount == 0)
@@ -4739,11 +4718,11 @@ bool UCustomizableInstancePrivate::BuildOrCopyRenderData(const TSharedRef<FUpdat
 		
 		FSkeletalMeshLODRenderData& LODResource = RenderData->LODRenderData[LODIndex];
 
-		const TMap<mu::FBoneName, TPair<FName, uint16>>& BoneInfoMap = OperationData->InstanceUpdateData.SkeletonsPerInstanceComponent[InstanceComponentIndex.GetValue()].BoneInfoMap;
+		const TMap<mu::FBoneName, TPair<FName, uint16>>& BoneInfoMap = Context->InstanceUpdateData.SkeletonsPerInstanceComponent[InstanceComponentIndex.GetValue()].BoneInfoMap;
 		
 		// Set active and required bones
 		{
-			const TArray<mu::FBoneName>& ActiveBones = OperationData->InstanceUpdateData.ActiveBones;
+			const TArray<mu::FBoneName>& ActiveBones = Context->InstanceUpdateData.ActiveBones;
 			LODResource.ActiveBoneIndices.Reserve(LOD.ActiveBoneCount);
 
 			for (uint32 Index = 0; Index < LOD.ActiveBoneCount; ++Index)
@@ -4764,14 +4743,14 @@ bool UCustomizableInstancePrivate::BuildOrCopyRenderData(const TSharedRef<FUpdat
 		for (int32 MeshSectionIndex = 0; MeshSectionIndex < MeshNumSurfaces; ++MeshSectionIndex)
 		{
 			uint32 MeshSurfaceId = LOD.Mesh->GetSurfaceId(MeshSectionIndex);
-			int32 InstanceSurfaceIndex = OperationData->MutableInstance->FindSurfaceById(InstanceComponentIndex.GetValue(), LODIndex, MeshSurfaceId);
+			int32 InstanceSurfaceIndex = Context->MutableInstance->FindSurfaceById(InstanceComponentIndex.GetValue(), LODIndex, MeshSurfaceId);
 			
 			if (InstanceSurfaceIndex < 0)
 			{
 				continue;
 			}
 			
-			uint32 SurfaceMetadataId = OperationData->MutableInstance->GetSurfaceCustomId(InstanceComponentIndex.GetValue(), LODIndex, InstanceSurfaceIndex);
+			uint32 SurfaceMetadataId = Context->MutableInstance->GetSurfaceCustomId(InstanceComponentIndex.GetValue(), LODIndex, InstanceSurfaceIndex);
 			
 			uint32 UsedSurfaceMetadataId = 0;
 			if (SurfaceMetadataId != 0)
@@ -4806,7 +4785,7 @@ bool UCustomizableInstancePrivate::BuildOrCopyRenderData(const TSharedRef<FUpdat
 		UnrealConversionUtils::SetupRenderSections(
 			LODResource,
 			LOD.Mesh.Get(),
-			OperationData->InstanceUpdateData.BoneMaps,
+			Context->InstanceUpdateData.BoneMaps,
 			BoneInfoMap,
 			LOD.FirstBoneMap,
 			MeshSurfacesMetadata);
@@ -4908,7 +4887,7 @@ bool UCustomizableInstancePrivate::BuildOrCopyRenderData(const TSharedRef<FUpdat
 
 		if (LODResource.DoesVertexBufferUse16BitBoneIndex() && !UCustomizableObjectSystem::GetInstance()->IsSupport16BitBoneIndexEnabled())
 		{
-			OperationData->UpdateResult = EUpdateResult::Error16BitBoneIndex;
+			Context->UpdateResult = EUpdateResult::Error16BitBoneIndex;
 
 			const FString Msg = FString::Printf(TEXT("Customizable Object [%s] requires of Skinning - 'Support 16 Bit Bone Index' to be enabled. Please, update the Project Settings."),
 				*CustomizableObject->GetName());
@@ -4926,11 +4905,11 @@ bool UCustomizableInstancePrivate::BuildOrCopyRenderData(const TSharedRef<FUpdat
 	
 
 	// Mutable hacky LOD Streaming
-	if (!OperationData->bStreamMeshLODs)
+	if (!Context->bStreamMeshLODs)
 	{
 		// Copy LODRenderData from the FirstRequestedLOD to the LODs below
-		const int32 FirstRequestedLOD = OperationData->GetFirstRequestedLOD()[ComponentName];
-		for (int32 LODIndex = OperationData->FirstLODAvailable[ComponentName]; LODIndex < FirstRequestedLOD; ++LODIndex)
+		const int32 FirstRequestedLOD = Context->GetFirstRequestedLOD()[ComponentName];
+		for (int32 LODIndex = Context->FirstLODAvailable[ComponentName]; LODIndex < FirstRequestedLOD; ++LODIndex)
 		{
 			TRACE_CPUPROFILER_EVENT_SCOPE_TEXT(*FString::Printf(TEXT("CopyRenderData: From LOD %d to LOD %d"), FirstRequestedLOD, LODIndex));
 
@@ -4952,11 +4931,14 @@ bool UCustomizableInstancePrivate::BuildOrCopyRenderData(const TSharedRef<FUpdat
 }
 
 
-UE::Tasks::FTask UCustomizableInstancePrivate::LoadAdditionalAssetsAndData(const TSharedRef<FUpdateContextPrivate>& OperationData)
+UE::Tasks::FTask LoadAdditionalAssetsAndData(const TSharedRef<FUpdateContextPrivate>& Context)
 {
 	MUTABLE_CPUPROFILER_SCOPE(UCustomizableInstancePrivate::LoadAdditionalAssetsAndDataAsync);
 
-	UCustomizableObject* CustomizableObject = GetPublic()->GetCustomizableObject();
+	UCustomizableObjectInstance* Instance = Context->Instance.Get();
+	UCustomizableInstancePrivate* InstancePrivate = Instance->GetPrivate();
+	
+	UCustomizableObject* CustomizableObject = Instance->GetCustomizableObject();
 
 	const UModelResources& ModelResources = CustomizableObject->GetPrivate()->GetModelResourcesChecked();
 	const TSharedPtr<FModelStreamableBulkData>& ModelStreamableBulkData = CustomizableObject->GetPrivate()->GetModelStreamableBulkData();
@@ -4965,20 +4947,20 @@ UE::Tasks::FTask UCustomizableInstancePrivate::LoadAdditionalAssetsAndData(const
 
 	TArray<FSoftObjectPath> AssetsToStream;
 
-	TArray<FInstanceUpdateData::FComponent>& Components = OperationData->InstanceUpdateData.Components;
+	TArray<FInstanceUpdateData::FComponent>& Components = Context->InstanceUpdateData.Components;
 
-	ObjectToInstanceIndexMap.Empty();
-	ReferencedMaterials.Empty();
+	InstancePrivate->ObjectToInstanceIndexMap.Empty();
+	InstancePrivate->ReferencedMaterials.Empty();
 
 	const int32 NumClothingAssets = ModelResources.ClothingAssetsData.Num();
-	ClothingPhysicsAssets.Reset(NumClothingAssets);
-	ClothingPhysicsAssets.SetNum(NumClothingAssets);
+	InstancePrivate->ClothingPhysicsAssets.Reset(NumClothingAssets);
+	InstancePrivate->ClothingPhysicsAssets.SetNum(NumClothingAssets);
 
-	GatheredAnimBPs.Empty();
-	AnimBPGameplayTags.Reset();
-	AnimBpPhysicsAssets.Reset();
+	InstancePrivate->GatheredAnimBPs.Empty();
+	InstancePrivate->AnimBPGameplayTags.Reset();
+	InstancePrivate->AnimBpPhysicsAssets.Reset();
 	
-	for (const FInstanceUpdateData::FSurface& Surface : OperationData->InstanceUpdateData.Surfaces)
+	for (const FInstanceUpdateData::FSurface& Surface : Context->InstanceUpdateData.Surfaces)
 	{
 		if (!Surface.Material)
 		{
@@ -4991,7 +4973,7 @@ UE::Tasks::FTask UCustomizableInstancePrivate::LoadAdditionalAssetsAndData(const
 		}
 		
 		const int32 MaterialIndex = Surface.Material->ReferenceID;
-		if (ObjectToInstanceIndexMap.Contains(MaterialIndex))
+		if (InstancePrivate->ObjectToInstanceIndexMap.Contains(MaterialIndex))
 		{
 			continue;
 		}
@@ -4999,8 +4981,8 @@ UE::Tasks::FTask UCustomizableInstancePrivate::LoadAdditionalAssetsAndData(const
 		TSoftObjectPtr<UMaterialInterface> AssetPtr = ModelResources.Materials.IsValidIndex(MaterialIndex) ? ModelResources.Materials[MaterialIndex] : nullptr;
 		UMaterialInterface* LoadedMaterial = AssetPtr.Get();
 
-		const int32 ReferencedMaterialsIndex = ReferencedMaterials.Add(LoadedMaterial);
-		ObjectToInstanceIndexMap.Add(MaterialIndex, ReferencedMaterialsIndex);
+		const int32 ReferencedMaterialsIndex = InstancePrivate->ReferencedMaterials.Add(LoadedMaterial);
+		InstancePrivate->ObjectToInstanceIndexMap.Add(MaterialIndex, ReferencedMaterialsIndex);
 
 		if (!LoadedMaterial && !AssetPtr.IsNull())
 		{
@@ -5008,7 +4990,7 @@ UE::Tasks::FTask UCustomizableInstancePrivate::LoadAdditionalAssetsAndData(const
 		}
 	}
 
-	for (const FInstanceUpdateData::FComponent& Component : OperationData->InstanceUpdateData.Components)
+	for (const FInstanceUpdateData::FComponent& Component : Context->InstanceUpdateData.Components)
 	{
 		if (!Component.OverlayMaterial)
 		{
@@ -5021,7 +5003,7 @@ UE::Tasks::FTask UCustomizableInstancePrivate::LoadAdditionalAssetsAndData(const
 		}
 		
 		const int32 MaterialIndex = Component.OverlayMaterial->ReferenceID;
-		if (ObjectToInstanceIndexMap.Contains(Component.OverlayMaterial->ReferenceID))
+		if (InstancePrivate->ObjectToInstanceIndexMap.Contains(Component.OverlayMaterial->ReferenceID))
 		{
 			continue;
 		}
@@ -5029,8 +5011,8 @@ UE::Tasks::FTask UCustomizableInstancePrivate::LoadAdditionalAssetsAndData(const
 		TSoftObjectPtr<UMaterialInterface> AssetPtr = ModelResources.Materials.IsValidIndex(MaterialIndex) ? ModelResources.Materials[Component.OverlayMaterial->ReferenceID] : nullptr;
 		UMaterialInterface* LoadedMaterial = AssetPtr.Get();
 
-		const int32 ReferencedMaterialsIndex = ReferencedMaterials.Add(LoadedMaterial);
-		ObjectToInstanceIndexMap.Add(MaterialIndex, ReferencedMaterialsIndex);
+		const int32 ReferencedMaterialsIndex = InstancePrivate->ReferencedMaterials.Add(LoadedMaterial);
+		InstancePrivate->ObjectToInstanceIndexMap.Add(MaterialIndex, ReferencedMaterialsIndex);
 
 		if (!LoadedMaterial && !AssetPtr.IsNull())
 		{
@@ -5040,17 +5022,17 @@ UE::Tasks::FTask UCustomizableInstancePrivate::LoadAdditionalAssetsAndData(const
 
 
 	// Load Skeletons required by the SubMeshes of the newly generated Mesh, will be merged later
-	for (int32 InstanceComponentIndex = 0; InstanceComponentIndex < OperationData->NumInstanceComponents; ++InstanceComponentIndex)
+	for (int32 InstanceComponentIndex = 0; InstanceComponentIndex < Context->NumInstanceComponents; ++InstanceComponentIndex)
 	{
-		FCustomizableObjectComponentIndex ObjectComponentIndex = OperationData->GetObjectComponentIndex(FCustomizableObjectInstanceComponentIndex(InstanceComponentIndex));
+		FCustomizableObjectComponentIndex ObjectComponentIndex = Context->GetObjectComponentIndex(FCustomizableObjectInstanceComponentIndex(InstanceComponentIndex));
 		if (!ObjectComponentIndex.IsValid())
 		{
 			continue;
 		}
 
-		const FInstanceUpdateData::FSkeletonData& SkeletonData = OperationData->InstanceUpdateData.SkeletonsPerInstanceComponent[InstanceComponentIndex];
+		const FInstanceUpdateData::FSkeletonData& SkeletonData = Context->InstanceUpdateData.SkeletonsPerInstanceComponent[InstanceComponentIndex];
 		
-		FCustomizableInstanceComponentData* ComponentData = GetComponentData(ObjectComponentIndex);
+		FCustomizableInstanceComponentData* ComponentData = InstancePrivate->GetComponentData(ObjectComponentIndex);
 		if (!ComponentData)
 		{
 			continue;
@@ -5091,18 +5073,18 @@ UE::Tasks::FTask UCustomizableInstancePrivate::LoadAdditionalAssetsAndData(const
 	}
 	
 	bool bHasInvalidMesh = false;
-	const bool bUpdateMeshes = DoComponentsNeedUpdate(GetPublic(), OperationData, bHasInvalidMesh);
+	const bool bUpdateMeshes = DoComponentsNeedUpdate(Context, bHasInvalidMesh);
 
 	// Load assets coming from SubMeshes of the newly generated Mesh
-	if (OperationData->InstanceUpdateData.LODs.Num())
+	if (Context->InstanceUpdateData.LODs.Num())
 	{
-		for (int32 InstanceComponentIndex = 0; InstanceComponentIndex < OperationData->InstanceUpdateData.Components.Num(); ++InstanceComponentIndex)
+		for (int32 InstanceComponentIndex = 0; InstanceComponentIndex < Context->InstanceUpdateData.Components.Num(); ++InstanceComponentIndex)
 		{
 			const FInstanceUpdateData::FComponent& Component = Components[InstanceComponentIndex];
-			FCustomizableInstanceComponentData* ComponentData = GetComponentData(Component.Id);
+			FCustomizableInstanceComponentData* ComponentData = InstancePrivate->GetComponentData(Component.Id);
 
-			TSharedPtr<const mu::FMesh> FirstComponentMesh = OperationData->InstanceUpdateData.LODs.IsValidIndex(Component.FirstLOD) ?
-				OperationData->InstanceUpdateData.LODs[Component.FirstLOD].Mesh :
+			TSharedPtr<const mu::FMesh> FirstComponentMesh = Context->InstanceUpdateData.LODs.IsValidIndex(Component.FirstLOD) ?
+				Context->InstanceUpdateData.LODs[Component.FirstLOD].Mesh :
 				nullptr;
 
 			if (FirstComponentMesh && FirstComponentMesh->IsReference())
@@ -5117,21 +5099,21 @@ UE::Tasks::FTask UCustomizableInstancePrivate::LoadAdditionalAssetsAndData(const
 			for (int32 AssetUserDataIndex : RefSkeletalMeshData.AssetUserDataIndices)
 			{
 #if !WITH_EDITOR
-				OperationData->StreamedResourceIndex.AddUnique(AssetUserDataIndex); // Used to hold/release streamed resources in non-editor builds.
+				Context->StreamedResourceIndex.AddUnique(AssetUserDataIndex); // Used to hold/release streamed resources in non-editor builds.
 #endif
 				ComponentData->StreamedResourceIndex.AddUnique(AssetUserDataIndex);
 			}
 
-			const FName ComponentName = OperationData->ComponentNames[Component.Id.GetValue()];
+			const FName ComponentName = Context->ComponentNames[Component.Id.GetValue()];
 
 			if (bUpdateMeshes)
 			{
 				// Morphs
 				{
 					// Data
-					for (int32 LODIndex = OperationData->GetFirstRequestedLOD()[ComponentName]; LODIndex < Component.LODCount; ++LODIndex)
+					for (int32 LODIndex = Context->GetFirstRequestedLOD()[ComponentName]; LODIndex < Component.LODCount; ++LODIndex)
 					{
-						const FInstanceUpdateData::FLOD& LOD = OperationData->InstanceUpdateData.LODs[Component.FirstLOD + LODIndex];
+						const FInstanceUpdateData::FLOD& LOD = Context->InstanceUpdateData.LODs[Component.FirstLOD + LODIndex];
 
 						TSharedPtr<const mu::FMesh> MutableMesh = LOD.Mesh;
 						if (!MutableMesh)
@@ -5139,17 +5121,17 @@ UE::Tasks::FTask UCustomizableInstancePrivate::LoadAdditionalAssetsAndData(const
 							continue;
 						}
 
-						LoadMorphTargetsData(StreamRequest, MutableMesh.ToSharedRef(), OperationData->InstanceUpdateData.RealTimeMorphTargetMeshData);
+						LoadMorphTargetsData(StreamRequest, MutableMesh.ToSharedRef(), Context->InstanceUpdateData.RealTimeMorphTargetMeshData);
 					}
 
 					// Metadata
-					const int32 FirstLOD = OperationData->bStreamMeshLODs ?
-						OperationData->FirstResidentLOD[ComponentName] :
-						OperationData->GetFirstRequestedLOD()[ComponentName];
+					const int32 FirstLOD = Context->bStreamMeshLODs ?
+						Context->FirstResidentLOD[ComponentName] :
+						Context->GetFirstRequestedLOD()[ComponentName];
 				
 					for (int32 LODIndex = FirstLOD; LODIndex < Component.LODCount; ++LODIndex)
 					{
-						const FInstanceUpdateData::FLOD& LOD = OperationData->InstanceUpdateData.LODs[Component.FirstLOD + LODIndex];
+						const FInstanceUpdateData::FLOD& LOD = Context->InstanceUpdateData.LODs[Component.FirstLOD + LODIndex];
 
 						TSharedPtr<const mu::FMesh> MutableMesh = LOD.Mesh;
 						if (!MutableMesh)
@@ -5157,7 +5139,7 @@ UE::Tasks::FTask UCustomizableInstancePrivate::LoadAdditionalAssetsAndData(const
 							continue;
 						}
 
-						LoadMorphTargetsMetadata(StreamRequest, MutableMesh.ToSharedRef(), OperationData->InstanceUpdateData.RealTimeMorphTargetMeshData);
+						LoadMorphTargetsMetadata(StreamRequest, MutableMesh.ToSharedRef(), Context->InstanceUpdateData.RealTimeMorphTargetMeshData);
 					}
 				}
 
@@ -5165,9 +5147,9 @@ UE::Tasks::FTask UCustomizableInstancePrivate::LoadAdditionalAssetsAndData(const
 				{
 					// Data
 					// From FirstResidentLOD instead of FirstRequestedLOD since clothing we generate all LODs, even the hacky streaming copied ones.
-					for (int32 LODIndex = OperationData->FirstResidentLOD[ComponentName]; LODIndex < Component.LODCount; ++LODIndex)
+					for (int32 LODIndex = Context->FirstResidentLOD[ComponentName]; LODIndex < Component.LODCount; ++LODIndex)
 					{
-						const FInstanceUpdateData::FLOD& LOD = OperationData->InstanceUpdateData.LODs[Component.FirstLOD + LODIndex];
+						const FInstanceUpdateData::FLOD& LOD = Context->InstanceUpdateData.LODs[Component.FirstLOD + LODIndex];
 
 						TSharedPtr<const mu::FMesh> MutableMesh = LOD.Mesh;
 						if (!MutableMesh)
@@ -5177,14 +5159,14 @@ UE::Tasks::FTask UCustomizableInstancePrivate::LoadAdditionalAssetsAndData(const
 
 						if (bUpdateMeshes)
 						{
-							LoadClothing(StreamRequest, MutableMesh.ToSharedRef(), OperationData->InstanceUpdateData.ClothingMeshData);
+							LoadClothing(StreamRequest, MutableMesh.ToSharedRef(), Context->InstanceUpdateData.ClothingMeshData);
 						}
 					}
 
 					// Metadata
-					for (int32 LODIndex = OperationData->FirstLODAvailable[ComponentName]; LODIndex < Component.LODCount; ++LODIndex)
+					for (int32 LODIndex = Context->FirstLODAvailable[ComponentName]; LODIndex < Component.LODCount; ++LODIndex)
 					{
-						const FInstanceUpdateData::FLOD& LOD = OperationData->InstanceUpdateData.LODs[Component.FirstLOD + LODIndex];
+						const FInstanceUpdateData::FLOD& LOD = Context->InstanceUpdateData.LODs[Component.FirstLOD + LODIndex];
 
 						TSharedPtr<const mu::FMesh> MutableMesh = LOD.Mesh;
 						if (!MutableMesh)
@@ -5205,7 +5187,7 @@ UE::Tasks::FTask UCustomizableInstancePrivate::LoadAdditionalAssetsAndData(const
 								const TMap<uint32, FClothingStreamable>& ClothingStreamables = ModelStreamableBulkData->ClothingStreamables;
 								if (const FClothingStreamable* ClothingStreamable = ClothingStreamables.Find(TypedResourceId.Id))
 								{
-									FClothingMeshData& ReadDestData = OperationData->InstanceUpdateData.ClothingMeshData.FindOrAdd(TypedResourceId.Id);
+									FClothingMeshData& ReadDestData = Context->InstanceUpdateData.ClothingMeshData.FindOrAdd(TypedResourceId.Id);
 									ReadDestData.ClothingAssetIndex = ClothingStreamable->ClothingAssetIndex;
 									ReadDestData.ClothingAssetLOD = ClothingStreamable->ClothingAssetLOD;
  								
@@ -5221,9 +5203,9 @@ UE::Tasks::FTask UCustomizableInstancePrivate::LoadAdditionalAssetsAndData(const
 									{
 										if (PhysicsAsset.Get())
 										{
-											if (ClothingPhysicsAssets.IsValidIndex(ClothingAssetIndex))
+											if (InstancePrivate->ClothingPhysicsAssets.IsValidIndex(ClothingAssetIndex))
 											{
-												ClothingPhysicsAssets[ClothingAssetIndex] = PhysicsAsset.Get();
+												InstancePrivate->ClothingPhysicsAssets[ClothingAssetIndex] = PhysicsAsset.Get();
 											}
 										}
 										else
@@ -5235,7 +5217,7 @@ UE::Tasks::FTask UCustomizableInstancePrivate::LoadAdditionalAssetsAndData(const
 								}
 								else
 								{
-									UE_LOG(LogMutable, Error, TEXT("Invalid streamed clothing data block [%d] found."), TypedResourceId.Id);
+									UE_LOG(LogMutable, Error, TEXT("Invalid streamed clothing data block [%llu] found."), TypedResourceId.Id);
 								}
 							}
 						}
@@ -5243,9 +5225,9 @@ UE::Tasks::FTask UCustomizableInstancePrivate::LoadAdditionalAssetsAndData(const
 				}
 			}
 			
-			for (int32 LODIndex = OperationData->FirstLODAvailable[ComponentName]; LODIndex < Component.LODCount; ++LODIndex)
+			for (int32 LODIndex = Context->FirstLODAvailable[ComponentName]; LODIndex < Component.LODCount; ++LODIndex)
 			{
-				const FInstanceUpdateData::FLOD& LOD = OperationData->InstanceUpdateData.LODs[Component.FirstLOD + LODIndex];
+				const FInstanceUpdateData::FLOD& LOD = Context->InstanceUpdateData.LODs[Component.FirstLOD + LODIndex];
 
 				TSharedPtr<const mu::FMesh> MutableMesh = LOD.Mesh;
 				if (!MutableMesh)
@@ -5263,7 +5245,7 @@ UE::Tasks::FTask UCustomizableInstancePrivate::LoadAdditionalAssetsAndData(const
 					{
 						const uint32 ResourceIndex = TypedResourceId.Id;
 #if !WITH_EDITOR
-						OperationData->StreamedResourceIndex.AddUnique(ResourceIndex); // Used to hold/release streamed resources in non-editor builds.
+						Context->StreamedResourceIndex.AddUnique(ResourceIndex); // Used to hold/release streamed resources in non-editor builds.
 #endif
 						ComponentData->StreamedResourceIndex.AddUnique(ResourceIndex);
 					}
@@ -5322,7 +5304,7 @@ UE::Tasks::FTask UCustomizableInstancePrivate::LoadAdditionalAssetsAndData(const
 
 									if (AnimBPAsset.Get())
 									{
-										GatheredAnimBPs.Add(AnimBPAsset.Get());
+										InstancePrivate->GatheredAnimBPs.Add(AnimBPAsset.Get());
 									}
 									else
 									{
@@ -5332,7 +5314,7 @@ UE::Tasks::FTask UCustomizableInstancePrivate::LoadAdditionalAssetsAndData(const
 								else if (!bIsSameAnimBp)
 								{
 									// Two submeshes should not have the same animation slot index
-									OperationData->UpdateResult = EUpdateResult::Warning;
+									Context->UpdateResult = EUpdateResult::Warning;
 
 									FString WarningMessage = FString::Printf(TEXT("Two submeshes have the same anim slot index [%s] in a Mutable Instance."), *SlotIndex.ToString());
 									UE_LOG(LogMutable, Warning, TEXT("%s"), *WarningMessage);
@@ -5346,7 +5328,7 @@ UE::Tasks::FTask UCustomizableInstancePrivate::LoadAdditionalAssetsAndData(const
 					}
 					else if (Tag.RemoveFromStart("__AnimBPTag:"))
 					{
-						AnimBPGameplayTags.AddTag(FGameplayTag::RequestGameplayTag(*Tag));
+						InstancePrivate->AnimBPGameplayTags.AddTag(FGameplayTag::RequestGameplayTag(*Tag));
 					}
 #if WITH_EDITORONLY_DATA
 					else if (Tag.RemoveFromStart("__MeshPath:"))
@@ -5397,26 +5379,26 @@ UE::Tasks::FTask UCustomizableInstancePrivate::LoadAdditionalAssetsAndData(const
 		}
 	}
 
-	for (TSoftObjectPtr<const UTexture>& TextureRef : PassThroughTexturesToLoad)
+	for (TSoftObjectPtr<const UTexture>& TextureRef : InstancePrivate->PassThroughTexturesToLoad)
 	{
 		AssetsToStream.Add(TextureRef.ToSoftObjectPath());
 	}
 
-	for (TSoftObjectPtr<const UStreamableRenderAsset>& MeshRef : PassThroughMeshesToLoad)
+	for (TSoftObjectPtr<const UStreamableRenderAsset>& MeshRef : InstancePrivate->PassThroughMeshesToLoad)
 	{
 		AssetsToStream.Add(MeshRef.ToSoftObjectPath());
 	}
 
 	// Copy FExtensionData Object node input from the Instance to the InstanceUpdateData
-	for (int32 ExtensionDataIndex = 0; ExtensionDataIndex < OperationData->MutableInstance->GetExtensionDataCount(); ExtensionDataIndex++)
+	for (int32 ExtensionDataIndex = 0; ExtensionDataIndex < Context->MutableInstance->GetExtensionDataCount(); ExtensionDataIndex++)
 	{
 		TSharedPtr<const mu::FExtensionData> ExtensionData;
 		FName Name;
-		OperationData->MutableInstance->GetExtensionData(ExtensionDataIndex, ExtensionData, Name);
+		Context->MutableInstance->GetExtensionData(ExtensionDataIndex, ExtensionData, Name);
 
 		check(ExtensionData);
 
-		FInstanceUpdateData::FNamedExtensionData& NewEntry = OperationData->InstanceUpdateData.ExtendedInputPins.AddDefaulted_GetRef();
+		FInstanceUpdateData::FNamedExtensionData& NewEntry = Context->InstanceUpdateData.ExtendedInputPins.AddDefaulted_GetRef();
 		NewEntry.Data = ExtensionData;
 		NewEntry.Name = Name;
 		check(NewEntry.Name != NAME_None);
@@ -5449,7 +5431,7 @@ UE::Tasks::FTask UCustomizableInstancePrivate::LoadAdditionalAssetsAndData(const
 		// Note that this just checks if the path is non-null, NOT if the object is loaded
 		check(!StreamedData.GetPath().IsNull());
 
-		OperationData->ExtensionStreamedResourceIndex.Add(ExtensionData->Index);
+		Context->ExtensionStreamedResourceIndex.Add(ExtensionData->Index);
 		AssetsToStream.Add(StreamedData.GetPath().ToSoftObjectPath());
 #endif
 	}
@@ -5472,7 +5454,7 @@ UE::Tasks::FTask UCustomizableInstancePrivate::LoadAdditionalAssetsAndData(const
 		
 		PrivateSystem->StreamableManager->RequestAsyncLoad(
 			AssetsToStream,
-			FStreamableDelegate::CreateUObject(this, &UCustomizableInstancePrivate::AdditionalAssetsAsyncLoaded, OperationData, Event),
+			FStreamableDelegate::CreateStatic(&AdditionalAssetsAsyncLoaded, Context, Event),
 			CVarMutableHighPriorityLoading.GetValueOnAnyThread() ? FStreamableManager::AsyncLoadHighPriority : FStreamableManager::DefaultAsyncLoadPriority);
 	}
 	
@@ -5480,7 +5462,7 @@ UE::Tasks::FTask UCustomizableInstancePrivate::LoadAdditionalAssetsAndData(const
 	UE::Tasks::FTask StreamingTask = StreamRequest.Stream();
 	Prerequisites.Add(StreamingTask);
 
-	return UE::Tasks::Launch(TEXT("CaptureOperationData"), [OperationData]() {}, // Keep a reference to make sure allocated memory is always alive.
+	return UE::Tasks::Launch(TEXT("CaptureContext"), [Context]() {}, // Keep a reference to make sure allocated memory is always alive. TODO Probably not necessary since AdditionalAssetsAsyncLoaded already captrues Context.
 	Prerequisites,
 	UE::Tasks::ETaskPriority::Inherit);
 }
@@ -5541,7 +5523,7 @@ TArray<UMaterialInterface*> UCustomizableObjectInstance::GetSkeletalMeshComponen
 }
 
 
-void UCustomizableInstancePrivate::AdditionalAssetsAsyncLoaded(const TSharedRef<FUpdateContextPrivate> OperationData, UE::Tasks::FTaskEvent Event)
+void AdditionalAssetsAsyncLoaded(const TSharedRef<FUpdateContextPrivate> Context, UE::Tasks::FTaskEvent Event)
 {
 	MUTABLE_CPUPROFILER_SCOPE(UCustomizableInstancePrivate::AdditionalAssetsAsyncLoaded);
 
@@ -5549,30 +5531,33 @@ void UCustomizableInstancePrivate::AdditionalAssetsAsyncLoaded(const TSharedRef<
 
 	Event.Trigger();
 
-	UCustomizableObjectPrivate* CustomizableObjectPrivate = GetPublic()->GetCustomizableObject()->GetPrivate();
+	UCustomizableObjectInstance* Instance = Context->Instance.Get();
+	UCustomizableInstancePrivate* InstancePrivate = Instance->GetPrivate();
+	
+	UCustomizableObjectPrivate* CustomizableObjectPrivate = Instance->GetCustomizableObject()->GetPrivate();
 
 	UModelResources& ModelResources = *CustomizableObjectPrivate->GetModelResources();
 
-	for (int32 ResourceIndex : OperationData->StreamedResourceIndex)
+	for (int32 ResourceIndex : Context->StreamedResourceIndex)
 	{
 		ModelResources.StreamedResourceData[ResourceIndex].Hold();
 	}
 	
-	for (int32 ResourceIndex : OperationData->ExtensionStreamedResourceIndex)
+	for (int32 ResourceIndex : Context->ExtensionStreamedResourceIndex)
 	{
 		ModelResources.StreamedExtensionData[ResourceIndex].Hold();
 	}
 	
 	// Loaded Materials
-	check(ObjectToInstanceIndexMap.Num() == ReferencedMaterials.Num());
+	check(InstancePrivate->ObjectToInstanceIndexMap.Num() == InstancePrivate->ReferencedMaterials.Num());
 
-	for (TPair<uint32, uint32> Pair : ObjectToInstanceIndexMap)
+	for (TPair<uint32, uint32> Pair : InstancePrivate->ObjectToInstanceIndexMap)
 	{
 		const TSoftObjectPtr<UMaterialInterface>& AssetPtr = ModelResources.Materials.IsValidIndex(Pair.Key) ? ModelResources.Materials[Pair.Key] : nullptr;
-		ReferencedMaterials[Pair.Value] = AssetPtr.Get();
+		InstancePrivate->ReferencedMaterials[Pair.Value] = AssetPtr.Get();
 
 #if WITH_EDITOR
-		if (!ReferencedMaterials[Pair.Value])
+		if (!InstancePrivate->ReferencedMaterials[Pair.Value])
 		{
 			if (!AssetPtr.IsNull())
 			{
@@ -5590,7 +5575,7 @@ void UCustomizableInstancePrivate::AdditionalAssetsAsyncLoaded(const TSharedRef<
 #endif
 	}
 
-	for (FCustomizableInstanceComponentData& ComponentData : ComponentsData)
+	for (FCustomizableInstanceComponentData& ComponentData : InstancePrivate->ComponentsData)
 	{
 		for (int32 ResourceIndex : ComponentData.StreamedResourceIndex)
 		{
@@ -5656,10 +5641,10 @@ void UCustomizableInstancePrivate::AdditionalAssetsAsyncLoaded(const TSharedRef<
 		{
 			const int32 AssetIndex = AssetToStream.Key;
 
-			if (ClothingPhysicsAssets.IsValidIndex(AssetIndex) && ModelResources.PhysicsAssets.IsValidIndex(AssetToStream.Value))
+			if (InstancePrivate->ClothingPhysicsAssets.IsValidIndex(AssetIndex) && ModelResources.PhysicsAssets.IsValidIndex(AssetToStream.Value))
 			{
 				const TSoftObjectPtr<UPhysicsAsset>& PhysicsAssetPtr = ModelResources.PhysicsAssets[AssetToStream.Value];
-				ClothingPhysicsAssets[AssetIndex] = PhysicsAssetPtr.Get();
+				InstancePrivate->ClothingPhysicsAssets[AssetIndex] = PhysicsAssetPtr.Get();
 			}
 		}
 		ComponentData.ClothingPhysicsAssetsToStream.Empty();
@@ -5669,9 +5654,9 @@ void UCustomizableInstancePrivate::AdditionalAssetsAsyncLoaded(const TSharedRef<
 		{
 			if (TSubclassOf<UAnimInstance> AnimBP = SlotAnimBP.Value.Get())
 			{
-				if (!GatheredAnimBPs.Contains(AnimBP))
+				if (!InstancePrivate->GatheredAnimBPs.Contains(AnimBP))
 				{
-					GatheredAnimBPs.Add(AnimBP);
+					InstancePrivate->GatheredAnimBPs.Add(AnimBP);
 				}
 			}
 #if WITH_EDITOR
@@ -5698,25 +5683,25 @@ void UCustomizableInstancePrivate::AdditionalAssetsAsyncLoaded(const TSharedRef<
 		ComponentData.PhysicsAssets.AdditionalPhysicsAssetsToLoad.Empty();
 	}
 
-	LoadedPassThroughTexturesPendingSetMaterial.Empty(PassThroughTexturesToLoad.Num());
+	Context->LoadedPassThroughTexturesPendingSetMaterial.Empty(InstancePrivate->PassThroughTexturesToLoad.Num());
 
-	for (TSoftObjectPtr<const UTexture>& TextureRef : PassThroughTexturesToLoad)
+	for (TSoftObjectPtr<const UTexture>& TextureRef : InstancePrivate->PassThroughTexturesToLoad)
 	{
 		ensure(TextureRef.IsValid());
-		LoadedPassThroughTexturesPendingSetMaterial.Add(TextureRef.Get());
+		Context->LoadedPassThroughTexturesPendingSetMaterial.Add(TStrongObjectPtr(TextureRef.Get()));
 	}
 
-	PassThroughTexturesToLoad.Empty();
+	InstancePrivate->PassThroughTexturesToLoad.Empty();
 
-	LoadedPassThroughMeshesPendingSetMaterial.Empty(PassThroughMeshesToLoad.Num());
+	Context->LoadedPassThroughMeshesPendingSetMaterial.Empty(InstancePrivate->PassThroughMeshesToLoad.Num());
 
-	for (TSoftObjectPtr<const UStreamableRenderAsset>& MeshRef : PassThroughMeshesToLoad)
+	for (TSoftObjectPtr<const UStreamableRenderAsset>& MeshRef : InstancePrivate->PassThroughMeshesToLoad)
 	{
 		ensure(MeshRef.IsValid());
-		LoadedPassThroughMeshesPendingSetMaterial.Add(MeshRef.Get());
+		Context->LoadedPassThroughMeshesPendingSetMaterial.Add(TStrongObjectPtr(MeshRef.Get()));
 	}
 
-	PassThroughMeshesToLoad.Empty();
+	InstancePrivate->PassThroughMeshesToLoad.Empty();
 	
 #if WITH_EDITOR
 	// TODO: Remove with UE-217665 when the underlying bug in the ColorPicker is solved
@@ -5817,7 +5802,7 @@ void UCustomizableInstancePrivate::ReuseTexture(UTexture2D* Texture, TSharedRef<
 }
 
 
-void UCustomizableInstancePrivate::BuildMaterials(const TSharedRef<FUpdateContextPrivate>& OperationData, UCustomizableObjectInstance* Public)
+void UCustomizableInstancePrivate::BuildMaterials(const TSharedRef<FUpdateContextPrivate>& Context, UCustomizableObjectInstance* Public)
 {
 	MUTABLE_CPUPROFILER_SCOPE(UCustomizableInstancePrivate::BuildMaterials)
 
@@ -5836,36 +5821,36 @@ void UCustomizableInstancePrivate::BuildMaterials(const TSharedRef<FUpdateContex
 	UpdateContext->CustomizableObjectPathName = CustomizableObject->GetPathName();
 	UpdateContext->InstancePathName = Public->GetPathName();
 	UpdateContext->System = UCustomizableObjectSystem::GetInstance()->GetPrivate()->MutableSystem;
-	UpdateContext->Model = OperationData->Model;
+	UpdateContext->Model = Context->Model;
 	UpdateContext->MeshIdRegistry = CustomizableObject->GetPrivate()->MeshIdRegistry;
 	UpdateContext->ImageIdRegistry = CustomizableObject->GetPrivate()->ImageIdRegistry;
     UpdateContext->MaterialIdRegistry = CustomizableObject->GetPrivate()->MaterialIdRegistry,
-	UpdateContext->ExternalResourceProvider = OperationData->ExternalResourceProvider;
+	UpdateContext->ExternalResourceProvider = Context->ExternalResourceProvider;
 	UpdateContext->ModelStreamableBulkData = CustomizableObject->GetPrivate()->GetModelStreamableBulkData();
-	UpdateContext->Parameters = OperationData->Parameters;
-	UpdateContext->State = OperationData->GetCapturedDescriptor().GetState();
+	UpdateContext->Parameters = Context->Parameters;
+	UpdateContext->State = Context->GetCapturedDescriptor().GetState();
 
 	// Cache the descriptor as a string if we want to later report it using our benchmark utility. 
 	if (FLogBenchmarkUtil::IsBenchmarkingReportingEnabled())
 	{
-		UpdateContext->CapturedDescriptor = OperationData->GetCapturedDescriptor().ToString();
+		UpdateContext->CapturedDescriptor = Context->GetCapturedDescriptor().ToString();
 		if (GWorld)
 		{
 			UpdateContext->bLevelBegunPlay = GWorld->GetBegunPlay();
 		}
 	}
 
-	const bool bReuseTextures = OperationData->bBake ? false : OperationData->bReuseInstanceTextures;
+	const bool bReuseTextures = Context->bBake ? false : Context->bReuseInstanceTextures;
 	
 	TArray<bool> RecreateRenderStateOnInstanceComponent;
-	RecreateRenderStateOnInstanceComponent.Init(false, OperationData->NumInstanceComponents);
+	RecreateRenderStateOnInstanceComponent.Init(false, Context->NumInstanceComponents);
 
 	TArray<bool> NotifyUpdateOnInstanceComponent;
-	NotifyUpdateOnInstanceComponent.Init(false, OperationData->NumInstanceComponents);
+	NotifyUpdateOnInstanceComponent.Init(false, Context->NumInstanceComponents);
 
-	for (int32 InstanceComponentIndex = 0; InstanceComponentIndex < OperationData->NumInstanceComponents; ++InstanceComponentIndex)
+	for (int32 InstanceComponentIndex = 0; InstanceComponentIndex < Context->NumInstanceComponents; ++InstanceComponentIndex)
 	{
-		const FInstanceUpdateData::FComponent& Component = OperationData->InstanceUpdateData.Components[InstanceComponentIndex];
+		const FInstanceUpdateData::FComponent& Component = Context->InstanceUpdateData.Components[InstanceComponentIndex];
 
 		const FCustomizableObjectComponentIndex ObjectComponentIndex = Component.Id;
 
@@ -5881,13 +5866,13 @@ void UCustomizableInstancePrivate::BuildMaterials(const TSharedRef<FUpdateContex
 			continue;
 		}
 
-		const bool bReuseMaterials = !OperationData->bBake && !OperationData->MeshChangedPerInstanceComponent[InstanceComponentIndex];
+		const bool bReuseMaterials = !Context->bBake && !Context->MeshChangedPerInstanceComponent[InstanceComponentIndex];
 
 		// If the mesh is not transient, it means it's pass-through so it should use material overrides and not be modified in any way
 		const bool bIsTransientMesh = SkeletalMesh->HasAllFlags(EObjectFlags::RF_Transient);
 
 		// It is not safe to replace the materials of a SkeletalMesh whose resources are initialized. Use overrides instead.
-		const bool bUseOverrideMaterialsOnly = !bIsTransientMesh || (OperationData->bUseMeshCache && SkeletalMesh->GetResourceForRendering()->IsInitialized());
+		const bool bUseOverrideMaterialsOnly = !bIsTransientMesh || (Context->bUseMeshCache && SkeletalMesh->GetResourceForRendering()->IsInitialized());
 
 		UMaterialInterface* OverlayMaterial = nullptr;
 
@@ -5934,13 +5919,13 @@ void UCustomizableInstancePrivate::BuildMaterials(const TSharedRef<FUpdateContex
 
 		MUTABLE_CPUPROFILER_SCOPE(BuildMaterials_LODLoop);
 
-		const int32 FirstLOD = OperationData->bStreamMeshLODs ?
-			OperationData->FirstLODAvailable[ComponentName] :
-			OperationData->GetFirstRequestedLOD()[ComponentName];
+		const int32 FirstLOD = Context->bStreamMeshLODs ?
+			Context->FirstLODAvailable[ComponentName] :
+			Context->GetFirstRequestedLOD()[ComponentName];
 
 		for (int32 LODIndex = FirstLOD; LODIndex < Component.LODCount; LODIndex++)
 		{
-			const FInstanceUpdateData::FLOD& LOD = OperationData->InstanceUpdateData.LODs[Component.FirstLOD+LODIndex];
+			const FInstanceUpdateData::FLOD& LOD = Context->InstanceUpdateData.LODs[Component.FirstLOD+LODIndex];
 			
 			if (!bUseOverrideMaterialsOnly && LODIndex < SkeletalMesh->GetLODNum())
 			{
@@ -5956,7 +5941,7 @@ void UCustomizableInstancePrivate::BuildMaterials(const TSharedRef<FUpdateContex
 
 			for (int32 SurfaceIndex = 0; SurfaceIndex < LOD.SurfaceCount; ++SurfaceIndex)
 			{
-				const FInstanceUpdateData::FSurface& Surface = OperationData->InstanceUpdateData.Surfaces[LOD.FirstSurface + SurfaceIndex];
+				const FInstanceUpdateData::FSurface& Surface = Context->InstanceUpdateData.Surfaces[LOD.FirstSurface + SurfaceIndex];
 
 				// Reuse MaterialSlot from the previous LOD.
 				if (const int32 MaterialIndex = SurfaceIdToMaterialIndex.Find(Surface.SharedSurfaceId); MaterialIndex != INDEX_NONE)
@@ -6074,7 +6059,7 @@ void UCustomizableInstancePrivate::BuildMaterials(const TSharedRef<FUpdateContex
 
 					for (int32 VectorIndex = 0; VectorIndex < Surface.VectorCount; ++VectorIndex)
 					{
-						const FInstanceUpdateData::FVector& Vector = OperationData->InstanceUpdateData.Vectors[Surface.FirstVector + VectorIndex];
+						const FInstanceUpdateData::FVector& Vector = Context->InstanceUpdateData.Vectors[Surface.FirstVector + VectorIndex];
 
 						// Decoding Material Layer from Mutable parameter name
 						FString EncodingString = "-MutableLayerParam:";
@@ -6106,7 +6091,7 @@ void UCustomizableInstancePrivate::BuildMaterials(const TSharedRef<FUpdateContex
 
 					for (int32 ScalarIndex = 0; ScalarIndex < Surface.ScalarCount; ++ScalarIndex)
 					{
-						const FInstanceUpdateData::FScalar& Scalar = OperationData->InstanceUpdateData.Scalars[Surface.FirstScalar + ScalarIndex];
+						const FInstanceUpdateData::FScalar& Scalar = Context->InstanceUpdateData.Scalars[Surface.FirstScalar + ScalarIndex];
 
 						// Decoding Material Layer from Mutable parameter name
 						FString EncodingString = "-MutableLayerParam:";
@@ -6144,14 +6129,14 @@ void UCustomizableInstancePrivate::BuildMaterials(const TSharedRef<FUpdateContex
 					FTextureCache& TextureCache = CustomizableObject->GetPrivate()->TextureCache;
 
 					FString CurrentState = Public->GetCurrentState();
-					bool bNeverStream = OperationData->bNeverStream;
+					bool bNeverStream = Context->bNeverStream;
 
-					check((bNeverStream && OperationData->MipsToSkip == 0) ||
-						(!bNeverStream && OperationData->MipsToSkip >= 0));
+					check((bNeverStream && Context->MipsToSkip == 0) ||
+						(!bNeverStream && Context->MipsToSkip >= 0));
 
 					for (int32 ImageIndex = 0; ImageIndex < Surface.ImageCount; ++ImageIndex)
 					{
-						const FInstanceUpdateData::FImage& Image = OperationData->InstanceUpdateData.Images[Surface.FirstImage + ImageIndex];
+						const FInstanceUpdateData::FImage& Image = Context->InstanceUpdateData.Images[Surface.FirstImage + ImageIndex];
 						FString KeyName = Image.Name.ToString();
 						TSharedPtr<const mu::FImage> MutableImage = Image.Image;
 
@@ -6161,9 +6146,9 @@ void UCustomizableInstancePrivate::BuildMaterials(const TSharedRef<FUpdateContex
 						const FTextureReuseCacheKey TextureReuseCacheRef = bReuseTextures ? FTextureReuseCacheKey{Image.BaseLOD, ObjectComponentIndex.GetValue(), Surface.SurfaceId, ImageIndex} : FTextureReuseCacheKey{};
 
 #if WITH_EDITOR
-						const FTextureCache::FId TextureCacheKey = FTextureCache::FId(Image.ImageID, OperationData->MipsToSkip, OperationData->bBake);
+						const FTextureCache::FId TextureCacheKey = FTextureCache::FId(Image.ImageID, Context->MipsToSkip, Context->bBake);
 #else
-						const FTextureCache::FId TextureCacheKey = FTextureCache::FId(Image.ImageID, OperationData->MipsToSkip);
+						const FTextureCache::FId TextureCacheKey = FTextureCache::FId(Image.ImageID, Context->MipsToSkip);
 #endif
 						
 						// If the mutable image is null, it must be in the cache
@@ -6221,12 +6206,12 @@ void UCustomizableInstancePrivate::BuildMaterials(const TSharedRef<FUpdateContex
 
 								FString MutableTextureName;
 								{
-									const uint32 MutableImageHash = GetTypeHash(*Image.ImageID.GetKey(), OperationData->bBake);
+									const uint32 MutableImageHash = GetTypeHash(*Image.ImageID.GetKey(), Context->bBake);
 									MutableTextureName = FString::Printf(TEXT("T_%s_%lu"), *Props.TextureParameterName, MutableImageHash);
 									MutableTextureName.ReplaceInline(TEXT(" "), TEXT("_"));
 
 									// If not for bake work with this resource as normal and make it unique to avoid collisions
-									if (!OperationData->bBake)
+									if (!Context->bBake)
 									{
 										MakeMutableGeneratedObjectNameUnique(MutableTextureName, MutableImageHash, UMaterialInstanceDynamic::StaticClass());
 									}
@@ -6270,11 +6255,11 @@ void UCustomizableInstancePrivate::BuildMaterials(const TSharedRef<FUpdateContex
 
 								if (MutableTexture)
 								{
-									if (OperationData->ImageToPlatformDataMap.Contains(Image.ImageID))
+									if (Context->ImageToPlatformDataMap.Contains(Image.ImageID))
 									{
 										SetTexturePropertiesFromMutableImageProps(MutableTexture, Props, bNeverStream);
 
-										FTexturePlatformData* PlatformData = OperationData->ImageToPlatformDataMap[Image.ImageID];
+										FTexturePlatformData* PlatformData = Context->ImageToPlatformDataMap[Image.ImageID];
 
 										if (ReusedTexture)
 										{
@@ -6291,7 +6276,7 @@ void UCustomizableInstancePrivate::BuildMaterials(const TSharedRef<FUpdateContex
 											MutableTexture->SetPlatformData(PlatformData);
 										}
 
-										OperationData->ImageToPlatformDataMap.Remove(Image.ImageID);
+										Context->ImageToPlatformDataMap.Remove(Image.ImageID);
 									}
 									else
 									{
@@ -6478,11 +6463,11 @@ void UCustomizableInstancePrivate::BuildMaterials(const TSharedRef<FUpdateContex
 								MIDName = MIDName.RightChop(MaterialInstanceConstantPrefix.Len());
 							}
 
-							const uint32 MaterialPlaceHolderHash = GetHash(MutableMaterialPlaceholder, OperationData->bBake);
+							const uint32 MaterialPlaceHolderHash = GetHash(MutableMaterialPlaceholder, Context->bBake);
 							MIDName = FString::Printf(TEXT("%s%s_%u"),TEXT("MID_"),*MIDName, MaterialPlaceHolderHash);
 
 							// Add a "tag" to better identify this asset during the bake
-							if (OperationData->bBake)
+							if (Context->bBake)
 							{
 								MIDName = FBakingConfiguration::BakedResourcePrefix + MIDName;
 							}
@@ -6579,10 +6564,10 @@ void UCustomizableInstancePrivate::BuildMaterials(const TSharedRef<FUpdateContex
 		if (!bUseOverrideMaterialsOnly)
 		{
 			// Mutable hacky LOD Streaming
-			if (!OperationData->bStreamMeshLODs)
+			if (!Context->bStreamMeshLODs)
 			{
 				// Copy data from the FirstLODAvailable into the LODs below.
-				for (int32 LODIndex = OperationData->FirstLODAvailable[ComponentName]; LODIndex < FirstLOD; ++LODIndex)
+				for (int32 LODIndex = Context->FirstLODAvailable[ComponentName]; LODIndex < FirstLOD; ++LODIndex)
 				{
 					SkeletalMesh->GetLODInfo(LODIndex)->LODMaterialMap = SkeletalMesh->GetLODInfo(FirstLOD)->LODMaterialMap;
 
@@ -6646,9 +6631,9 @@ void UCustomizableInstancePrivate::BuildMaterials(const TSharedRef<FUpdateContex
 			const int32 ObjectComponentIndex = ModelResources.ComponentNamesPerObjectComponent.IndexOfByKey(ComponentName);
 			
 			int32 InstanceComponentIndex = -1;
-			for (int32 CurrentInstanceIndex=0; CurrentInstanceIndex<OperationData->InstanceUpdateData.Components.Num(); ++CurrentInstanceIndex)
+			for (int32 CurrentInstanceIndex=0; CurrentInstanceIndex<Context->InstanceUpdateData.Components.Num(); ++CurrentInstanceIndex)
 			{
-				if (OperationData->InstanceUpdateData.Components[CurrentInstanceIndex].Id.GetValue() == ObjectComponentIndex)
+				if (Context->InstanceUpdateData.Components[CurrentInstanceIndex].Id.GetValue() == ObjectComponentIndex)
 				{
 					InstanceComponentIndex = CurrentInstanceIndex;
 					break;
@@ -6685,10 +6670,6 @@ void UCustomizableInstancePrivate::BuildMaterials(const TSharedRef<FUpdateContex
 		MUTABLE_CPUPROFILER_SCOPE(BuildMaterials_Exchange);
 
 		Exchange(GeneratedTextures, NewGeneratedTextures);
-
-		// All pass-through textures and meshes have been set, no need to keep referencing them from the instance
-		LoadedPassThroughTexturesPendingSetMaterial.Empty();
-		LoadedPassThroughMeshesPendingSetMaterial.Empty();
 	}
 }
 
@@ -7238,7 +7219,6 @@ void UCustomizableInstancePrivate::RegenerateImportedModels()
 {
 	MUTABLE_CPUPROFILER_SCOPE(RegenerateImportedModels);
 
-	using UVsVectorType = typename TDecay<decltype(DeclVal<FSoftSkinVertex>().UVs[0])>::Type;
 	struct FMeshDataConvertJob
 	{
 		int32 NumIndices = 0;
diff --git a/Engine/Plugins/Mutable/Source/CustomizableObject/Private/MuCO/CustomizableObjectSystem.cpp b/Engine/Plugins/Mutable/Source/CustomizableObject/Private/MuCO/CustomizableObjectSystem.cpp
--- a/Engine/Plugins/Mutable/Source/CustomizableObject/Private/MuCO/CustomizableObjectSystem.cpp
+++ b/Engine/Plugins/Mutable/Source/CustomizableObject/Private/MuCO/CustomizableObjectSystem.cpp
@@ -3351,7 +3351,7 @@ namespace impl
 		// which will safely access the GlobalExternalImages map, and then just get the cached image or issue a disk read
 
 		// Copy data generated in the mutable thread over to the instance
-		ObjectInstancePrivateData->PrepareForUpdate(OperationData);
+		ObjectInstancePrivateData->PrepareForUpdate();
 
 		// Task: Mutable GetImages
 		//-------------------------------------------------------------
@@ -3369,7 +3369,7 @@ namespace impl
 
 		// Next Task: Load Unreal Assets
 		//-------------------------------------------------------------
-		UE::Tasks::FTask Game_LoadUnrealAssets = ObjectInstancePrivateData->LoadAdditionalAssetsAndData(OperationData);
+		UE::Tasks::FTask Game_LoadUnrealAssets = LoadAdditionalAssetsAndData(OperationData);
 
 		// Next-next Task: Convert Resources
 		//-------------------------------------------------------------
