MuCO: Fix unreachable code in server builds

#rnx

diff --git a/Engine/Plugins/Mutable/Source/CustomizableObject/Private/MuCO/CustomizableObjectSystem.cpp b/Engine/Plugins/Mutable/Source/CustomizableObject/Private/MuCO/CustomizableObjectSystem.cpp
--- a/Engine/Plugins/Mutable/Source/CustomizableObject/Private/MuCO/CustomizableObjectSystem.cpp
+++ b/Engine/Plugins/Mutable/Source/CustomizableObject/Private/MuCO/CustomizableObjectSystem.cpp
@@ -3805,7 +3805,7 @@ int32 UCustomizableObjectSystem::TickInternal(const bool bBlocking)
 	// Building instances is not enabled in servers. If at some point relevant collision or animation data is necessary for server logic this will need to be changed.
 #if UE_SERVER
 	return 0;
-#endif
+#else // !UE_SERVER
 
 	if (!Private)
 	{
@@ -4147,6 +4147,7 @@ int32 UCustomizableObjectSystem::TickInternal(const bool bBlocking)
 	RemainingWork += GetPrivate()->StreamableManager->Tick(bBlocking);
 	
 	return RemainingWork;
+#endif // !UE_SERVER
 }
 
 
