[Mutable] Debugger improvements

- Added support for Dataless parameters.
- Unified Debugger and Editor compilation options.
- Reordered some widgets have a better use of the available space.

diff --git a/Engine/Plugins/Mutable/Source/CustomizableObjectEditor/Private/MuCOE/AssetDefinition_CustomizableObject.cpp b/Engine/Plugins/Mutable/Source/CustomizableObjectEditor/Private/MuCOE/AssetDefinition_CustomizableObject.cpp
--- a/Engine/Plugins/Mutable/Source/CustomizableObjectEditor/Private/MuCOE/AssetDefinition_CustomizableObject.cpp
+++ b/Engine/Plugins/Mutable/Source/CustomizableObjectEditor/Private/MuCOE/AssetDefinition_CustomizableObject.cpp
@@ -148,29 +148,6 @@ namespace MenuExtension_CustomizableObject
 		}
 	}
 
-
-	void ExecuteDebug(const FToolMenuContext& InContext)
-	{
-		const UContentBrowserAssetContextMenuContext* Context = UContentBrowserAssetContextMenuContext::FindContextWithAssets(InContext);
-		check(Context);
-		
-		for (UCustomizableObject* Object : Context->LoadSelectedObjects<UCustomizableObject>())
-		{
-			if (!AssetViewUtils::OpenEditorForAsset(Object))
-			{
-				continue;
-			}
-
-			if (IAssetEditorInstance* AssetEditor = GEditor->GetEditorSubsystem<UAssetEditorSubsystem>()->FindEditorForAsset(Object, true))
-			{
-				if (const FCustomizableObjectEditor* CustomizableObjectEditor = StaticCast<FCustomizableObjectEditor*>(AssetEditor))
-				{
-					CustomizableObjectEditor->DebugObject();
-				}
-			}
-		}
-	}
-
 	enum class ERecompileCO
 	{
 		RCO_All,
@@ -375,20 +352,6 @@ namespace MenuExtension_CustomizableObject
 			        InSection.AddMenuEntry("CustomizableObject_ExecuteNewInstance", Label, ToolTip, Icon, UIAction);
 		        }
 
-		        // Debug
-		        {
-			        const TAttribute<FText> Label = NSLOCTEXT("UAssetDefinition_CustomizableObject",
-			                                                  "CustomizableObject_Debug", "Debug");
-			        const TAttribute<FText> ToolTip = NSLOCTEXT("UAssetDefinition_CustomizableObject",
-			                                                    "CustomizableObject_DebugTooltip",
-			                                                    "Open the debugger for the customizable object.");
-			        const FSlateIcon Icon = FSlateIcon();
-
-			        FToolUIAction UIAction;
-			        UIAction.ExecuteAction = FToolMenuExecuteAction::CreateStatic(&ExecuteDebug);
-			        InSection.AddMenuEntry("CustomizableObject_ExecuteDebug", Label, ToolTip, Icon, UIAction);
-		        }
-
 		        // Recompile submenu
 		        {
 			        const TAttribute<FText> Label = NSLOCTEXT("UAssetDefinition_CustomizableObject",
diff --git a/Engine/Plugins/Mutable/Source/CustomizableObjectEditor/Private/MuCOE/CustomizableObjectEditor.cpp b/Engine/Plugins/Mutable/Source/CustomizableObjectEditor/Private/MuCOE/CustomizableObjectEditor.cpp
--- a/Engine/Plugins/Mutable/Source/CustomizableObjectEditor/Private/MuCOE/CustomizableObjectEditor.cpp
+++ b/Engine/Plugins/Mutable/Source/CustomizableObjectEditor/Private/MuCOE/CustomizableObjectEditor.cpp
@@ -3,6 +3,7 @@
 #include "MuCOE/CustomizableObjectEditor.h"
 
 #include "CustomizableObjectCompiler.h"
+#include "CustomizableObjectCompileRunnable.h"
 #include "CustomizableObjectEditorModule.h"
 #include "CustomizableObjectEditorPerformanceAnalyzer.h"
 #include "AssetRegistry/ARFilter.h"
@@ -49,6 +50,8 @@
 #include "MuCOE/SCustomizableObjectEditorViewport.h"
 #include "PropertyEditorModule.h"
 #include "ScopedTransaction.h"
+#include "SMutableCodeViewer.h"
+#include "SMutableGraphViewer.h"
 #include "MuCO/CustomizableObjectInstancePrivate.h"
 #include "MuCO/CustomizableObjectPrivate.h"
 #include "UObject/EnumProperty.h"
@@ -57,8 +60,9 @@
 #include "Widgets/Input/STextComboBox.h"
 #include "Widgets/Layout/SScrollBox.h"
 #include "Widgets/Notifications/SNotificationList.h"
-#include "SMutableObjectViewer.h"
+#include "MuCO/BoneNames.h"
 #include "MuCO/CustomizableObjectCustomVersion.h"
+#include "MuCO/UnrealMutableImageProvider.h"
 #include "MuCOE/CompileRequest.h"
 #include "MuCOE/CustomizableObjectInstanceFactory.h"
 
@@ -473,11 +477,17 @@ void FCustomizableObjectEditor::BindCommands()
 
 	// Debug and options
 	ToolkitCommands->MapAction(
-		Commands.Debug,
-		FExecuteAction::CreateSP(this, &FCustomizableObjectEditor::DebugObject),
+		Commands.GraphViewer,
+		FExecuteAction::CreateSP(this, &FCustomizableObjectEditor::GraphViewer),
 		FCanExecuteAction(),
 		FIsActionChecked());
 
+	ToolkitCommands->MapAction(
+		Commands.CodeViewer,
+		FExecuteAction::CreateSP(this, &FCustomizableObjectEditor::CodeViewer),
+		FCanExecuteAction(),
+		FIsActionChecked());
+	
 	// References
 	ToolkitCommands->MapAction(
 		Commands.CompileGatherReferences,
@@ -1173,7 +1183,45 @@ TSharedRef<SWidget> FCustomizableObjectEditor::GenerateCompileOptionsMenuContent
 	// Debugging options
 	MenuBuilder.BeginSection("Debugger", LOCTEXT("MutableDebugger", "Debugger"));
 	{
-		MenuBuilder.AddMenuEntry(FCustomizableObjectEditorCommands::Get().Debug);
+		// Platform
+		TargetPlatformNames.Empty();
+		
+		if (ITargetPlatformManagerModule* TargetPlatformManager = GetTargetPlatformManager())
+		{
+			const ITargetPlatform* CurrentPlatform = TargetPlatformManager->GetRunningTargetPlatform();
+			for (const ITargetPlatform* Platform : TargetPlatformManager->GetTargetPlatforms())
+			{
+				const TSharedPtr<FString> ThisPlatform = MakeShareable(new FString(Platform->PlatformName()));
+				TargetPlatformNames.Add(ThisPlatform);
+
+				if (Platform == CurrentPlatform)
+				{
+					DebugTargetPlatformName = ThisPlatform;
+				}
+			}
+		}
+		
+		TSharedRef<STextComboBox> DebugPlatformCombo = SNew(STextComboBox)
+			.OptionsSource(&TargetPlatformNames)
+			.InitiallySelectedItem(DebugTargetPlatformName)
+			.OnSelectionChanged(this, &FCustomizableObjectEditor::OnChangeDebugPlatform);
+
+		MenuBuilder.AddWidget(DebugPlatformCombo, LOCTEXT("MutableDebugPlatform", "Target Platform")); 
+
+		MenuBuilder.AddMenuEntry(
+				LOCTEXT("ForceLargeLODBias", "Force a large texture LODBias."),
+				LOCTEXT("ForceLargeLODBiasTooltip", "This is useful to test compilation of special cook modes."),
+				FSlateIcon(),
+				FUIAction(
+					FExecuteAction::CreateLambda([this]() { bDebugForceLargeLODBias = !bDebugForceLargeLODBias; }),
+					FCanExecuteAction(),
+					FIsActionChecked::CreateLambda([this]() { return bDebugForceLargeLODBias; })),
+				NAME_None,
+				EUserInterfaceActionType::ToggleButton
+			);
+		
+		MenuBuilder.AddMenuEntry(FCustomizableObjectEditorCommands::Get().GraphViewer);
+		MenuBuilder.AddMenuEntry(FCustomizableObjectEditorCommands::Get().CodeViewer);
 	}
 	MenuBuilder.EndSection();
 
@@ -1311,12 +1359,138 @@ void FCustomizableObjectEditor::CompileObject(bool bOnlySelectedParameters, bool
 }
 
 
-void FCustomizableObjectEditor::DebugObject() const
+const ITargetPlatform* GetTargetPlatform(const TSharedPtr<FString>& TargetPlatformName)
+{
+	ITargetPlatformManagerModule* TargetPlatformManager = GetTargetPlatformManager();
+	if (!TargetPlatformManager)
+	{
+		return nullptr;
+	}
+
+	for (const ITargetPlatform* Platform : TargetPlatformManager->GetTargetPlatforms())
+	{
+		if (*TargetPlatformName.Get() == Platform->PlatformName())
+		{
+			return Platform;
+		}
+	}
+
+	return nullptr;
+}
+
+
+mu::Ptr<mu::Node> FCustomizableObjectEditor::DebuggerCompile(TSharedPtr<mu::FModel>& Model,
+	TArray<TSoftObjectPtr<const UTexture>>& RuntimeTextures,
+	TArray<FMutableSourceTextureData>& CompilerTextures,
+	TArray<TSoftObjectPtr<const UStreamableRenderAsset>>& RuntimeMeshes,
+	TArray<FMutableSourceMeshData>& CompilerMeshes) const
+{
+	FCompilationOptions CompileOptions = CustomizableObject->GetPrivate()->GetCompileOptions();
+	CompileOptions.TargetPlatform = GetTargetPlatform(DebugTargetPlatformName);
+	CompileOptions.bForceLargeLODBias = bDebugForceLargeLODBias;
+	
+	if (CompileOptions.bForceLargeLODBias)
+	{
+		// Debug compile with many different biases
+		constexpr int32 MaxBias = 15;
+		for (int32 Bias = 0; Bias < MaxBias; ++Bias)
+		{
+			CompileOptions.DebugBias = Bias;
+
+			TSharedRef<FCustomizableObjectCompiler> Compiler = MakeShared<FCustomizableObjectCompiler>();
+			RuntimeTextures.Empty();
+			CompilerTextures.Empty();
+			RuntimeMeshes.Empty();
+			CompilerMeshes.Empty();
+			mu::Ptr<mu::Node> RootNode = Compiler->Export(CustomizableObject, CompileOptions, RuntimeTextures, CompilerTextures, RuntimeMeshes, CompilerMeshes);
+			if (!RootNode)
+			{
+				// TODO: Show errors
+				ensure(false);
+				return nullptr;
+			}
+
+			// Do the compilation to Mutable Code synchronously.
+			TSharedPtr<FCustomizableObjectCompileRunnable> CompileTask = MakeShareable(new FCustomizableObjectCompileRunnable(RootNode, Compiler));
+			CompileTask->Options = CompileOptions;
+			CompileTask->ReferencedTextures = CompilerTextures;
+			CompileTask->ReferencedMeshes = CompilerMeshes;
+			CompileTask->Init();
+			CompileTask->Run();
+		}
+	}
+
+	// Convert from Unreal graph to Mutable graph.
+	TSharedRef<FCustomizableObjectCompiler> Compiler = MakeShared<FCustomizableObjectCompiler>();
+	RuntimeTextures.Empty();
+	CompilerTextures.Empty();
+	RuntimeMeshes.Empty();
+	CompilerMeshes.Empty();
+	mu::Ptr<mu::Node> RootNode = Compiler->Export(CustomizableObject, CompileOptions, RuntimeTextures, CompilerTextures, RuntimeMeshes, CompilerMeshes);
+	if (!RootNode)
+	{
+		// TODO: Show errors
+		ensure(false);
+		return nullptr;
+	}
+
+	// Do the compilation to Mutable Code synchronously.
+	TSharedPtr<FCustomizableObjectCompileRunnable> CompileTask = MakeShareable(new FCustomizableObjectCompileRunnable(RootNode, Compiler));
+	CompileTask->Options = CompileOptions;
+	CompileTask->ReferencedTextures = CompilerTextures;
+	CompileTask->ReferencedMeshes = CompilerMeshes;
+	CompileTask->Init();
+	CompileTask->Run();
+
+	Model = CompileTask->Model;
+	
+	return RootNode;
+}
+
+
+void FCustomizableObjectEditor::GraphViewer() const
+{
+	TSharedPtr<mu::FModel> Model;
+	TArray<TSoftObjectPtr<const UTexture>> RuntimeTextures;
+	TArray<FMutableSourceTextureData> CompilerTextures;
+	TArray<TSoftObjectPtr<const UStreamableRenderAsset>> RuntimeMeshes;
+	TArray<FMutableSourceMeshData> CompilerMeshes;
+
+	mu::Ptr<mu::Node> RootNode = DebuggerCompile(Model, RuntimeTextures, CompilerTextures, RuntimeMeshes, CompilerMeshes);
+	
+	const TSharedPtr<SDockTab> NewMutableObjectTab = SNew(SDockTab)
+	.Label(LOCTEXT("GraphViewer", "Graph Viewer"))
+	[
+		SNew(SMutableGraphViewer, RootNode)
+		.ReferencedRuntimeTextures(RuntimeTextures)
+		.ReferencedCompileTextures(CompilerTextures)
+		.ReferencedRuntimeMeshes(RuntimeMeshes)
+		.ReferencedCompileMeshes(CompilerMeshes)
+	];
+	
+	// Spawn the debugger tab alongside the Graph Tab 
+	TabManager->InsertNewDocumentTab(GraphTabId, FTabManager::ESearchPreference::PreferLiveTab, NewMutableObjectTab.ToSharedRef());
+}
+
+
+void FCustomizableObjectEditor::CodeViewer() const
 {
+	TSharedPtr<mu::FModel> Model;
+	TArray<TSoftObjectPtr<const UTexture>> RuntimeTextures;
+	TArray<FMutableSourceTextureData> CompilerTextures;
+	TArray<TSoftObjectPtr<const UStreamableRenderAsset>> RuntimeMeshes;
+	TArray<FMutableSourceMeshData> CompilerMeshes;
+
+	mu::Ptr<mu::Node> RootNode = DebuggerCompile(Model, RuntimeTextures, CompilerTextures, RuntimeMeshes, CompilerMeshes);
+	
+	const TSharedRef<FBoneNames> BoneNames = MakeShared<FBoneNames>(*CustomizableObject->GetPrivate()->GetModelResources());
+	TSharedPtr<FUnrealMutableResourceProvider> ExternalResourceProvider = MakeShared<FUnrealMutableResourceProvider>(BoneNames);
+	
 	const TSharedPtr<SDockTab> NewMutableObjectTab = SNew(SDockTab)
-	.Label(FText::FromString(TEXT("Debugger")))
+	.Label(LOCTEXT("CodeViewer", "Code Viewer"))
 	[
-		SNew(SMutableObjectViewer, CustomizableObject)
+		SNew(SMutableCodeViewer, *CustomizableObject, *PreviewInstance, Model)
+		.ExternalResourceProvider(ExternalResourceProvider)
 	];
 	
 	// Spawn the debugger tab alongside the Graph Tab 
@@ -1770,4 +1944,10 @@ void FCustomizableObjectEditor::OnCustomizableObjectStatusChanged(FCustomizableO
 	}
 }
 
+
+void FCustomizableObjectEditor::OnChangeDebugPlatform(TSharedPtr<FString> NewSelection, ESelectInfo::Type SelectInfo)
+{
+	DebugTargetPlatformName = NewSelection;
+}
+
 #undef LOCTEXT_NAMESPACE
diff --git a/Engine/Plugins/Mutable/Source/CustomizableObjectEditor/Private/MuCOE/CustomizableObjectEditor.h b/Engine/Plugins/Mutable/Source/CustomizableObjectEditor/Private/MuCOE/CustomizableObjectEditor.h
--- a/Engine/Plugins/Mutable/Source/CustomizableObjectEditor/Private/MuCOE/CustomizableObjectEditor.h
+++ b/Engine/Plugins/Mutable/Source/CustomizableObjectEditor/Private/MuCOE/CustomizableObjectEditor.h
@@ -7,10 +7,13 @@
 #include "GraphEditor.h"
 #include "Misc/NotifyHook.h"
 #include "MuCO/CustomizableObjectPrivate.h"
+#include "MuT/Node.h"
 #include "Widgets/Input/SNumericDropDown.h"
 
 #include "CustomizableObjectEditor.generated.h"
 
+struct FMutableSourceTextureData;
+struct FMutableSourceMeshData;
 enum class ECustomizableObjectProjectorType : uint8;
 namespace ESelectInfo { enum Type : int; }
 
@@ -166,8 +169,14 @@ public:
 	/** Callback to notify the editor when the PreviewInstance has been updated */
 	void OnUpdatePreviewInstance(UCustomizableObjectInstance* Instance);
 
-	/** Debug the object as a raw mutable data in the internal tools. */
-	void DebugObject() const;
+	mu::Ptr<mu::Node> DebuggerCompile(TSharedPtr<mu::FModel>& Model,
+		TArray<TSoftObjectPtr<const UTexture>>& RuntimeTextures,
+		TArray<FMutableSourceTextureData>& CompilerTextures,
+		TArray<TSoftObjectPtr<const UStreamableRenderAsset>>& RuntimeMeshes,
+		TArray<FMutableSourceMeshData>& CompilerMeshes) const;
+	
+	void GraphViewer() const;
+	void CodeViewer() const;
 
 	/** Clear game asset references saved in the Customizable Object. */
 	void ClearGatheredReferences();
@@ -246,6 +255,8 @@ private:
 
 	void OnPostCompile();
 
+	void OnChangeDebugPlatform(TSharedPtr<FString> NewSelection, ESelectInfo::Type SelectInfo);
+	
 public:
 	void OnCustomizableObjectStatusChanged(FCustomizableObjectStatus::EState PreviousState, FCustomizableObjectStatus::EState CurrentState);
 
@@ -320,6 +331,11 @@ private:
 	
 	EGizmoType GizmoType = EGizmoType::Hidden;
 
+	TArray<TSharedPtr<FString>> TargetPlatformNames;
+	TSharedPtr<FString> DebugTargetPlatformName;
+	
+	bool bDebugForceLargeLODBias = false;
+
 protected:
 	/** @return the documentation location for this editor */
 	virtual FString GetDocumentationLink() const override;
diff --git a/Engine/Plugins/Mutable/Source/CustomizableObjectEditor/Private/MuCOE/CustomizableObjectEditorActions.cpp b/Engine/Plugins/Mutable/Source/CustomizableObjectEditor/Private/MuCOE/CustomizableObjectEditorActions.cpp
--- a/Engine/Plugins/Mutable/Source/CustomizableObjectEditor/Private/MuCOE/CustomizableObjectEditorActions.cpp
+++ b/Engine/Plugins/Mutable/Source/CustomizableObjectEditor/Private/MuCOE/CustomizableObjectEditorActions.cpp
@@ -26,7 +26,8 @@ void FCustomizableObjectEditorCommands::RegisterCommands()
 	UI_COMMAND(CompileOnlySelected, "Compile Only Selected", "Compile the source graph of the customizable object and update the previews, only for the selected options in the preview. The rest of options are discarded. If they are selected, press again this button to see their effect in the preview. \nActive if the CVar Mutable.Enabled is set to true.", EUserInterfaceActionType::Button, FInputChord());
 	UI_COMMAND(ResetCompileOptions, "Reset Compilation Options", "Set reasonable defaults for the compilation options.", EUserInterfaceActionType::Button, FInputChord());
 	UI_COMMAND(CompileOptions_UseDiskCompilation, "Enable compiling using the disk as memory.", "This is very slow but supports compiling huge objects. It requires a lot of free space in the OS disk.", EUserInterfaceActionType::ToggleButton, FInputChord());
-	UI_COMMAND(Debug, "Debug", "Open the CustomizableObject debugger tab for this object.", EUserInterfaceActionType::Button, FInputChord());
+	UI_COMMAND(GraphViewer, "Graph Viewer", "Open the Customizable Object Graph Viewer.", EUserInterfaceActionType::Button, FInputChord());
+	UI_COMMAND(CodeViewer, "Code Viewer", "Open the Customizable Object Code Viewer.", EUserInterfaceActionType::Button, FInputChord());
 
 	UI_COMMAND(PerformanceAnalyzer, "Performance Analyzer", "Open the Performance Analyzer window to set up and perform all tests relevant to Customizable Objects.", EUserInterfaceActionType::Button, FInputChord());
 	UI_COMMAND(TextureAnalyzer, "Texture Memory Analyzer", "Open the Texture Analyzer window to check all the information of the textures created by Mutable.", EUserInterfaceActionType::Button, FInputChord());
diff --git a/Engine/Plugins/Mutable/Source/CustomizableObjectEditor/Private/MuCOE/CustomizableObjectEditorActions.h b/Engine/Plugins/Mutable/Source/CustomizableObjectEditor/Private/MuCOE/CustomizableObjectEditorActions.h
--- a/Engine/Plugins/Mutable/Source/CustomizableObjectEditor/Private/MuCOE/CustomizableObjectEditorActions.h
+++ b/Engine/Plugins/Mutable/Source/CustomizableObjectEditor/Private/MuCOE/CustomizableObjectEditorActions.h
@@ -18,7 +18,8 @@ public:
 	TSharedPtr< FUICommandInfo > CompileOnlySelected;
 	TSharedPtr< FUICommandInfo > ResetCompileOptions;
 	TSharedPtr< FUICommandInfo > CompileOptions_UseDiskCompilation;
-	TSharedPtr< FUICommandInfo > Debug;
+	TSharedPtr< FUICommandInfo > GraphViewer;
+	TSharedPtr< FUICommandInfo > CodeViewer;
 
 	TSharedPtr< FUICommandInfo > PerformanceAnalyzer;
 	TSharedPtr< FUICommandInfo > ResetPerformanceReportOptions;
diff --git a/Engine/Plugins/Mutable/Source/CustomizableObjectEditor/Private/MuCOE/CustomizableObjectEditorStyle.cpp b/Engine/Plugins/Mutable/Source/CustomizableObjectEditor/Private/MuCOE/CustomizableObjectEditorStyle.cpp
--- a/Engine/Plugins/Mutable/Source/CustomizableObjectEditor/Private/MuCOE/CustomizableObjectEditorStyle.cpp
+++ b/Engine/Plugins/Mutable/Source/CustomizableObjectEditor/Private/MuCOE/CustomizableObjectEditorStyle.cpp
@@ -82,8 +82,10 @@ TSharedRef<FSlateStyleSet> FCustomizableObjectEditorStyle::Create()
 	Style->Set("CustomizableObjectEditor.CompileOnlySelected.Small", new IMAGE_BRUSH_SVG(Style, "Editor/Slate/Starship/MainToolbar/compile", Icon20x20));
 	Style->Set("CustomizableObjectEditor.ResetCompileOptions", new IMAGE_BRUSH_SVG(Style, "Editor/Slate/Starship/MainToolbar/compile", Icon40x40));
 	Style->Set("CustomizableObjectEditor.ResetCompileOptions.Small", new IMAGE_BRUSH_SVG(Style, "Editor/Slate/Starship/MainToolbar/compile", Icon20x20));
-	Style->Set("CustomizableObjectEditor.Debug", new IMAGE_BRUSH_SVG(Style, "Editor/Slate/Starship/Common/Debug", Icon40x40));
-	Style->Set("CustomizableObjectEditor.Debug.Small", new IMAGE_BRUSH_SVG(Style, "Editor/Slate/Starship/Common/Debug", Icon20x20));
+	Style->Set("CustomizableObjectEditor.GraphViewer", new IMAGE_BRUSH_SVG(Style, "Editor/Slate/Starship/Common/BlueprintDebugger", Icon40x40));
+	Style->Set("CustomizableObjectEditor.GraphViewer.Small", new IMAGE_BRUSH_SVG(Style, "Editor/Slate/Starship/Common/BlueprintDebugger", Icon20x20));
+	Style->Set("CustomizableObjectEditor.CodeViewer", new IMAGE_BRUSH_SVG(Style, "Editor/Slate/Starship/Common/Debug", Icon40x40));
+	Style->Set("CustomizableObjectEditor.CodeViewer.Small", new IMAGE_BRUSH_SVG(Style, "Editor/Slate/Starship/Common/Debug", Icon20x20));
 	Style->Set("CustomizableObjectEditor.TextureAnalyzer", new IMAGE_BRUSH_SVG(Style, "Editor/Slate/Starship/Common/IssueTracker", Icon40x40));
 	Style->Set("CustomizableObjectEditor.TextureAnalyzer.Small", new IMAGE_BRUSH_SVG(Style, "Editor/Slate/Starship/Common/IssueTracker", Icon20x20));
 	Style->Set("CustomizableObjectEditor.PerformanceReport", new IMAGE_BRUSH_SVG(Style, "Editor/Slate/Starship/MainToolbar/simulate", Icon40x40));
diff --git a/Engine/Plugins/Mutable/Source/CustomizableObjectEditor/Private/MuCOE/SMutableCodeViewer.cpp b/Engine/Plugins/Mutable/Source/CustomizableObjectEditor/Private/MuCOE/SMutableCodeViewer.cpp
--- a/Engine/Plugins/Mutable/Source/CustomizableObjectEditor/Private/MuCOE/SMutableCodeViewer.cpp
+++ b/Engine/Plugins/Mutable/Source/CustomizableObjectEditor/Private/MuCOE/SMutableCodeViewer.cpp
@@ -16,7 +16,6 @@
 #include "MuCOE/SMutableLayoutViewer.h"
 #include "MuCOE/SMutableMeshViewer.h"
 #include "MuCOE/SMutableInstanceViewer.h"
-#include "MuCOE/SMutableParametersWidget.h"
 #include "MuCOE/SMutableProjectorViewer.h"
 #include "MuCOE/SMutableScalarViewer.h"
 #include "MuCOE/SMutableSkeletonViewer.h"
@@ -39,6 +38,7 @@
 #include "Widgets/Input/SSearchBox.h"
 #include "Engine/SkeletalMesh.h"
 #include "Internationalization/Regex.h"
+#include "MuCO/CustomizableObjectInstancePrivate.h"
 
 class FExtender;
 class FReferenceCollector;
@@ -275,18 +275,6 @@ private:
 };
 
 
-void SMutableCodeViewer::AddReferencedObjects(FReferenceCollector& Collector)
-{
-	// Add UObjects here if we own any at some point
-	//Collector.AddReferencedObject(CustomizableObject);
-}
-
-
-FString SMutableCodeViewer::GetReferencerName() const
-{
-	return TEXT("SMutableCodeViewer");
-}
-
 void SMutableCodeViewer::ClearSelectedTreeRow() const
 {
 	check(TreeView);
@@ -294,19 +282,18 @@ void SMutableCodeViewer::ClearSelectedTreeRow() const
 }
 
 
-void SMutableCodeViewer::Construct(const FArguments& InArgs,
-	const TSharedPtr<mu::FModel>& InMutableModel)
+void SMutableCodeViewer::Construct(const FArguments& InArgs, UCustomizableObject& InObject, UCustomizableObjectInstance& InInstance, const TSharedPtr<mu::FModel>& InMutableModel)
 {
 	// Min width allowed for the column. Needed to avoid having issues with the constants space being to small
 	// and then getting too tall on the y axis crashing the UI drawer.
 	constexpr float MinParametersCollWidth = 200;
 
+	Object = TStrongObjectPtr(&InObject);
+	Instance = TStrongObjectPtr(&InInstance);
+
 	MutableModel = InMutableModel;
-	MeshIdRegistry = InArgs._MeshIdRegistry;
-	ImageIdRegistry = InArgs._ImageIdRegistry;
-	MaterialIdRegistry = InArgs._MaterialIdRegistry;
 	ExternalResourceProvider = InArgs._ExternalResourceProvider;
-	PreviewParameters = mu::FModel::NewParameters(MutableModel);
+	PreviewParameters = Instance->GetPrivate()->GetDescriptor().GetParameters();
 
 	RootNodes.Empty();
 	RootNodeAddresses.Empty();
@@ -352,10 +339,6 @@ void SMutableCodeViewer::Construct(const FArguments& InArgs,
 	ToolbarBuilder.SetLabelVisibility(EVisibility::Visible);
 	ToolbarBuilder.SetStyle(&FAppStyle::Get(), "SlimToolBar");
 
-	ToolbarBuilder.AddWidget(
-		SNew(STextBlock)
-		.Text(FText::FromString(InArgs._DataTag)));
-
 	TSharedRef<SScrollBar> TreeVertScrollBar =
 		SNew(SScrollBar).
 		Orientation(EOrientation::Orient_Vertical).
@@ -393,6 +376,7 @@ void SMutableCodeViewer::Construct(const FArguments& InArgs,
 				
 					// Search by name
 					+ SHorizontalBox::Slot()
+					.Padding(2,2)
 					.AutoWidth()
 					[
 						SNew(SHorizontalBox)
@@ -402,7 +386,7 @@ void SMutableCodeViewer::Construct(const FArguments& InArgs,
 						.VAlign(VAlign_Center)
 						[
 							SNew(STextBlock)
-							.Text(LOCTEXT("SelectedOperationByStringLabel","Search Operation by String :"))
+							.Text(LOCTEXT("SelectedOperationByStringLabel","Search Operation by String "))
 						]
 					
 						+ SHorizontalBox::Slot()
@@ -431,11 +415,12 @@ void SMutableCodeViewer::Construct(const FArguments& InArgs,
 						.VAlign(VAlign_Center)
 						[
 							SNew(STextBlock)
-							.Text(LOCTEXT("OperationToSearchRegexLabel","Is RegEx?"))
+							.Text(LOCTEXT("Regex", "Regex"))
 						]
 				
 						+ SHorizontalBox::Slot()
 						.AutoWidth()
+						.Padding(4, 0, 0, 0)
 						.VAlign(VAlign_Center)
 						[
 							SNew(SCheckBox)
@@ -468,7 +453,7 @@ void SMutableCodeViewer::Construct(const FArguments& InArgs,
 							.VAlign(VAlign_Center)
 							[
 								SNew(STextBlock)
-								.Text(LOCTEXT("SelectedOperationTypeLabel","Search Operation Type :"))
+								.Text(LOCTEXT("SelectedOperationTypeLabel","Search Operation Type "))
 							]
 				
 							// ComboBox used to select one or another Op_Type for tree navigation purposes
@@ -590,25 +575,37 @@ void SMutableCodeViewer::Construct(const FArguments& InArgs,
 				.MinSize(MinParametersCollWidth)
 				[
 					// Splitter managing both parameter and constant panels
-					SNew(SSplitter)
-					.Orientation(EOrientation::Orient_Vertical)
-					+SSplitter::Slot()
+					SNew(SVerticalBox)
+					+ SVerticalBox::Slot()
+					.AutoHeight()
 					[
 						SNew(SVerticalBox)
-
 						+ SVerticalBox::Slot()
+						.Padding(4, 4, 4,0)
+						.AutoHeight()
+						[
+							SNew(SButton)
+							.Text(LOCTEXT("RefreshParameters", "Refresh Parameters"))
+							.OnClicked(this, &SMutableCodeViewer::UpdateParameters)
+						]
+						
+						+ SVerticalBox::Slot()
+						.Padding(4, 8, 4, 0)
 						.AutoHeight()
 						[
 							SNew(SHorizontalBox)
 
 							+ SHorizontalBox::Slot()
+							.Padding(4, 0, 0, 0)
+							.VAlign(VAlign_Center)
 							.AutoWidth()
 							[
 								SNew(STextBlock)
-								.Text(LOCTEXT("SkipMipsLabel", "Skip mips on generate :"))
+								.Text(LOCTEXT("SkipMipsLabel", "Skip mips on generate "))
 								.Visibility(this, &SMutableCodeViewer::IsMipSkipVisible)
 							]
 							+ SHorizontalBox::Slot()
+							.Padding(4, 0, 0, 0)
 							[
 								SNew(SNumericEntryBox<int32>)
 								.Visibility(this, &SMutableCodeViewer::IsMipSkipVisible)
@@ -621,20 +618,12 @@ void SMutableCodeViewer::Construct(const FArguments& InArgs,
 								.OnValueChanged(this, &SMutableCodeViewer::OnCurrentMipSkipChanged)
 							]
 						]
-
-						+ SVerticalBox::Slot()
-						[
-							SNew(SScrollBox)
-							+ SScrollBox::Slot()
-							[
-								SAssignNew(ParametersWidget, SMutableParametersWidget)
-								.OnParametersValueChanged(this, &SMutableCodeViewer::OnPreviewParameterValueChanged)
-								.Parameters(PreviewParameters)
-							]
-						]
 					]
 
-					+ SSplitter::Slot()
+					+ SVerticalBox::Slot()
+					.AutoHeight()
+					.Padding(8, 8, 0, 0)
+					.VAlign(VAlign_Top)
 					[
 						// Generate a new Constants panel to show the data stored on the current mutable program
 						SAssignNew(ConstantsWidget, SMutableConstantsWidget,
@@ -662,7 +651,16 @@ void SMutableCodeViewer::Construct(const FArguments& InArgs,
 	// how the user is point to interact with the view.
 	bShouldRecalculateStates = true;
 	// Now, on expansion or contraction the states will get recalculated
-} 
+}
+
+
+FReply SMutableCodeViewer::UpdateParameters()
+{
+	PreviewParameters = Instance->GetPrivate()->GetDescriptor().GetParameters();
+	bIsPreviewPendingUpdate = true;
+
+	return FReply::Handled();
+}
 
 
 EVisibility SMutableCodeViewer::IsMipSkipVisible() const 
@@ -2825,7 +2823,7 @@ void SMutableCodeViewer::Tick(const FGeometry& AllottedGeometry, const double In
 	mu::FSettings Settings;
 	const TSharedPtr<mu::FSystem> System = MakeShared<mu::FSystem>(Settings);
 	
-	System->GetPrivate()->BeginBuild(MutableModel, MeshIdRegistry, ImageIdRegistry, MaterialIdRegistry);
+	System->GetPrivate()->BeginBuild(MutableModel, Object->GetPrivate()->MeshIdRegistry, Object->GetPrivate()->ImageIdRegistry, Object->GetPrivate()->MaterialIdRegistry);
 
 	switch (OperationDataType)
 	{
diff --git a/Engine/Plugins/Mutable/Source/CustomizableObjectEditor/Private/MuCOE/SMutableCodeViewer.h b/Engine/Plugins/Mutable/Source/CustomizableObjectEditor/Private/MuCOE/SMutableCodeViewer.h
--- a/Engine/Plugins/Mutable/Source/CustomizableObjectEditor/Private/MuCOE/SMutableCodeViewer.h
+++ b/Engine/Plugins/Mutable/Source/CustomizableObjectEditor/Private/MuCOE/SMutableCodeViewer.h
@@ -11,6 +11,8 @@
 #include "MuR/Mesh.h"
 #include "UObject/SoftObjectPtr.h"
 
+class UCustomizableObjectInstance;
+class UCustomizableObject;
 class FUnrealMutableResourceProvider;
 class STableViewBase;
 namespace ESelectInfo { enum Type : int; }
@@ -46,43 +48,30 @@ struct FGeometry;
 /** This widget shows the internal Mutable Code for debugging purposes. 
  * This is not the Unreal source graph in the UCustomizableObject, but the actual Mutable virtual machine graph.
  */
-class SMutableCodeViewer :
-	public SCompoundWidget,
-	public FGCObject
+class SMutableCodeViewer : public SCompoundWidget
 {
 public:
 
 	SLATE_BEGIN_ARGS(SMutableCodeViewer) {}
 
-	/** User-visible tag to identify the source of the data shown. */
-	SLATE_ARGUMENT(FString, DataTag)
-	SLATE_ARGUMENT(TSharedPtr<mu::FMeshIdRegistry>, MeshIdRegistry)
-	SLATE_ARGUMENT(TSharedPtr<mu::FImageIdRegistry>, ImageIdRegistry)
-	SLATE_ARGUMENT(TSharedPtr<mu::FMaterialIdRegistry>, MaterialIdRegistry)
 	SLATE_ARGUMENT(TSharedPtr<FUnrealMutableResourceProvider>, ExternalResourceProvider)
 		
 	SLATE_END_ARGS()
 
-	void Construct(const FArguments& InArgs,
-		const TSharedPtr<mu::FModel>& InMutableModel);
+	void Construct(const FArguments& InArgs, UCustomizableObject& InObject, UCustomizableObjectInstance& InInstance, const TSharedPtr<mu::FModel>& InMutableModel);
 
 	// SWidget interface
 	virtual void Tick(const FGeometry& AllottedGeometry, const double InCurrentTime, const float InDeltaTime) override;
-
-	// FGCObject interface
-	virtual void AddReferencedObjects(FReferenceCollector& Collector) override;
-	virtual FString GetReferencerName() const override;
-
+	
 	/** Clear the selected operation of the tree view.*/
 	void ClearSelectedTreeRow() const;
 	
 private:
-
+	TStrongObjectPtr<UCustomizableObject> Object;
+	TStrongObjectPtr<UCustomizableObjectInstance> Instance;
+	
 	/** The Mutable Model that we are showing. */
 	TSharedPtr<mu::FModel> MutableModel;
-	TSharedPtr<mu::FMeshIdRegistry> MeshIdRegistry;
-	TSharedPtr<mu::FImageIdRegistry> ImageIdRegistry;
-	TSharedPtr<mu::FMaterialIdRegistry> MaterialIdRegistry;
 
 	TSharedPtr<FUnrealMutableResourceProvider> ExternalResourceProvider;
 	
@@ -786,6 +775,8 @@ private:
 	/** */
 	bool bSelectedOperationIsImage = false;
 
+	FReply UpdateParameters();
+	
 	/** Stores a list of strings based on the possible types of operations mutable define to be used by the UI */
 	EVisibility IsMipSkipVisible() const;
 
diff --git a/Engine/Plugins/Mutable/Source/CustomizableObjectEditor/Private/MuCOE/SMutableGraphViewer.cpp b/Engine/Plugins/Mutable/Source/CustomizableObjectEditor/Private/MuCOE/SMutableGraphViewer.cpp
--- a/Engine/Plugins/Mutable/Source/CustomizableObjectEditor/Private/MuCOE/SMutableGraphViewer.cpp
+++ b/Engine/Plugins/Mutable/Source/CustomizableObjectEditor/Private/MuCOE/SMutableGraphViewer.cpp
@@ -296,65 +296,44 @@ FString SMutableGraphViewer::GetReferencerName() const
 
 void SMutableGraphViewer::Construct(const FArguments& InArgs, const mu::NodePtr& InRootNode)
 {
-	DataTag = InArgs._DataTag;
 	ReferencedRuntimeTextures = InArgs._ReferencedRuntimeTextures;
 	ReferencedCompileTextures = InArgs._ReferencedCompileTextures;
 	RootNode = InRootNode;
-
-	FToolBarBuilder ToolbarBuilder(TSharedPtr<const FUICommandList>(), FMultiBoxCustomization::None, TSharedPtr<FExtender>(), true);
-	ToolbarBuilder.SetLabelVisibility(EVisibility::Visible);
-	ToolbarBuilder.SetStyle(&FAppStyle::Get(), "SlimToolBar");
-
-	ToolbarBuilder.AddWidget(SNew(STextBlock).Text(MakeAttributeLambda([this]() { return FText::FromString(DataTag); })));
-
+	  
 	ChildSlot
 	[
-		SNew(SVerticalBox)
-		+ SVerticalBox::Slot()
-		.AutoHeight()
-		.VAlign(VAlign_Center)
-		[
-			ToolbarBuilder.MakeWidget()
-		]
-		+ SVerticalBox::Slot()
-		.VAlign(VAlign_Fill)
+		SNew(SSplitter)
+		.Orientation(EOrientation::Orient_Horizontal)
+		+ SSplitter::Slot()
+		.Value(0.25f)
 		[
-			SNew(SSplitter)
-			.Orientation(EOrientation::Orient_Horizontal)
-			+ SSplitter::Slot()
-			.Value(0.25f)
+			SNew(SBorder)
+			.BorderImage(UE_MUTABLE_GET_BRUSH("ToolPanel.GroupBorder"))
+			.Padding(FMargin(4.0f, 4.0f))
 			[
-				SNew(SBorder)
-				.BorderImage(UE_MUTABLE_GET_BRUSH("ToolPanel.GroupBorder"))
-				.Padding(FMargin(4.0f, 4.0f))
-				[
-					SAssignNew(TreeView, STreeView<TSharedPtr<FMutableGraphTreeElement>>)
-					.TreeItemsSource(&RootNodes)
-					.OnGenerateRow(this,&SMutableGraphViewer::GenerateRowForNodeTree)
-					.OnGetChildren(this, &SMutableGraphViewer::GetChildrenForInfo)
-					.OnSetExpansionRecursive(this, &SMutableGraphViewer::TreeExpandRecursive)
-					.OnContextMenuOpening(this, &SMutableGraphViewer::OnTreeContextMenuOpening)
-					.SelectionMode(ESelectionMode::Single)
-					.HeaderRow
-					(
-						SNew(SHeaderRow)
-						+ SHeaderRow::Column(MutableGraphTreeViewColumns::Name)
-						.FillWidth(25.f)
-						.DefaultLabel(LOCTEXT("Node Name", "Node Name"))
-					)
-				]
-			]
-			+ SSplitter::Slot()
-			.Value(0.75f)
-			[
-				SNew(SBorder)
-				.BorderImage(UE_MUTABLE_GET_BRUSH("ToolPanel.GroupBorder"))
-				.Padding(FMargin(4.0f, 4.0f))
-				//[
-				//	SubjectsTreeView->AsShared()
-				//]
+				SAssignNew(TreeView, STreeView<TSharedPtr<FMutableGraphTreeElement>>)
+				.TreeItemsSource(&RootNodes)
+				.OnGenerateRow(this,&SMutableGraphViewer::GenerateRowForNodeTree)
+				.OnGetChildren(this, &SMutableGraphViewer::GetChildrenForInfo)
+				.OnSetExpansionRecursive(this, &SMutableGraphViewer::TreeExpandRecursive)
+				.OnContextMenuOpening(this, &SMutableGraphViewer::OnTreeContextMenuOpening)
+				.SelectionMode(ESelectionMode::Single)
+				.HeaderRow
+				(
+					SNew(SHeaderRow)
+					+ SHeaderRow::Column(MutableGraphTreeViewColumns::Name)
+					.FillWidth(25.f)
+					.DefaultLabel(LOCTEXT("Node Name", "Node Name"))
+				)
 			]
 		]
+		+ SSplitter::Slot()
+		.Value(0.75f)
+		[
+			SNew(SBorder)
+			.BorderImage(UE_MUTABLE_GET_BRUSH("ToolPanel.GroupBorder"))
+			.Padding(FMargin(4.0f, 4.0f))
+		]
 	];
 	
 	RebuildTree();
diff --git a/Engine/Plugins/Mutable/Source/CustomizableObjectEditor/Private/MuCOE/SMutableGraphViewer.h b/Engine/Plugins/Mutable/Source/CustomizableObjectEditor/Private/MuCOE/SMutableGraphViewer.h
--- a/Engine/Plugins/Mutable/Source/CustomizableObjectEditor/Private/MuCOE/SMutableGraphViewer.h
+++ b/Engine/Plugins/Mutable/Source/CustomizableObjectEditor/Private/MuCOE/SMutableGraphViewer.h
@@ -31,10 +31,7 @@ public:
 
 	// SWidget interface
 	SLATE_BEGIN_ARGS(SMutableGraphViewer) {}
-
-	/** User-visible tag to indentify the source of the data shown. */
-	SLATE_ARGUMENT(FString, DataTag)
-
+		
 	SLATE_ARGUMENT(TArray<TSoftObjectPtr<const UTexture>>, ReferencedRuntimeTextures)
 	SLATE_ARGUMENT(TArray<FMutableSourceTextureData>, ReferencedCompileTextures)
 	SLATE_ARGUMENT(TArray<TSoftObjectPtr<const UStreamableRenderAsset>>, ReferencedRuntimeMeshes)
@@ -49,10 +46,6 @@ public:
 	virtual FString GetReferencerName() const override;
 
 private:
-
-	/** Tag for the user to identify the data being shown. */
-	FString DataTag;
-
 	/** The Mutable Graph to show, represented by its root. */
 	mu::NodePtr RootNode;
 
diff --git a/Engine/Plugins/Mutable/Source/CustomizableObjectEditor/Private/MuCOE/SMutableImageViewer.cpp b/Engine/Plugins/Mutable/Source/CustomizableObjectEditor/Private/MuCOE/SMutableImageViewer.cpp
--- a/Engine/Plugins/Mutable/Source/CustomizableObjectEditor/Private/MuCOE/SMutableImageViewer.cpp
+++ b/Engine/Plugins/Mutable/Source/CustomizableObjectEditor/Private/MuCOE/SMutableImageViewer.cpp
@@ -180,57 +180,61 @@ void SMutableImageViewer::Construct(const FArguments& InArgs)
 		.Padding(4.0f)
 		.AutoHeight()
 		[
-			SNew(SHorizontalBox)
-			+ SHorizontalBox::Slot()
-			.AutoWidth()
-			.VAlign(VAlign_Center)
-			.Padding(FMargin(4.0f, 4.0f, 4.0f, 4.0f))
+			SNew(SVerticalBox)
+			+ SVerticalBox::Slot()
+			.AutoHeight()
 			[
 				SNew(STextBlock)
 				.Text(this, &SMutableImageViewer::GetImageDescriptionLabel)
 			]
-			+ SHorizontalBox::Slot()
-			.AutoWidth()
-			.VAlign(VAlign_Center)
-			.Padding(FMargin(4.0f, 4.0f, 0.0f, 4.0f))
-			[
-				SNew(STextBlock)
-				.Text(LOCTEXT("ShowLOD","Show LOD: "))
-			]
-			+ SHorizontalBox::Slot()
-			.FillWidth(1.0f)
-			.VAlign(VAlign_Center)
-			.Padding(FMargin(0.0f, 4.0f, 4.0f, 4.0f))
-			[
-				SNew(SNumericEntryBox<int32>)
-				.Visibility(this, &SMutableImageViewer::IsLODSelectionVisible)
-				.AllowSpin(true)
-				.MinValue(0)
-				.MaxValue(this, &SMutableImageViewer::GetImageLODMaxValue)
-				.MinSliderValue(0)
-				.MaxSliderValue(this, &SMutableImageViewer::GetImageLODMaxValue)
-				.Value(this, &SMutableImageViewer::GetCurrentImageLOD)
-				.OnValueChanged(this, &SMutableImageViewer::OnCurrentLODChanged)
-			]
-			+ SHorizontalBox::Slot()
-			.AutoWidth()
-			.Padding(FMargin(4.0f, 4.0f, 4.0f, 4.0f))
+
+			+ SVerticalBox::Slot()
+			.AutoHeight()
 			[
-				SNew(SSegmentedControl<EMutableImageChannels>)
-				.Value(this, &SMutableImageViewer::GetImageChannels)
-				.OnValueChanged(this, &SMutableImageViewer::SetImageChannels)
-				+ SSegmentedControl<EMutableImageChannels>::Slot(EMutableImageChannels::RGBA)
-				.Text(FText::FromString(TEXT("RGBA")))
-				+ SSegmentedControl<EMutableImageChannels>::Slot(EMutableImageChannels::RGB)
-				.Text(FText::FromString(TEXT("RGB")))
-				+ SSegmentedControl<EMutableImageChannels>::Slot(EMutableImageChannels::R)
-				.Text(FText::FromString(TEXT("R")))
-				+ SSegmentedControl<EMutableImageChannels>::Slot(EMutableImageChannels::G)
-				.Text(FText::FromString(TEXT("G")))
-				+ SSegmentedControl<EMutableImageChannels>::Slot(EMutableImageChannels::B)
-				.Text(FText::FromString(TEXT("B")))
-				+ SSegmentedControl<EMutableImageChannels>::Slot(EMutableImageChannels::A)
-				.Text(FText::FromString(TEXT("A")))
+				SNew(SHorizontalBox)
+				+ SHorizontalBox::Slot()
+				.AutoWidth()
+				.VAlign(VAlign_Center)
+				.Padding(FMargin(0, 4.0f, 0.0f, 4.0f))
+				[
+					SNew(STextBlock)
+					.Text(LOCTEXT("ShowLOD","Show LOD "))
+				]
+				+ SHorizontalBox::Slot()
+				.FillWidth(1.0f)
+				.VAlign(VAlign_Center)
+				.Padding(FMargin(0.0f, 4.0f, 4.0f, 4.0f))
+				[
+					SNew(SNumericEntryBox<int32>)
+					.Visibility(this, &SMutableImageViewer::IsLODSelectionVisible)
+					.AllowSpin(true)
+					.MinValue(0)
+					.MaxValue(this, &SMutableImageViewer::GetImageLODMaxValue)
+					.MinSliderValue(0)
+					.MaxSliderValue(this, &SMutableImageViewer::GetImageLODMaxValue)
+					.Value(this, &SMutableImageViewer::GetCurrentImageLOD)
+					.OnValueChanged(this, &SMutableImageViewer::OnCurrentLODChanged)
+				]
+				+ SHorizontalBox::Slot()
+				.AutoWidth()
+				.Padding(FMargin(4.0f, 4.0f, 4.0f, 4.0f))
+				[
+					SNew(SSegmentedControl<EMutableImageChannels>)
+					.Value(this, &SMutableImageViewer::GetImageChannels)
+					.OnValueChanged(this, &SMutableImageViewer::SetImageChannels)
+					+ SSegmentedControl<EMutableImageChannels>::Slot(EMutableImageChannels::RGBA)
+					.Text(FText::FromString(TEXT("RGBA")))
+					+ SSegmentedControl<EMutableImageChannels>::Slot(EMutableImageChannels::RGB)
+					.Text(FText::FromString(TEXT("RGB")))
+					+ SSegmentedControl<EMutableImageChannels>::Slot(EMutableImageChannels::R)
+					.Text(FText::FromString(TEXT("R")))
+					+ SSegmentedControl<EMutableImageChannels>::Slot(EMutableImageChannels::G)
+					.Text(FText::FromString(TEXT("G")))
+					+ SSegmentedControl<EMutableImageChannels>::Slot(EMutableImageChannels::B)
+					.Text(FText::FromString(TEXT("B")))
+					+ SSegmentedControl<EMutableImageChannels>::Slot(EMutableImageChannels::A)
+					.Text(FText::FromString(TEXT("A")))
+				]
 			]
 		]
 
diff --git a/Engine/Plugins/Mutable/Source/CustomizableObjectEditor/Private/MuCOE/SMutableMeshViewer.cpp b/Engine/Plugins/Mutable/Source/CustomizableObjectEditor/Private/MuCOE/SMutableMeshViewer.cpp
--- a/Engine/Plugins/Mutable/Source/CustomizableObjectEditor/Private/MuCOE/SMutableMeshViewer.cpp
+++ b/Engine/Plugins/Mutable/Source/CustomizableObjectEditor/Private/MuCOE/SMutableMeshViewer.cpp
@@ -178,7 +178,7 @@ void SMutableMeshViewer::Construct(const FArguments& InArgs)
 	ChildSlot
 	[
 		SNew(SSplitter)
-
+		.Orientation(Orient_Vertical)
 		+ SSplitter::Slot()
 		.Value(TablesSplitterValue)
 		[
@@ -282,8 +282,11 @@ TSharedRef<SWidget> SMutableMeshViewer::GenerateDataTableSlates()
 	const FText MeshIdPrefixTitle = LOCTEXT("MeshIdPrefixTitle", "Mesh ID prefix : ");
 	const FText MeshFlagsTitle = LOCTEXT("MeshFlagsTitle", "Mesh flags : ");
 	const FText BuffersTitle = LOCTEXT("BuffersTitle", "Buffers");
-	
-	return SNew(SScrollBox)
+
+	return SNew(SSplitter)
+	+ SSplitter::Slot()
+	[
+		 SNew(SScrollBox)
 		+ SScrollBox::Slot()
 		[
 			SNew(SVerticalBox)
@@ -498,6 +501,15 @@ TSharedRef<SWidget> SMutableMeshViewer::GenerateDataTableSlates()
 					.OnGenerateRow(this, &SMutableMeshViewer::GenerateTagRow)
 				]
 			// ---------------------------------
+		]
+	]
+	
+	+ SSplitter::Slot()
+	[
+		SNew(SScrollBox)
+	   + SScrollBox::Slot()
+		[
+			SNew(SVerticalBox)
 
 			// UVs
 			+SVerticalBox::Slot()
@@ -519,7 +531,8 @@ TSharedRef<SWidget> SMutableMeshViewer::GenerateDataTableSlates()
 					.UVLayout(this, &SMutableMeshViewer::GetUVs)
 			]
 			// ---------------------------------
-		];
+		]
+	];
 }
 
 
diff --git a/Engine/Plugins/Mutable/Source/CustomizableObjectEditor/Private/MuCOE/SMutableObjectViewer.cpp b/Engine/Plugins/Mutable/Source/CustomizableObjectEditor/Private/MuCOE/SMutableObjectViewer.cpp
deleted file mode 100644
diff --git a/Engine/Plugins/Mutable/Source/CustomizableObjectEditor/Private/MuCOE/SMutableObjectViewer.h b/Engine/Plugins/Mutable/Source/CustomizableObjectEditor/Private/MuCOE/SMutableObjectViewer.h
deleted file mode 100644
diff --git a/Engine/Plugins/Mutable/Source/CustomizableObjectEditor/Private/MuCOE/SMutableParametersWidget.cpp b/Engine/Plugins/Mutable/Source/CustomizableObjectEditor/Private/MuCOE/SMutableParametersWidget.cpp
deleted file mode 100644
diff --git a/Engine/Plugins/Mutable/Source/CustomizableObjectEditor/Private/MuCOE/SMutableParametersWidget.h b/Engine/Plugins/Mutable/Source/CustomizableObjectEditor/Private/MuCOE/SMutableParametersWidget.h
deleted file mode 100644
