[Mutable] Implemented the new node Break Material.
+ Added the new concept of lazy branches into the mutable core.
+ Created the new node interface and widgets.

#rnx

diff --git a/Engine/Plugins/Mutable/Source/CustomizableObject/Internal/MuCO/CustomizableObjectPrivate.h b/Engine/Plugins/Mutable/Source/CustomizableObject/Internal/MuCO/CustomizableObjectPrivate.h
--- a/Engine/Plugins/Mutable/Source/CustomizableObject/Internal/MuCO/CustomizableObjectPrivate.h
+++ b/Engine/Plugins/Mutable/Source/CustomizableObject/Internal/MuCO/CustomizableObjectPrivate.h
@@ -1689,6 +1689,8 @@ enum class ECustomizableObjectVersions : int32
 
 	CacheReferenceTextureSize,
 
+	MaterialBreak,
+
 	// -----<new versions can be added above this line>--------
 	LastCustomizableObjectVersion
 };
diff --git a/Engine/Plugins/Mutable/Source/CustomizableObjectEditor/Private/MuCOE/COVariable.cpp b/Engine/Plugins/Mutable/Source/CustomizableObjectEditor/Private/MuCOE/COVariable.cpp
new file mode 100644
--- /dev/null
+++ b/Engine/Plugins/Mutable/Source/CustomizableObjectEditor/Private/MuCOE/COVariable.cpp
@@ -0,0 +1,161 @@
+// Copyright Epic Games, Inc. All Rights Reserved.
+
+#include "COVariable.h"
+
+#include "DetailLayoutBuilder.h"
+#include "MuCOE/EdGraphSchema_CustomizableObject.h"
+#include "MuCOE/Nodes/CustomizableObjectNode.h"
+#include "IDetailsView.h"
+#include "PropertyCustomizationHelpers.h"
+#include "Widgets/Input/SComboBox.h"
+#include "Widgets/Input/SEditableTextBox.h"
+#include "Widgets/SWidget.h"
+
+
+TSharedRef<IPropertyTypeCustomization> FCOVariableCustomzation::MakeInstance()
+{
+	return MakeShared<FCOVariableCustomzation>();
+}
+
+
+void FCOVariableCustomzation::CustomizeHeader(TSharedRef<IPropertyHandle> StructPropertyHandle, FDetailWidgetRow& InHeaderRow, IPropertyTypeCustomizationUtils& InCustomizationUtils)
+{
+	// Get the property handles to edit the variable values
+	NamePropertyHandle = StructPropertyHandle->GetChildHandle("Name");
+	TypePropertyHandle = StructPropertyHandle->GetChildHandle("Type");
+	check(NamePropertyHandle && TypePropertyHandle);
+
+	//Get parent node if it's defined inside one
+	TArray<UObject*> OuterObjects;
+	StructPropertyHandle->GetOuterObjects(OuterObjects);
+
+	if (OuterObjects[0])
+	{
+		BaseObjectNode = Cast<UCustomizableObjectNode>(OuterObjects[0]->GetTypedOuter(UCustomizableObjectNode::StaticClass()));
+
+		//You can add more types here. This class is not meant to work just with nodes.
+	}
+
+	// Get which type of pins we can create
+	if (BaseObjectNode.IsValid())
+	{
+		AllowedVariableTypes = BaseObjectNode->GetAllowedPinViewerCreationTypes();
+	}
+	else
+	{
+		AllowedVariableTypes = UEdGraphSchema_CustomizableObject::SupportedMacroPinTypes;
+	}
+
+	InHeaderRow.NameContent()
+	[
+		SNew(SEditableTextBox)
+		.Text(this, &FCOVariableCustomzation::GetVariableName)
+		.OnTextChanged(this, &FCOVariableCustomzation::OnVariableNameChanged)
+		.OnTextCommitted(this, &FCOVariableCustomzation::OnVariableNameCommitted)
+		.Font(IDetailLayoutBuilder::GetDetailFont())
+	]
+	.ValueContent()
+	[
+		// Combo button that summons the dropdown menu
+		SNew(SComboBox<FName>)
+		.OptionsSource(&AllowedVariableTypes)
+		.OnGenerateWidget(this, &FCOVariableCustomzation::OnGenerateVariableTypeRow)
+		.OnSelectionChanged_Lambda
+		([this](FName Item, ESelectInfo::Type)
+			{
+				TypePropertyHandle->SetValue(Item);
+
+				if (BaseObjectNode.IsValid())
+				{
+					BaseObjectNode->ReconstructNode();
+				}
+			}
+		)
+		.Content()
+		[
+			GenerateCurrentSelectectedTypeWidget()
+		]
+	]
+	.OverrideResetToDefault(FResetToDefaultOverride::Create(FSimpleDelegate::CreateSP(this, &FCOVariableCustomzation::ResetToDefault)));
+}
+
+
+FText FCOVariableCustomzation::GetVariableName() const
+{
+	FName VariableName;
+	NamePropertyHandle->GetValue(VariableName);
+
+	return FText::FromName(VariableName);
+}
+
+
+void FCOVariableCustomzation::ResetToDefault()
+{
+	NamePropertyHandle->SetValue(FName("NewVar"));
+	TypePropertyHandle->SetValue(AllowedVariableTypes[0]);
+}
+
+
+void FCOVariableCustomzation::OnVariableNameChanged(const FText& InNewText)
+{
+	FName NewName = FName(InNewText.ToString());
+
+	if (InNewText.IsEmpty() || NewName.IsNone())
+	{
+		return;
+	}
+
+	NamePropertyHandle->SetValue(NewName);
+}
+
+
+void FCOVariableCustomzation::OnVariableNameCommitted(const FText& InNewText, ETextCommit::Type InTextCommit)
+{
+	FName NewName = FName(InNewText.ToString());
+
+	if (InNewText.IsEmpty() || NewName.IsNone())
+	{
+		return;
+	}
+
+	NamePropertyHandle->SetValue(NewName);
+
+	if (BaseObjectNode.IsValid())
+	{
+		BaseObjectNode->ReconstructNode();
+	}
+}
+
+
+TSharedRef<SWidget> FCOVariableCustomzation::OnGenerateVariableTypeRow(FName Type)
+{
+	const FSlateBrush* IconBrush = FAppStyle::GetBrush(TEXT("Kismet.VariableList.TypeIcon"));
+
+	TSharedPtr<SHorizontalBox> RowWidget = SNew(SHorizontalBox)
+		+ SHorizontalBox::Slot()
+		.AutoWidth()
+		[
+			SNew(SImage)
+				.Image(IconBrush)
+				.ColorAndOpacity(UEdGraphSchema_CustomizableObject::GetPinTypeColor(Type))
+		]
+		+ SHorizontalBox::Slot()
+		.AutoWidth()
+		.Padding(7.5f, 0.0f, 0.0f, 0.0f)
+		[
+			SNew(STextBlock)
+				.Text(Type.IsValid() ? UEdGraphSchema_CustomizableObject::GetPinCategoryFriendlyName(Type) : FText::FromString("Invalid"))
+				.Font(IDetailLayoutBuilder::GetDetailFont())
+		];
+
+	return RowWidget.ToSharedRef();
+}
+
+
+TSharedRef<SWidget> FCOVariableCustomzation::GenerateCurrentSelectectedTypeWidget()
+{
+	FName CurrentType;
+	TypePropertyHandle->GetValue(CurrentType);
+
+	return OnGenerateVariableTypeRow(CurrentType);
+}
diff --git a/Engine/Plugins/Mutable/Source/CustomizableObjectEditor/Private/MuCOE/COVariable.h b/Engine/Plugins/Mutable/Source/CustomizableObjectEditor/Private/MuCOE/COVariable.h
new file mode 100644
--- /dev/null
+++ b/Engine/Plugins/Mutable/Source/CustomizableObjectEditor/Private/MuCOE/COVariable.h
@@ -0,0 +1,72 @@
+// Copyright Epic Games, Inc. All Rights Reserved.
+
+#pragma once
+
+#include "IPropertyTypeCustomization.h"
+
+#include "COVariable.generated.h"
+
+
+class UCustomizableObjectNode;
+class SWidget;
+
+namespace ETextCommit { enum Type : int; }
+
+/** This struct can beused to add new variables to nodes. */
+
+USTRUCT()
+struct FCOVariable
+{
+	GENERATED_USTRUCT_BODY()
+
+	/** Name of this variable. Can be eddited from any details. */
+	UPROPERTY(EditAnywhere, Category = CustomizableObject)
+	FName Name;
+
+	/** The type of this variable. Ususally a pin category */
+	UPROPERTY(EditAnywhere, Category = CustomizableObject)
+	FName Type;
+
+	/** Id to identify this variable. */
+	UPROPERTY()
+	FGuid Id;
+
+};
+
+
+class FCOVariableCustomzation: public IPropertyTypeCustomization
+{
+public:
+	static TSharedRef<IPropertyTypeCustomization> MakeInstance();
+
+	// IPropertyTypeCustomization interface
+	void CustomizeHeader(TSharedRef<IPropertyHandle> StructPropertyHandle, class FDetailWidgetRow& HeaderRow, IPropertyTypeCustomizationUtils& StructCustomizationUtils) override;
+	void CustomizeChildren(TSharedRef<IPropertyHandle> StructPropertyHandle, class IDetailChildrenBuilder& ChildBuilder, IPropertyTypeCustomizationUtils& StructCustomizationUtils) override {};
+
+private:
+
+	/** Returns the name of the specified variable. Needed to update automatically the text if the name of the variable changes. */
+	FText GetVariableName() const;
+
+	void ResetToDefault();
+
+	/** Editable Text Box Callbacks */
+	void OnVariableNameChanged(const FText& InNewText);
+	void OnVariableNameCommitted(const FText& InNewText, ETextCommit::Type InTextCommit);
+
+	/** ComboBox Callbacks */
+	TSharedRef<SWidget> OnGenerateVariableTypeRow(FName Option);
+	TSharedRef<SWidget> GenerateCurrentSelectectedTypeWidget();
+
+private:
+
+	/** Weak pointer to the Customizable Object node that contains this property */
+	TWeakObjectPtr<UCustomizableObjectNode> BaseObjectNode;
+
+	TArray<FName> AllowedVariableTypes;
+
+	/** Runtime Parameter Name property of a State */
+	TSharedPtr<IPropertyHandle> NamePropertyHandle;
+	TSharedPtr<IPropertyHandle> TypePropertyHandle;
+
+};
\ No newline at end of file
diff --git a/Engine/Plugins/Mutable/Source/CustomizableObjectEditor/Private/MuCOE/CustomizableObjectEditorModule.cpp b/Engine/Plugins/Mutable/Source/CustomizableObjectEditor/Private/MuCOE/CustomizableObjectEditorModule.cpp
--- a/Engine/Plugins/Mutable/Source/CustomizableObjectEditor/Private/MuCOE/CustomizableObjectEditorModule.cpp
+++ b/Engine/Plugins/Mutable/Source/CustomizableObjectEditor/Private/MuCOE/CustomizableObjectEditorModule.cpp
@@ -32,6 +32,7 @@
 #include "MuCOE/CustomizableObjectMacroLibrary/CustomizableObjectMacroDetails.h"
 #include "MuCOE/CustomizableObjectVersionBridge.h"
 #include "MuCOE/CustomizableObjectNodeObjectGroupDetails.h"
+#include "MuCOE/Nodes/CONodeMaterialBreak.h"
 #include "MuCOE/Nodes/CustomizableObjectNodeCopyMaterial.h"
 #include "MuCOE/Nodes/CustomizableObjectNodeExternalPin.h"
 #include "MuCOE/Nodes/CustomizableObjectNodeExternalPinDetails.h"
@@ -88,6 +89,7 @@
 #include "MuCOE/Widgets/CustomizableObjectVariationCustomization.h"
 #include "MuCOE/Widgets/CustomizableObjectLODReductionSettings.h"
 #include "MuCOE/Widgets/CustomizableObjectNodeTableCompilationFilterEditor.h"
+#include "MuCOE/COVariable.h"
 #include "PropertyEditorModule.h"
 #include "AssetRegistry/AssetRegistryModule.h"
 #include "MuCO/CustomizableObjectInstance.h"
@@ -275,6 +277,7 @@ void FCustomizableObjectEditorModule::StartupModule()
 	RegisterCustomDetails(PropertyModule, UCustomizableObjectNodeComponentMesh::StaticClass(), FOnGetDetailCustomizationInstance::CreateStatic(&FCustomizableObjectNodeComponentMeshDetails::MakeInstance));
 	RegisterCustomDetails(PropertyModule, UCustomizableObjectNodeMacroInstance::StaticClass(), FOnGetDetailCustomizationInstance::CreateStatic(&FCustomizableObjectNodeMacroInstanceDetails::MakeInstance));
 	RegisterCustomDetails(PropertyModule, UCustomizableObjectNodeParameter::StaticClass(), FOnGetDetailCustomizationInstance::CreateStatic(&FCustomizableObjectNodeParameterDetails::MakeInstance));
+	RegisterCustomDetails(PropertyModule, UCONodeMaterialBreak::StaticClass(), FOnGetDetailCustomizationInstance::CreateStatic(&FCustomizableObjectNodeDetails::MakeInstance));
 
 	// Other Objects
 	RegisterCustomDetails(PropertyModule, UCustomizableObject::StaticClass(), FOnGetDetailCustomizationInstance::CreateStatic(&FCustomizableObjectDetails::MakeInstance));
@@ -294,6 +297,7 @@ void FCustomizableObjectEditorModule::StartupModule()
 	PropertyModule.RegisterCustomPropertyTypeLayout(FCustomizableObjectVariation::StaticStruct()->GetFName(), FOnGetPropertyTypeCustomizationInstance::CreateStatic(&FCustomizableObjectVariationCustomization::MakeInstance));
 	PropertyModule.RegisterCustomPropertyTypeLayout(FCustomizableObjectTextureVariation::StaticStruct()->GetFName(), FOnGetPropertyTypeCustomizationInstance::CreateStatic(&FCustomizableObjectVariationCustomization::MakeInstance));
 	PropertyModule.RegisterCustomPropertyTypeLayout(FTableNodeCompilationFilter::StaticStruct()->GetFName(), FOnGetPropertyTypeCustomizationInstance::CreateStatic(&FCustomizableObjectNodeTableCompilationFilterEditor::MakeInstance));
+	PropertyModule.RegisterCustomPropertyTypeLayout(FCOVariable::StaticStruct()->GetFName(), FOnGetPropertyTypeCustomizationInstance::CreateStatic(&FCOVariableCustomzation::MakeInstance));
 
 	PropertyModule.NotifyCustomizationModuleChanged();
 
diff --git a/Engine/Plugins/Mutable/Source/CustomizableObjectEditor/Private/MuCOE/EdGraphSchema_CustomizableObject.cpp b/Engine/Plugins/Mutable/Source/CustomizableObjectEditor/Private/MuCOE/EdGraphSchema_CustomizableObject.cpp
--- a/Engine/Plugins/Mutable/Source/CustomizableObjectEditor/Private/MuCOE/EdGraphSchema_CustomizableObject.cpp
+++ b/Engine/Plugins/Mutable/Source/CustomizableObjectEditor/Private/MuCOE/EdGraphSchema_CustomizableObject.cpp
@@ -17,6 +17,7 @@
 #include "MuCOE/CustomizableObjectGraph.h"
 #include "MuCOE/CustomizableObjectMacroLibrary/CustomizableObjectGraphEditorToolkit.h"
 #include "MuCOE/CustomizableObjectEditor.h"
+#include "MuCOE/Nodes/CONodeMaterialBreak.h"
 #include "MuCOE/Nodes/CONodeMaterialConstant.h"
 #include "MuCOE/Nodes/CONodeMaterialVariation.h"
 #include "MuCOE/Nodes/CONodeMaterialSwitch.h"
@@ -632,7 +633,8 @@ void UEdGraphSchema_CustomizableObject::GetGraphContextActions(FGraphContextMenu
 		{
 			ContextMenuBuilder.CreateTemplateNode<UCONodeMaterialConstant>(),
 			ContextMenuBuilder.CreateTemplateNode<UCONodeMaterialVariation>(),
-			ContextMenuBuilder.CreateTemplateNode<UCONodeMaterialSwitch>()
+			ContextMenuBuilder.CreateTemplateNode<UCONodeMaterialSwitch>(),
+			ContextMenuBuilder.CreateTemplateNode<UCONodeMaterialBreak>()
 		};
 
 		AddNewNodeCategoryActionsFiltered(MaterialTemplateNodes, ContextMenuBuilder, NC_Material.ToString(), 2, Filter);
diff --git a/Engine/Plugins/Mutable/Source/CustomizableObjectEditor/Private/MuCOE/GenerateMutableSource/GenerateMutableSource.h b/Engine/Plugins/Mutable/Source/CustomizableObjectEditor/Private/MuCOE/GenerateMutableSource/GenerateMutableSource.h
--- a/Engine/Plugins/Mutable/Source/CustomizableObjectEditor/Private/MuCOE/GenerateMutableSource/GenerateMutableSource.h
+++ b/Engine/Plugins/Mutable/Source/CustomizableObjectEditor/Private/MuCOE/GenerateMutableSource/GenerateMutableSource.h
@@ -18,6 +18,7 @@
 #include "MuCOE/Nodes/CustomizableObjectNodeTable.h"
 #include "MuT/NodeComponentNew.h"
 #include "MuT/NodeImageConstant.h"
+#include "MuT/NodeImageMaterialBreak.h"
 #include "MuT/NodeMeshConstant.h"
 #include "MuT/NodeMeshApplyPose.h"
 #include "MuT/NodeModifierMeshClipWithMesh.h"
@@ -26,6 +27,8 @@
 #include "MuT/NodeScalarEnumParameter.h"
 #include "MuT/NodeScalarParameter.h"
 #include "MuT/NodeSurfaceNew.h"
+#include "MuT/NodeMaterialConstant.h"
+#include "MuT/NodeMaterialParameter.h"
 #include "MuT/Table.h"
 #include "MuR/Skeleton.h"
 #include "MuR/Mesh.h"
@@ -42,8 +45,10 @@ class FCustomizableObjectCompiler;
 class UTextureLODSettings;
 class UAnimInstance;
 class UCompositeDataTable;
+class UCONodeMaterialConstant;
 class UCustomizableObjectNodeMacroInstance;
 class UCustomizableObjectNodeMaterial;
+class UCustomizableObjectNodeMaterialParameter;
 class UCustomizableObjectNodeMeshMorph;
 class UCustomizableObjectNodeMeshMorphStackDefinition;
 class UCustomizableObjectNodeObjectGroup;
@@ -890,6 +895,19 @@ struct FMutableGraphGenerationContext
 	TMap<FName, TObjectPtr<UTexture>> TextureParameterDefaultValues;
 	TMap<FName, TObjectPtr<USkeletalMesh>> SkeletalMeshParameterDefaultValues;
 	TMap<FName, TObjectPtr<UMaterialInterface>> MaterialParameterDefaultValues;
+
+	
+	struct FMaterialBreakParameter {
+		FName ParameterName;
+		EMaterialParameterType ParameterType = EMaterialParameterType::None;
+	};
+
+	/** Name of the material break Node that is being compiled.
+	NOTE: This is set when the material pin of a material breake node is compiled and then it must be reset to null. */
+	FMaterialBreakParameter CurrentMaterialBreakParameter;
+
+	TMap<const UCustomizableObjectNodeMaterialParameter*, mu::Ptr<mu::NodeMaterialParameter>> MaterialParameterNodesCache;
+	TMap<const UCONodeMaterialConstant*, mu::Ptr<mu::NodeMaterialConstant>> MaterialConstantNodesCache;
 };
 
 
diff --git a/Engine/Plugins/Mutable/Source/CustomizableObjectEditor/Private/MuCOE/GenerateMutableSource/GenerateMutableSourceColor.cpp b/Engine/Plugins/Mutable/Source/CustomizableObjectEditor/Private/MuCOE/GenerateMutableSource/GenerateMutableSourceColor.cpp
--- a/Engine/Plugins/Mutable/Source/CustomizableObjectEditor/Private/MuCOE/GenerateMutableSource/GenerateMutableSourceColor.cpp
+++ b/Engine/Plugins/Mutable/Source/CustomizableObjectEditor/Private/MuCOE/GenerateMutableSource/GenerateMutableSourceColor.cpp
@@ -8,6 +8,7 @@
 #include "MuCOE/GenerateMutableSource/GenerateMutableSourceFloat.h"
 #include "MuCOE/GenerateMutableSource/GenerateMutableSourceImage.h"
 #include "MuCOE/GenerateMutableSource/GenerateMutableSourceMacro.h"
+#include "MuCOE/GenerateMutableSource/GenerateMutableSourceMaterial.h"
 #include "MuCOE/GenerateMutableSource/GenerateMutableSourceTable.h"
 #include "MuCOE/GenerateMutableSource/GenerateMutableSource.h"
 #include "MuCOE/GraphTraversal.h"
@@ -19,6 +20,7 @@
 #include "MuCOE/Nodes/CustomizableObjectNodeColorVariation.h"
 #include "MuCOE/Nodes/CustomizableObjectNodeColorToSRGB.h"
 #include "MuCOE/Nodes/CustomizableObjectNodeMacroInstance.h"
+#include "MuCOE/Nodes/CONodeMaterialBreak.h"
 #include "MuCOE/Nodes/CustomizableObjectNodeTable.h"
 #include "MuCOE/Nodes/CustomizableObjectNodeTextureSample.h"
 #include "MuCOE/Nodes/CustomizableObjectNodeTunnel.h"
@@ -31,6 +33,7 @@
 #include "MuT/NodeColourTable.h"
 #include "MuT/NodeColourVariation.h"
 #include "MuT/NodeColorToSRGB.h"
+#include "MuT/NodeColourMaterialBreak.h"
 #include "MuT/NodeImageBinarise.h"
 
 
@@ -373,6 +376,42 @@ mu::Ptr<mu::NodeColour> GenerateMutableSourceColor(const UEdGraphPin* Pin, FMuta
 		Result = GenerateMutableSourceMacro<mu::NodeColour>(*Pin, GenerationContext, GenerateMutableSourceColor);
 	}
 
+	else if (const UCONodeMaterialBreak* TypedNodeMaterialBreak = Cast<UCONodeMaterialBreak>(Node))
+	{
+		Result = new mu::NodeColourConstant();
+		check(TypedNodeMaterialBreak->MaterialPinRef.Get());
+
+		if (const UEdGraphPin* LinkedMaterial = FollowInputPin(*TypedNodeMaterialBreak->MaterialPinRef.Get()))
+		{
+			mu::Ptr<mu::NodeColourMaterialBreak> MaterialBreak = new mu::NodeColourMaterialBreak();
+
+			FName ParameterName = TypedNodeMaterialBreak->GetPinParameterName(*Pin);
+			check(!ParameterName.IsNone());
+
+			// Set the name of the parameter that represents the pin
+			MaterialBreak->ParameterName = ParameterName;
+
+			// Generate the material that is going to be broken into parameters
+			GenerationContext.CurrentMaterialBreakParameter = FMutableGraphGenerationContext::FMaterialBreakParameter{ ParameterName, EMaterialParameterType::Vector };
+			mu::Ptr<mu::NodeMaterial> LinkedNode = GenerateMutableSourceMaterial(LinkedMaterial, GenerationContext);
+			MaterialBreak->MaterialSource = LinkedNode;
+			GenerationContext.CurrentMaterialBreakParameter = FMutableGraphGenerationContext::FMaterialBreakParameter{ };
+
+			if (LinkedNode)
+			{
+				Result = MaterialBreak;
+			}
+			else
+			{
+				GenerationContext.Log(LOCTEXT("MaterialBreakNodeError", "Something went wrong generation the Material Node linked to the Material Break Node."), Node);
+			}
+		}
+		else
+		{
+			GenerationContext.Log(LOCTEXT("MaterialBreakPinError", "Material Break Nodes must be linked to an input Material Node."), Node);
+		}
+		}
+
 	else
 	{
 		GenerationContext.Log(LOCTEXT("UnimplementedNode", "Node type not implemented yet."), Node);
diff --git a/Engine/Plugins/Mutable/Source/CustomizableObjectEditor/Private/MuCOE/GenerateMutableSource/GenerateMutableSourceFloat.cpp b/Engine/Plugins/Mutable/Source/CustomizableObjectEditor/Private/MuCOE/GenerateMutableSource/GenerateMutableSourceFloat.cpp
--- a/Engine/Plugins/Mutable/Source/CustomizableObjectEditor/Private/MuCOE/GenerateMutableSource/GenerateMutableSourceFloat.cpp
+++ b/Engine/Plugins/Mutable/Source/CustomizableObjectEditor/Private/MuCOE/GenerateMutableSource/GenerateMutableSourceFloat.cpp
@@ -10,8 +10,10 @@
 #include "MuCOE/EdGraphSchema_CustomizableObject.h"
 #include "MuCOE/GenerateMutableSource/GenerateMutableSourceImage.h"
 #include "MuCOE/GenerateMutableSource/GenerateMutableSourceMacro.h"
+#include "MuCOE/GenerateMutableSource/GenerateMutableSourceMaterial.h"
 #include "MuCOE/GenerateMutableSource/GenerateMutableSourceTable.h"
 #include "MuCOE/GraphTraversal.h"
+#include "MuCOE/Nodes/CONodeMaterialBreak.h"
 #include "MuCOE/Nodes/CustomizableObjectNodeCurve.h"
 #include "MuCOE/Nodes/CustomizableObjectNodeEnumParameter.h"
 #include "MuCOE/Nodes/CustomizableObjectNodeFloatConstant.h"
@@ -28,6 +30,7 @@
 #include "MuT/NodeScalarTable.h"
 #include "MuT/NodeScalarVariation.h"
 #include "MuT/NodeScalarArithmeticOperation.h"
+#include "MuT/NodeScalarMaterialBreak.h"
 
 #define LOCTEXT_NAMESPACE "CustomizableObjectEditor"
 
@@ -396,6 +399,42 @@ mu::Ptr<mu::NodeScalar> GenerateMutableSourceFloat(const UEdGraphPin* Pin, FMuta
 		Result = GenerateMutableSourceMacro<mu::NodeScalar>(*Pin, GenerationContext, GenerateMutableSourceFloat);
 	}
 
+	else if (const UCONodeMaterialBreak* TypedNodeMaterialBreak = Cast<UCONodeMaterialBreak>(Node))
+	{
+		Result = new mu::NodeScalarConstant();
+		check(TypedNodeMaterialBreak->MaterialPinRef.Get());
+
+		if (const UEdGraphPin* LinkedMaterial = FollowInputPin(*TypedNodeMaterialBreak->MaterialPinRef.Get()))
+		{
+			mu::Ptr<mu::NodeScalarMaterialBreak> MaterialBreak = new mu::NodeScalarMaterialBreak();
+
+			FName ParameterName = TypedNodeMaterialBreak->GetPinParameterName(*Pin);
+			check(!ParameterName.IsNone());
+
+			// Set the name of the parameter that represents the pin
+			MaterialBreak->ParameterName = ParameterName;
+
+			// Generate the material that is going to be broken into parameters
+			GenerationContext.CurrentMaterialBreakParameter = FMutableGraphGenerationContext::FMaterialBreakParameter{ ParameterName, EMaterialParameterType::Scalar };
+			mu::Ptr<mu::NodeMaterial> LinkedNode = GenerateMutableSourceMaterial(LinkedMaterial, GenerationContext);
+			MaterialBreak->MaterialSource = LinkedNode;
+			GenerationContext.CurrentMaterialBreakParameter = FMutableGraphGenerationContext::FMaterialBreakParameter{ };
+
+			if (LinkedNode)
+			{
+				Result = MaterialBreak;
+			}
+			else
+			{
+				GenerationContext.Log(LOCTEXT("MaterialBreakNodeError", "Something went wrong generation the Material Node linked to the Material Break Node."), Node);
+			}
+		}
+		else
+		{
+			GenerationContext.Log(LOCTEXT("MaterialBreakPinError", "Material Break Nodes must be linked to an input Material Node."), Node);
+		}
+	}
+
 	else
 	{
 		GenerationContext.Log(LOCTEXT("UnimplementedNode", "Node type not implemented yet."), Node);
diff --git a/Engine/Plugins/Mutable/Source/CustomizableObjectEditor/Private/MuCOE/GenerateMutableSource/GenerateMutableSourceImage.cpp b/Engine/Plugins/Mutable/Source/CustomizableObjectEditor/Private/MuCOE/GenerateMutableSource/GenerateMutableSourceImage.cpp
--- a/Engine/Plugins/Mutable/Source/CustomizableObjectEditor/Private/MuCOE/GenerateMutableSource/GenerateMutableSourceImage.cpp
+++ b/Engine/Plugins/Mutable/Source/CustomizableObjectEditor/Private/MuCOE/GenerateMutableSource/GenerateMutableSourceImage.cpp
@@ -9,10 +9,12 @@
 #include "MuCOE/EdGraphSchema_CustomizableObject.h"
 #include "MuCOE/GenerateMutableSource/GenerateMutableSourceColor.h"
 #include "MuCOE/GenerateMutableSource/GenerateMutableSourceFloat.h"
+#include "MuCOE/GenerateMutableSource/GenerateMutableSourceMaterial.h"
 #include "MuCOE/GenerateMutableSource/GenerateMutableSourceMesh.h"
 #include "MuCOE/GenerateMutableSource/GenerateMutableSourceProjector.h"
 #include "MuCOE/GenerateMutableSource/GenerateMutableSourceTable.h"
 #include "MuCOE/GraphTraversal.h"
+#include "MuCOE/Nodes/CONodeMaterialBreak.h"
 #include "MuCOE/Nodes/CustomizableObjectNodeMacroInstance.h"
 #include "MuCOE/Nodes/CustomizableObjectNodeMeshInterface.h"
 #include "MuCOE/Nodes/CustomizableObjectNodeTable.h"
@@ -1218,6 +1220,42 @@ mu::Ptr<mu::NodeImage> GenerateMutableSourceImage(const UEdGraphPin* Pin, FMutab
 		GenerationContext.MacroNodesStack.Push(MacroInstanceNode);
 	}
 
+	else if (const UCONodeMaterialBreak* TypedNodeMaterialBreak = Cast<UCONodeMaterialBreak>(Node))
+	{
+		Result = new mu::NodeImageConstant();
+		check(TypedNodeMaterialBreak->MaterialPinRef.Get());
+
+		if (const UEdGraphPin* LinkedMaterial = FollowInputPin(*TypedNodeMaterialBreak->MaterialPinRef.Get()))
+		{
+			mu::Ptr<mu::NodeImageMaterialBreak> MaterialBreak = new mu::NodeImageMaterialBreak();
+
+			FName ParameterName = TypedNodeMaterialBreak->GetPinParameterName(*Pin);
+			check(!ParameterName.IsNone());
+
+			// Set the name of the parameter that represents the pin
+			MaterialBreak->ParameterName = ParameterName;
+
+			// Generate the material that is going to be broken into parameters
+			GenerationContext.CurrentMaterialBreakParameter = FMutableGraphGenerationContext::FMaterialBreakParameter{ ParameterName, EMaterialParameterType::Texture };
+			mu::Ptr<mu::NodeMaterial> LinkedNode = GenerateMutableSourceMaterial(LinkedMaterial, GenerationContext);
+			MaterialBreak->MaterialSource = LinkedNode;
+			GenerationContext.CurrentMaterialBreakParameter = FMutableGraphGenerationContext::FMaterialBreakParameter{ };
+
+			if (LinkedNode)
+			{
+				Result = MaterialBreak;
+			}
+			else
+			{
+				GenerationContext.Log(LOCTEXT("MaterialBreakNodeError", "Something went wrong generation the Material Node linked to the Material Break Node."), Node);
+			}
+		}
+		else
+		{
+			GenerationContext.Log(LOCTEXT("MaterialBreakPinError", "Material Break Nodes must be linked to an input Material Node."), Node);
+		}
+	}
+
 	else
 	{
 		GenerationContext.Log(LOCTEXT("UnimplementedNode", "Node type not implemented yet."), Node);
diff --git a/Engine/Plugins/Mutable/Source/CustomizableObjectEditor/Private/MuCOE/GenerateMutableSource/GenerateMutableSourceMaterial.cpp b/Engine/Plugins/Mutable/Source/CustomizableObjectEditor/Private/MuCOE/GenerateMutableSource/GenerateMutableSourceMaterial.cpp
--- a/Engine/Plugins/Mutable/Source/CustomizableObjectEditor/Private/MuCOE/GenerateMutableSource/GenerateMutableSourceMaterial.cpp
+++ b/Engine/Plugins/Mutable/Source/CustomizableObjectEditor/Private/MuCOE/GenerateMutableSource/GenerateMutableSourceMaterial.cpp
@@ -2,10 +2,13 @@
 
 #include "MuCOE/GenerateMutableSource/GenerateMutableSourceMaterial.h"
 
+#include "Engine/TextureLODSettings.h"
+#include "Interfaces/ITargetPlatform.h"
 #include "MuCOE/CustomizableObjectCompiler.h"
 #include "MuCOE/CustomizableObjectGraph.h"
 #include "MuCOE/EdGraphSchema_CustomizableObject.h"
 #include "MuCOE/GenerateMutableSource/GenerateMutableSourceFloat.h"
+#include "MuCOE/GenerateMutableSource/GenerateMutableSourceImage.h"
 #include "MuCOE/GenerateMutableSource/GenerateMutableSourceMacro.h"
 #include "MuCOE/GenerateMutableSource/GenerateMutableSourceTable.h"
 #include "MuCOE/GraphTraversal.h"
@@ -14,10 +17,14 @@
 #include "MuCOE/Nodes/CustomizableObjectNodeMacroInstance.h"
 #include "MuCOE/Nodes/CustomizableObjectNodeTable.h"
 #include "MuCOE/Nodes/CustomizableObjectNodeTunnel.h"
+#include "MuCOE/Nodes/CONodeMaterialBreak.h"
 #include "MuCOE/Nodes/CONodeMaterialConstant.h"
 #include "MuCOE/Nodes/CONodeMaterialSwitch.h"
 #include "MuCOE/Nodes/CONodeMaterialVariation.h"
 #include "MuCO/LoadUtils.h"
+#include "MuT/NodeImageFormat.h"
+#include "MuT/NodeImageFromMaterialParameter.h"
+#include "MuT/NodeImageResize.h"
 #include "MuT/NodeMaterialConstant.h"
 #include "MuT/NodeMaterialSwitch.h"
 #include "MuT/NodeMaterialTable.h"
@@ -61,14 +68,99 @@ mu::Ptr<mu::NodeMaterial> GenerateMutableSourceMaterial(const UEdGraphPin* Pin,
 	
 	if (const UCONodeMaterialConstant* TypedNodeMaterialConstant = Cast<UCONodeMaterialConstant>(Node))
 	{
-		mu::Ptr<mu::NodeMaterialConstant> MaterialNode = new mu::NodeMaterialConstant();
-		Result = MaterialNode;
+		mu::Ptr<mu::NodeMaterialConstant> MaterialNode;
+		bCacheNode = false; //custom node cache
 
 		if (UMaterialInterface* Material = GenerationContext.LoadObject(TypedNodeMaterialConstant->Material))
 		{
-			GenerationContext.CurrentReferencedMaterialIndex = GenerationContext.ReferencedMaterials.AddUnique(Material);
-			MaterialNode->MaterialId = GenerationContext.CurrentReferencedMaterialIndex;
+			if (GenerationContext.MaterialConstantNodesCache.Contains(TypedNodeMaterialConstant))
+			{
+				MaterialNode = GenerationContext.MaterialConstantNodesCache[TypedNodeMaterialConstant];
+			}
+			else
+			{
+				MaterialNode = new mu::NodeMaterialConstant();
+
+				GenerationContext.CurrentReferencedMaterialIndex = GenerationContext.ReferencedMaterials.AddUnique(Material);
+				MaterialNode->MaterialId = GenerationContext.CurrentReferencedMaterialIndex;
+
+				GenerationContext.MaterialConstantNodesCache.Add(TypedNodeMaterialConstant, MaterialNode);
+			}
+
+			// Store the parameter that needs to be processed
+			if (GenerationContext.CurrentMaterialBreakParameter.ParameterType != EMaterialParameterType::None)
+			{
+				if (!TypedNodeMaterialConstant->Material)
+				{
+					GenerationContext.Log(LOCTEXT("NoRefernceMaterialToBreak", "Could not find a Reference Material to break."), TypedNodeMaterialConstant);
+				}
+				else
+				{
+					switch (GenerationContext.CurrentMaterialBreakParameter.ParameterType)
+					{
+					case EMaterialParameterType::Texture:
+					{
+						UTexture* Texture = nullptr;
+						bool bParameterFound = Material->GetTextureParameterValue(GenerationContext.CurrentMaterialBreakParameter.ParameterName, Texture);
+						UTexture2D* BaseTexture = Cast<UTexture2D>(Texture);
+
+						if (bParameterFound && BaseTexture)
+						{
+							TSharedPtr<mu::FImage> ImageConstant = GenerateImageConstant(BaseTexture, GenerationContext, false);
+							const uint32 MipsToSkip = ComputeLODBiasForTexture(GenerationContext, *BaseTexture);
+
+							mu::Ptr<mu::NodeImageConstant> ConstantImageNode = new mu::NodeImageConstant;
+							ConstantImageNode->SetValue(ImageConstant);
+
+							const UTextureLODSettings& TextureLODSettings = GenerationContext.CompilationContext->Options.TargetPlatform->GetTextureLODSettings();
+							const FTextureLODGroup& LODGroupInfo = TextureLODSettings.GetTextureLODGroup(BaseTexture->LODGroup);
+							const FCompilationOptions& CompilationOptions = GenerationContext.CompilationContext->Options;
+
+							ConstantImageNode->SourceDataDescriptor.OptionalMaxLODSize = LODGroupInfo.OptionalMaxLODSize;
+							ConstantImageNode->SourceDataDescriptor.OptionalLODBias = LODGroupInfo.OptionalLODBias;
+							ConstantImageNode->SourceDataDescriptor.NumNonOptionalLODs = CompilationOptions.MinDiskMips;
+
+							const FString TextureName = GetNameSafe(BaseTexture).ToLower();
+							ConstantImageNode->SourceDataDescriptor.SourceId = CityHash32(reinterpret_cast<const char*>(*TextureName), TextureName.Len() * sizeof(FString::ElementType));
+
+							mu::Ptr<mu::NodeImage> FinalNodeImag = ResizeTextureByNumMips(ConstantImageNode, MipsToSkip);
+
+							if (FinalNodeImag)
+							{
+								MaterialNode->ImageValues.Add(GenerationContext.CurrentMaterialBreakParameter.ParameterName, FinalNodeImag);
+							}
+						}
+
+						break;
+					}
+					case EMaterialParameterType::Vector:
+					{
+						FLinearColor Color;
+						if (Material->GetVectorParameterValue(GenerationContext.CurrentMaterialBreakParameter.ParameterName, Color))
+						{
+							MaterialNode->ColorValues.Add(GenerationContext.CurrentMaterialBreakParameter.ParameterName, Color);
+						}
+
+						break;
+					}
+					case EMaterialParameterType::Scalar:
+					{
+						float Scalar;
+						if (Material->GetScalarParameterValue(GenerationContext.CurrentMaterialBreakParameter.ParameterName, Scalar))
+						{
+							MaterialNode->ScalarValues.Add(GenerationContext.CurrentMaterialBreakParameter.ParameterName, Scalar);
+						}
+
+						break;
+					}
+					default:
+						break;
+					}
+				}
+			}
 		}
+
+		Result = MaterialNode;
 	}
 
 	else if (const UCONodeMaterialSwitch* TypedNodeMaterialSwitch = Cast<UCONodeMaterialSwitch>(Node))
@@ -171,26 +263,103 @@ mu::Ptr<mu::NodeMaterial> GenerateMutableSourceMaterial(const UEdGraphPin* Pin,
 			}
 		}
 	}
-	
+
 	else if (const UCustomizableObjectNodeMaterialParameter* MaterialParameterNode = Cast<UCustomizableObjectNodeMaterialParameter>(Node))
 	{
-		mu::Ptr<mu::NodeMaterialParameter> MaterialNode = new mu::NodeMaterialParameter();
-		Result = MaterialNode;
+		bCacheNode = false;
+		mu::Ptr<mu::NodeMaterialParameter> MaterialNode;
 
-		MaterialNode->Name = MaterialParameterNode->GetParameterName(&GenerationContext.MacroNodesStack);
-		MaterialNode->UID = GenerationContext.GetNodeIdUnique(Node).ToString();
+		if (GenerationContext.MaterialParameterNodesCache.Contains(MaterialParameterNode))
+		{
+			MaterialNode = GenerationContext.MaterialParameterNodesCache[MaterialParameterNode];
+		}
+		else
+		{
+			MaterialNode = new mu::NodeMaterialParameter();
+			MaterialNode->Name = MaterialParameterNode->GetParameterName(&GenerationContext.MacroNodesStack);
+			MaterialNode->UID = GenerationContext.GetNodeIdUnique(Node).ToString();
 
-		if (MaterialParameterNode->DefaultValue)
+			if (MaterialParameterNode->DefaultValue)
+			{
+				GenerationContext.MaterialParameterDefaultValues.Add(FName(MaterialParameterNode->GetParameterName()), MutablePrivate::LoadObject(MaterialParameterNode->DefaultValue));
+			}
+
+			GenerationContext.ParameterUIDataMap.Add(MaterialParameterNode->GetParameterName(&GenerationContext.MacroNodesStack), FMutableParameterData(
+				MaterialParameterNode->ParamUIMetadata,
+				EMutableParameterType::Material));
+
+			GenerationContext.MaterialParameterNodesCache.Add(MaterialParameterNode, MaterialNode);
+		}
+
+		// Store the parameter that needs to be processed
+		if (GenerationContext.CurrentMaterialBreakParameter.ParameterType != EMaterialParameterType::None)
 		{
-			GenerationContext.MaterialParameterDefaultValues.Add(FName(MaterialParameterNode->GetParameterName()), MutablePrivate::LoadObject(MaterialParameterNode->DefaultValue));
+			if (!MaterialParameterNode->ReferenceValue)
+			{
+				GenerationContext.Log(LOCTEXT("NoRefernceMaterialToBreak", "Could not find a Reference Material to break."), MaterialParameterNode);
+			}
+			else
+			{
+				switch (GenerationContext.CurrentMaterialBreakParameter.ParameterType)
+				{
+				case EMaterialParameterType::Texture:
+				{
+					//TODO(Max): Move this to the generatemutablesourceimage:
+					mu::Ptr<mu::NodeImageFromMaterialParameter> TextureNode = new mu::NodeImageFromMaterialParameter();
+					TextureNode->ImageParameterName = GenerationContext.CurrentMaterialBreakParameter.ParameterName;
+
+					// Force the same format that the default texture if any.
+					mu::Ptr<mu::NodeImageFormat> FormatNode = new mu::NodeImageFormat();
+
+					// Force an "easy format" on the texture.
+					FormatNode->Format = mu::EImageFormat::RGBA_UByte;
+					FormatNode->Source = TextureNode;
+
+					// Resize image node.
+					mu::Ptr<mu::NodeImageResize> ResizeNode = new mu::NodeImageResize();
+					ResizeNode->Base = FormatNode;
+					ResizeNode->bRelative = false;
+
+					mu::FImageSize TextureSize(1);
+
+					//Find reference Texture
+					UTexture* Texture = nullptr;
+					bool bParameterFound = MaterialParameterNode->ReferenceValue->GetTextureParameterValue(GenerationContext.CurrentMaterialBreakParameter.ParameterName, Texture);
+					UTexture2D* ReferenceTexture = Cast<UTexture2D>(Texture);
+
+					if (bParameterFound && ReferenceTexture)
+					{
+						const uint32 LODBias = ComputeLODBiasForTexture(GenerationContext, *ReferenceTexture);
+						TextureSize.X = FMath::Max(ReferenceTexture->Source.GetSizeX() >> LODBias, 1);
+						TextureSize.Y = FMath::Max(ReferenceTexture->Source.GetSizeY() >> LODBias, 1);
+					}
+
+					ResizeNode->SizeX = TextureSize.X;
+					ResizeNode->SizeY = TextureSize.Y;
+
+					MaterialNode->ImageParameterNodes.Add(GenerationContext.CurrentMaterialBreakParameter.ParameterName.ToString(), ResizeNode);
+
+					break;
+				}
+				case EMaterialParameterType::Vector:
+				{
+					MaterialNode->ColorParameterNames.Add(GenerationContext.CurrentMaterialBreakParameter.ParameterName.ToString());
+					break;
+				}
+				case EMaterialParameterType::Scalar:
+				{
+					MaterialNode->ScalarParameterNames.Add(GenerationContext.CurrentMaterialBreakParameter.ParameterName.ToString());
+					break;
+				}
+				default:
+					break;
+				}
+			}
 		}
 
-		GenerationContext.ParameterUIDataMap.Add(MaterialParameterNode->GetParameterName(&GenerationContext.MacroNodesStack), FMutableParameterData(
-			MaterialParameterNode->ParamUIMetadata,
-			EMutableParameterType::Material));
+		Result = MaterialNode;
 	}
 
-	
 	else if (const UCustomizableObjectNodeTable* TypedNodeTable = Cast<UCustomizableObjectNodeTable>(Node))
 	{
 		//This node will add a default value in case of error
diff --git a/Engine/Plugins/Mutable/Source/CustomizableObjectEditor/Private/MuCOE/Nodes/CONodeMaterialBreak.cpp b/Engine/Plugins/Mutable/Source/CustomizableObjectEditor/Private/MuCOE/Nodes/CONodeMaterialBreak.cpp
new file mode 100644
--- /dev/null
+++ b/Engine/Plugins/Mutable/Source/CustomizableObjectEditor/Private/MuCOE/Nodes/CONodeMaterialBreak.cpp
@@ -0,0 +1,152 @@
+// Copyright Epic Games, Inc. All Rights Reserved.
+
+#include "MuCOE/Nodes/CONodeMaterialBreak.h"
+
+#include "DetailsViewArgs.h"
+#include "MuCOE/EdGraphSchema_CustomizableObject.h"
+#include "PropertyEditorModule.h"
+
+
+#include UE_INLINE_GENERATED_CPP_BY_NAME(CONodeMaterialBreak)
+
+class UCustomizableObjectNodeRemapPins;
+struct FPropertyChangedEvent;
+struct FSlateBrush;
+
+#define LOCTEXT_NAMESPACE "CustomizableObjectEditor"
+
+
+bool UCustomizableObjectNodeBreakMaterialRemapPins::Equal(const UCustomizableObjectNode& Node, const UEdGraphPin& OldPin, const UEdGraphPin& NewPin) const
+{
+	UCONodeMaterialBreakPinData* OldPinData = Cast<UCONodeMaterialBreakPinData>(Node.GetPinData(OldPin));
+	UCONodeMaterialBreakPinData* NewPinData = Cast<UCONodeMaterialBreakPinData>(Node.GetPinData(NewPin));
+
+	if (OldPinData && NewPinData)
+	{
+		return OldPin.Direction == NewPin.Direction && OldPinData->VariableData.Type == NewPinData->VariableData.Type && OldPinData->VariableData.Id == NewPinData->VariableData.Id;
+	}
+	else
+	{
+		return Super::Equal(Node, OldPin, NewPin);
+	}
+}
+
+
+void UCONodeMaterialBreak::AllocateDefaultPins(UCustomizableObjectNodeRemapPins* RemapPins)
+{
+	const FName Type = UEdGraphSchema_CustomizableObject::PC_Material;
+	const FName PinName = UEdGraphSchema_CustomizableObject::GetPinCategoryName(Type);
+	const FText PinFriendlyName = UEdGraphSchema_CustomizableObject::GetPinCategoryFriendlyName(Type);
+	
+	UEdGraphPin* MaterialPin = CustomCreatePin(EGPD_Input, Type, PinName);
+	MaterialPin->PinFriendlyName = PinFriendlyName;
+
+	MaterialPinRef = FEdGraphPinReference(MaterialPin);
+
+	// Get old pins
+	TArray<UEdGraphPin*> OutputPins;
+	GetOutputPins(OutputPins);
+
+	// Regenerate the new pins from the old pin's pindata which may have changed
+	for (const UEdGraphPin* Pin : OutputPins)
+	{
+		if (UCONodeMaterialBreakPinData* PinData = Cast<UCONodeMaterialBreakPinData>(GetPinData(*Pin)))
+		{
+			// Use the same pin data as the old pin since the new one is a copy
+			UEdGraphPin* NewPin = CustomCreatePin(EGPD_Output, PinData->VariableData.Type, PinData->VariableData.Name, PinData);
+		}
+	}
+}
+
+
+FText UCONodeMaterialBreak::GetNodeTitle(ENodeTitleType::Type TitleType) const
+{
+	return LOCTEXT("Material_Break", "Break Material");
+}
+
+
+FLinearColor UCONodeMaterialBreak::GetNodeTitleColor() const
+{
+	return UEdGraphSchema_CustomizableObject::GetPinTypeColor(UEdGraphSchema_CustomizableObject::PC_Material);
+}
+
+
+FText UCONodeMaterialBreak::GetTooltipText() const
+{
+	return LOCTEXT("Material_Constant_Tooltip", "Allows to get the parameter value from the input material node.");
+}
+
+
+TSharedPtr<IDetailsView> UCONodeMaterialBreak::CustomizePinDetails(const UEdGraphPin& Pin) const
+{
+	if (UCONodeMaterialBreakPinData* PinData = Cast<UCONodeMaterialBreakPinData>(GetPinData(Pin)))
+	{
+		FPropertyEditorModule& EditModule = FModuleManager::Get().GetModuleChecked<FPropertyEditorModule>("PropertyEditor");
+
+		FDetailsViewArgs DetailsViewArgs;
+		DetailsViewArgs.bAllowSearch = false;
+		DetailsViewArgs.bHideSelectionTip = true;
+
+		const TSharedRef<IDetailsView> SettingsView = EditModule.CreateDetailView(DetailsViewArgs);
+		SettingsView->SetObject(PinData);
+
+		return SettingsView;
+	}
+	else
+	{
+		return nullptr;
+	}
+}
+
+
+bool UCONodeMaterialBreak::HasPinViewer() const
+{
+	return true;
+}
+
+
+bool UCONodeMaterialBreak::CanCreatePinsFromPinViewer() const
+{
+	return true;
+}
+
+
+UCustomizableObjectNodeRemapPins* UCONodeMaterialBreak::CreateRemapPinsDefault() const
+{
+	return NewObject<UCustomizableObjectNodeBreakMaterialRemapPins>();
+}
+
+
+TArray<FName> UCONodeMaterialBreak::GetAllowedPinViewerCreationTypes() const
+{
+	return { UEdGraphSchema_CustomizableObject::PC_Image, UEdGraphSchema_CustomizableObject::PC_Color, UEdGraphSchema_CustomizableObject::PC_Float };
+}
+
+
+void UCONodeMaterialBreak::CreatePinFromPinViewer()
+{
+	FName NewPinCategorery = UEdGraphSchema_CustomizableObject::PC_Float;
+	
+	UCONodeMaterialBreakPinData* PinData = NewObject<UCONodeMaterialBreakPinData>(this);
+	PinData->VariableData.Name = "NewPin";
+	PinData->VariableData.Type = NewPinCategorery;
+	PinData->VariableData.Id = FGuid::NewGuid();
+
+	CustomCreatePin(EEdGraphPinDirection::EGPD_Output, NewPinCategorery, "NewPin", PinData);
+
+	GetGraph()->NotifyGraphChanged();
+}
+
+
+FName UCONodeMaterialBreak::GetPinParameterName(const UEdGraphPin& Pin) const
+{
+	if (const UCONodeMaterialBreakPinData* PinData = Cast<UCONodeMaterialBreakPinData>(GetPinData(Pin)))
+	{
+		return PinData->VariableData.Name;
+	}
+
+	return NAME_None;
+}
+
+
+#undef LOCTEXT_NAMESPACE
diff --git a/Engine/Plugins/Mutable/Source/CustomizableObjectEditor/Private/MuCOE/Nodes/CONodeMaterialBreak.h b/Engine/Plugins/Mutable/Source/CustomizableObjectEditor/Private/MuCOE/Nodes/CONodeMaterialBreak.h
new file mode 100644
--- /dev/null
+++ b/Engine/Plugins/Mutable/Source/CustomizableObjectEditor/Private/MuCOE/Nodes/CONodeMaterialBreak.h
@@ -0,0 +1,71 @@
+// Copyright Epic Games, Inc. All Rights Reserved.
+
+#pragma once
+
+#include "MuCOE/Nodes/CustomizableObjectNode.h"
+#include "MuCOE/COVariable.h"
+#include "MuCOE/RemapPins/CustomizableObjectNodeRemapPinsByName.h"
+
+#include "CONodeMaterialBreak.generated.h"
+
+#define UE_API CUSTOMIZABLEOBJECTEDITOR_API
+
+namespace ENodeTitleType { enum Type : int; }
+class UCustomizableObjectNodeRemapPins;
+
+
+UCLASS()
+class UCustomizableObjectNodeBreakMaterialRemapPins: public UCustomizableObjectNodeRemapPinsByName
+{
+public:
+	GENERATED_BODY()
+
+	virtual bool Equal(const UCustomizableObjectNode& Node, const UEdGraphPin& OldPin, const UEdGraphPin& NewPin) const override;
+};
+
+
+/** Base class for all Material Parameters. */
+UCLASS(MinimalAPI)
+class UCONodeMaterialBreakPinData : public UCustomizableObjectNodePinData
+{
+	GENERATED_BODY()
+
+public:
+
+	UPROPERTY(EditAnywhere, Category = NoCategory)
+	FCOVariable VariableData;
+
+};
+
+
+UCLASS(MinimalAPI)
+class UCONodeMaterialBreak : public UCustomizableObjectNode
+{
+public:
+
+	GENERATED_BODY()
+
+	// Begin EdGraphNode interface
+	UE_API virtual FText GetNodeTitle(ENodeTitleType::Type TitleType) const override;
+	UE_API virtual FLinearColor GetNodeTitleColor() const override;
+	UE_API virtual FText GetTooltipText() const override;
+
+	// UCustomizableObjectNode interface
+	UE_API virtual void AllocateDefaultPins(UCustomizableObjectNodeRemapPins* RemapPins) override;
+	UE_API virtual TSharedPtr<IDetailsView> CustomizePinDetails(const UEdGraphPin& Pin) const override;
+	UE_API virtual bool HasPinViewer() const override;
+	UE_API virtual bool CanCreatePinsFromPinViewer() const override;
+	UE_API virtual void CreatePinFromPinViewer() override;
+	UE_API virtual UCustomizableObjectNodeRemapPins* CreateRemapPinsDefault() const override;
+	UE_API virtual TArray<FName> GetAllowedPinViewerCreationTypes() const override;
+
+	// Own interface
+	FName GetPinParameterName(const UEdGraphPin& Pin) const;
+
+public: 
+
+	UPROPERTY()
+	FEdGraphPinReference MaterialPinRef;
+};
+
+#undef UE_API
diff --git a/Engine/Plugins/Mutable/Source/CustomizableObjectEditor/Private/MuCOE/Nodes/CustomizableObjectNode.cpp b/Engine/Plugins/Mutable/Source/CustomizableObjectEditor/Private/MuCOE/Nodes/CustomizableObjectNode.cpp
--- a/Engine/Plugins/Mutable/Source/CustomizableObjectEditor/Private/MuCOE/Nodes/CustomizableObjectNode.cpp
+++ b/Engine/Plugins/Mutable/Source/CustomizableObjectEditor/Private/MuCOE/Nodes/CustomizableObjectNode.cpp
@@ -767,6 +767,12 @@ bool UCustomizableObjectNode::HasPinViewer() const
 }
 
 
+bool UCustomizableObjectNode::CanCreatePinsFromPinViewer() const
+{
+	return false;
+}
+
+
 TSharedPtr<IDetailsView> UCustomizableObjectNode::CustomizePinDetails(const UEdGraphPin& Pin) const
 {
 	return nullptr;
@@ -823,4 +829,10 @@ bool UCustomizableObjectNode::CreateStaticStringPinWidget() const
 }
 
 
+TArray<FName> UCustomizableObjectNode::GetAllowedPinViewerCreationTypes() const
+{
+	return UEdGraphSchema_CustomizableObject::SupportedMacroPinTypes;
+}
+
+
 #undef LOCTEXT_NAMESPACE
diff --git a/Engine/Plugins/Mutable/Source/CustomizableObjectEditor/Private/MuCOE/Nodes/CustomizableObjectNodeDetails.cpp b/Engine/Plugins/Mutable/Source/CustomizableObjectEditor/Private/MuCOE/Nodes/CustomizableObjectNodeDetails.cpp
--- a/Engine/Plugins/Mutable/Source/CustomizableObjectEditor/Private/MuCOE/Nodes/CustomizableObjectNodeDetails.cpp
+++ b/Engine/Plugins/Mutable/Source/CustomizableObjectEditor/Private/MuCOE/Nodes/CustomizableObjectNodeDetails.cpp
@@ -5,8 +5,11 @@
 #include "DetailCategoryBuilder.h"
 #include "DetailLayoutBuilder.h"
 #include "DetailWidgetRow.h"
+#include "Input/Reply.h"
 #include "MuCOE/Nodes/CustomizableObjectNode.h"
 #include "MuCOE/PinViewer/SPinViewer.h"
+#include "Widgets/Input/SButton.h"
+#include "Widgets/Images/SImage.h"
 
 class IDetailLayoutBuilder;
 
@@ -32,6 +35,33 @@ void FCustomizableObjectNodeDetails::CustomizeDetails(IDetailLayoutBuilder& Deta
 			}
 			
 			IDetailCategoryBuilder& PinViewerCategoryBuilder = DetailBuilder.EditCategory("PinViewer", LOCTEXT("PinViewer", "Pins"), ECategoryPriority::Uncommon);
+			
+			if (Node->CanCreatePinsFromPinViewer())
+			{
+				FText TooltipText = LOCTEXT("NodeAddPinTooltip", "Create a new pin.");
+
+				PinViewerCategoryBuilder.HeaderContent
+				(
+					SNew(SHorizontalBox)
+					+ SHorizontalBox::Slot()
+					.HAlign(HAlign_Right)
+					[
+						SNew(SButton)
+						.ButtonStyle(FAppStyle::Get(), "SimpleButton")
+						.ContentPadding(FMargin(1, 0))
+						.OnClicked(this, &FCustomizableObjectNodeDetails::AddNewPin, Node)
+						.HAlign(HAlign_Right)
+						.ToolTipText(TooltipText)
+						.VAlign(VAlign_Center)
+						[
+							SNew(SImage)
+							.Image(FAppStyle::Get().GetBrush("Icons.PlusCircle"))
+							.ColorAndOpacity(FSlateColor::UseForeground())
+						]
+					]
+				);
+			}
+
 			PinViewerCategoryBuilder.AddCustomRow(LOCTEXT("PinViewerDetailsCategory", "PinViewer")).ShouldAutoExpand(true)
 			[
 				SNew(SPinViewer).Node(Node)
@@ -45,7 +75,22 @@ void FCustomizableObjectNodeDetails::CustomizeDetails(IDetailLayoutBuilder& Deta
 void FCustomizableObjectNodeDetails::CustomizeDetails(const TSharedPtr<IDetailLayoutBuilder>& DetailBuilder)
 {
 	IDetailCustomization::CustomizeDetails(DetailBuilder);
+
+	DetailBuilderPtr = DetailBuilder;
 }
 
 
+FReply FCustomizableObjectNodeDetails::AddNewPin(UCustomizableObjectNode* Node)
+{
+	if (Node && DetailBuilderPtr)
+	{
+		Node->CreatePinFromPinViewer();
+		DetailBuilderPtr->ForceRefreshDetails();
+
+		return FReply::Handled();
+	}
+
+	return FReply::Unhandled();
+}
+
 #undef LOCTEXT_NAMESPACE
diff --git a/Engine/Plugins/Mutable/Source/CustomizableObjectEditor/Private/MuCOE/Nodes/CustomizableObjectNodeDetails.h b/Engine/Plugins/Mutable/Source/CustomizableObjectEditor/Private/MuCOE/Nodes/CustomizableObjectNodeDetails.h
--- a/Engine/Plugins/Mutable/Source/CustomizableObjectEditor/Private/MuCOE/Nodes/CustomizableObjectNodeDetails.h
+++ b/Engine/Plugins/Mutable/Source/CustomizableObjectEditor/Private/MuCOE/Nodes/CustomizableObjectNodeDetails.h
@@ -4,7 +4,9 @@
 
 #include "IDetailCustomization.h"
 
+class FReply;
 class IDetailLayoutBuilder;
+class UCustomizableObjectNode;
 
 
 /** Base of all UCustomizableObjectNode details. */
@@ -17,4 +19,12 @@ public:
 	// IDetailCustomization interface
 	virtual void CustomizeDetails(IDetailLayoutBuilder& DetailBuilder) override;
 	virtual void CustomizeDetails(const TSharedPtr<IDetailLayoutBuilder>& DetailBuilder) override;
+
+	//
+	FReply AddNewPin(UCustomizableObjectNode* Node);
+
+private: 
+
+	// Pointer to the DetailBuilder. Needed to refresh the details when we create a new pin from the pin viewer.
+	TSharedPtr<IDetailLayoutBuilder> DetailBuilderPtr;
 };
diff --git a/Engine/Plugins/Mutable/Source/CustomizableObjectEditor/Private/MuCOE/SMutableCodeViewer.cpp b/Engine/Plugins/Mutable/Source/CustomizableObjectEditor/Private/MuCOE/SMutableCodeViewer.cpp
--- a/Engine/Plugins/Mutable/Source/CustomizableObjectEditor/Private/MuCOE/SMutableCodeViewer.cpp
+++ b/Engine/Plugins/Mutable/Source/CustomizableObjectEditor/Private/MuCOE/SMutableCodeViewer.cpp
@@ -1777,6 +1777,17 @@ void SMutableCodeViewer::GenerateElementRecursive(const int32& InStateIndex, mu:
 		break;
 	}
 
+	case mu::EOpType::SC_MATERIAL_BREAK:
+	case mu::EOpType::CO_MATERIAL_BREAK:
+	case mu::EOpType::IM_MATERIAL_BREAK:
+	{
+		mu::OP::MaterialBreakArgs Args = InProgram.GetOpArgs<mu::OP::MaterialBreakArgs>(InParentAddress);
+		AddOpFunc(Args.Material, TEXT("Material "));
+		break;
+	}
+
+	//TODO(Max) Add MI_Parameter and MI_CONSTANT
+
 	// Add here more operation types to define how they are exposed in the tree (set a Caption)
 		
 	default:
diff --git a/Engine/Plugins/Mutable/Source/CustomizableObjectEditor/Private/MuCOE/SMutableGraphViewer.cpp b/Engine/Plugins/Mutable/Source/CustomizableObjectEditor/Private/MuCOE/SMutableGraphViewer.cpp
--- a/Engine/Plugins/Mutable/Source/CustomizableObjectEditor/Private/MuCOE/SMutableGraphViewer.cpp
+++ b/Engine/Plugins/Mutable/Source/CustomizableObjectEditor/Private/MuCOE/SMutableGraphViewer.cpp
@@ -127,6 +127,7 @@ static const char* MutableNodeNames[] =
 	"ImageVariation",
 	"ImageNormalComposite",
 	"ImageTransform",
+	"ImageFromMaterialParameter",
 
 	"Bool",
 	"BoolConstant",
@@ -212,6 +213,10 @@ static const char* MutableNodeNames[] =
 	"MaterialSwitch",
 	"MaterialVariation",
 	"MaterialParameter",
+
+	"ImageMaterialBreak",
+	"ScalarMaterialBreak",
+	"ColorMaterialBreak",
 };
 
 static_assert(UE_ARRAY_COUNT(MutableNodeNames) == SIZE_T(mu::Node::EType::Count));
diff --git a/Engine/Plugins/Mutable/Source/CustomizableObjectEditor/Public/MuCOE/Nodes/CustomizableObjectNode.h b/Engine/Plugins/Mutable/Source/CustomizableObjectEditor/Public/MuCOE/Nodes/CustomizableObjectNode.h
--- a/Engine/Plugins/Mutable/Source/CustomizableObjectEditor/Public/MuCOE/Nodes/CustomizableObjectNode.h
+++ b/Engine/Plugins/Mutable/Source/CustomizableObjectEditor/Public/MuCOE/Nodes/CustomizableObjectNode.h
@@ -274,6 +274,11 @@ public:
 	/** @return true if the details should spawn a Pin Viewer. */
 	UE_API virtual bool HasPinViewer() const;
 
+	/** @return true if the pin viewer's category should spawn a button to add new pins */
+	UE_API virtual bool CanCreatePinsFromPinViewer() const;
+
+	UE_API virtual void CreatePinFromPinViewer() {};
+
 	/** Returns pin custom details. Override if required. */
 	UE_API virtual TSharedPtr<IDetailsView> CustomizePinDetails(const UEdGraphPin& Pin) const;
 
@@ -300,6 +305,9 @@ public:
 
 	/** If a node has a static string pin, this function determines if it should construct an input text widget next to the pin or not. */
 	UE_API virtual bool CreateStaticStringPinWidget() const;
+
+	// Returns an array with all the PC_Types that can be created from the pin viewer
+	UE_API virtual TArray<FName> GetAllowedPinViewerCreationTypes() const;
 	
 	UPROPERTY()
 	FPostEditChangePropertyDelegate PostEditChangePropertyDelegate;
diff --git a/Engine/Plugins/Mutable/Source/MutableRuntime/Internal/MuR/Image.h b/Engine/Plugins/Mutable/Source/MutableRuntime/Internal/MuR/Image.h
--- a/Engine/Plugins/Mutable/Source/MutableRuntime/Internal/MuR/Image.h
+++ b/Engine/Plugins/Mutable/Source/MutableRuntime/Internal/MuR/Image.h
@@ -6,6 +6,7 @@
 
 #include "MuR/MutableMath.h"
 #include "MuR/MutableTrace.h"
+#include "MuR/Operations.h"
 #include "MuR/Ptr.h"
 #include "MuR/RefCounted.h"
 #include "MuR/Serialisation.h"
@@ -13,6 +14,7 @@
 #include "HAL/LowLevelMemTracker.h"
 #include "HAL/Platform.h"
 #include "Misc/AssertionMacros.h"
+#include "Engine/Texture.h"
 
 #include "MuR/ImageDataStorage.h"
 
@@ -126,9 +128,13 @@ namespace mu
 		uint16 RelevancyMinY = 0;
 		uint16 RelevancyMaxY = 0;
 
+		TStrongObjectPtr<UTexture> Texture;
+
 		/** Only valid if the right flags are set, this identifies a referenced image. */
 		uint32 ReferenceID = 0;
 
+		OP::ADDRESS LazyAddress = 0;
+
 		/** Pixel data for all lods. */
 		FImageDataStorage DataStorage;
 
@@ -148,6 +154,7 @@ namespace mu
 			Result->RelevancyMaxY = RelevancyMaxY;
 			Result->DataStorage = DataStorage;
 			Result->ReferenceID = ReferenceID;
+			Result->LazyAddress = LazyAddress;
 
 			return Result;
 		}
@@ -168,6 +175,7 @@ namespace mu
 			RelevancyMaxY = Other->RelevancyMaxY;
 			DataStorage = Other->DataStorage;
 			ReferenceID = Other->ReferenceID;
+			LazyAddress = Other->LazyAddress;
 		}
 
 		void CopyMove(FImage* Other)
@@ -184,7 +192,8 @@ namespace mu
 			RelevancyMinY = Other->RelevancyMinY;
 			RelevancyMaxY = Other->RelevancyMaxY;
 			ReferenceID = Other->ReferenceID;
-			
+			LazyAddress = Other->LazyAddress;
+
 			DataStorage = MoveTemp(Other->DataStorage);
 		}
 
@@ -200,7 +209,8 @@ namespace mu
 		{
 			return 
 				(DataStorage == Other.DataStorage) &&
-				(ReferenceID == Other.ReferenceID);
+				(ReferenceID == Other.ReferenceID) &&
+				(LazyAddress == Other.LazyAddress);
 		}
 
 		//-----------------------------------------------------------------------------------------
@@ -243,6 +253,8 @@ namespace mu
 		//! Get the rect inside the image bounding the non-black content of the image.
 		UE_API void GetNonBlackRect(FImageRect& OutRect) const;
 		UE_API void GetNonBlackRect_Reference(FImageRect& OutRect) const;
+
+		UE_API bool IsLazy() const;
 	};
 
 	/** This struct contains image operations that may need to allocate and free temporary images to work.
diff --git a/Engine/Plugins/Mutable/Source/MutableRuntime/Internal/MuR/ModelPrivate.h b/Engine/Plugins/Mutable/Source/MutableRuntime/Internal/MuR/ModelPrivate.h
--- a/Engine/Plugins/Mutable/Source/MutableRuntime/Internal/MuR/ModelPrivate.h
+++ b/Engine/Plugins/Mutable/Source/MutableRuntime/Internal/MuR/ModelPrivate.h
@@ -373,6 +373,9 @@ namespace mu
 		 */
 		TArray<TArray<uint16>> ParameterLists;
 
+		/** Constant Material Data */
+		TArray<TSharedPtr<const FMaterial>> ConstantMaterials;
+
 #if WITH_EDITOR
 		/** 
 		 * State of the program. True unless the streamed resources were destroyed,
@@ -407,6 +410,7 @@ namespace mu
             Arch << Parameters;
             Arch << Ranges;
             Arch << ParameterLists;
+			Arch << ConstantMaterials;
         }
 
         /** */
@@ -436,6 +440,7 @@ namespace mu
             Arch >> Parameters;
             Arch >> Ranges;
             Arch >> ParameterLists;
+            Arch >> ConstantMaterials;
         }
 
         /** Debug method that sanity-checks the program with a variety of tests. */
@@ -453,6 +458,7 @@ namespace mu
 					return ConstantImageLODsStreamed.Contains(RomIndex);
 				case uint32(ERomDataType::Mesh):
 					return ConstantMeshesStreamed.Contains(RomIndex);
+				//TODO(Max): Make materials loadable and unloadable from a rom
 				default:
 					check(false);
 					break;
@@ -496,6 +502,7 @@ namespace mu
 				}
 				break;
 			}
+			//TODO(Max): Make materials loadable and unloadable from a rom
 			default:
 				check(false);
 				break;
@@ -676,6 +683,11 @@ namespace mu
             return index;
         }
 
+		OP::ADDRESS AddConstant(TSharedPtr<const FMaterial> Material)
+		{
+			OP::ADDRESS index = OP::ADDRESS(ConstantMaterials.AddUnique(Material));
+			return index;
+		}
 
 		inline TSharedPtr<const FImage> GetImageLOD(FConstantResourceIndex Index) const
 		{
diff --git a/Engine/Plugins/Mutable/Source/MutableRuntime/Internal/MuR/Operations.h b/Engine/Plugins/Mutable/Source/MutableRuntime/Internal/MuR/Operations.h
--- a/Engine/Plugins/Mutable/Source/MutableRuntime/Internal/MuR/Operations.h
+++ b/Engine/Plugins/Mutable/Source/MutableRuntime/Internal/MuR/Operations.h
@@ -2,7 +2,7 @@
 
 #pragma once
 
-#include "MuR/Image.h"
+#include "MuR/ImageTypes.h"
 #include "MuR/Types.h"
 #include "MuR/SerialisationPrivate.h"
 #include "HAL/PlatformMath.h"
@@ -16,6 +16,7 @@
 
 namespace mu
 {
+	enum class EImageFormat : uint8;
 
     //!
     enum class EOpType : uint16
@@ -79,6 +80,14 @@ namespace mu
 		ED_SWITCH,
 		MI_SWITCH,
 
+		//! Selects a parameter value from a Material
+		SC_MATERIAL_BREAK,
+		CO_MATERIAL_BREAK,
+		IM_MATERIAL_BREAK,
+
+		//! Converts a texture of a material parameter into a texture parameter to process it at runtime.
+		IM_PARAMETER_FROM_MATERIAL,
+
         //-----------------------------------------------------------------------------------------
         // Boolean operations
         //-----------------------------------------------------------------------------------------
@@ -384,6 +393,8 @@ namespace mu
     {
         typedef ADDRESS_TYPE ADDRESS;
         typedef ADDRESS_TYPE CONSTANT_STRING_ADDRESS;
+        typedef ADDRESS_TYPE STRING_ADDRESS;
+		typedef ADDRESS_TYPE LIST_UINT32_ADDRESS;
 
         /** Arguments for every operation type. */
         struct BoolConstantArgs
@@ -431,6 +442,16 @@ namespace mu
 			ADDRESS variable;
 		};
 
+		struct MaterialParameterArgs
+		{
+			ADDRESS Variable;
+
+			LIST_UINT32_ADDRESS ScalarParameterNames;
+			LIST_UINT32_ADDRESS ColorParameterNames;
+			LIST_UINT32_ADDRESS ImageParameterNames;
+			LIST_UINT32_ADDRESS ImageParameterAddress;
+		};
+
 		struct MeshParameterArgs
 		{
 			ADDRESS variable;
@@ -1085,6 +1106,18 @@ namespace mu
 			ADDRESS Mesh;
 			uint8 LayoutIndex;
 		};
+
+		struct MaterialBreakArgs
+		{
+			ADDRESS Material;
+			STRING_ADDRESS ParameterName;
+		};
+
+		struct MaterialBreakImageParameterArgs
+		{
+			ADDRESS MaterialParameter;
+			STRING_ADDRESS ParameterName;
+		};
     };
 
 
diff --git a/Engine/Plugins/Mutable/Source/MutableRuntime/Private/MuR/CodeRunner.cpp b/Engine/Plugins/Mutable/Source/MutableRuntime/Private/MuR/CodeRunner.cpp
--- a/Engine/Plugins/Mutable/Source/MutableRuntime/Private/MuR/CodeRunner.cpp
+++ b/Engine/Plugins/Mutable/Source/MutableRuntime/Private/MuR/CodeRunner.cpp
@@ -158,6 +158,7 @@ namespace mu
 				InSettings, InSystem, InExecutionStrategy, InModel, InParams, At, InLODMask, ExecutionOptions, InImageLOD, Type, ExternalResourceProvider);
 	}
 
+
     CodeRunner::CodeRunner(FPrivateToken PrivateToken, 
 		const FSettings& InSettings,
 		FSystem::Private* InSystem,
@@ -1178,11 +1179,10 @@ namespace mu
 		case EOpType::MI_CONSTANT:
 		{
 			OP::ResourceConstantArgs Args = Program.GetOpArgs<OP::ResourceConstantArgs>(Item.At);
+			TSharedPtr<const FMaterial> SourceConst = Program.ConstantMaterials[Args.value];
+			check(SourceConst);
 
-			TSharedPtr<FMaterial> Source = MakeShared<FMaterial>();
-			Source->ReferenceID = Args.value;
-
-			StoreMaterial(Item, Source);
+			StoreMaterial(Item, SourceConst);
 			break;
 		}
 
@@ -5424,6 +5424,73 @@ namespace mu
 			break;
 		}
 
+		case EOpType::IM_MATERIAL_BREAK:
+		{
+			const OP::MaterialBreakArgs Args = Program.GetOpArgs<OP::MaterialBreakArgs>(Item.At);
+
+			switch (Item.Stage)
+			{
+			case 0:
+			{
+				AddOp(FScheduledOp(Item.At, Item, 1),
+					FScheduledOp(Args.Material, Item));
+				break;
+			}
+			case 1:
+			{
+				TSharedPtr<const FMaterial> Material = LoadMaterial(FCacheAddress(Args.Material, Item));
+
+				if (!Material)
+				{
+					StoreImage(Item, nullptr);
+					break;
+				}
+
+				check(Args.ParameterName < (uint32)Program.ConstantStrings.Num());
+				const FString& ParameterName = Program.ConstantStrings[Args.ParameterName];
+				const TSharedPtr<const FImage>* Image = Material->ImageParameters.Find(FName(ParameterName));
+				
+				if (!Image)
+				{
+					StoreImage(Item, nullptr);
+				}
+				else
+				{
+					if (Image->Get()->IsLazy())
+					{
+						FScheduledOp ItemNextStage(Item.At, Item, 2);
+						ItemNextStage.CustomState = Image->Get()->LazyAddress;
+
+						AddOp(ItemNextStage,
+							FScheduledOp(Image->Get()->LazyAddress, Item));
+					}
+					else
+					{
+						StoreImage(Item, *Image);
+					}
+				}
+				break;
+			}
+			case 2:
+			{
+				TSharedPtr<const FMaterial> Material = LoadMaterial(FCacheAddress(Args.Material, Item));
+
+				check(Args.ParameterName < (uint32)Program.ConstantStrings.Num());
+				const FString& ParameterName = Program.ConstantStrings[Args.ParameterName];
+				const TSharedPtr<const FImage>* Image = Material->ImageParameters.Find(FName(ParameterName));
+
+				TSharedPtr<const FImage> pImage = LoadImage(FCacheAddress(Image->Get()->LazyAddress, Item));
+				StoreImage(Item, pImage);
+				break;
+			}
+
+			default:
+				unimplemented();
+			}
+
+			break;
+		}
+
         default:
             if (type!=EOpType::NONE)
             {
@@ -5682,7 +5749,7 @@ namespace mu
             break;
         }
     }
-
+	
     //---------------------------------------------------------------------------------------------
     void CodeRunner::RunCode_Int(const FScheduledOp& Item, const FParameters* pParams, const FModel* InModel )
     {
@@ -5843,6 +5910,44 @@ namespace mu
             break;
         }
 
+		case EOpType::SC_MATERIAL_BREAK:
+		{
+			const OP::MaterialBreakArgs Args = Program.GetOpArgs<OP::MaterialBreakArgs>(Item.At);
+
+			switch (Item.Stage)
+			{
+			case 0:
+			{
+				AddOp(FScheduledOp(Item.At, Item, 1),
+					FScheduledOp(Args.Material, Item));
+				break;
+			}
+			case 1:
+			{
+				TSharedPtr<const FMaterial> Material = LoadMaterial(FCacheAddress(Args.Material, Item));
+
+				if (!Material)
+				{
+					StoreScalar(Item, 0.0f);
+					break;
+				}
+
+				check(Args.ParameterName < (uint32)Program.ConstantStrings.Num());
+				const FString& ParameterName = Program.ConstantStrings[Args.ParameterName];
+				const float Value = *Material->ScalarParameters.Find(FName(ParameterName));
+
+				StoreScalar(Item, Value);
+
+				break;
+			}
+
+			default:
+				unimplemented();
+			}
+
+			break;
+		}
+
         default:
             check( false );
             break;
@@ -6126,6 +6231,44 @@ namespace mu
 			break;
 		}
 
+		case EOpType::CO_MATERIAL_BREAK:
+		{
+			const OP::MaterialBreakArgs Args = Program.GetOpArgs<OP::MaterialBreakArgs>(Item.At);
+
+			switch (Item.Stage)
+			{
+			case 0:
+			{
+				AddOp(FScheduledOp(Item.At, Item, 1),
+					FScheduledOp(Args.Material, Item));
+				break;
+			}
+			case 1:
+			{
+				TSharedPtr<const FMaterial> Material = LoadMaterial(FCacheAddress(Args.Material, Item));
+
+				if (!Material)
+				{
+					StoreColor(Item, FVector4f::Zero());
+					break;
+				}
+
+				check(Args.ParameterName < (uint32)Program.ConstantStrings.Num());
+				const FString& ParameterName = Program.ConstantStrings[Args.ParameterName];
+				const FVector4f Value = *Material->ColorParameters.Find(FName(ParameterName));
+
+				StoreColor(Item, Value);
+
+				break;
+			}
+
+			default:
+				unimplemented();
+			}
+			
+			break;
+		}
+
         default:
             check( false );
             break;
@@ -6218,13 +6361,61 @@ namespace mu
 			MUTABLE_CPUPROFILER_SCOPE(RunCodeImageDesc_MI_PARAMETER);
 			check(Item.Stage == 0);
 
-			OP::ParameterArgs Args = Program.GetOpArgs<OP::ParameterArgs>(Item.At);
-			TSharedPtr<FRangeIndex> Index = BuildCurrentOpRangeIndex(Item, Args.variable);
+			OP::MaterialParameterArgs Args = Program.GetOpArgs<OP::MaterialParameterArgs >(Item.At);
+			TSharedPtr<FRangeIndex> Index = BuildCurrentOpRangeIndex(Item, Args.Variable);
 
 			TSharedPtr<FMaterial> Material = MakeShared<FMaterial>();
-			Material->Material = TStrongObjectPtr(InParams->GetMaterialValue(Args.variable, Index.Get()));
-			
-			StoreMaterial(Item, Material);
+			Material->Material = TStrongObjectPtr(InParams->GetMaterialValue(Args.Variable, Index.Get()));
+
+			if (Material->Material)
+			{
+				// Images
+				const TArray<uint32>& ImageParameterNames = Program.ConstantUInt32Lists[Args.ImageParameterNames];
+				const TArray<uint32>& ImageParameterOperations = Program.ConstantUInt32Lists[Args.ImageParameterAddress];
+
+				for (int32 ParameterIndex = 0; ParameterIndex < ImageParameterNames.Num(); ++ParameterIndex)
+				{
+					const FString& ParameterName = Program.ConstantStrings[ImageParameterNames[ParameterIndex]];
+
+					UTexture* Texture;
+					bool bParameterFound = Material->Material->GetTextureParameterValue(FName(ParameterName), Texture);
+
+					if (bParameterFound)
+					{
+						TSharedPtr<FImage> NewParameterImage = MakeShared<FImage>();
+						NewParameterImage->Texture = TStrongObjectPtr(Texture);
+						NewParameterImage->LazyAddress = ImageParameterOperations[ParameterIndex];
+
+						Material->ImageParameters.Add(FName(ParameterName), NewParameterImage);
+					}
+				}
+
+				// Scalars
+				const TArray<uint32>& ScalarParameterNames = Program.ConstantUInt32Lists[Args.ScalarParameterNames];
+
+				for (int32 ParameterIndex = 0; ParameterIndex < ScalarParameterNames.Num(); ++ParameterIndex)
+				{
+					const FString& ParameterName = Program.ConstantStrings[ScalarParameterNames[ParameterIndex]];
+
+					float ScalarValue;
+					Material->Material->GetScalarParameterValue(FName(ParameterName), ScalarValue);
+					Material->ScalarParameters.Add(FName(ParameterName), ScalarValue);
+				}
+
+				// Vectors
+				const TArray<uint32>& ColorParameterNames = Program.ConstantUInt32Lists[Args.ColorParameterNames];
+
+				for (int32 ParameterIndex = 0; ParameterIndex < ColorParameterNames.Num(); ++ParameterIndex)
+				{
+					const FString& ParameterName = Program.ConstantStrings[ColorParameterNames[ParameterIndex]];
+
+					FLinearColor VectorValue;
+					Material->Material->GetVectorParameterValue(FName(ParameterName), VectorValue);
+					Material->ColorParameters.Add(FName(ParameterName), VectorValue);
+				}
+
+				StoreMaterial(Item, Material);
+			}
 			break;
 		}
 		default: unimplemented();
@@ -7583,6 +7774,124 @@ namespace mu
 			break;
 		}
 
+		case EOpType::IM_MATERIAL_BREAK:
+		{
+			MUTABLE_CPUPROFILER_SCOPE(RunCodeImageDesc_IM_MATERIAL_BREAK);
+			const OP::MaterialBreakArgs Args = Program.GetOpArgs<OP::MaterialBreakArgs>(Item.At);
+
+			switch (Item.Stage)
+			{
+			case 0:
+			{
+				FScheduledOp MaterialOp(Args.Material, Item);
+				MaterialOp.Type = FScheduledOp::EType::Full;
+
+				AddOp(FScheduledOp(Item.At, Item, 1),
+					MaterialOp);
+				break;
+			}
+			case 1:
+			{
+				FCacheAddress CacheMaterialOp(Args.Material, Item);
+				CacheMaterialOp.Type = FScheduledOp::EType::Full;
+
+				TSharedPtr<const FMaterial> Material = LoadMaterial(CacheMaterialOp);
+
+				if (!Material)
+				{
+					ImageDescResults.FindOrAdd(Item.At);
+					StoreValidDesc(Item);
+				}
+				else
+				{
+					check(Args.ParameterName < (uint32)Program.ConstantStrings.Num());
+					const FString& ParameterName = Program.ConstantStrings[Args.ParameterName];
+					const TSharedPtr<const FImage>* Image = Material->ImageParameters.Find(FName(ParameterName));
+
+					if (!Image)
+					{
+						ImageDescResults.FindOrAdd(Item.At);
+						StoreValidDesc(Item);
+					}
+					else
+					{
+						if (Image->Get()->IsLazy())
+						{
+							FScheduledOp ItemNextStage(Item.At, Item, 2);
+							ItemNextStage.CustomState = Image->Get()->LazyAddress;
+
+							AddOp(ItemNextStage,
+								FScheduledOp(Image->Get()->LazyAddress, Item));
+						}
+					}
+				}
+
+				break;
+			}
+			case 2:
+			{
+				FCacheAddress CacheMaterialOp(Args.Material, Item);
+				CacheMaterialOp.Type = FScheduledOp::EType::Full;
+
+				TSharedPtr<const FMaterial> Material = LoadMaterial(CacheMaterialOp);
+
+				if (!Material)
+				{
+					ImageDescResults.FindOrAdd(Item.At);
+				}
+				else
+				{
+					check(Args.ParameterName < (uint32)Program.ConstantStrings.Num());
+					const FString& ParameterName = Program.ConstantStrings[Args.ParameterName];
+					const TSharedPtr<const FImage>* Image = Material->ImageParameters.Find(FName(ParameterName));
+
+					FExtendedImageDesc& Result = ImageDescResults.FindOrAdd(Item.At);
+					Result = ImageDescResults.FindChecked(Image->Get()->LazyAddress);
+				}
+
+				StoreValidDesc(Item);
+				break;
+			}
+
+			default:
+				check(false);
+			}
+
+			break;
+		}
+
+		case EOpType::IM_PARAMETER_FROM_MATERIAL:
+		{
+			check(Item.Stage == 0);
+
+			const OP::MaterialBreakImageParameterArgs Args = Program.GetOpArgs<OP::MaterialBreakImageParameterArgs>(Item.At);
+
+			TSharedPtr<FRangeIndex> Index = BuildCurrentOpRangeIndex(Item, Args.MaterialParameter);
+
+			// Get material parameter from the array of parameters
+			TSharedPtr<FMaterial> Material = MakeShared<FMaterial>();
+			Material->Material = TStrongObjectPtr(pParams->GetMaterialValue(Args.MaterialParameter, Index.Get()));
+
+			// Get the texture parameter name
+			check(Args.ParameterName < (uint32)Program.ConstantStrings.Num());
+			const FString& ParameterName = Program.ConstantStrings[Args.ParameterName];
+
+			UTexture* Image = nullptr;
+
+			// Get the parameter texture from the UMaterial
+			if (Material->Material)
+			{
+				Material->Material->GetTextureParameterValue(FName(ParameterName), Image);
+			}
+
+			FExtendedImageDesc& Result = ImageDescResults.FindOrAdd(Item.At);
+			Result = GetExternalImageDesc(Image);
+
+			StoreValidDesc(Item);
+
+			break;
+		}
+
 		default:
 			if (Type != EOpType::NONE)
 			{
diff --git a/Engine/Plugins/Mutable/Source/MutableRuntime/Private/MuR/CodeRunnerBegin.cpp b/Engine/Plugins/Mutable/Source/MutableRuntime/Private/MuR/CodeRunnerBegin.cpp
--- a/Engine/Plugins/Mutable/Source/MutableRuntime/Private/MuR/CodeRunnerBegin.cpp
+++ b/Engine/Plugins/Mutable/Source/MutableRuntime/Private/MuR/CodeRunnerBegin.cpp
@@ -384,6 +384,7 @@ namespace mu
 			case EOpType::ME_PARAMETER:
 			case EOpType::IM_PARAMETER:
 			case EOpType::MI_PARAMETER:
+			case EOpType::IM_PARAMETER_FROM_MATERIAL:
 				{
 					StoreNone(Item);
 					break;
@@ -2331,7 +2332,19 @@ namespace mu
 			case EOpType::MI_CONSTANT:
 				{
 					OP::ResourceConstantArgs Args = Program.GetOpArgs<OP::ResourceConstantArgs>(Item.At);
-					MI_REFERENCE_ID.Push(Args.value);
+					TSharedPtr<const FMaterial> SourceConst = Program.ConstantMaterials[Args.value];
+					MI_REFERENCE_ID.Add(SourceConst->ReferenceID);
+
+					StoreNone(Item);
+					break;
+				}
+
+			case EOpType::IM_MATERIAL_BREAK:
+			case EOpType::SC_MATERIAL_BREAK:
+			case EOpType::CO_MATERIAL_BREAK:
+				{
+					OP::MaterialBreakArgs Args = Program.GetOpArgs<OP::MaterialBreakArgs>(Item.At);
+					PushOp(FScheduledOpInline(Args.Material, Item));
 
 					StoreNone(Item);
 					break;
diff --git a/Engine/Plugins/Mutable/Source/MutableRuntime/Private/MuR/CodeRunnerTasked.cpp b/Engine/Plugins/Mutable/Source/MutableRuntime/Private/MuR/CodeRunnerTasked.cpp
--- a/Engine/Plugins/Mutable/Source/MutableRuntime/Private/MuR/CodeRunnerTasked.cpp
+++ b/Engine/Plugins/Mutable/Source/MutableRuntime/Private/MuR/CodeRunnerTasked.cpp
@@ -2605,6 +2605,41 @@ namespace mu
 			break;
 		}
 
+		case EOpType::IM_PARAMETER_FROM_MATERIAL:
+		{
+			const OP::MaterialBreakImageParameterArgs Args = Program.GetOpArgs<OP::MaterialBreakImageParameterArgs>(Item.At);
+			TSharedPtr<FRangeIndex> Index = BuildCurrentOpRangeIndex(Item, Args.MaterialParameter);
+
+			// Get material parameter from the array of parameters
+			TSharedPtr<FMaterial> Material = MakeShared<FMaterial>();
+			Material->Material = TStrongObjectPtr(Params->GetMaterialValue(Args.MaterialParameter, Index.Get()));
+
+			// Get the texture parameter name
+			check(Args.ParameterName < (uint32)Program.ConstantStrings.Num());
+			const FString& ParameterName = Program.ConstantStrings[Args.ParameterName];
+
+			// Get the parameter texture from the UMaterial
+			if (Material->Material)
+			{
+				UTexture* Texture;
+				bool bParameterFound = Material->Material->GetTextureParameterValue(FName(ParameterName), Texture);
+
+				if (bParameterFound && Texture)
+				{
+					check(ImageLOD < TNumericLimits<uint8>::Max() && ImageLOD >= 0);
+					check(ImageLOD + static_cast<int32>(Item.ExecutionOptions) < TNumericLimits<uint8>::Max());
+
+					const uint8 MipmapsToSkip = Item.ExecutionOptions + static_cast<uint8>(ImageLOD);
+
+					CodeRunner::FExternalResourceId FullId;
+					FullId.ImageParameter = Texture;
+					Issued = MakeShared<FImageExternalLoadTask>(Item, MipmapsToSkip, FullId);
+				}
+			}
+
+			break;
+		}
+
 		case EOpType::IM_REFERENCE:
 		{
 			OP::ResourceReferenceArgs Args = Model->GetPrivate()->Program.GetOpArgs<OP::ResourceReferenceArgs>(Item.At);
diff --git a/Engine/Plugins/Mutable/Source/MutableRuntime/Private/MuR/Image.cpp b/Engine/Plugins/Mutable/Source/MutableRuntime/Private/MuR/Image.cpp
--- a/Engine/Plugins/Mutable/Source/MutableRuntime/Private/MuR/Image.cpp
+++ b/Engine/Plugins/Mutable/Source/MutableRuntime/Private/MuR/Image.cpp
@@ -8,7 +8,6 @@
 #include "MuR/ImagePrivate.h"
 #include "MuR/MutableMath.h"
 #include "MuR/MutableTrace.h"
-
 #include "ImageCore.h"
 
 namespace mu
@@ -73,12 +72,14 @@ namespace mu
 		// Remove non-persistent flags.
 		uint8 flags = Flags & ~IF_HAS_RELEVANCY_MAP;
 		Arch << flags;
+		Arch << LazyAddress;
 	}
 
 	void FImage::Unserialise(FInputArchive& Arch)
 	{
 		Arch >> DataStorage;
 		Arch >> Flags;
+		Arch >> LazyAddress;
 	}
 
 
@@ -618,6 +619,12 @@ namespace mu
 	 }
 
 
+	bool FImage::IsLazy() const
+	{
+		return LazyAddress != 0;
+	}
+
+
     //---------------------------------------------------------------------------------------------
     void FImage::GetNonBlackRect(FImageRect& rect) const
     {            
diff --git a/Engine/Plugins/Mutable/Source/MutableRuntime/Private/MuR/Material.cpp b/Engine/Plugins/Mutable/Source/MutableRuntime/Private/MuR/Material.cpp
--- a/Engine/Plugins/Mutable/Source/MutableRuntime/Private/MuR/Material.cpp
+++ b/Engine/Plugins/Mutable/Source/MutableRuntime/Private/MuR/Material.cpp
@@ -3,26 +3,68 @@
 
 #include "MuR/Material.h"
 
-#include "Containers/UnrealString.h"
 #include "HAL/LowLevelMemTracker.h"
-#include "HAL/PlatformCrt.h"
-#include "HAL/UnrealMemory.h"
 #include "MuR/MutableTrace.h"
+#include "MuR/Serialisation.h"
+#include "MuR/SerialisationPrivate.h"
+#include "MuR/Image.h"
+
 
 namespace mu
 {
+	void FMaterial::Serialise(const FMaterial* MaterialPtr, FOutputArchive& Arch)
+	{
+		Arch << *MaterialPtr;
+	}
+
+
+	TSharedPtr<FMaterial> FMaterial::StaticUnserialise(FInputArchive& Arch)
+	{
+		MUTABLE_CPUPROFILER_SCOPE(MaterialUnserialise)
+		LLM_SCOPE_BYNAME(TEXT("MutableRuntime"));
+
+		TSharedPtr<FMaterial> Result = MakeShared<FMaterial>();
+		Arch >> *Result;
+
+		return Result;
+	}
+
+
+	void FMaterial::Serialise(FOutputArchive& Arch) const
+	{
+		Arch << ReferenceID;
+		Arch << ImageParameters;
+		Arch << ColorParameters;
+		Arch << ScalarParameters;
+	}
+
+
+	void FMaterial::Unserialise(FInputArchive& Arch)
+	{
+		Arch >>ReferenceID;
+		Arch >> ImageParameters;
+		Arch >> ColorParameters;
+		Arch >> ScalarParameters;
+	}
+
+
 	TSharedPtr<FMaterial> FMaterial::Clone() const
 	{
 		TSharedPtr<FMaterial> Result = MakeShared<FMaterial>();
 
 		Result->Material = Material;
 		Result->ReferenceID = ReferenceID;
+		Result->ImageParameters = ImageParameters;
+		Result->ColorParameters = ColorParameters;
+		Result->ScalarParameters = ScalarParameters;
 
 		return Result;
 	}
 
+
 	bool FMaterial::operator==(const FMaterial& Other) const
 	{
-		return ReferenceID == Other.ReferenceID && Material == Other.Material;
+		return ReferenceID == Other.ReferenceID && Material == Other.Material	&&
+			ImageParameters == Other.ImageParameters &&	ColorParameters == Other.ColorParameters && ScalarParameters == Other.ScalarParameters;
 	};
 }
diff --git a/Engine/Plugins/Mutable/Source/MutableRuntime/Private/MuR/Operations.cpp b/Engine/Plugins/Mutable/Source/MutableRuntime/Private/MuR/Operations.cpp
--- a/Engine/Plugins/Mutable/Source/MutableRuntime/Private/MuR/Operations.cpp
+++ b/Engine/Plugins/Mutable/Source/MutableRuntime/Private/MuR/Operations.cpp
@@ -3,6 +3,7 @@
 
 #include "MuR/Operations.h"
 
+#include "MuR/Image.h"
 #include "MuR/ModelPrivate.h"
 
 
@@ -64,6 +65,11 @@ namespace mu
 		{ EDataType::ExtensionData,	},	// ED_SWITCH
 		{ EDataType::Material,		},	// MI_SWITCH
 
+		{ EDataType::Scalar,		},	// SC_MATERIAL_BREAK
+		{ EDataType::Color,			},	// CO_MATERIAL_BREAK
+		{ EDataType::Image,			},	// IM_MATERIAL_BREAK
+		{ EDataType::Image,			},	// IM_PARAMETER_FROM_MATERIAL
+
 		{ EDataType::Bool,			},	// BO_EQUAL_SC_CONST
 		{ EDataType::Bool,			},	// BO_AND
 		{ EDataType::Bool,			},	// BO_OR
@@ -187,6 +193,7 @@ namespace mu
 		case EOpType::MI_PARAMETER:
 		case EOpType::IM_REFERENCE:
 		case EOpType::ME_REFERENCE:
+		case EOpType::IM_PARAMETER_FROM_MATERIAL:
 			break;
 
         case EOpType::SC_CURVE:
@@ -278,6 +285,16 @@ namespace mu
             break;
         }
 
+		case EOpType::SC_MATERIAL_BREAK:
+		case EOpType::CO_MATERIAL_BREAK:
+		case EOpType::IM_MATERIAL_BREAK:
+		{
+			OP::MaterialBreakArgs Args = program.GetOpArgs<OP::MaterialBreakArgs>(at);
+			f(Args.Material);
+			f(Args.ParameterName);
+			break;
+		}
+
         //-------------------------------------------------------------------------------------
         case EOpType::BO_EQUAL_INT_CONST:
         {
diff --git a/Engine/Plugins/Mutable/Source/MutableRuntime/Public/MuR/Material.h b/Engine/Plugins/Mutable/Source/MutableRuntime/Public/MuR/Material.h
--- a/Engine/Plugins/Mutable/Source/MutableRuntime/Public/MuR/Material.h
+++ b/Engine/Plugins/Mutable/Source/MutableRuntime/Public/MuR/Material.h
@@ -2,7 +2,6 @@
 
 #pragma once
 
-#include "HAL/PlatformMath.h"
 #include "MuR/Ptr.h"
 #include "MuR/RefCounted.h"
 #include "MuR/Types.h"
@@ -14,6 +13,7 @@
 
 namespace mu
 {
+	class FImage;
 
     /** Material type resource */
 	class FMaterial : public FResource
@@ -24,20 +24,35 @@ namespace mu
 		// Life cycle
 		//-----------------------------------------------------------------------------------------
 
+		//! Serialisation
+		static UE_API void Serialise(const FMaterial* MaterialPtr, FOutputArchive& Arch);
+		static UE_API TSharedPtr<FMaterial> StaticUnserialise(FInputArchive& Arch);
+
         //! Clone this material
 		UE_API TSharedPtr<FMaterial> Clone() const;
 
 		// Resource interface
-		int32 GetDataSize() const override { return 0; };
+		virtual int32 GetDataSize() const override { return 0; };
 
 		//-----------------------------------------------------------------------------------------
 		// Own interface
 		//-----------------------------------------------------------------------------------------
 		UE_API bool operator==(const FMaterial& Other) const;
 
+
+		UE_API void Serialise(FOutputArchive&) const;
+		UE_API void Unserialise(FInputArchive&);
+
+
+	public:
+
 		TStrongObjectPtr<UMaterialInterface> Material;
 
 		int32 ReferenceID = INDEX_NONE;
+
+		TMap<FName, TSharedPtr<const FImage>> ImageParameters;
+		TMap<FName, FVector4f> ColorParameters;
+		TMap<FName, float> ScalarParameters;
 	};
 }
 
diff --git a/Engine/Plugins/Mutable/Source/MutableTools/Private/MuT/ASTOpConstantResource.cpp b/Engine/Plugins/Mutable/Source/MutableTools/Private/MuT/ASTOpConstantResource.cpp
--- a/Engine/Plugins/Mutable/Source/MutableTools/Private/MuT/ASTOpConstantResource.cpp
+++ b/Engine/Plugins/Mutable/Source/MutableTools/Private/MuT/ASTOpConstantResource.cpp
@@ -232,10 +232,13 @@ namespace mu
 
 
 	//-------------------------------------------------------------------------------------------------
-	//-------------------------------------------------------------------------------------------------
-	//-------------------------------------------------------------------------------------------------
-	void ASTOpConstantResource::ForEachChild(const TFunctionRef<void(ASTChild&)>)
+	void ASTOpConstantResource::ForEachChild(const TFunctionRef<void(ASTChild&)> f)
 	{
+		/** Image Parameters are evaluated at runtime. */
+		for (TPair<FName, ASTChild >& Element : ImageOperations)
+		{
+			f(Element.Value);
+		}
 	}
 
 
@@ -247,21 +250,28 @@ namespace mu
 			const ASTOpConstantResource* Other = static_cast<const ASTOpConstantResource*>(&OtherUntyped);
 			return Type == Other->Type && ValueHash == Other->ValueHash &&
 				LoadedValue == Other->LoadedValue && Proxy == Other->Proxy
-				&& SourceDataDescriptor == Other->SourceDataDescriptor;
+				&& SourceDataDescriptor == Other->SourceDataDescriptor
+				&& ImageOperations == Other->ImageOperations;
 		}
 		return false;
 	}
 
 
 	//-------------------------------------------------------------------------------------------------
-	mu::Ptr<ASTOp> ASTOpConstantResource::Clone(MapChildFuncRef) const
+	mu::Ptr<ASTOp> ASTOpConstantResource::Clone(MapChildFuncRef MapChild) const
 	{
 		Ptr<ASTOpConstantResource> n = new ASTOpConstantResource();
 		n->Type = Type;
 		n->Proxy = Proxy;
 		n->LoadedValue = LoadedValue;
 		n->ValueHash = ValueHash;
-		n->SourceDataDescriptor = SourceDataDescriptor; 
+		n->SourceDataDescriptor = SourceDataDescriptor;
+
+		for (const TPair<FName, ASTChild>& ImageOperation : ImageOperations)
+		{
+			n->ImageOperations.Add(ImageOperation.Key, ASTChild(n, MapChild(ImageOperation.Value.child())));
+		}
+		
 		return n;
 	}
 
@@ -664,6 +674,24 @@ namespace mu
 					Args.value = Program.AddConstant(pTyped);
 					break;
 				}
+				case EOpType::MI_CONSTANT:
+				{
+					TSharedPtr<FMaterial> MaterialValue = StaticCastSharedPtr<const FMaterial>(GetValue())->Clone();
+					check(MaterialValue);
+
+					// Store linked child images into the material constant
+					for (const TPair<FName, ASTChild>& ImageOperation : ImageOperations)
+					{
+						TSharedPtr<FImage> NewImage = MakeShared<FImage>();
+						NewImage->LazyAddress = ImageOperation.Value->linkedAddress;
+
+						MaterialValue->ImageParameters.Add(FName(ImageOperation.Key), NewImage);
+					}
+
+					Args.value = Program.AddConstant(MaterialValue);
+					
+					break;
+				}
 				default:
 					check(false);
 				}
@@ -703,6 +731,13 @@ namespace mu
 			Result.m_lods = ConstImage->GetLODCount();
 			Result.m_size = ConstImage->GetSize();
 		}
+		else if (Type == EOpType::MI_CONSTANT)
+		{
+			if (TMap<FName, ASTChild>::TConstIterator Iterator = ImageOperations.CreateConstIterator())
+			{
+				Result = Iterator.Value()->GetImageDesc();
+			}
+		}
 		else
 		{
 			check(false);
@@ -933,6 +968,23 @@ namespace mu
 			LoadedValue = Resource;
 			break;
 		}
+		
+		case EOpType::MI_CONSTANT:
+		{
+			TSharedPtr<const FMaterial> Resource = StaticCastSharedPtr<const FMaterial>(Value);
+
+			FOutputHashStream stream;
+			{
+				MUTABLE_CPUPROFILER_SCOPE(Serialize);
+				FOutputArchive arch(&stream);
+				FMaterial::Serialise(Resource.Get(), arch);
+			}
+
+			ValueHash = stream.GetHash();
+			LoadedValue = Resource;
+
+			break;
+		}
 
 		default:
 			LoadedValue = Value;
diff --git a/Engine/Plugins/Mutable/Source/MutableTools/Private/MuT/ASTOpConstantResource.h b/Engine/Plugins/Mutable/Source/MutableTools/Private/MuT/ASTOpConstantResource.h
--- a/Engine/Plugins/Mutable/Source/MutableTools/Private/MuT/ASTOpConstantResource.h
+++ b/Engine/Plugins/Mutable/Source/MutableTools/Private/MuT/ASTOpConstantResource.h
@@ -38,6 +38,9 @@ namespace mu
 		/** Source data descriptor. */
 		FSourceDataDescriptor SourceDataDescriptor;
 
+		/** Image Parameters are evaluated at runtime. */
+		TMap<FName, ASTChild> ImageOperations;
+
 	public:
 
 		~ASTOpConstantResource() override;
diff --git a/Engine/Plugins/Mutable/Source/MutableTools/Private/MuT/ASTOpMaterialBreak.cpp b/Engine/Plugins/Mutable/Source/MutableTools/Private/MuT/ASTOpMaterialBreak.cpp
new file mode 100644
--- /dev/null
+++ b/Engine/Plugins/Mutable/Source/MutableTools/Private/MuT/ASTOpMaterialBreak.cpp
@@ -0,0 +1,127 @@
+// Copyright Epic Games, Inc. All Rights Reserved.
+
+#include "MuT/ASTOpMaterialBreak.h"
+
+#include "MuT/ASTOpBoolEqualIntConst.h"
+#include "HAL/PlatformMath.h"
+#include "Misc/AssertionMacros.h"
+#include "MuR/ModelPrivate.h"
+#include "MuR/Parameters.h"
+#include "MuR/RefCounted.h"
+#include "MuR/Types.h"
+
+
+namespace mu
+{
+	ASTOpMaterialBreak::ASTOpMaterialBreak()
+		: Material(this)
+	{
+	}
+
+	ASTOpMaterialBreak::~ASTOpMaterialBreak()
+	{
+		// Explicit call needed to avoid recursive destruction
+		ASTOp::RemoveChildren();
+	}
+
+
+	void ASTOpMaterialBreak::ForEachChild(const TFunctionRef<void(ASTChild&)> f)
+	{
+		f(Material);
+	}
+
+
+	bool ASTOpMaterialBreak::IsEqual(const ASTOp& OtherUntyped) const
+	{
+		if (OtherUntyped.GetOpType() == GetOpType())
+		{
+			const ASTOpMaterialBreak* Other = static_cast<const ASTOpMaterialBreak*>(&OtherUntyped);
+			return Type == Other->Type &&
+				Material == Other->Material &&
+				ParameterName == Other->ParameterName;
+		}
+		return false;
+	}
+
+
+	mu::Ptr<ASTOp> ASTOpMaterialBreak::Clone(MapChildFuncRef MapChild) const
+	{
+		Ptr<ASTOpMaterialBreak> Result = new ASTOpMaterialBreak();
+		Result->Type = Type;
+		Result->ParameterName = ParameterName;
+		Result->Material = MapChild(Material.child());
+
+		return Result;
+	}
+
+
+	EOpType ASTOpMaterialBreak::GetOpType() const
+	{ 
+		return Type; 
+	}
+
+
+	uint64 ASTOpMaterialBreak::Hash() const
+	{
+		uint64 res = std::hash<uint64>()(uint64(Type));
+		hash_combine(res, GetTypeHash(ParameterName));
+		hash_combine(res, Material.child().get());
+		
+		return res;
+	}
+
+
+	void ASTOpMaterialBreak::Assert()
+	{
+		ASTOp::Assert();
+	}
+
+
+	void ASTOpMaterialBreak::Link(FProgram& Program, FLinkerOptions*)
+	{
+		// Already linked?
+		if (!linkedAddress)
+		{
+			linkedAddress = (OP::ADDRESS)Program.OpAddress.Num();
+			Program.OpAddress.Add((uint32)Program.ByteCode.Num());
+			AppendCode(Program.ByteCode, Type);
+
+			OP::MaterialBreakArgs Args;
+			FMemory::Memzero(Args);
+
+			if (Material)
+			{
+				Args.Material = Material->linkedAddress;
+			}
+
+			Args.ParameterName = Program.AddConstant(ParameterName.ToString());
+
+			AppendCode(Program.ByteCode, Args);
+		}
+	}
+
+
+	FImageDesc ASTOpMaterialBreak::GetImageDesc(bool returnBestOption, FGetImageDescContext* Context) const
+	{
+		check(Type == EOpType::IM_MATERIAL_BREAK);
+
+		FImageDesc Result;
+
+		if (Material)
+		{
+			Result = Material->GetImageDesc();
+		}
+
+		return Result;
+	}
+
+	FSourceDataDescriptor ASTOpMaterialBreak::GetSourceDataDescriptor(FGetSourceDataDescriptorContext* Context) const
+	{
+		if (Material)
+		{
+			return Material->GetSourceDataDescriptor(Context);
+		}
+
+		return {};
+	}
+}
diff --git a/Engine/Plugins/Mutable/Source/MutableTools/Private/MuT/ASTOpMaterialBreak.h b/Engine/Plugins/Mutable/Source/MutableTools/Private/MuT/ASTOpMaterialBreak.h
new file mode 100644
--- /dev/null
+++ b/Engine/Plugins/Mutable/Source/MutableTools/Private/MuT/ASTOpMaterialBreak.h
@@ -0,0 +1,46 @@
+// Copyright Epic Games, Inc. All Rights Reserved.
+
+#pragma once
+
+#include "HAL/Platform.h"
+#include "MuR/Operations.h"
+#include "MuR/Ptr.h"
+#include "MuT/AST.h"
+
+
+namespace mu
+{
+	struct FProgram;
+
+	/** Parameter operation. */
+	class ASTOpMaterialBreak final : public ASTOp
+	{
+	public:
+
+		//! Type of parameter
+		EOpType Type = EOpType::NONE;
+
+		ASTChild Material;
+
+		FName ParameterName;
+
+	public:
+
+		ASTOpMaterialBreak();
+		ASTOpMaterialBreak(const ASTOpMaterialBreak&) = delete;
+		~ASTOpMaterialBreak() override;
+
+		// ASTOp interface.
+		virtual EOpType GetOpType() const override;
+		virtual uint64 Hash() const override;
+		virtual void ForEachChild(const TFunctionRef<void(ASTChild&)>) override;
+		virtual bool IsEqual(const ASTOp&) const override;
+		virtual Ptr<ASTOp> Clone(MapChildFuncRef ) const override;
+		virtual void Assert() override;
+		virtual void Link(FProgram&, FLinkerOptions*) override;
+		virtual FImageDesc GetImageDesc(bool, FGetImageDescContext*) const override;
+		virtual FSourceDataDescriptor GetSourceDataDescriptor(FGetSourceDataDescriptorContext*) const override;
+	};
+
+}
+
diff --git a/Engine/Plugins/Mutable/Source/MutableTools/Private/MuT/ASTOpMaterialBreakImageParameter.cpp b/Engine/Plugins/Mutable/Source/MutableTools/Private/MuT/ASTOpMaterialBreakImageParameter.cpp
new file mode 100644
--- /dev/null
+++ b/Engine/Plugins/Mutable/Source/MutableTools/Private/MuT/ASTOpMaterialBreakImageParameter.cpp
@@ -0,0 +1,98 @@
+// Copyright Epic Games, Inc. All Rights Reserved.
+
+#include "MuT/ASTOpMaterialBreakImageParameter.h"
+
+#include "MuR/Types.h"
+
+
+namespace mu
+{
+	ASTOpImageFromMaterialParameter::ASTOpImageFromMaterialParameter()
+	{
+	}
+
+	ASTOpImageFromMaterialParameter::~ASTOpImageFromMaterialParameter()
+	{
+		// Explicit call needed to avoid recursive destruction
+		ASTOp::RemoveChildren();
+	}
+
+
+	void ASTOpImageFromMaterialParameter::ForEachChild(const TFunctionRef<void(ASTChild&)> f)
+	{
+	}
+
+
+	bool ASTOpImageFromMaterialParameter::IsEqual(const ASTOp& OtherUntyped) const
+	{
+		if (OtherUntyped.GetOpType() == GetOpType())
+		{
+			const ASTOpImageFromMaterialParameter* Other = static_cast<const ASTOpImageFromMaterialParameter*>(&OtherUntyped);
+			return MaterialParameter == Other->MaterialParameter &&
+				ParameterName == Other->ParameterName;
+		}
+		return false;
+	}
+
+
+	mu::Ptr<ASTOp> ASTOpImageFromMaterialParameter::Clone(MapChildFuncRef MapChild) const
+	{
+		Ptr<ASTOpImageFromMaterialParameter> Result = new ASTOpImageFromMaterialParameter();
+		Result->ParameterName = ParameterName;
+		Result->MaterialParameter = MaterialParameter;
+
+		return Result;
+	}
+
+
+	EOpType ASTOpImageFromMaterialParameter::GetOpType() const
+	{ 
+		return EOpType::IM_PARAMETER_FROM_MATERIAL;
+	}
+
+
+	uint64 ASTOpImageFromMaterialParameter::Hash() const
+	{
+		uint64 res = std::hash<EOpType>()(GetOpType());
+		hash_combine(res, GetTypeHash(ParameterName));
+		hash_combine(res, GetTypeHash(MaterialParameter.Name));
+		
+		return res;
+	}
+
+
+	void ASTOpImageFromMaterialParameter::Assert()
+	{
+		ASTOp::Assert();
+	}
+
+
+	void ASTOpImageFromMaterialParameter::Link(FProgram& Program, FLinkerOptions*)
+	{
+		// Already linked?
+		if (!linkedAddress)
+		{
+			int32 LinkedParameterIndex = Program.Parameters.Find(MaterialParameter);
+			check(LinkedParameterIndex != INDEX_NONE); // check ASTOpParameter.cpp if this happens
+
+			linkedAddress = (OP::ADDRESS)Program.OpAddress.Num();
+			Program.OpAddress.Add((uint32)Program.ByteCode.Num());
+			AppendCode(Program.ByteCode, GetOpType());
+
+			OP::MaterialBreakImageParameterArgs Args;
+			FMemory::Memzero(Args);
+
+			Args.MaterialParameter = OP::ADDRESS(LinkedParameterIndex);
+			Args.ParameterName = Program.AddConstant(ParameterName.ToString());
+
+			AppendCode(Program.ByteCode, Args);
+		}
+	}
+
+
+	FImageDesc ASTOpImageFromMaterialParameter::GetImageDesc(bool, FGetImageDescContext*) const
+	{
+		check(GetOpType() == EOpType::IM_PARAMETER_FROM_MATERIAL);
+		return FImageDesc();
+	}
+}
diff --git a/Engine/Plugins/Mutable/Source/MutableTools/Private/MuT/ASTOpMaterialBreakImageParameter.h b/Engine/Plugins/Mutable/Source/MutableTools/Private/MuT/ASTOpMaterialBreakImageParameter.h
new file mode 100644
--- /dev/null
+++ b/Engine/Plugins/Mutable/Source/MutableTools/Private/MuT/ASTOpMaterialBreakImageParameter.h
@@ -0,0 +1,42 @@
+// Copyright Epic Games, Inc. All Rights Reserved.
+
+#pragma once
+
+#include "HAL/Platform.h"
+#include "MuR/Operations.h"
+#include "MuR/Ptr.h"
+#include "MuT/AST.h"
+
+
+namespace mu
+{
+	struct FProgram;
+
+	/** Parameter operation. */
+	class ASTOpImageFromMaterialParameter final : public ASTOp
+	{
+	public:
+
+		FParameterDesc MaterialParameter;
+		FName ParameterName;
+
+	public:
+
+		ASTOpImageFromMaterialParameter();
+		ASTOpImageFromMaterialParameter(const ASTOpImageFromMaterialParameter&) = delete;
+		~ASTOpImageFromMaterialParameter() override;
+
+		// ASTOp interface.
+		virtual EOpType GetOpType() const override;
+		virtual uint64 Hash() const override;
+		virtual void ForEachChild(const TFunctionRef<void(ASTChild&)>) override;
+		virtual bool IsEqual(const ASTOp&) const override;
+		virtual Ptr<ASTOp> Clone(MapChildFuncRef ) const override;
+		virtual void Assert() override;
+		virtual void Link(FProgram&, FLinkerOptions*) override;
+		virtual FImageDesc GetImageDesc(bool, FGetImageDescContext*) const override;
+
+	};
+
+}
+
diff --git a/Engine/Plugins/Mutable/Source/MutableTools/Private/MuT/ASTOpParameter.cpp b/Engine/Plugins/Mutable/Source/MutableTools/Private/MuT/ASTOpParameter.cpp
--- a/Engine/Plugins/Mutable/Source/MutableTools/Private/MuT/ASTOpParameter.cpp
+++ b/Engine/Plugins/Mutable/Source/MutableTools/Private/MuT/ASTOpParameter.cpp
@@ -28,6 +28,11 @@ namespace mu
 		{
 			f(c.rangeSize);
 		}
+		
+		for (TPair<FString, ASTChild>& Element : ImageOperations)
+		{
+			f(Element.Value);
+		}
 	}
 
 
@@ -41,7 +46,10 @@ namespace mu
 				LODIndex == Other->LODIndex &&
 				SectionIndex == Other->SectionIndex &&
 				MeshID == Other->MeshID &&
-				Ranges == Other->Ranges;
+				Ranges == Other->Ranges &&
+				ColorParameterNames == Other->ColorParameterNames &&
+				ScalarParameterNames == Other->ScalarParameterNames &&
+				ImageOperations == Other->ImageOperations;
 		}
 		return false;
 	}
@@ -59,6 +67,15 @@ namespace mu
 		{
 			n->Ranges.Emplace(n.get(), MapChild(c.rangeSize.child()), c.rangeName, c.rangeUID);
 		}
+
+		n->ColorParameterNames = ColorParameterNames;
+		n->ScalarParameterNames = ScalarParameterNames;
+		
+		for (const TPair<FString, ASTChild>& ImageOperation : ImageOperations)
+		{
+			n->ImageOperations.Add(ImageOperation.Key, ASTChild(n, MapChild(ImageOperation.Value.child())));
+		}
+
 		return n;
 	}
 
@@ -76,6 +93,7 @@ namespace mu
 		hash_combine(res, LODIndex);
 		hash_combine(res, SectionIndex);
 		hash_combine(res, Parameter.Name.Len());
+
 		return res;
 	}
 
@@ -111,7 +129,7 @@ namespace mu
 				Args.Section = SectionIndex;
 				Args.MeshID = MeshID;
 
-				for (const auto& d : Ranges)
+				for (const FRangeData& d : Ranges)
 				{
 					OP::ADDRESS sizeAt = 0;
 					uint16 rangeId = 0;
@@ -121,13 +139,48 @@ namespace mu
 
 				AppendCode(Program.ByteCode, Args);
 			}
+
+			else if (Type == EOpType::MI_PARAMETER)
+			{
+				OP::MaterialParameterArgs Args;
+				FMemory::Memzero(Args);
+				Args.Variable = (OP::ADDRESS)LinkedParameterIndex;
+				
+				// Material Parameters
+				Args.ScalarParameterNames = Program.AddConstant(ScalarParameterNames);
+				Args.ColorParameterNames = Program.AddConstant(ColorParameterNames);
+
+				TArray<FString> ImageParameterNames;
+				TArray<uint32> ImageParameterOperations;
+
+				for (const TPair<FString, ASTChild>& Element : ImageOperations)
+				{
+					ImageParameterNames.Add(Element.Key);
+					ImageParameterOperations.Add(Element.Value->linkedAddress);
+				}
+
+				Args.ImageParameterNames = Program.AddConstant(ImageParameterNames);
+				Args.ImageParameterAddress = Program.AddConstant(ImageParameterOperations);
+
+				// Ranges
+				for (const FRangeData& d : Ranges)
+				{
+					OP::ADDRESS sizeAt = 0;
+					uint16 rangeId = 0;
+					LinkRange(Program, d, sizeAt, rangeId);
+					Program.Parameters[LinkedParameterIndex].Ranges.Add(rangeId);
+				}
+
+				AppendCode(Program.ByteCode, Args);
+			}
+
 			else
 			{
 				OP::ParameterArgs Args;
 				FMemory::Memzero(Args);
 				Args.variable = (OP::ADDRESS)LinkedParameterIndex;
 
-				for (const auto& d : Ranges)
+				for (const FRangeData& d : Ranges)
 				{
 					OP::ADDRESS sizeAt = 0;
 					uint16 rangeId = 0;
@@ -179,8 +232,17 @@ namespace mu
 
 	FImageDesc ASTOpParameter::GetImageDesc(bool, FGetImageDescContext*) const
 	{
-		check(Type == EOpType::IM_PARAMETER);
-		return FImageDesc();
+		FImageDesc Result;
+
+		if (Type == EOpType::MI_PARAMETER)
+		{
+			if (TMap<FString, ASTChild>::TConstIterator Iterator = ImageOperations.CreateConstIterator())
+			{
+				Result = Iterator.Value()->GetImageDesc();
+			}
+		}
+		
+		return Result;
 	}
 
 
diff --git a/Engine/Plugins/Mutable/Source/MutableTools/Private/MuT/ASTOpParameter.h b/Engine/Plugins/Mutable/Source/MutableTools/Private/MuT/ASTOpParameter.h
--- a/Engine/Plugins/Mutable/Source/MutableTools/Private/MuT/ASTOpParameter.h
+++ b/Engine/Plugins/Mutable/Source/MutableTools/Private/MuT/ASTOpParameter.h
@@ -34,6 +34,13 @@ namespace mu
 		//** Ranges adding dimensions to this parameter. */
 		TArray<FRangeData> Ranges;
 
+		/** List of all the parameters that this material can be broken into. */
+		TArray<FString> ScalarParameterNames;
+		TArray<FString> ColorParameterNames;
+
+		/** Image Parameters are evaluated at runtime. */
+		TMap<FString, ASTChild> ImageOperations;
+
 		/** Index of the parameter in the final program parameter list. 
 		* This is generated at link time, because the parameters may be reordered. 
 		*/
diff --git a/Engine/Plugins/Mutable/Source/MutableTools/Private/MuT/CodeGenerator.h b/Engine/Plugins/Mutable/Source/MutableTools/Private/MuT/CodeGenerator.h
--- a/Engine/Plugins/Mutable/Source/MutableTools/Private/MuT/CodeGenerator.h
+++ b/Engine/Plugins/Mutable/Source/MutableTools/Private/MuT/CodeGenerator.h
@@ -53,16 +53,19 @@ namespace mu
 	class NodeColourTable;
 	class NodeColourVariation;
 	class NodeColorToSRGB;
+	class NodeColourMaterialBreak;
 	class NodeImageBinarise;
 	class NodeImageColourMap;
 	class NodeImageConditional;
 	class NodeImageConstant;
 	class NodeImageFormat;
+	class NodeImageFromMaterialParameter;
 	class NodeImageInterpolate;
 	class NodeImageInvert;
 	class NodeImageLayer;
 	class NodeImageLayerColour;
 	class NodeImageLuminance;
+	class NodeImageMaterialBreak;
 	class NodeImageMipmap;
 	class NodeImageMultiLayer;
 	class NodeImageNormalComposite;
@@ -76,6 +79,7 @@ namespace mu
 	class NodeImageTransform;
 	class NodeImageVariation;
 	class NodeMaterial;
+	class NodeMaterialBreak;
 	class NodeMaterialConstant;
 	class NodeMaterialTable;
 	class NodeMaterialSwitch;
@@ -104,6 +108,7 @@ namespace mu
 	class NodeScalarSwitch;
 	class NodeScalarTable;
 	class NodeScalarVariation;
+	class NodeScalarMaterialBreak;
 	class NodeStringConstant;
 	class NodeStringParameter;
 	class NodeMatrix;
@@ -424,28 +429,6 @@ namespace mu
 		void GenerateTableSwitchNode(const NODE_TABLE& node, Ptr<NODE_SWITCH_TYPE>& OutResult, F&& GenerateOption);
 
 
-		//-----------------------------------------------------------------------------------------
-		// Materials
-
-		/** Options that affect the generation of Materials. It is like list of what required data we want while parsing down the material node graph. */
-		struct FMaterialGenerationOptions : public FGenericGenerationOptions
-		{
-			
-		};
-
-		/** */
-		struct FMaterialGenerationResult
-		{
-			Ptr<ASTOp> op;
-		};
-
-		void GenerateMaterial(const FMaterialGenerationOptions&, FMaterialGenerationResult& result, const Ptr<const NodeMaterial>& node);
-		void GenerateMaterial_Constant(const FMaterialGenerationOptions&, FMaterialGenerationResult& result, const NodeMaterialConstant* node);
-		void GenerateMaterial_Switch(const FMaterialGenerationOptions&, FMaterialGenerationResult& result, const NodeMaterialSwitch* node);
-		void GenerateMaterial_Table(const FMaterialGenerationOptions&, FMaterialGenerationResult& result, const NodeMaterialTable* node);
-		void GenerateMaterial_Variation(const FMaterialGenerationOptions&, FMaterialGenerationResult& result, const NodeMaterialVariation* node);
-		void GenerateMaterial_Parameter(const FMaterialGenerationOptions&, FMaterialGenerationResult& result, const NodeMaterialParameter* node);
-
 		//-----------------------------------------------------------------------------------------
 		// Images
 		
@@ -475,6 +458,9 @@ namespace mu
 			uint64 LayoutBlockId = FLayoutBlock::InvalidBlockId;
 			TSharedPtr<const FLayout> LayoutToApply;
 
+			/**(Optional) Parameter Descriptor of the material where this image belongs. */
+			FParameterDesc MaterialParameter;
+
 			friend FORCEINLINE uint32 GetTypeHash(const FImageGenerationOptions& InKey)
 			{
 				uint32 KeyHash = GetTypeHash((FGenericGenerationOptions&)InKey);
@@ -552,6 +538,8 @@ namespace mu
 		void GenerateImage_Variation(const FImageGenerationOptions&, FImageGenerationResult&, const NodeImageVariation*);
 		void GenerateImage_NormalComposite(const FImageGenerationOptions&, FImageGenerationResult&, const NodeImageNormalComposite*);
 		void GenerateImage_Transform(const FImageGenerationOptions&, FImageGenerationResult&, const NodeImageTransform*);
+		void GenerateImage_MaterialBreak(const FImageGenerationOptions&, FImageGenerationResult&, const NodeImageMaterialBreak*);
+		void GenerateImage_FromMaterialParameter(const FImageGenerationOptions&, FImageGenerationResult&, const NodeImageFromMaterialParameter*);
 
     	void GenerateCrop(const FImageGenerationOptions& Options, FImageGenerationResult& Result, const NodeImage&);
 
@@ -582,6 +570,50 @@ namespace mu
 		/** Generate all the operations to apply the block patching on top of the BlockOp, and masking with PatchMask. */
 		Ptr<ASTOp> GenerateImageBlockPatch(Ptr<ASTOp> BlockOp, const NodeModifierSurfaceEdit::FTexture&, TSharedPtr<FImage> PatchMask, Ptr<ASTOp> ConditionOp, const FImageGenerationOptions&);
 
+		//-----------------------------------------------------------------------------------------
+		// Materials	
+
+		/** Options that affect the generation of Materials. It is like list of what required data we want while parsing down the material node graph. */
+		struct FMaterialGenerationOptions : public FGenericGenerationOptions
+		{
+			/** The id of the component that we are currently generating. */
+			int32 ComponentId = -1;
+
+			/** The LOD being generated. */
+			int32 LODIndex = -1;
+
+			/** */
+			CompilerOptions::TextureLayoutStrategy ImageLayoutStrategy = CompilerOptions::TextureLayoutStrategy::None;
+
+			/** If different than {0,0} this is the mandatory size of the image that needs to be generated. */
+			UE::Math::TIntVector2<int32> RectSize = { 0, 0 };
+
+			/** Layout block that we are trying to generate if any. */
+			uint64 LayoutBlockId = FLayoutBlock::InvalidBlockId;
+			TSharedPtr<const FLayout> LayoutToApply;
+		};
+
+		/** */
+		struct FMaterialGenerationResult
+		{
+			Ptr<ASTOp> op;
+		};
+
+		typedef TMap<FGeneratedCacheKey, FMaterialGenerationResult> FGeneratedMaterialsMap;
+		struct FGeneratedMaterials
+		{
+			UE::FMutex Mutex;
+			FGeneratedMaterialsMap Map;
+		};
+		FGeneratedMaterials GeneratedMaterials;
+
+		void GenerateMaterial(const FMaterialGenerationOptions&, FMaterialGenerationResult& result, const Ptr<const NodeMaterial>& node);
+		void GenerateMaterial_Constant(const FMaterialGenerationOptions&, FMaterialGenerationResult& result, const NodeMaterialConstant* node);
+		void GenerateMaterial_Switch(const FMaterialGenerationOptions&, FMaterialGenerationResult& result, const NodeMaterialSwitch* node);
+		void GenerateMaterial_Table(const FMaterialGenerationOptions&, FMaterialGenerationResult& result, const NodeMaterialTable* node);
+		void GenerateMaterial_Variation(const FMaterialGenerationOptions&, FMaterialGenerationResult& result, const NodeMaterialVariation* node);
+		void GenerateMaterial_Parameter(const FMaterialGenerationOptions&, FMaterialGenerationResult& result, const NodeMaterialParameter* node);
+
         //-----------------------------------------------------------------------------------------
         // Meshes
 
@@ -853,6 +885,7 @@ namespace mu
 		void GenerateScalar_Arithmetic(FScalarGenerationResult&, const FGenericGenerationOptions&, const Ptr<const NodeScalarArithmeticOperation>&);
 		void GenerateScalar_Variation(FScalarGenerationResult&, const FGenericGenerationOptions&, const Ptr<const NodeScalarVariation>&);
 		void GenerateScalar_Table(FScalarGenerationResult&, const FGenericGenerationOptions&, const Ptr<const NodeScalarTable>&);
+		void GenerateScalar_MaterialBreak(FScalarGenerationResult&, const FGenericGenerationOptions&, const Ptr<const NodeScalarMaterialBreak>&);
 		Ptr<ASTOp> GenerateMissingScalarCode(const TCHAR* strWhere, float value, const void* errorContext);
 
 		//-----------------------------------------------------------------------------------------
@@ -880,6 +913,7 @@ namespace mu
 		void GenerateColor_Variation(FColorGenerationResult&, const FGenericGenerationOptions&, const Ptr<const NodeColourVariation>&);
 		void GenerateColor_ToSRGB(FColorGenerationResult&, const FGenericGenerationOptions&, const Ptr<const NodeColorToSRGB>&);
 		void GenerateColor_Table(FColorGenerationResult&, const FGenericGenerationOptions&, const Ptr<const NodeColourTable>&);
+		void GenerateColor_MaterialBreak(FColorGenerationResult&, const FGenericGenerationOptions&, const Ptr<const NodeColourMaterialBreak>&);
 		Ptr<ASTOp> GenerateMissingColourCode(const TCHAR* strWhere, const void* errorContext);
 
 		//-----------------------------------------------------------------------------------------
diff --git a/Engine/Plugins/Mutable/Source/MutableTools/Private/MuT/CodeGenerator_Color.cpp b/Engine/Plugins/Mutable/Source/MutableTools/Private/MuT/CodeGenerator_Color.cpp
--- a/Engine/Plugins/Mutable/Source/MutableTools/Private/MuT/CodeGenerator_Color.cpp
+++ b/Engine/Plugins/Mutable/Source/MutableTools/Private/MuT/CodeGenerator_Color.cpp
@@ -21,6 +21,7 @@
 #include "MuT/ASTOpColorSampleImage.h"
 #include "MuT/ASTOpColorFromScalars.h"
 #include "MuT/ASTOpColorToSRGB.h"
+#include "MuT/ASTOpMaterialBreak.h"
 #include "MuT/CodeGenerator.h"
 #include "MuT/CodeGenerator_FirstPass.h"
 #include "MuT/ErrorLog.h"
@@ -29,6 +30,7 @@
 #include "MuT/NodeColourArithmeticOperation.h"
 #include "MuT/NodeColourConstant.h"
 #include "MuT/NodeColourFromScalars.h"
+#include "MuT/NodeColourMaterialBreak.h"
 #include "MuT/NodeColourParameter.h"
 #include "MuT/NodeColourSampleImage.h"
 #include "MuT/NodeColourSwitch.h"
@@ -118,6 +120,11 @@ namespace mu
 			const NodeColourTable* Table = static_cast<const NodeColourTable*>(Untyped.get());
 			GenerateColor_Table(Result, Options, Table);
 		}
+		else if (Untyped->GetType() == NodeColourMaterialBreak::GetStaticType())
+		{
+			const NodeColourMaterialBreak* Table = static_cast<const NodeColourMaterialBreak*>(Untyped.get());
+			GenerateColor_MaterialBreak(Result, Options, Table);
+		}
 		else
 		{
 			check(false);
@@ -540,6 +547,26 @@ namespace mu
 	}
 
 
+	//-------------------------------------------------------------------------------------------------
+	void CodeGenerator::GenerateColor_MaterialBreak(FColorGenerationResult& Result, const FGenericGenerationOptions& Options, const Ptr<const NodeColourMaterialBreak>& Typed)
+	{
+		Ptr<ASTOpMaterialBreak> MaterialOp = new ASTOpMaterialBreak();
+
+		// Set the parameter name that this break node will generate
+		MaterialOp->ParameterName = Typed->ParameterName;
+
+		// Generate Material Source
+		FMaterialGenerationOptions MaterialSourceOptions;
+		FMaterialGenerationResult MaterialSourceResult;
+
+		GenerateMaterial(MaterialSourceOptions, MaterialSourceResult, Typed->MaterialSource);
+		MaterialOp->Material = MaterialSourceResult.op;
+		MaterialOp->Type = EOpType::CO_MATERIAL_BREAK;
+
+		Result.op = MaterialOp;
+	}
+
+
 	//---------------------------------------------------------------------------------------------
 	mu::Ptr<ASTOp> CodeGenerator::GenerateMissingColourCode(const TCHAR* StrWhere, const void* ErrorContext)
 	{
diff --git a/Engine/Plugins/Mutable/Source/MutableTools/Private/MuT/CodeGenerator_Image.cpp b/Engine/Plugins/Mutable/Source/MutableTools/Private/MuT/CodeGenerator_Image.cpp
--- a/Engine/Plugins/Mutable/Source/MutableTools/Private/MuT/CodeGenerator_Image.cpp
+++ b/Engine/Plugins/Mutable/Source/MutableTools/Private/MuT/CodeGenerator_Image.cpp
@@ -38,6 +38,8 @@
 #include "MuT/ASTOpImagePlainColor.h"
 #include "MuT/ASTOpImageDisplace.h"
 #include "MuT/ASTOpImageInvert.h"
+#include "MuT/ASTOpMaterialBreak.h"
+#include "MuT/ASTOpMaterialBreakImageParameter.h"
 #include "MuT/ASTOpMeshExtractLayoutBlocks.h"
 #include "MuT/ASTOpMeshFormat.h"
 #include "MuT/ASTOpMeshProject.h"
@@ -59,11 +61,13 @@
 #include "MuT/NodeImageConditional.h"
 #include "MuT/NodeImageConstant.h"
 #include "MuT/NodeImageFormat.h"
+#include "MuT/NodeImageFromMaterialParameter.h"
 #include "MuT/NodeImageInterpolate.h"
 #include "MuT/NodeImageInvert.h"
 #include "MuT/NodeImageLayer.h"
 #include "MuT/NodeImageLayerColour.h"
 #include "MuT/NodeImageLuminance.h"
+#include "MuT/NodeImageMaterialBreak.h"
 #include "MuT/NodeImageMipmap.h"
 #include "MuT/NodeImageMultiLayer.h"
 #include "MuT/NodeImageNormalComposite.h"
@@ -134,6 +138,8 @@ namespace mu
 		case Node::EType::ImageVariation: GenerateImage_Variation(Options, Result, static_cast<const NodeImageVariation*>(Node)); break;
 		case Node::EType::ImageNormalComposite: GenerateImage_NormalComposite(Options, Result, static_cast<const NodeImageNormalComposite*>(Node)); break;
 		case Node::EType::ImageTransform: GenerateImage_Transform(Options, Result, static_cast<const NodeImageTransform*>(Node)); break;
+		case Node::EType::ImageMaterialBreak: GenerateImage_MaterialBreak(Options, Result, static_cast<const NodeImageMaterialBreak*>(Node)); break;
+		case Node::EType::ImageFromMaterialParameter: GenerateImage_FromMaterialParameter(Options, Result, static_cast<const NodeImageFromMaterialParameter*>(Node)); break;
 		default: check(false);
 		}
 
@@ -1879,6 +1885,42 @@ namespace mu
 	}
 
 
+	void CodeGenerator::GenerateImage_MaterialBreak(const FImageGenerationOptions& InOptions, FImageGenerationResult& Result, const NodeImageMaterialBreak* Node)
+	{
+		Ptr<ASTOpMaterialBreak> MaterialOp = new ASTOpMaterialBreak();
+
+		// Set the parameter name that this break node will generate
+		MaterialOp->ParameterName = Node->ParameterName;
+
+		// Generate Material Source
+		FMaterialGenerationResult MaterialSourceResult;
+		FMaterialGenerationOptions MaterialSourceOptions;
+		// Temporary fix
+		MaterialSourceOptions.ComponentId = InOptions.ComponentId;
+		MaterialSourceOptions.LODIndex = InOptions.LODIndex;
+		MaterialSourceOptions.ImageLayoutStrategy = InOptions.ImageLayoutStrategy;
+		MaterialSourceOptions.RectSize = InOptions.RectSize;
+		MaterialSourceOptions.LayoutBlockId = InOptions.LayoutBlockId;
+		MaterialSourceOptions.LayoutToApply = InOptions.LayoutToApply;
+
+		GenerateMaterial(MaterialSourceOptions, MaterialSourceResult, Node->MaterialSource);
+		MaterialOp->Material = MaterialSourceResult.op;
+		MaterialOp->Type = EOpType::IM_MATERIAL_BREAK;
+
+		Result.op = MaterialOp;
+	}
+
+
+	void CodeGenerator::GenerateImage_FromMaterialParameter(const FImageGenerationOptions& InOptions, FImageGenerationResult& Result, const NodeImageFromMaterialParameter* InNode)
+	{
+		Ptr<ASTOpImageFromMaterialParameter> ImageOp = new ASTOpImageFromMaterialParameter();
+		ImageOp->MaterialParameter = InOptions.MaterialParameter;
+		ImageOp->ParameterName = InNode->ImageParameterName;
+
+		Result.op = ImageOp;
+	}
+
+
 	TSharedPtr<FImage> CodeGenerator::GenerateMissingImage(EImageFormat Format)
 	{
 		MUTABLE_CPUPROFILER_SCOPE(GenerateMissingImage);
diff --git a/Engine/Plugins/Mutable/Source/MutableTools/Private/MuT/CodeGenerator_Material.cpp b/Engine/Plugins/Mutable/Source/MutableTools/Private/MuT/CodeGenerator_Material.cpp
--- a/Engine/Plugins/Mutable/Source/MutableTools/Private/MuT/CodeGenerator_Material.cpp
+++ b/Engine/Plugins/Mutable/Source/MutableTools/Private/MuT/CodeGenerator_Material.cpp
@@ -10,6 +10,10 @@
 #include "MuT/AST.h"
 #include "MuT/ASTOpConditional.h"
 #include "MuT/ASTOpConstantMaterial.h"
+#include "MuT/ASTOpConstantResource.h"
+#include "MuT/ASTOpImagePixelFormat.h"
+#include "MuT/ASTOpImageResize.h"
+#include "MuT/ASTOpMaterialBreak.h"
 #include "MuT/ASTOpSwitch.h"
 #include "MuT/CodeGenerator.h"
 #include "MuT/CodeGenerator_FirstPass.h"
@@ -32,30 +36,98 @@ namespace mu
 
 	void CodeGenerator::GenerateMaterial(const FMaterialGenerationOptions& InOptions, FMaterialGenerationResult& Result , const Ptr<const NodeMaterial>& InUntypedNode)
 	{
+		// Early out
 		if (!InUntypedNode)
 		{
+			Result = FMaterialGenerationResult();
 			return;
 		}
 
-		const NodeMaterial* Node = InUntypedNode.get();
+		// See if it was already generated
+		FGeneratedCacheKey Key;
+		Key.Node = InUntypedNode;
+		Key.Options = InOptions;
+		{
+			UE::TUniqueLock Lock(GeneratedMaterials.Mutex);
+			FGeneratedMaterialsMap::ValueType* Found = GeneratedMaterials.Map.Find(Key);
+			if (Found)
+			{
+				Result = *Found;
+				return;
+			}
+		}
 
 		// Generate for each different type of node
+		const NodeMaterial* Node = InUntypedNode.get();
 		switch (Node->GetType()->Type)
 		{
-		case Node::EType::MaterialConstant: GenerateMaterial_Constant(InOptions, Result, static_cast<const NodeMaterialConstant*>(Node)); break;
-		case Node::EType::MaterialSwitch: GenerateMaterial_Switch(InOptions, Result, static_cast<const NodeMaterialSwitch*>(Node)); break;
-		case Node::EType::MaterialVariation: GenerateMaterial_Variation(InOptions, Result, static_cast<const NodeMaterialVariation*>(Node)); break;
-		case Node::EType::MaterialTable: GenerateMaterial_Table(InOptions, Result, static_cast<const NodeMaterialTable*>(Node)); break;
-		case Node::EType::MaterialParameter: GenerateMaterial_Parameter(InOptions, Result, static_cast<const NodeMaterialParameter*>(Node)); break;
-		default: check(false);
+		case Node::EType::MaterialConstant:
+			GenerateMaterial_Constant(InOptions, Result, static_cast<const NodeMaterialConstant*>(Node)); 
+			break;
+		
+		case Node::EType::MaterialSwitch:
+			GenerateMaterial_Switch(InOptions, Result, static_cast<const NodeMaterialSwitch*>(Node));
+			break;
+
+		case Node::EType::MaterialVariation:
+			GenerateMaterial_Variation(InOptions, Result, static_cast<const NodeMaterialVariation*>(Node));
+			break;
+
+		case Node::EType::MaterialTable:
+			GenerateMaterial_Table(InOptions, Result, static_cast<const NodeMaterialTable*>(Node));
+			break;
+
+		case Node::EType::MaterialParameter:
+			GenerateMaterial_Parameter(InOptions, Result, static_cast<const NodeMaterialParameter*>(Node));
+			break;
+
+		default: 
+			check(false);
 		}
 	}
 
+
 	void CodeGenerator::GenerateMaterial_Constant(const FMaterialGenerationOptions& Options, FMaterialGenerationResult& Result, const NodeMaterialConstant* Node)
 	{
-		Result.op = new ASTOpConstantMaterial(Node->MaterialId);
+		// Constant Material Resource
+		TSharedPtr<FMaterial> MaterialResource = MakeShared<FMaterial>();
+		MaterialResource->ReferenceID = Node->MaterialId;
+		MaterialResource->ColorParameters = Node->ColorValues;
+		MaterialResource->ScalarParameters = Node->ScalarValues;
+
+		// Constant Material Operation
+		Ptr<ASTOpConstantResource> ConstantMaterial = new ASTOpConstantResource();
+		ConstantMaterial->Type = EOpType::MI_CONSTANT;
+
+		//TODO: each image should have its own generation options
+		FImageGenerationOptions ImageOptions(Options.ComponentId, Options.LODIndex);
+		ImageOptions.ImageLayoutStrategy = Options.ImageLayoutStrategy;
+		ImageOptions.RectSize = Options.RectSize;
+		ImageOptions.LayoutBlockId = Options.LayoutBlockId;
+		ImageOptions.LayoutToApply = Options.LayoutToApply;
+
+		const TArray<TPair<FName, Ptr<NodeImage>>>& ImageValuesArray = Node->ImageValues.Array();
+
+		//Store the images into the operation so that we can link them later
+		for (int32 ImageIndex = 0; ImageIndex < ImageValuesArray.Num(); ++ImageIndex)
+		{
+			FName ImageParameterName = ImageValuesArray[ImageIndex].Key;
+			Ptr<NodeImage> ImageNode = ImageValuesArray[ImageIndex].Value;
+			FImageGenerationResult ImageResult;
+
+			GenerateImage(ImageOptions, ImageResult, ImageNode);
+
+			// Store lazy branches
+			ConstantMaterial->ImageOperations.Add(ImageParameterName, ASTChild(ConstantMaterial, ImageResult.op));
+		}
+
+		// Store the constant material
+		ConstantMaterial->SetValue(MaterialResource, CompilerOptions->OptimisationOptions.DiskCacheContext);
+
+		Result.op = ConstantMaterial;
 	}
 
+
 	void CodeGenerator::GenerateMaterial_Switch(const FMaterialGenerationOptions& Options, FMaterialGenerationResult& Result, const NodeMaterialSwitch* Node)
 	{
 		if (Node->Options.Num() == 0)
@@ -103,6 +175,7 @@ namespace mu
 		Result.op = SwitchOp;
 	}
 
+
 	void CodeGenerator::GenerateMaterial_Variation(const FMaterialGenerationOptions& Options, FMaterialGenerationResult& Result, const NodeMaterialVariation* Node)
 	{
 		Ptr<ASTOp> LastOp;
@@ -161,6 +234,7 @@ namespace mu
 		Result.op = LastOp;
 	}
 
+
 	void CodeGenerator::GenerateMaterial_Table(const FMaterialGenerationOptions& Options, FMaterialGenerationResult& Result, const NodeMaterialTable* Node)
 	{
 		Result.op = GenerateTableSwitch<NodeMaterialTable, ETableColumnType::Material, EOpType::MI_SWITCH>(*Node,
@@ -178,6 +252,7 @@ namespace mu
 			});
 	}
 
+
 	void CodeGenerator::GenerateMaterial_Parameter(const FMaterialGenerationOptions& InOptions, FMaterialGenerationResult& Result, const NodeMaterialParameter* InNode)
 	{
 		{
@@ -205,6 +280,27 @@ namespace mu
 			GenerateRange(RangeResult, InOptions, RangeNode);
 			MaterialOp->Ranges.Emplace(MaterialOp.get(), RangeResult.sizeOp, RangeResult.rangeName, RangeResult.rangeUID);
 		}
+		
+		// Copy Scalar and Color Names
+		MaterialOp->ColorParameterNames = InNode->ColorParameterNames;
+		MaterialOp->ScalarParameterNames = InNode->ScalarParameterNames;
+
+		// Generate Image Nodes
+		FImageGenerationResult ImageResult;
+		FImageGenerationOptions ImageOptions(InOptions.ComponentId, InOptions.LODIndex);
+		ImageOptions.ImageLayoutStrategy = InOptions.ImageLayoutStrategy;
+		ImageOptions.RectSize = InOptions.RectSize;
+		ImageOptions.LayoutBlockId = InOptions.LayoutBlockId;
+		ImageOptions.LayoutToApply = InOptions.LayoutToApply;
+		ImageOptions.MaterialParameter = MaterialOp->Parameter;
+
+		for (const TPair<FString, Ptr<NodeImage>>& ImageNode : InNode->ImageParameterNodes)
+		{
+			GenerateImage(ImageOptions, ImageResult, ImageNode.Value);
+
+			// Store lazy branches
+			MaterialOp->ImageOperations.Add(ImageNode.Key, ASTChild(MaterialOp, ImageResult.op));
+		}
 
 		Result.op = MaterialOp;
 
diff --git a/Engine/Plugins/Mutable/Source/MutableTools/Private/MuT/CodeGenerator_Scalar.cpp b/Engine/Plugins/Mutable/Source/MutableTools/Private/MuT/CodeGenerator_Scalar.cpp
--- a/Engine/Plugins/Mutable/Source/MutableTools/Private/MuT/CodeGenerator_Scalar.cpp
+++ b/Engine/Plugins/Mutable/Source/MutableTools/Private/MuT/CodeGenerator_Scalar.cpp
@@ -17,6 +17,7 @@
 #include "MuT/ASTOpScalarCurve.h"
 #include "MuT/ASTOpScalarArithmetic.h"
 #include "MuT/ASTOpConstantScalar.h"
+#include "MuT/ASTOpMaterialBreak.h"
 #include "MuT/ASTOpSwitch.h"
 #include "MuT/CodeGenerator.h"
 #include "MuT/CodeGenerator_FirstPass.h"
@@ -29,6 +30,7 @@
 #include "MuT/NodeScalarConstant.h"
 #include "MuT/NodeScalarCurve.h"
 #include "MuT/NodeScalarEnumParameter.h"
+#include "MuT/NodeScalarMaterialBreak.h"
 #include "MuT/NodeScalarParameter.h"
 #include "MuT/NodeScalarSwitch.h"
 #include "MuT/NodeScalarTable.h"
@@ -103,6 +105,11 @@ namespace mu
 			const NodeScalarTable* Table = static_cast<const NodeScalarTable*>(Untyped.get());
 			GenerateScalar_Table(result, Options, Table);
 		}
+		else if (Untyped->GetType() == NodeScalarMaterialBreak::GetStaticType())
+		{
+			const NodeScalarMaterialBreak* Table = static_cast<const NodeScalarMaterialBreak*>(Untyped.get());
+			GenerateScalar_MaterialBreak(result, Options, Table);
+		}
 		else
 		{
 			check(false);
@@ -422,6 +429,25 @@ namespace mu
 	}
 
 
+	void CodeGenerator::GenerateScalar_MaterialBreak(FScalarGenerationResult& Result, const FGenericGenerationOptions& Options, const Ptr<const NodeScalarMaterialBreak>& Node)
+	{
+		Ptr<ASTOpMaterialBreak> MaterialOp = new ASTOpMaterialBreak();
+
+		// Set the parameter name that this break node will generate
+		MaterialOp->ParameterName = Node->ParameterName;
+
+		// Generate Material Source
+		FMaterialGenerationOptions MaterialSourceOptions;
+		FMaterialGenerationResult MaterialSourceResult;
+
+		GenerateMaterial(MaterialSourceOptions, MaterialSourceResult, Node->MaterialSource);
+		MaterialOp->Material = MaterialSourceResult.op;
+		MaterialOp->Type = EOpType::SC_MATERIAL_BREAK;
+
+		Result.op = MaterialOp;
+	}
+
+
 	mu::Ptr<ASTOp> CodeGenerator::GenerateMissingScalarCode(const TCHAR* strWhere, float value, const void* errorContext)
 	{
 		// Log a warning
diff --git a/Engine/Plugins/Mutable/Source/MutableTools/Private/MuT/CodeOptimiser.cpp b/Engine/Plugins/Mutable/Source/MutableTools/Private/MuT/CodeOptimiser.cpp
--- a/Engine/Plugins/Mutable/Source/MutableTools/Private/MuT/CodeOptimiser.cpp
+++ b/Engine/Plugins/Mutable/Source/MutableTools/Private/MuT/CodeOptimiser.cpp
@@ -550,6 +550,8 @@ namespace mu
 					case EOpType::IM_PARAMETER:
 					case EOpType::ME_PARAMETER:
 					case EOpType::MA_PARAMETER:
+					case EOpType::MI_PARAMETER:
+					case EOpType::IM_PARAMETER_FROM_MATERIAL:
 						bIsConstantSubgraph = false;
 						break;
 					default:
diff --git a/Engine/Plugins/Mutable/Source/MutableTools/Private/MuT/Compiler.cpp b/Engine/Plugins/Mutable/Source/MutableTools/Private/MuT/Compiler.cpp
--- a/Engine/Plugins/Mutable/Source/MutableTools/Private/MuT/Compiler.cpp
+++ b/Engine/Plugins/Mutable/Source/MutableTools/Private/MuT/Compiler.cpp
@@ -103,6 +103,11 @@ namespace mu
 		{ false,	{ 0, 0, 0, 0, 0, 0, 0, 0, 0, 0, 0, 0, 0, 0, 0, 0, 0, 0, 0, 0, 0, 0, 0, 0, 0, 0, 0, 0, 0 }		},	// ED_SWITCH
 		{ false,	{ 0, 0, 0, 0, 0, 0, 0, 0, 0, 0, 0, 0, 0, 0, 0, 0, 0, 0, 0, 0, 0, 0, 0, 0, 0, 0, 0, 0, 0 }		},	// MI_SWITCH
 
+		{ false,	{ 0, 0, 0, 0, 0, 0, 0, 0, 0, 0, 0, 0, 0, 0, 0, 0, 0, 0, 0, 0, 0, 0, 0, 0, 0, 0, 0, 0, 0 }		},	// SC_MATERIAL_BREAK
+		{ false,	{ 0, 0, 0, 0, 0, 0, 0, 0, 0, 0, 0, 0, 0, 0, 0, 0, 0, 0, 0, 0, 0, 0, 0, 0, 0, 0, 0, 0, 0 }		},	// CO_MATERIAL_BREAK
+		{ false,	{ 0, 0, 0, 0, 0, 0, 0, 0, 0, 0, 0, 0, 0, 0, 0, 0, 0, 0, 0, 0, 0, 0, 0, 0, 0, 0, 0, 0, 0 }		},	// IM_MATERIAL_BREAK
+		{ false,	{ 0, 0, 0, 0, 0, 0, 0, 0, 0, 0, 0, 0, 0, 0, 0, 0, 0, 0, 0, 0, 0, 0, 0, 0, 0, 0, 0, 0, 0 }		},	// IM_PARAMETER_FROM_MATERIAL
+
 		{ false,	{ 0, 0, 0, 0, 0, 0, 0, 0, 0, 0, 0, 0, 0, 0, 0, 0, 0, 0, 0, 0, 0, 0, 0, 0, 0, 0, 0, 0, 0 }		},	// BO_EQUAL_SC_CONST
 		{ false,	{ 0, 0, 0, 0, 0, 0, 0, 0, 0, 0, 0, 0, 0, 0, 0, 0, 0, 0, 0, 0, 0, 0, 0, 0, 0, 0, 0, 0, 0 }		},	// BO_AND
 		{ false,	{ 0, 0, 0, 0, 0, 0, 0, 0, 0, 0, 0, 0, 0, 0, 0, 0, 0, 0, 0, 0, 0, 0, 0, 0, 0, 0, 0, 0, 0 }		},	// BO_OR
diff --git a/Engine/Plugins/Mutable/Source/MutableTools/Private/MuT/ErrorLog.cpp b/Engine/Plugins/Mutable/Source/MutableTools/Private/MuT/ErrorLog.cpp
--- a/Engine/Plugins/Mutable/Source/MutableTools/Private/MuT/ErrorLog.cpp
+++ b/Engine/Plugins/Mutable/Source/MutableTools/Private/MuT/ErrorLog.cpp
@@ -248,6 +248,11 @@ namespace mu
 		TEXT("ED_SWITCH        "),
 		TEXT("MI_SWITCH        "),
 
+		TEXT("SC_MATERIAL_BREAK"),
+		TEXT("CO_MATERIAL_BREAK"),
+		TEXT("IM_MATERIAL_BREAK"),
+		TEXT("IM_PARAMETER_FROM_MATERIAL"),
+
 		TEXT("BO_EQUAL_SC_CONST"),
 		TEXT("BO_AND           "),
 		TEXT("BO_OR            "),
diff --git a/Engine/Plugins/Mutable/Source/MutableTools/Private/MuT/Node.cpp b/Engine/Plugins/Mutable/Source/MutableTools/Private/MuT/Node.cpp
--- a/Engine/Plugins/Mutable/Source/MutableTools/Private/MuT/Node.cpp
+++ b/Engine/Plugins/Mutable/Source/MutableTools/Private/MuT/Node.cpp
@@ -25,6 +25,7 @@
 #include "MuT/NodeImageFormat.h"
 #include "MuT/NodeImageBinarise.h"
 #include "MuT/NodeImageConditional.h"
+#include "MuT/NodeImageFromMaterialParameter.h"
 #include "MuT/NodeImageInterpolate.h"
 #include "MuT/NodeImageInvert.h"
 #include "MuT/NodeImageLayer.h"
@@ -53,6 +54,7 @@
 #include "MuT/NodeColourFromScalars.h"
 #include "MuT/NodeColourTable.h"
 #include "MuT/NodeColorToSRGB.h"
+#include "MuT/NodeColourMaterialBreak.h"
 #include "MuT/NodeScalarSwitch.h"
 #include "MuT/NodeScalarConstant.h"
 #include "MuT/NodeScalarVariation.h"
@@ -61,6 +63,7 @@
 #include "MuT/NodeScalarEnumParameter.h"
 #include "MuT/NodeScalarTable.h"
 #include "MuT/NodeScalarCurve.h"
+#include "MuT/NodeScalarMaterialBreak.h"
 #include "MuT/NodeMeshConstant.h"
 #include "MuT/NodeMeshFragment.h"
 #include "MuT/NodeMeshClipMorphPlane.h"
@@ -94,6 +97,7 @@
 #include "MuT/NodeRange.h"
 #include "MuT/NodeRangeFromScalar.h"
 #include "MuT/NodeMaterial.h"
+#include "MuT/NodeImageMaterialBreak.h"
 #include "MuT/NodeMaterialConstant.h"
 #include "MuT/NodeMaterialTable.h"
 #include "MuT/NodeMaterialSwitch.h"
@@ -166,6 +170,7 @@ namespace mu
 	FNodeType NodeImageTransform::StaticType = FNodeType(Node::EType::ImageTransform, NodeImage::GetStaticType());
 	FNodeType NodeImageVariation::StaticType = FNodeType(Node::EType::ImageVariation, NodeImage::GetStaticType());
 	FNodeType NodeImageColourMap::StaticType = FNodeType(Node::EType::ImageColorMap, NodeImage::GetStaticType());
+	FNodeType NodeImageFromMaterialParameter::StaticType = FNodeType(Node::EType::ImageFromMaterialParameter, NodeImage::GetStaticType());
 
 	FNodeType NodeColour::StaticType = FNodeType(Node::EType::Color, Node::GetStaticType());
 	FNodeType NodeColourConstant::StaticType = FNodeType(Node::EType::ColorConstant, NodeColour::GetStaticType());
@@ -224,6 +229,10 @@ namespace mu
 	FNodeType NodeMaterialSwitch::StaticType = FNodeType(Node::EType::MaterialSwitch, NodeMaterial::GetStaticType());
 	FNodeType NodeMaterialVariation::StaticType = FNodeType(Node::EType::MaterialVariation, NodeMaterial::GetStaticType());
 	FNodeType NodeMaterialParameter::StaticType = FNodeType(Node::EType::MaterialParameter, NodeMaterial::GetStaticType());
+	
+	FNodeType NodeImageMaterialBreak::StaticType = FNodeType(Node::EType::ImageMaterialBreak, NodeImage::GetStaticType());
+	FNodeType NodeScalarMaterialBreak::StaticType = FNodeType(Node::EType::ScalarMaterialBreak, NodeScalar::GetStaticType());
+	FNodeType NodeColourMaterialBreak::StaticType = FNodeType(Node::EType::ColourMaterialBreak, NodeColour::GetStaticType());
 
 
 	FNodeType::FNodeType()
diff --git a/Engine/Plugins/Mutable/Source/MutableTools/Public/MuT/Node.h b/Engine/Plugins/Mutable/Source/MutableTools/Public/MuT/Node.h
--- a/Engine/Plugins/Mutable/Source/MutableTools/Public/MuT/Node.h
+++ b/Engine/Plugins/Mutable/Source/MutableTools/Public/MuT/Node.h
@@ -75,6 +75,7 @@ namespace mu
 			ImageVariation,
 			ImageNormalComposite,
 			ImageTransform,
+			ImageFromMaterialParameter,
 
 			Bool,
 			BoolConstant,
@@ -160,6 +161,10 @@ namespace mu
 			MaterialSwitch,
 			MaterialVariation,
 			MaterialParameter,
+			
+			ImageMaterialBreak,
+			ScalarMaterialBreak,
+			ColourMaterialBreak,
 
 			Count
 		};
diff --git a/Engine/Plugins/Mutable/Source/MutableTools/Public/MuT/NodeColourMaterialBreak.h b/Engine/Plugins/Mutable/Source/MutableTools/Public/MuT/NodeColourMaterialBreak.h
new file mode 100644
--- /dev/null
+++ b/Engine/Plugins/Mutable/Source/MutableTools/Public/MuT/NodeColourMaterialBreak.h
@@ -0,0 +1,39 @@
+// Copyright Epic Games, Inc. All Rights Reserved.
+
+#pragma once
+
+#include "MuT/Node.h"
+#include "MuT/NodeColour.h"
+#include "MuR/Ptr.h"
+#include "MuT/NodeMaterial.h"
+
+#define UE_API MUTABLETOOLS_API
+
+namespace mu
+{
+	/** NodeColourMaterialBreak class. */
+	class NodeColourMaterialBreak : public NodeColour
+	{
+	public:
+
+		Ptr<NodeMaterial> MaterialSource;
+		FName ParameterName;
+
+	public:
+
+		// Node Interface
+		virtual const FNodeType* GetType() const override { return &StaticType; }
+		static const FNodeType* GetStaticType() { return &StaticType; }
+
+	protected:
+
+		/** Forbidden. Manage with the Ptr<> template. */
+		~NodeColourMaterialBreak() {}
+
+	private:
+
+		static UE_API FNodeType StaticType;
+	};
+}
+
+#undef UE_API
diff --git a/Engine/Plugins/Mutable/Source/MutableTools/Public/MuT/NodeImageFromMaterialParameter.h b/Engine/Plugins/Mutable/Source/MutableTools/Public/MuT/NodeImageFromMaterialParameter.h
new file mode 100644
--- /dev/null
+++ b/Engine/Plugins/Mutable/Source/MutableTools/Public/MuT/NodeImageFromMaterialParameter.h
@@ -0,0 +1,34 @@
+// Copyright Epic Games, Inc. All Rights Reserved.
+
+#pragma once
+#include "MuR/Ptr.h"
+#include "MuR/RefCounted.h"
+#include "MuT/Node.h"
+#include "MuT/NodeImage.h"
+
+#define UE_API MUTABLETOOLS_API
+
+namespace mu
+{
+	class NodeImageFromMaterialParameter : public NodeImage
+	{
+	public:
+
+		FName ImageParameterName;
+
+	public:
+
+		// Node Interface
+		virtual const FNodeType* GetType() const override { return &StaticType; }
+		static const FNodeType* GetStaticType() { return &StaticType; }
+
+	protected:
+
+		/** Forbidden. Manage with the Ptr<> template. */
+		~NodeImageFromMaterialParameter() {}
+
+	private:
+
+		static UE_API FNodeType StaticType;
+	};
+}
\ No newline at end of file
diff --git a/Engine/Plugins/Mutable/Source/MutableTools/Public/MuT/NodeImageMaterialBreak.h b/Engine/Plugins/Mutable/Source/MutableTools/Public/MuT/NodeImageMaterialBreak.h
new file mode 100644
--- /dev/null
+++ b/Engine/Plugins/Mutable/Source/MutableTools/Public/MuT/NodeImageMaterialBreak.h
@@ -0,0 +1,39 @@
+// Copyright Epic Games, Inc. All Rights Reserved.
+
+#pragma once
+
+#include "MuT/Node.h"
+#include "MuT/NodeImage.h"
+#include "MuR/Ptr.h"
+#include "MuT/NodeMaterial.h"
+
+#define UE_API MUTABLETOOLS_API
+
+namespace mu
+{
+	/** NodeImageMaterialBreak class. */
+	class NodeImageMaterialBreak : public NodeImage
+	{
+	public:
+
+		Ptr<NodeMaterial> MaterialSource;
+		FName ParameterName;
+
+	public:
+
+		// Node Interface
+		virtual const FNodeType* GetType() const override { return &StaticType; }
+		static const FNodeType* GetStaticType() { return &StaticType; }
+
+	protected:
+
+		/** Forbidden. Manage with the Ptr<> template. */
+		~NodeImageMaterialBreak() {}
+
+	private:
+
+		static UE_API FNodeType StaticType;
+	};
+}
+
+#undef UE_API
diff --git a/Engine/Plugins/Mutable/Source/MutableTools/Public/MuT/NodeMaterialConstant.h b/Engine/Plugins/Mutable/Source/MutableTools/Public/MuT/NodeMaterialConstant.h
--- a/Engine/Plugins/Mutable/Source/MutableTools/Public/MuT/NodeMaterialConstant.h
+++ b/Engine/Plugins/Mutable/Source/MutableTools/Public/MuT/NodeMaterialConstant.h
@@ -6,13 +6,15 @@
 #include "MuR/Ptr.h"
 #include "MuR/RefCounted.h"
 #include "MuT/NodeMaterial.h"
+#include "MuT/NodeImage.h"
+#include "MuR/Types.h"
+#include "Math/Vector4.h"
 
 #define UE_API MUTABLETOOLS_API
 
 namespace mu
 {
-
-    /** NodeMaterialConstant class. */
+	/** NodeMaterialConstant class. */
 	class NodeMaterialConstant : public NodeMaterial
 	{
 	public:
@@ -23,6 +25,10 @@ namespace mu
 
 		int32 MaterialId = INDEX_NONE;
 
+		TMap<FName, FVector4f> ColorValues;
+		TMap<FName, float> ScalarValues;
+		TMap<FName, Ptr<NodeImage>> ImageValues;
+
 	protected:
 
 		// Forbidden. Manage with the Ptr<> template.
@@ -31,10 +37,7 @@ namespace mu
 	private:
 
 		static UE_API FNodeType StaticType;
-
 	};
-
-
 }
 
 #undef UE_API
diff --git a/Engine/Plugins/Mutable/Source/MutableTools/Public/MuT/NodeMaterialParameter.h b/Engine/Plugins/Mutable/Source/MutableTools/Public/MuT/NodeMaterialParameter.h
--- a/Engine/Plugins/Mutable/Source/MutableTools/Public/MuT/NodeMaterialParameter.h
+++ b/Engine/Plugins/Mutable/Source/MutableTools/Public/MuT/NodeMaterialParameter.h
@@ -2,11 +2,14 @@
 
 #pragma once
 
+#include "Containers/Map.h"
 #include "HAL/Platform.h"
 #include "MuR/Ptr.h"
 #include "MuR/RefCounted.h"
 #include "MuT/NodeRange.h"
 #include "MuT/NodeMaterial.h"
+#include "MuT/NodeImage.h"
+
 
 #define UE_API MUTABLETOOLS_API
 
@@ -27,6 +30,10 @@ namespace mu
 		/** Ranges for the parameter in case it is multidimensional. */
 		TArray<Ptr<NodeRange>> Ranges;
 
+		TArray<FString> ColorParameterNames;
+		TArray<FString> ScalarParameterNames;
+		TMap<FString, Ptr<NodeImage>> ImageParameterNodes;
+
 	public:
 
 		// Node Interface
diff --git a/Engine/Plugins/Mutable/Source/MutableTools/Public/MuT/NodeScalarMaterialBreak.h b/Engine/Plugins/Mutable/Source/MutableTools/Public/MuT/NodeScalarMaterialBreak.h
new file mode 100644
--- /dev/null
+++ b/Engine/Plugins/Mutable/Source/MutableTools/Public/MuT/NodeScalarMaterialBreak.h
@@ -0,0 +1,39 @@
+// Copyright Epic Games, Inc. All Rights Reserved.
+
+#pragma once
+
+#include "MuT/Node.h"
+#include "MuT/NodeScalar.h"
+#include "MuR/Ptr.h"
+#include "MuT/NodeMaterial.h"
+
+#define UE_API MUTABLETOOLS_API
+
+namespace mu
+{
+	/** NodeMaterialTable class. */
+	class NodeScalarMaterialBreak : public NodeScalar
+	{
+	public:
+
+		Ptr<NodeMaterial> MaterialSource;
+		FName ParameterName;
+
+	public:
+
+		// Node Interface
+		virtual const FNodeType* GetType() const override { return &StaticType; }
+		static const FNodeType* GetStaticType() { return &StaticType; }
+
+	protected:
+
+		/** Forbidden. Manage with the Ptr<> template. */
+		~NodeScalarMaterialBreak() {}
+
+	private:
+
+		static UE_API FNodeType StaticType;
+	};
+}
+
+#undef UE_API
