[Mutable] Fixed some Hawk's suggestions.

#rnx

diff --git a/Engine/Plugins/Mutable/Source/CustomizableObjectEditor/Private/MuCOE/COVariable.cpp b/Engine/Plugins/Mutable/Source/CustomizableObjectEditor/Private/MuCOE/COVariable.cpp
--- a/Engine/Plugins/Mutable/Source/CustomizableObjectEditor/Private/MuCOE/COVariable.cpp
+++ b/Engine/Plugins/Mutable/Source/CustomizableObjectEditor/Private/MuCOE/COVariable.cpp
@@ -29,7 +29,7 @@ void FCOVariableCustomzation::CustomizeHeader(TSharedRef<IPropertyHandle> Struct
 	TArray<UObject*> OuterObjects;
 	StructPropertyHandle->GetOuterObjects(OuterObjects);
 
-	if (OuterObjects[0])
+	if (OuterObjects.Num() > 0)
 	{
 		BaseObjectNode = Cast<UCustomizableObjectNode>(OuterObjects[0]->GetTypedOuter(UCustomizableObjectNode::StaticClass()));
 
@@ -92,7 +92,7 @@ FText FCOVariableCustomzation::GetVariableName() const
 void FCOVariableCustomzation::ResetToDefault()
 {
 	NamePropertyHandle->SetValue(FName("NewVar"));
-	TypePropertyHandle->SetValue(AllowedVariableTypes[0]);
+	TypePropertyHandle->SetValue(AllowedVariableTypes.Num() ? AllowedVariableTypes[0] : NAME_None);
 }
 
 
diff --git a/Engine/Plugins/Mutable/Source/CustomizableObjectEditor/Private/MuCOE/COVariable.h b/Engine/Plugins/Mutable/Source/CustomizableObjectEditor/Private/MuCOE/COVariable.h
--- a/Engine/Plugins/Mutable/Source/CustomizableObjectEditor/Private/MuCOE/COVariable.h
+++ b/Engine/Plugins/Mutable/Source/CustomizableObjectEditor/Private/MuCOE/COVariable.h
@@ -65,7 +65,7 @@ private:
 
 	TArray<FName> AllowedVariableTypes;
 
-	/** Runtime Parameter Name property of a State */
+	/** Name and type property handles */
 	TSharedPtr<IPropertyHandle> NamePropertyHandle;
 	TSharedPtr<IPropertyHandle> TypePropertyHandle;
 
diff --git a/Engine/Plugins/Mutable/Source/CustomizableObjectEditor/Private/MuCOE/CustomizableObjectEditorModule.cpp b/Engine/Plugins/Mutable/Source/CustomizableObjectEditor/Private/MuCOE/CustomizableObjectEditorModule.cpp
--- a/Engine/Plugins/Mutable/Source/CustomizableObjectEditor/Private/MuCOE/CustomizableObjectEditorModule.cpp
+++ b/Engine/Plugins/Mutable/Source/CustomizableObjectEditor/Private/MuCOE/CustomizableObjectEditorModule.cpp
@@ -344,6 +344,8 @@ void FCustomizableObjectEditorModule::ShutdownModule()
 		PropertyModule.UnregisterCustomPropertyTypeLayout(MeshReshapeBoneReferenceUStructName);
 		PropertyModule.UnregisterCustomPropertyTypeLayout(BoneToRemoveUStructName);
 		PropertyModule.UnregisterCustomPropertyTypeLayout(NAME_StrProperty);
+		PropertyModule.UnregisterCustomPropertyTypeLayout(FCOVariable::StaticStruct()->GetFName());
+
 
 		PropertyModule.NotifyCustomizationModuleChanged();
 	}
diff --git a/Engine/Plugins/Mutable/Source/CustomizableObjectEditor/Private/MuCOE/GenerateMutableSource/GenerateMutableSource.h b/Engine/Plugins/Mutable/Source/CustomizableObjectEditor/Private/MuCOE/GenerateMutableSource/GenerateMutableSource.h
--- a/Engine/Plugins/Mutable/Source/CustomizableObjectEditor/Private/MuCOE/GenerateMutableSource/GenerateMutableSource.h
+++ b/Engine/Plugins/Mutable/Source/CustomizableObjectEditor/Private/MuCOE/GenerateMutableSource/GenerateMutableSource.h
@@ -903,7 +903,7 @@ struct FMutableGraphGenerationContext
 	};
 
 	/** Name of the material break Node that is being compiled.
-	NOTE: This is set when the material pin of a material breake node is compiled and then it must be reset to null. */
+	NOTE: This is set when the material pin of a material breake node is compiled and then it must be reset to empty. */
 	FMaterialBreakParameter CurrentMaterialBreakParameter;
 
 	TMap<const UCustomizableObjectNodeMaterialParameter*, mu::Ptr<mu::NodeMaterialParameter>> MaterialParameterNodesCache;
diff --git a/Engine/Plugins/Mutable/Source/CustomizableObjectEditor/Private/MuCOE/Nodes/CONodeMaterialBreak.h b/Engine/Plugins/Mutable/Source/CustomizableObjectEditor/Private/MuCOE/Nodes/CONodeMaterialBreak.h
--- a/Engine/Plugins/Mutable/Source/CustomizableObjectEditor/Private/MuCOE/Nodes/CONodeMaterialBreak.h
+++ b/Engine/Plugins/Mutable/Source/CustomizableObjectEditor/Private/MuCOE/Nodes/CONodeMaterialBreak.h
@@ -24,7 +24,7 @@ public:
 };
 
 
-/** Base class for all Material Parameters. */
+
 UCLASS(MinimalAPI)
 class UCONodeMaterialBreakPinData : public UCustomizableObjectNodePinData
 {
diff --git a/Engine/Plugins/Mutable/Source/MutableTools/Private/MuT/ASTOpMaterialBreak.h b/Engine/Plugins/Mutable/Source/MutableTools/Private/MuT/ASTOpMaterialBreak.h
--- a/Engine/Plugins/Mutable/Source/MutableTools/Private/MuT/ASTOpMaterialBreak.h
+++ b/Engine/Plugins/Mutable/Source/MutableTools/Private/MuT/ASTOpMaterialBreak.h
@@ -12,7 +12,7 @@ namespace mu
 {
 	struct FProgram;
 
-	/** Parameter operation. */
+	/** Break Material operation. */
 	class ASTOpMaterialBreak final : public ASTOp
 	{
 	public:
diff --git a/Engine/Plugins/Mutable/Source/MutableTools/Public/MuT/NodeImageFromMaterialParameter.h b/Engine/Plugins/Mutable/Source/MutableTools/Public/MuT/NodeImageFromMaterialParameter.h
--- a/Engine/Plugins/Mutable/Source/MutableTools/Public/MuT/NodeImageFromMaterialParameter.h
+++ b/Engine/Plugins/Mutable/Source/MutableTools/Public/MuT/NodeImageFromMaterialParameter.h
@@ -31,4 +31,6 @@ namespace mu
 
 		static UE_API FNodeType StaticType;
 	};
-}
\ No newline at end of file
+}
+
+#undef UE_API
diff --git a/Engine/Plugins/Mutable/Source/MutableTools/Public/MuT/NodeScalarMaterialBreak.h b/Engine/Plugins/Mutable/Source/MutableTools/Public/MuT/NodeScalarMaterialBreak.h
--- a/Engine/Plugins/Mutable/Source/MutableTools/Public/MuT/NodeScalarMaterialBreak.h
+++ b/Engine/Plugins/Mutable/Source/MutableTools/Public/MuT/NodeScalarMaterialBreak.h
@@ -11,7 +11,7 @@
 
 namespace mu
 {
-	/** NodeMaterialTable class. */
+	/** NodeScalarMaterialBreak class. */
 	class NodeScalarMaterialBreak : public NodeScalar
 	{
 	public:
