[Backout] - CL44849948
#rnx
Original CL Desc
-----------------------------------------------------------------
[Mutable] Implemented the new node Break Material.
+ Added the new concept of lazy branches into the mutable core.
+ Created the new node interface and widgets.

#rnx

diff --git a/Engine/Plugins/Mutable/Source/CustomizableObject/Internal/MuCO/CustomizableObjectPrivate.h b/Engine/Plugins/Mutable/Source/CustomizableObject/Internal/MuCO/CustomizableObjectPrivate.h
--- a/Engine/Plugins/Mutable/Source/CustomizableObject/Internal/MuCO/CustomizableObjectPrivate.h
+++ b/Engine/Plugins/Mutable/Source/CustomizableObject/Internal/MuCO/CustomizableObjectPrivate.h
@@ -1593,103 +1593,105 @@ enum class ECustomizableObjectVersions : int32
 
 	RemovedUnnecessarySerializationVersioning,
 
-	AddTextureCompressionSettingCompilationInfo,
+		AddTextureCompressionSettingCompilationInfo,
 
-	RestructureConstantImageData,
+		RestructureConstantImageData,
 
-	RestructureConstantMeshData,
+		RestructureConstantMeshData,
 
-	RestructureRomData,
+		RestructureRomData,
 
-	RestructureRomDataRemovingRomHash,
+		RestructureRomDataRemovingRomHash,
 
-	ModifiedRomCompiledDataSerialization,
+		ModifiedRomCompiledDataSerialization,
 
-	ModelResourcesExtensionData,
+		ModelResourcesExtensionData,
 
-	LODsPerComponent,
+		LODsPerComponent,
 
-	LODsPerComponentTypeMismatch,
+		LODsPerComponentTypeMismatch,
 
-	ImageHiResLODsUseLODGroupInfo,
+		ImageHiResLODsUseLODGroupInfo,
 
-	MovedTableRowNoneGenerationToUnreal,
+		MovedTableRowNoneGenerationToUnreal,
 
-	RemoveObsoletMeshInterpolateAndGeometryOp,
+		RemoveObsoletMeshInterpolateAndGeometryOp,
 
-	RemoveObsoleteDataTypesFromEnum,
+		RemoveObsoleteDataTypesFromEnum,
 
-	ConvertModelResourcesToUObject,
+		ConvertModelResourcesToUObject,
 
-	RemoveObsoletImageGradientOp,
+		RemoveObsoletImageGradientOp,
 
-	MeshReferencesExtendedForCompilation,
+		MeshReferencesExtendedForCompilation,
 
-	RemoveObsoleteBoolOps,
+		RemoveObsoleteBoolOps,
 
-	AddOverlayMaterials,
+		AddOverlayMaterials,
 
-	PrefetchHighQualityMipMaps,
+		PrefetchHighQualityMipMaps,
 
-	AddedMeshParameterOp,
+		AddedMeshParameterOp,
 
-	AddedMeshParameterOpForDDCPollution,
+		AddedMeshParameterOpForDDCPollution,
 
-	ExtendedMeshParameterArgumentsWithLODAndSection,
-	
-	AddAssetUserDataEditor,
-	
-	MeshDataRomSplit,
+		ExtendedMeshParameterArgumentsWithLODAndSection,
 
-	MeshDataRomSplitBackout,
+		AddAssetUserDataEditor,
 
-	MovedLODSettingsToMeshComponentNode,
-	
-	AddedMeshPrepareLayoutOp,
+		MeshDataRomSplit,
 
-	AddedMeshIDToMeshParamOp,
+		MeshDataRomSplitBackout,
 
-	ClothMorphMeshMetadata,
+		MovedLODSettingsToMeshComponentNode,
 
-	AddedMeshIDToMeshParamOpBackout,
+		AddedMeshPrepareLayoutOp,
 
-	MeshDataRomSplitSerializationFix,
+		AddedMeshIDToMeshParamOp,
 
-	ReaddAddedMeshIDToMeshParamOp,
+		ClothMorphMeshMetadata,
 
-	AddConnectedChildObjectComponentsToPrepass,
+		AddedMeshIDToMeshParamOpBackout,
 
-	FixMeshReusalDueToLayouts,
+		MeshDataRomSplitSerializationFix,
 
-	IncorrectBonePoseMerging,
+		ReaddAddedMeshIDToMeshParamOp,
 
-	CoreParameterUObjects,
+		AddConnectedChildObjectComponentsToPrepass,
 
-	MoveMeshMetadataToOperation,
-	
-	MovePhysicsBodiesToRoms,
+		FixMeshReusalDueToLayouts,
+
+		IncorrectBonePoseMerging,
+
+		CoreParameterUObjects,
+
+		MoveMeshMetadataToOperation,
+
+		MovePhysicsBodiesToRoms,
+
+		SwitchOpCompactRepresentation,
 
-	SwitchOpCompactRepresentation,
+		AddedBlockMasksToRuntimeLayouts,
 
-	AddedBlockMasksToRuntimeLayouts,
+		StoreMipTailInASingleRom,
 
-	StoreMipTailInASingleRom,
+		FixDDCCrashesAndLoadTimes,
 
-	FixDDCCrashesAndLoadTimes,
+		CompactBindingDataStructure,
 
-	CompactBindingDataStructure,
+		MorphMeshParameters,
 
-	MorphMeshParameters,
+		MaterialParameter,
 
-	MaterialParameter,
+		AddLinearTosRGB,
 
-	AddLinearTosRGB,
+		AddMaterialArgs,
 
-	AddMaterialArgs,
+		CacheReferenceTextureSize,
 
-	CacheReferenceTextureSize,
+		MaterialBreak,
 
-	MaterialBreak,
+		BackoutMaterialBreak,
 
 	// -----<new versions can be added above this line>--------
 	LastCustomizableObjectVersion
diff --git a/Engine/Plugins/Mutable/Source/CustomizableObjectEditor/Private/MuCOE/COVariable.cpp b/Engine/Plugins/Mutable/Source/CustomizableObjectEditor/Private/MuCOE/COVariable.cpp
deleted file mode 100644
diff --git a/Engine/Plugins/Mutable/Source/CustomizableObjectEditor/Private/MuCOE/COVariable.h b/Engine/Plugins/Mutable/Source/CustomizableObjectEditor/Private/MuCOE/COVariable.h
deleted file mode 100644
diff --git a/Engine/Plugins/Mutable/Source/CustomizableObjectEditor/Private/MuCOE/CustomizableObjectEditorModule.cpp b/Engine/Plugins/Mutable/Source/CustomizableObjectEditor/Private/MuCOE/CustomizableObjectEditorModule.cpp
--- a/Engine/Plugins/Mutable/Source/CustomizableObjectEditor/Private/MuCOE/CustomizableObjectEditorModule.cpp
+++ b/Engine/Plugins/Mutable/Source/CustomizableObjectEditor/Private/MuCOE/CustomizableObjectEditorModule.cpp
@@ -32,7 +32,6 @@
 #include "MuCOE/CustomizableObjectMacroLibrary/CustomizableObjectMacroDetails.h"
 #include "MuCOE/CustomizableObjectVersionBridge.h"
 #include "MuCOE/CustomizableObjectNodeObjectGroupDetails.h"
-#include "MuCOE/Nodes/CONodeMaterialBreak.h"
 #include "MuCOE/Nodes/CustomizableObjectNodeCopyMaterial.h"
 #include "MuCOE/Nodes/CustomizableObjectNodeExternalPin.h"
 #include "MuCOE/Nodes/CustomizableObjectNodeExternalPinDetails.h"
@@ -89,7 +88,6 @@
 #include "MuCOE/Widgets/CustomizableObjectVariationCustomization.h"
 #include "MuCOE/Widgets/CustomizableObjectLODReductionSettings.h"
 #include "MuCOE/Widgets/CustomizableObjectNodeTableCompilationFilterEditor.h"
-#include "MuCOE/COVariable.h"
 #include "PropertyEditorModule.h"
 #include "AssetRegistry/AssetRegistryModule.h"
 #include "MuCO/CustomizableObjectInstance.h"
@@ -277,7 +275,6 @@ void FCustomizableObjectEditorModule::StartupModule()
 	RegisterCustomDetails(PropertyModule, UCustomizableObjectNodeComponentMesh::StaticClass(), FOnGetDetailCustomizationInstance::CreateStatic(&FCustomizableObjectNodeComponentMeshDetails::MakeInstance));
 	RegisterCustomDetails(PropertyModule, UCustomizableObjectNodeMacroInstance::StaticClass(), FOnGetDetailCustomizationInstance::CreateStatic(&FCustomizableObjectNodeMacroInstanceDetails::MakeInstance));
 	RegisterCustomDetails(PropertyModule, UCustomizableObjectNodeParameter::StaticClass(), FOnGetDetailCustomizationInstance::CreateStatic(&FCustomizableObjectNodeParameterDetails::MakeInstance));
-	RegisterCustomDetails(PropertyModule, UCONodeMaterialBreak::StaticClass(), FOnGetDetailCustomizationInstance::CreateStatic(&FCustomizableObjectNodeDetails::MakeInstance));
 
 	// Other Objects
 	RegisterCustomDetails(PropertyModule, UCustomizableObject::StaticClass(), FOnGetDetailCustomizationInstance::CreateStatic(&FCustomizableObjectDetails::MakeInstance));
@@ -297,7 +294,6 @@ void FCustomizableObjectEditorModule::StartupModule()
 	PropertyModule.RegisterCustomPropertyTypeLayout(FCustomizableObjectVariation::StaticStruct()->GetFName(), FOnGetPropertyTypeCustomizationInstance::CreateStatic(&FCustomizableObjectVariationCustomization::MakeInstance));
 	PropertyModule.RegisterCustomPropertyTypeLayout(FCustomizableObjectTextureVariation::StaticStruct()->GetFName(), FOnGetPropertyTypeCustomizationInstance::CreateStatic(&FCustomizableObjectVariationCustomization::MakeInstance));
 	PropertyModule.RegisterCustomPropertyTypeLayout(FTableNodeCompilationFilter::StaticStruct()->GetFName(), FOnGetPropertyTypeCustomizationInstance::CreateStatic(&FCustomizableObjectNodeTableCompilationFilterEditor::MakeInstance));
-	PropertyModule.RegisterCustomPropertyTypeLayout(FCOVariable::StaticStruct()->GetFName(), FOnGetPropertyTypeCustomizationInstance::CreateStatic(&FCOVariableCustomzation::MakeInstance));
 
 	PropertyModule.NotifyCustomizationModuleChanged();
 
diff --git a/Engine/Plugins/Mutable/Source/CustomizableObjectEditor/Private/MuCOE/EdGraphSchema_CustomizableObject.cpp b/Engine/Plugins/Mutable/Source/CustomizableObjectEditor/Private/MuCOE/EdGraphSchema_CustomizableObject.cpp
--- a/Engine/Plugins/Mutable/Source/CustomizableObjectEditor/Private/MuCOE/EdGraphSchema_CustomizableObject.cpp
+++ b/Engine/Plugins/Mutable/Source/CustomizableObjectEditor/Private/MuCOE/EdGraphSchema_CustomizableObject.cpp
@@ -17,7 +17,6 @@
 #include "MuCOE/CustomizableObjectGraph.h"
 #include "MuCOE/CustomizableObjectMacroLibrary/CustomizableObjectGraphEditorToolkit.h"
 #include "MuCOE/CustomizableObjectEditor.h"
-#include "MuCOE/Nodes/CONodeMaterialBreak.h"
 #include "MuCOE/Nodes/CONodeMaterialConstant.h"
 #include "MuCOE/Nodes/CONodeMaterialVariation.h"
 #include "MuCOE/Nodes/CONodeMaterialSwitch.h"
@@ -633,8 +632,7 @@ void UEdGraphSchema_CustomizableObject::GetGraphContextActions(FGraphContextMenu
 		{
 			ContextMenuBuilder.CreateTemplateNode<UCONodeMaterialConstant>(),
 			ContextMenuBuilder.CreateTemplateNode<UCONodeMaterialVariation>(),
-			ContextMenuBuilder.CreateTemplateNode<UCONodeMaterialSwitch>(),
-			ContextMenuBuilder.CreateTemplateNode<UCONodeMaterialBreak>()
+			ContextMenuBuilder.CreateTemplateNode<UCONodeMaterialSwitch>()
 		};
 
 		AddNewNodeCategoryActionsFiltered(MaterialTemplateNodes, ContextMenuBuilder, NC_Material.ToString(), 2, Filter);
diff --git a/Engine/Plugins/Mutable/Source/CustomizableObjectEditor/Private/MuCOE/GenerateMutableSource/GenerateMutableSource.h b/Engine/Plugins/Mutable/Source/CustomizableObjectEditor/Private/MuCOE/GenerateMutableSource/GenerateMutableSource.h
--- a/Engine/Plugins/Mutable/Source/CustomizableObjectEditor/Private/MuCOE/GenerateMutableSource/GenerateMutableSource.h
+++ b/Engine/Plugins/Mutable/Source/CustomizableObjectEditor/Private/MuCOE/GenerateMutableSource/GenerateMutableSource.h
@@ -18,7 +18,6 @@
 #include "MuCOE/Nodes/CustomizableObjectNodeTable.h"
 #include "MuT/NodeComponentNew.h"
 #include "MuT/NodeImageConstant.h"
-#include "MuT/NodeImageMaterialBreak.h"
 #include "MuT/NodeMeshConstant.h"
 #include "MuT/NodeMeshApplyPose.h"
 #include "MuT/NodeModifierMeshClipWithMesh.h"
@@ -27,8 +26,6 @@
 #include "MuT/NodeScalarEnumParameter.h"
 #include "MuT/NodeScalarParameter.h"
 #include "MuT/NodeSurfaceNew.h"
-#include "MuT/NodeMaterialConstant.h"
-#include "MuT/NodeMaterialParameter.h"
 #include "MuT/Table.h"
 #include "MuR/Skeleton.h"
 #include "MuR/Mesh.h"
@@ -45,10 +42,8 @@ class FCustomizableObjectCompiler;
 class UTextureLODSettings;
 class UAnimInstance;
 class UCompositeDataTable;
-class UCONodeMaterialConstant;
 class UCustomizableObjectNodeMacroInstance;
 class UCustomizableObjectNodeMaterial;
-class UCustomizableObjectNodeMaterialParameter;
 class UCustomizableObjectNodeMeshMorph;
 class UCustomizableObjectNodeMeshMorphStackDefinition;
 class UCustomizableObjectNodeObjectGroup;
@@ -895,19 +890,6 @@ struct FMutableGraphGenerationContext
 	TMap<FName, TObjectPtr<UTexture>> TextureParameterDefaultValues;
 	TMap<FName, TObjectPtr<USkeletalMesh>> SkeletalMeshParameterDefaultValues;
 	TMap<FName, TObjectPtr<UMaterialInterface>> MaterialParameterDefaultValues;
-
-	
-	struct FMaterialBreakParameter {
-		FName ParameterName;
-		EMaterialParameterType ParameterType = EMaterialParameterType::None;
-	};
-
-	/** Name of the material break Node that is being compiled.
-	NOTE: This is set when the material pin of a material breake node is compiled and then it must be reset to null. */
-	FMaterialBreakParameter CurrentMaterialBreakParameter;
-
-	TMap<const UCustomizableObjectNodeMaterialParameter*, mu::Ptr<mu::NodeMaterialParameter>> MaterialParameterNodesCache;
-	TMap<const UCONodeMaterialConstant*, mu::Ptr<mu::NodeMaterialConstant>> MaterialConstantNodesCache;
 };
 
 
diff --git a/Engine/Plugins/Mutable/Source/CustomizableObjectEditor/Private/MuCOE/GenerateMutableSource/GenerateMutableSourceColor.cpp b/Engine/Plugins/Mutable/Source/CustomizableObjectEditor/Private/MuCOE/GenerateMutableSource/GenerateMutableSourceColor.cpp
--- a/Engine/Plugins/Mutable/Source/CustomizableObjectEditor/Private/MuCOE/GenerateMutableSource/GenerateMutableSourceColor.cpp
+++ b/Engine/Plugins/Mutable/Source/CustomizableObjectEditor/Private/MuCOE/GenerateMutableSource/GenerateMutableSourceColor.cpp
@@ -8,7 +8,6 @@
 #include "MuCOE/GenerateMutableSource/GenerateMutableSourceFloat.h"
 #include "MuCOE/GenerateMutableSource/GenerateMutableSourceImage.h"
 #include "MuCOE/GenerateMutableSource/GenerateMutableSourceMacro.h"
-#include "MuCOE/GenerateMutableSource/GenerateMutableSourceMaterial.h"
 #include "MuCOE/GenerateMutableSource/GenerateMutableSourceTable.h"
 #include "MuCOE/GenerateMutableSource/GenerateMutableSource.h"
 #include "MuCOE/GraphTraversal.h"
@@ -20,7 +19,6 @@
 #include "MuCOE/Nodes/CustomizableObjectNodeColorVariation.h"
 #include "MuCOE/Nodes/CustomizableObjectNodeColorToSRGB.h"
 #include "MuCOE/Nodes/CustomizableObjectNodeMacroInstance.h"
-#include "MuCOE/Nodes/CONodeMaterialBreak.h"
 #include "MuCOE/Nodes/CustomizableObjectNodeTable.h"
 #include "MuCOE/Nodes/CustomizableObjectNodeTextureSample.h"
 #include "MuCOE/Nodes/CustomizableObjectNodeTunnel.h"
@@ -33,7 +31,6 @@
 #include "MuT/NodeColourTable.h"
 #include "MuT/NodeColourVariation.h"
 #include "MuT/NodeColorToSRGB.h"
-#include "MuT/NodeColourMaterialBreak.h"
 #include "MuT/NodeImageBinarise.h"
 
 
@@ -376,42 +373,6 @@ mu::Ptr<mu::NodeColour> GenerateMutableSourceColor(const UEdGraphPin* Pin, FMuta
 		Result = GenerateMutableSourceMacro<mu::NodeColour>(*Pin, GenerationContext, GenerateMutableSourceColor);
 	}
 
-	else if (const UCONodeMaterialBreak* TypedNodeMaterialBreak = Cast<UCONodeMaterialBreak>(Node))
-	{
-		Result = new mu::NodeColourConstant();
-		check(TypedNodeMaterialBreak->MaterialPinRef.Get());
-
-		if (const UEdGraphPin* LinkedMaterial = FollowInputPin(*TypedNodeMaterialBreak->MaterialPinRef.Get()))
-		{
-			mu::Ptr<mu::NodeColourMaterialBreak> MaterialBreak = new mu::NodeColourMaterialBreak();
-
-			FName ParameterName = TypedNodeMaterialBreak->GetPinParameterName(*Pin);
-			check(!ParameterName.IsNone());
-
-			// Set the name of the parameter that represents the pin
-			MaterialBreak->ParameterName = ParameterName;
-
-			// Generate the material that is going to be broken into parameters
-			GenerationContext.CurrentMaterialBreakParameter = FMutableGraphGenerationContext::FMaterialBreakParameter{ ParameterName, EMaterialParameterType::Vector };
-			mu::Ptr<mu::NodeMaterial> LinkedNode = GenerateMutableSourceMaterial(LinkedMaterial, GenerationContext);
-			MaterialBreak->MaterialSource = LinkedNode;
-			GenerationContext.CurrentMaterialBreakParameter = FMutableGraphGenerationContext::FMaterialBreakParameter{ };
-
-			if (LinkedNode)
-			{
-				Result = MaterialBreak;
-			}
-			else
-			{
-				GenerationContext.Log(LOCTEXT("MaterialBreakNodeError", "Something went wrong generation the Material Node linked to the Material Break Node."), Node);
-			}
-		}
-		else
-		{
-			GenerationContext.Log(LOCTEXT("MaterialBreakPinError", "Material Break Nodes must be linked to an input Material Node."), Node);
-		}
-		}
-
 	else
 	{
 		GenerationContext.Log(LOCTEXT("UnimplementedNode", "Node type not implemented yet."), Node);
diff --git a/Engine/Plugins/Mutable/Source/CustomizableObjectEditor/Private/MuCOE/GenerateMutableSource/GenerateMutableSourceFloat.cpp b/Engine/Plugins/Mutable/Source/CustomizableObjectEditor/Private/MuCOE/GenerateMutableSource/GenerateMutableSourceFloat.cpp
--- a/Engine/Plugins/Mutable/Source/CustomizableObjectEditor/Private/MuCOE/GenerateMutableSource/GenerateMutableSourceFloat.cpp
+++ b/Engine/Plugins/Mutable/Source/CustomizableObjectEditor/Private/MuCOE/GenerateMutableSource/GenerateMutableSourceFloat.cpp
@@ -10,10 +10,8 @@
 #include "MuCOE/EdGraphSchema_CustomizableObject.h"
 #include "MuCOE/GenerateMutableSource/GenerateMutableSourceImage.h"
 #include "MuCOE/GenerateMutableSource/GenerateMutableSourceMacro.h"
-#include "MuCOE/GenerateMutableSource/GenerateMutableSourceMaterial.h"
 #include "MuCOE/GenerateMutableSource/GenerateMutableSourceTable.h"
 #include "MuCOE/GraphTraversal.h"
-#include "MuCOE/Nodes/CONodeMaterialBreak.h"
 #include "MuCOE/Nodes/CustomizableObjectNodeCurve.h"
 #include "MuCOE/Nodes/CustomizableObjectNodeEnumParameter.h"
 #include "MuCOE/Nodes/CustomizableObjectNodeFloatConstant.h"
@@ -30,7 +28,6 @@
 #include "MuT/NodeScalarTable.h"
 #include "MuT/NodeScalarVariation.h"
 #include "MuT/NodeScalarArithmeticOperation.h"
-#include "MuT/NodeScalarMaterialBreak.h"
 
 #define LOCTEXT_NAMESPACE "CustomizableObjectEditor"
 
@@ -399,42 +396,6 @@ mu::Ptr<mu::NodeScalar> GenerateMutableSourceFloat(const UEdGraphPin* Pin, FMuta
 		Result = GenerateMutableSourceMacro<mu::NodeScalar>(*Pin, GenerationContext, GenerateMutableSourceFloat);
 	}
 
-	else if (const UCONodeMaterialBreak* TypedNodeMaterialBreak = Cast<UCONodeMaterialBreak>(Node))
-	{
-		Result = new mu::NodeScalarConstant();
-		check(TypedNodeMaterialBreak->MaterialPinRef.Get());
-
-		if (const UEdGraphPin* LinkedMaterial = FollowInputPin(*TypedNodeMaterialBreak->MaterialPinRef.Get()))
-		{
-			mu::Ptr<mu::NodeScalarMaterialBreak> MaterialBreak = new mu::NodeScalarMaterialBreak();
-
-			FName ParameterName = TypedNodeMaterialBreak->GetPinParameterName(*Pin);
-			check(!ParameterName.IsNone());
-
-			// Set the name of the parameter that represents the pin
-			MaterialBreak->ParameterName = ParameterName;
-
-			// Generate the material that is going to be broken into parameters
-			GenerationContext.CurrentMaterialBreakParameter = FMutableGraphGenerationContext::FMaterialBreakParameter{ ParameterName, EMaterialParameterType::Scalar };
-			mu::Ptr<mu::NodeMaterial> LinkedNode = GenerateMutableSourceMaterial(LinkedMaterial, GenerationContext);
-			MaterialBreak->MaterialSource = LinkedNode;
-			GenerationContext.CurrentMaterialBreakParameter = FMutableGraphGenerationContext::FMaterialBreakParameter{ };
-
-			if (LinkedNode)
-			{
-				Result = MaterialBreak;
-			}
-			else
-			{
-				GenerationContext.Log(LOCTEXT("MaterialBreakNodeError", "Something went wrong generation the Material Node linked to the Material Break Node."), Node);
-			}
-		}
-		else
-		{
-			GenerationContext.Log(LOCTEXT("MaterialBreakPinError", "Material Break Nodes must be linked to an input Material Node."), Node);
-		}
-	}
-
 	else
 	{
 		GenerationContext.Log(LOCTEXT("UnimplementedNode", "Node type not implemented yet."), Node);
diff --git a/Engine/Plugins/Mutable/Source/CustomizableObjectEditor/Private/MuCOE/GenerateMutableSource/GenerateMutableSourceImage.cpp b/Engine/Plugins/Mutable/Source/CustomizableObjectEditor/Private/MuCOE/GenerateMutableSource/GenerateMutableSourceImage.cpp
--- a/Engine/Plugins/Mutable/Source/CustomizableObjectEditor/Private/MuCOE/GenerateMutableSource/GenerateMutableSourceImage.cpp
+++ b/Engine/Plugins/Mutable/Source/CustomizableObjectEditor/Private/MuCOE/GenerateMutableSource/GenerateMutableSourceImage.cpp
@@ -9,12 +9,10 @@
 #include "MuCOE/EdGraphSchema_CustomizableObject.h"
 #include "MuCOE/GenerateMutableSource/GenerateMutableSourceColor.h"
 #include "MuCOE/GenerateMutableSource/GenerateMutableSourceFloat.h"
-#include "MuCOE/GenerateMutableSource/GenerateMutableSourceMaterial.h"
 #include "MuCOE/GenerateMutableSource/GenerateMutableSourceMesh.h"
 #include "MuCOE/GenerateMutableSource/GenerateMutableSourceProjector.h"
 #include "MuCOE/GenerateMutableSource/GenerateMutableSourceTable.h"
 #include "MuCOE/GraphTraversal.h"
-#include "MuCOE/Nodes/CONodeMaterialBreak.h"
 #include "MuCOE/Nodes/CustomizableObjectNodeMacroInstance.h"
 #include "MuCOE/Nodes/CustomizableObjectNodeMeshInterface.h"
 #include "MuCOE/Nodes/CustomizableObjectNodeTable.h"
@@ -1220,42 +1218,6 @@ mu::Ptr<mu::NodeImage> GenerateMutableSourceImage(const UEdGraphPin* Pin, FMutab
 		GenerationContext.MacroNodesStack.Push(MacroInstanceNode);
 	}
 
-	else if (const UCONodeMaterialBreak* TypedNodeMaterialBreak = Cast<UCONodeMaterialBreak>(Node))
-	{
-		Result = new mu::NodeImageConstant();
-		check(TypedNodeMaterialBreak->MaterialPinRef.Get());
-
-		if (const UEdGraphPin* LinkedMaterial = FollowInputPin(*TypedNodeMaterialBreak->MaterialPinRef.Get()))
-		{
-			mu::Ptr<mu::NodeImageMaterialBreak> MaterialBreak = new mu::NodeImageMaterialBreak();
-
-			FName ParameterName = TypedNodeMaterialBreak->GetPinParameterName(*Pin);
-			check(!ParameterName.IsNone());
-
-			// Set the name of the parameter that represents the pin
-			MaterialBreak->ParameterName = ParameterName;
-
-			// Generate the material that is going to be broken into parameters
-			GenerationContext.CurrentMaterialBreakParameter = FMutableGraphGenerationContext::FMaterialBreakParameter{ ParameterName, EMaterialParameterType::Texture };
-			mu::Ptr<mu::NodeMaterial> LinkedNode = GenerateMutableSourceMaterial(LinkedMaterial, GenerationContext);
-			MaterialBreak->MaterialSource = LinkedNode;
-			GenerationContext.CurrentMaterialBreakParameter = FMutableGraphGenerationContext::FMaterialBreakParameter{ };
-
-			if (LinkedNode)
-			{
-				Result = MaterialBreak;
-			}
-			else
-			{
-				GenerationContext.Log(LOCTEXT("MaterialBreakNodeError", "Something went wrong generation the Material Node linked to the Material Break Node."), Node);
-			}
-		}
-		else
-		{
-			GenerationContext.Log(LOCTEXT("MaterialBreakPinError", "Material Break Nodes must be linked to an input Material Node."), Node);
-		}
-	}
-
 	else
 	{
 		GenerationContext.Log(LOCTEXT("UnimplementedNode", "Node type not implemented yet."), Node);
diff --git a/Engine/Plugins/Mutable/Source/CustomizableObjectEditor/Private/MuCOE/GenerateMutableSource/GenerateMutableSourceMaterial.cpp b/Engine/Plugins/Mutable/Source/CustomizableObjectEditor/Private/MuCOE/GenerateMutableSource/GenerateMutableSourceMaterial.cpp
--- a/Engine/Plugins/Mutable/Source/CustomizableObjectEditor/Private/MuCOE/GenerateMutableSource/GenerateMutableSourceMaterial.cpp
+++ b/Engine/Plugins/Mutable/Source/CustomizableObjectEditor/Private/MuCOE/GenerateMutableSource/GenerateMutableSourceMaterial.cpp
@@ -2,13 +2,10 @@
 
 #include "MuCOE/GenerateMutableSource/GenerateMutableSourceMaterial.h"
 
-#include "Engine/TextureLODSettings.h"
-#include "Interfaces/ITargetPlatform.h"
 #include "MuCOE/CustomizableObjectCompiler.h"
 #include "MuCOE/CustomizableObjectGraph.h"
 #include "MuCOE/EdGraphSchema_CustomizableObject.h"
 #include "MuCOE/GenerateMutableSource/GenerateMutableSourceFloat.h"
-#include "MuCOE/GenerateMutableSource/GenerateMutableSourceImage.h"
 #include "MuCOE/GenerateMutableSource/GenerateMutableSourceMacro.h"
 #include "MuCOE/GenerateMutableSource/GenerateMutableSourceTable.h"
 #include "MuCOE/GraphTraversal.h"
@@ -17,14 +14,10 @@
 #include "MuCOE/Nodes/CustomizableObjectNodeMacroInstance.h"
 #include "MuCOE/Nodes/CustomizableObjectNodeTable.h"
 #include "MuCOE/Nodes/CustomizableObjectNodeTunnel.h"
-#include "MuCOE/Nodes/CONodeMaterialBreak.h"
 #include "MuCOE/Nodes/CONodeMaterialConstant.h"
 #include "MuCOE/Nodes/CONodeMaterialSwitch.h"
 #include "MuCOE/Nodes/CONodeMaterialVariation.h"
 #include "MuCO/LoadUtils.h"
-#include "MuT/NodeImageFormat.h"
-#include "MuT/NodeImageFromMaterialParameter.h"
-#include "MuT/NodeImageResize.h"
 #include "MuT/NodeMaterialConstant.h"
 #include "MuT/NodeMaterialSwitch.h"
 #include "MuT/NodeMaterialTable.h"
@@ -68,99 +61,14 @@ mu::Ptr<mu::NodeMaterial> GenerateMutableSourceMaterial(const UEdGraphPin* Pin,
 	
 	if (const UCONodeMaterialConstant* TypedNodeMaterialConstant = Cast<UCONodeMaterialConstant>(Node))
 	{
-		mu::Ptr<mu::NodeMaterialConstant> MaterialNode;
-		bCacheNode = false; //custom node cache
+		mu::Ptr<mu::NodeMaterialConstant> MaterialNode = new mu::NodeMaterialConstant();
+		Result = MaterialNode;
 
 		if (UMaterialInterface* Material = GenerationContext.LoadObject(TypedNodeMaterialConstant->Material))
 		{
-			if (GenerationContext.MaterialConstantNodesCache.Contains(TypedNodeMaterialConstant))
-			{
-				MaterialNode = GenerationContext.MaterialConstantNodesCache[TypedNodeMaterialConstant];
-			}
-			else
-			{
-				MaterialNode = new mu::NodeMaterialConstant();
-
-				GenerationContext.CurrentReferencedMaterialIndex = GenerationContext.ReferencedMaterials.AddUnique(Material);
-				MaterialNode->MaterialId = GenerationContext.CurrentReferencedMaterialIndex;
-
-				GenerationContext.MaterialConstantNodesCache.Add(TypedNodeMaterialConstant, MaterialNode);
-			}
-
-			// Store the parameter that needs to be processed
-			if (GenerationContext.CurrentMaterialBreakParameter.ParameterType != EMaterialParameterType::None)
-			{
-				if (!TypedNodeMaterialConstant->Material)
-				{
-					GenerationContext.Log(LOCTEXT("NoRefernceMaterialToBreak", "Could not find a Reference Material to break."), TypedNodeMaterialConstant);
-				}
-				else
-				{
-					switch (GenerationContext.CurrentMaterialBreakParameter.ParameterType)
-					{
-					case EMaterialParameterType::Texture:
-					{
-						UTexture* Texture = nullptr;
-						bool bParameterFound = Material->GetTextureParameterValue(GenerationContext.CurrentMaterialBreakParameter.ParameterName, Texture);
-						UTexture2D* BaseTexture = Cast<UTexture2D>(Texture);
-
-						if (bParameterFound && BaseTexture)
-						{
-							TSharedPtr<mu::FImage> ImageConstant = GenerateImageConstant(BaseTexture, GenerationContext, false);
-							const uint32 MipsToSkip = ComputeLODBiasForTexture(GenerationContext, *BaseTexture);
-
-							mu::Ptr<mu::NodeImageConstant> ConstantImageNode = new mu::NodeImageConstant;
-							ConstantImageNode->SetValue(ImageConstant);
-
-							const UTextureLODSettings& TextureLODSettings = GenerationContext.CompilationContext->Options.TargetPlatform->GetTextureLODSettings();
-							const FTextureLODGroup& LODGroupInfo = TextureLODSettings.GetTextureLODGroup(BaseTexture->LODGroup);
-							const FCompilationOptions& CompilationOptions = GenerationContext.CompilationContext->Options;
-
-							ConstantImageNode->SourceDataDescriptor.OptionalMaxLODSize = LODGroupInfo.OptionalMaxLODSize;
-							ConstantImageNode->SourceDataDescriptor.OptionalLODBias = LODGroupInfo.OptionalLODBias;
-							ConstantImageNode->SourceDataDescriptor.NumNonOptionalLODs = CompilationOptions.MinDiskMips;
-
-							const FString TextureName = GetNameSafe(BaseTexture).ToLower();
-							ConstantImageNode->SourceDataDescriptor.SourceId = CityHash32(reinterpret_cast<const char*>(*TextureName), TextureName.Len() * sizeof(FString::ElementType));
-
-							mu::Ptr<mu::NodeImage> FinalNodeImag = ResizeTextureByNumMips(ConstantImageNode, MipsToSkip);
-
-							if (FinalNodeImag)
-							{
-								MaterialNode->ImageValues.Add(GenerationContext.CurrentMaterialBreakParameter.ParameterName, FinalNodeImag);
-							}
-						}
-
-						break;
-					}
-					case EMaterialParameterType::Vector:
-					{
-						FLinearColor Color;
-						if (Material->GetVectorParameterValue(GenerationContext.CurrentMaterialBreakParameter.ParameterName, Color))
-						{
-							MaterialNode->ColorValues.Add(GenerationContext.CurrentMaterialBreakParameter.ParameterName, Color);
-						}
-
-						break;
-					}
-					case EMaterialParameterType::Scalar:
-					{
-						float Scalar;
-						if (Material->GetScalarParameterValue(GenerationContext.CurrentMaterialBreakParameter.ParameterName, Scalar))
-						{
-							MaterialNode->ScalarValues.Add(GenerationContext.CurrentMaterialBreakParameter.ParameterName, Scalar);
-						}
-
-						break;
-					}
-					default:
-						break;
-					}
-				}
-			}
+			GenerationContext.CurrentReferencedMaterialIndex = GenerationContext.ReferencedMaterials.AddUnique(Material);
+			MaterialNode->MaterialId = GenerationContext.CurrentReferencedMaterialIndex;
 		}
-
-		Result = MaterialNode;
 	}
 
 	else if (const UCONodeMaterialSwitch* TypedNodeMaterialSwitch = Cast<UCONodeMaterialSwitch>(Node))
@@ -263,103 +171,26 @@ mu::Ptr<mu::NodeMaterial> GenerateMutableSourceMaterial(const UEdGraphPin* Pin,
 			}
 		}
 	}
-
+	
 	else if (const UCustomizableObjectNodeMaterialParameter* MaterialParameterNode = Cast<UCustomizableObjectNodeMaterialParameter>(Node))
 	{
-		bCacheNode = false;
-		mu::Ptr<mu::NodeMaterialParameter> MaterialNode;
-
-		if (GenerationContext.MaterialParameterNodesCache.Contains(MaterialParameterNode))
-		{
-			MaterialNode = GenerationContext.MaterialParameterNodesCache[MaterialParameterNode];
-		}
-		else
-		{
-			MaterialNode = new mu::NodeMaterialParameter();
-			MaterialNode->Name = MaterialParameterNode->GetParameterName(&GenerationContext.MacroNodesStack);
-			MaterialNode->UID = GenerationContext.GetNodeIdUnique(Node).ToString();
-
-			if (MaterialParameterNode->DefaultValue)
-			{
-				GenerationContext.MaterialParameterDefaultValues.Add(FName(MaterialParameterNode->GetParameterName()), MutablePrivate::LoadObject(MaterialParameterNode->DefaultValue));
-			}
+		mu::Ptr<mu::NodeMaterialParameter> MaterialNode = new mu::NodeMaterialParameter();
+		Result = MaterialNode;
 
-			GenerationContext.ParameterUIDataMap.Add(MaterialParameterNode->GetParameterName(&GenerationContext.MacroNodesStack), FMutableParameterData(
-				MaterialParameterNode->ParamUIMetadata,
-				EMutableParameterType::Material));
+		MaterialNode->Name = MaterialParameterNode->GetParameterName(&GenerationContext.MacroNodesStack);
+		MaterialNode->UID = GenerationContext.GetNodeIdUnique(Node).ToString();
 
-			GenerationContext.MaterialParameterNodesCache.Add(MaterialParameterNode, MaterialNode);
-		}
-
-		// Store the parameter that needs to be processed
-		if (GenerationContext.CurrentMaterialBreakParameter.ParameterType != EMaterialParameterType::None)
+		if (MaterialParameterNode->DefaultValue)
 		{
-			if (!MaterialParameterNode->ReferenceValue)
-			{
-				GenerationContext.Log(LOCTEXT("NoRefernceMaterialToBreak", "Could not find a Reference Material to break."), MaterialParameterNode);
-			}
-			else
-			{
-				switch (GenerationContext.CurrentMaterialBreakParameter.ParameterType)
-				{
-				case EMaterialParameterType::Texture:
-				{
-					//TODO(Max): Move this to the generatemutablesourceimage:
-					mu::Ptr<mu::NodeImageFromMaterialParameter> TextureNode = new mu::NodeImageFromMaterialParameter();
-					TextureNode->ImageParameterName = GenerationContext.CurrentMaterialBreakParameter.ParameterName;
-
-					// Force the same format that the default texture if any.
-					mu::Ptr<mu::NodeImageFormat> FormatNode = new mu::NodeImageFormat();
-
-					// Force an "easy format" on the texture.
-					FormatNode->Format = mu::EImageFormat::RGBA_UByte;
-					FormatNode->Source = TextureNode;
-
-					// Resize image node.
-					mu::Ptr<mu::NodeImageResize> ResizeNode = new mu::NodeImageResize();
-					ResizeNode->Base = FormatNode;
-					ResizeNode->bRelative = false;
-
-					mu::FImageSize TextureSize(1);
-
-					//Find reference Texture
-					UTexture* Texture = nullptr;
-					bool bParameterFound = MaterialParameterNode->ReferenceValue->GetTextureParameterValue(GenerationContext.CurrentMaterialBreakParameter.ParameterName, Texture);
-					UTexture2D* ReferenceTexture = Cast<UTexture2D>(Texture);
-
-					if (bParameterFound && ReferenceTexture)
-					{
-						const uint32 LODBias = ComputeLODBiasForTexture(GenerationContext, *ReferenceTexture);
-						TextureSize.X = FMath::Max(ReferenceTexture->Source.GetSizeX() >> LODBias, 1);
-						TextureSize.Y = FMath::Max(ReferenceTexture->Source.GetSizeY() >> LODBias, 1);
-					}
-
-					ResizeNode->SizeX = TextureSize.X;
-					ResizeNode->SizeY = TextureSize.Y;
-
-					MaterialNode->ImageParameterNodes.Add(GenerationContext.CurrentMaterialBreakParameter.ParameterName.ToString(), ResizeNode);
-
-					break;
-				}
-				case EMaterialParameterType::Vector:
-				{
-					MaterialNode->ColorParameterNames.Add(GenerationContext.CurrentMaterialBreakParameter.ParameterName.ToString());
-					break;
-				}
-				case EMaterialParameterType::Scalar:
-				{
-					MaterialNode->ScalarParameterNames.Add(GenerationContext.CurrentMaterialBreakParameter.ParameterName.ToString());
-					break;
-				}
-				default:
-					break;
-				}
-			}
+			GenerationContext.MaterialParameterDefaultValues.Add(FName(MaterialParameterNode->GetParameterName()), MutablePrivate::LoadObject(MaterialParameterNode->DefaultValue));
 		}
 
-		Result = MaterialNode;
+		GenerationContext.ParameterUIDataMap.Add(MaterialParameterNode->GetParameterName(&GenerationContext.MacroNodesStack), FMutableParameterData(
+			MaterialParameterNode->ParamUIMetadata,
+			EMutableParameterType::Material));
 	}
 
+	
 	else if (const UCustomizableObjectNodeTable* TypedNodeTable = Cast<UCustomizableObjectNodeTable>(Node))
 	{
 		//This node will add a default value in case of error
diff --git a/Engine/Plugins/Mutable/Source/CustomizableObjectEditor/Private/MuCOE/Nodes/CONodeMaterialBreak.cpp b/Engine/Plugins/Mutable/Source/CustomizableObjectEditor/Private/MuCOE/Nodes/CONodeMaterialBreak.cpp
deleted file mode 100644
diff --git a/Engine/Plugins/Mutable/Source/CustomizableObjectEditor/Private/MuCOE/Nodes/CONodeMaterialBreak.h b/Engine/Plugins/Mutable/Source/CustomizableObjectEditor/Private/MuCOE/Nodes/CONodeMaterialBreak.h
deleted file mode 100644
diff --git a/Engine/Plugins/Mutable/Source/CustomizableObjectEditor/Private/MuCOE/Nodes/CustomizableObjectNode.cpp b/Engine/Plugins/Mutable/Source/CustomizableObjectEditor/Private/MuCOE/Nodes/CustomizableObjectNode.cpp
--- a/Engine/Plugins/Mutable/Source/CustomizableObjectEditor/Private/MuCOE/Nodes/CustomizableObjectNode.cpp
+++ b/Engine/Plugins/Mutable/Source/CustomizableObjectEditor/Private/MuCOE/Nodes/CustomizableObjectNode.cpp
@@ -767,12 +767,6 @@ bool UCustomizableObjectNode::HasPinViewer() const
 }
 
 
-bool UCustomizableObjectNode::CanCreatePinsFromPinViewer() const
-{
-	return false;
-}
-
-
 TSharedPtr<IDetailsView> UCustomizableObjectNode::CustomizePinDetails(const UEdGraphPin& Pin) const
 {
 	return nullptr;
@@ -829,10 +823,4 @@ bool UCustomizableObjectNode::CreateStaticStringPinWidget() const
 }
 
 
-TArray<FName> UCustomizableObjectNode::GetAllowedPinViewerCreationTypes() const
-{
-	return UEdGraphSchema_CustomizableObject::SupportedMacroPinTypes;
-}
-
-
 #undef LOCTEXT_NAMESPACE
diff --git a/Engine/Plugins/Mutable/Source/CustomizableObjectEditor/Private/MuCOE/Nodes/CustomizableObjectNodeDetails.cpp b/Engine/Plugins/Mutable/Source/CustomizableObjectEditor/Private/MuCOE/Nodes/CustomizableObjectNodeDetails.cpp
--- a/Engine/Plugins/Mutable/Source/CustomizableObjectEditor/Private/MuCOE/Nodes/CustomizableObjectNodeDetails.cpp
+++ b/Engine/Plugins/Mutable/Source/CustomizableObjectEditor/Private/MuCOE/Nodes/CustomizableObjectNodeDetails.cpp
@@ -5,11 +5,8 @@
 #include "DetailCategoryBuilder.h"
 #include "DetailLayoutBuilder.h"
 #include "DetailWidgetRow.h"
-#include "Input/Reply.h"
 #include "MuCOE/Nodes/CustomizableObjectNode.h"
 #include "MuCOE/PinViewer/SPinViewer.h"
-#include "Widgets/Input/SButton.h"
-#include "Widgets/Images/SImage.h"
 
 class IDetailLayoutBuilder;
 
@@ -35,33 +32,6 @@ void FCustomizableObjectNodeDetails::CustomizeDetails(IDetailLayoutBuilder& Deta
 			}
 			
 			IDetailCategoryBuilder& PinViewerCategoryBuilder = DetailBuilder.EditCategory("PinViewer", LOCTEXT("PinViewer", "Pins"), ECategoryPriority::Uncommon);
-			
-			if (Node->CanCreatePinsFromPinViewer())
-			{
-				FText TooltipText = LOCTEXT("NodeAddPinTooltip", "Create a new pin.");
-
-				PinViewerCategoryBuilder.HeaderContent
-				(
-					SNew(SHorizontalBox)
-					+ SHorizontalBox::Slot()
-					.HAlign(HAlign_Right)
-					[
-						SNew(SButton)
-						.ButtonStyle(FAppStyle::Get(), "SimpleButton")
-						.ContentPadding(FMargin(1, 0))
-						.OnClicked(this, &FCustomizableObjectNodeDetails::AddNewPin, Node)
-						.HAlign(HAlign_Right)
-						.ToolTipText(TooltipText)
-						.VAlign(VAlign_Center)
-						[
-							SNew(SImage)
-							.Image(FAppStyle::Get().GetBrush("Icons.PlusCircle"))
-							.ColorAndOpacity(FSlateColor::UseForeground())
-						]
-					]
-				);
-			}
-
 			PinViewerCategoryBuilder.AddCustomRow(LOCTEXT("PinViewerDetailsCategory", "PinViewer")).ShouldAutoExpand(true)
 			[
 				SNew(SPinViewer).Node(Node)
@@ -75,22 +45,7 @@ void FCustomizableObjectNodeDetails::CustomizeDetails(IDetailLayoutBuilder& Deta
 void FCustomizableObjectNodeDetails::CustomizeDetails(const TSharedPtr<IDetailLayoutBuilder>& DetailBuilder)
 {
 	IDetailCustomization::CustomizeDetails(DetailBuilder);
-
-	DetailBuilderPtr = DetailBuilder;
 }
 
 
-FReply FCustomizableObjectNodeDetails::AddNewPin(UCustomizableObjectNode* Node)
-{
-	if (Node && DetailBuilderPtr)
-	{
-		Node->CreatePinFromPinViewer();
-		DetailBuilderPtr->ForceRefreshDetails();
-
-		return FReply::Handled();
-	}
-
-	return FReply::Unhandled();
-}
-
 #undef LOCTEXT_NAMESPACE
diff --git a/Engine/Plugins/Mutable/Source/CustomizableObjectEditor/Private/MuCOE/Nodes/CustomizableObjectNodeDetails.h b/Engine/Plugins/Mutable/Source/CustomizableObjectEditor/Private/MuCOE/Nodes/CustomizableObjectNodeDetails.h
--- a/Engine/Plugins/Mutable/Source/CustomizableObjectEditor/Private/MuCOE/Nodes/CustomizableObjectNodeDetails.h
+++ b/Engine/Plugins/Mutable/Source/CustomizableObjectEditor/Private/MuCOE/Nodes/CustomizableObjectNodeDetails.h
@@ -4,9 +4,7 @@
 
 #include "IDetailCustomization.h"
 
-class FReply;
 class IDetailLayoutBuilder;
-class UCustomizableObjectNode;
 
 
 /** Base of all UCustomizableObjectNode details. */
@@ -19,12 +17,4 @@ public:
 	// IDetailCustomization interface
 	virtual void CustomizeDetails(IDetailLayoutBuilder& DetailBuilder) override;
 	virtual void CustomizeDetails(const TSharedPtr<IDetailLayoutBuilder>& DetailBuilder) override;
-
-	//
-	FReply AddNewPin(UCustomizableObjectNode* Node);
-
-private: 
-
-	// Pointer to the DetailBuilder. Needed to refresh the details when we create a new pin from the pin viewer.
-	TSharedPtr<IDetailLayoutBuilder> DetailBuilderPtr;
 };
diff --git a/Engine/Plugins/Mutable/Source/CustomizableObjectEditor/Private/MuCOE/SMutableCodeViewer.cpp b/Engine/Plugins/Mutable/Source/CustomizableObjectEditor/Private/MuCOE/SMutableCodeViewer.cpp
--- a/Engine/Plugins/Mutable/Source/CustomizableObjectEditor/Private/MuCOE/SMutableCodeViewer.cpp
+++ b/Engine/Plugins/Mutable/Source/CustomizableObjectEditor/Private/MuCOE/SMutableCodeViewer.cpp
@@ -1777,17 +1777,6 @@ void SMutableCodeViewer::GenerateElementRecursive(const int32& InStateIndex, mu:
 		break;
 	}
 
-	case mu::EOpType::SC_MATERIAL_BREAK:
-	case mu::EOpType::CO_MATERIAL_BREAK:
-	case mu::EOpType::IM_MATERIAL_BREAK:
-	{
-		mu::OP::MaterialBreakArgs Args = InProgram.GetOpArgs<mu::OP::MaterialBreakArgs>(InParentAddress);
-		AddOpFunc(Args.Material, TEXT("Material "));
-		break;
-	}
-
-	//TODO(Max) Add MI_Parameter and MI_CONSTANT
-
 	// Add here more operation types to define how they are exposed in the tree (set a Caption)
 		
 	default:
diff --git a/Engine/Plugins/Mutable/Source/CustomizableObjectEditor/Private/MuCOE/SMutableGraphViewer.cpp b/Engine/Plugins/Mutable/Source/CustomizableObjectEditor/Private/MuCOE/SMutableGraphViewer.cpp
--- a/Engine/Plugins/Mutable/Source/CustomizableObjectEditor/Private/MuCOE/SMutableGraphViewer.cpp
+++ b/Engine/Plugins/Mutable/Source/CustomizableObjectEditor/Private/MuCOE/SMutableGraphViewer.cpp
@@ -127,7 +127,6 @@ static const char* MutableNodeNames[] =
 	"ImageVariation",
 	"ImageNormalComposite",
 	"ImageTransform",
-	"ImageFromMaterialParameter",
 
 	"Bool",
 	"BoolConstant",
@@ -213,10 +212,6 @@ static const char* MutableNodeNames[] =
 	"MaterialSwitch",
 	"MaterialVariation",
 	"MaterialParameter",
-
-	"ImageMaterialBreak",
-	"ScalarMaterialBreak",
-	"ColorMaterialBreak",
 };
 
 static_assert(UE_ARRAY_COUNT(MutableNodeNames) == SIZE_T(mu::Node::EType::Count));
diff --git a/Engine/Plugins/Mutable/Source/CustomizableObjectEditor/Public/MuCOE/Nodes/CustomizableObjectNode.h b/Engine/Plugins/Mutable/Source/CustomizableObjectEditor/Public/MuCOE/Nodes/CustomizableObjectNode.h
--- a/Engine/Plugins/Mutable/Source/CustomizableObjectEditor/Public/MuCOE/Nodes/CustomizableObjectNode.h
+++ b/Engine/Plugins/Mutable/Source/CustomizableObjectEditor/Public/MuCOE/Nodes/CustomizableObjectNode.h
@@ -274,11 +274,6 @@ public:
 	/** @return true if the details should spawn a Pin Viewer. */
 	UE_API virtual bool HasPinViewer() const;
 
-	/** @return true if the pin viewer's category should spawn a button to add new pins */
-	UE_API virtual bool CanCreatePinsFromPinViewer() const;
-
-	UE_API virtual void CreatePinFromPinViewer() {};
-
 	/** Returns pin custom details. Override if required. */
 	UE_API virtual TSharedPtr<IDetailsView> CustomizePinDetails(const UEdGraphPin& Pin) const;
 
@@ -305,9 +300,6 @@ public:
 
 	/** If a node has a static string pin, this function determines if it should construct an input text widget next to the pin or not. */
 	UE_API virtual bool CreateStaticStringPinWidget() const;
-
-	// Returns an array with all the PC_Types that can be created from the pin viewer
-	UE_API virtual TArray<FName> GetAllowedPinViewerCreationTypes() const;
 	
 	UPROPERTY()
 	FPostEditChangePropertyDelegate PostEditChangePropertyDelegate;
diff --git a/Engine/Plugins/Mutable/Source/MutableRuntime/Internal/MuR/Image.h b/Engine/Plugins/Mutable/Source/MutableRuntime/Internal/MuR/Image.h
--- a/Engine/Plugins/Mutable/Source/MutableRuntime/Internal/MuR/Image.h
+++ b/Engine/Plugins/Mutable/Source/MutableRuntime/Internal/MuR/Image.h
@@ -6,7 +6,6 @@
 
 #include "MuR/MutableMath.h"
 #include "MuR/MutableTrace.h"
-#include "MuR/Operations.h"
 #include "MuR/Ptr.h"
 #include "MuR/RefCounted.h"
 #include "MuR/Serialisation.h"
@@ -14,7 +13,6 @@
 #include "HAL/LowLevelMemTracker.h"
 #include "HAL/Platform.h"
 #include "Misc/AssertionMacros.h"
-#include "Engine/Texture.h"
 
 #include "MuR/ImageDataStorage.h"
 
@@ -128,13 +126,9 @@ namespace mu
 		uint16 RelevancyMinY = 0;
 		uint16 RelevancyMaxY = 0;
 
-		TStrongObjectPtr<UTexture> Texture;
-
 		/** Only valid if the right flags are set, this identifies a referenced image. */
 		uint32 ReferenceID = 0;
 
-		OP::ADDRESS LazyAddress = 0;
-
 		/** Pixel data for all lods. */
 		FImageDataStorage DataStorage;
 
@@ -154,7 +148,6 @@ namespace mu
 			Result->RelevancyMaxY = RelevancyMaxY;
 			Result->DataStorage = DataStorage;
 			Result->ReferenceID = ReferenceID;
-			Result->LazyAddress = LazyAddress;
 
 			return Result;
 		}
@@ -175,7 +168,6 @@ namespace mu
 			RelevancyMaxY = Other->RelevancyMaxY;
 			DataStorage = Other->DataStorage;
 			ReferenceID = Other->ReferenceID;
-			LazyAddress = Other->LazyAddress;
 		}
 
 		void CopyMove(FImage* Other)
@@ -192,8 +184,7 @@ namespace mu
 			RelevancyMinY = Other->RelevancyMinY;
 			RelevancyMaxY = Other->RelevancyMaxY;
 			ReferenceID = Other->ReferenceID;
-			LazyAddress = Other->LazyAddress;
-
+			
 			DataStorage = MoveTemp(Other->DataStorage);
 		}
 
@@ -209,8 +200,7 @@ namespace mu
 		{
 			return 
 				(DataStorage == Other.DataStorage) &&
-				(ReferenceID == Other.ReferenceID) &&
-				(LazyAddress == Other.LazyAddress);
+				(ReferenceID == Other.ReferenceID);
 		}
 
 		//-----------------------------------------------------------------------------------------
@@ -253,8 +243,6 @@ namespace mu
 		//! Get the rect inside the image bounding the non-black content of the image.
 		UE_API void GetNonBlackRect(FImageRect& OutRect) const;
 		UE_API void GetNonBlackRect_Reference(FImageRect& OutRect) const;
-
-		UE_API bool IsLazy() const;
 	};
 
 	/** This struct contains image operations that may need to allocate and free temporary images to work.
diff --git a/Engine/Plugins/Mutable/Source/MutableRuntime/Internal/MuR/ModelPrivate.h b/Engine/Plugins/Mutable/Source/MutableRuntime/Internal/MuR/ModelPrivate.h
--- a/Engine/Plugins/Mutable/Source/MutableRuntime/Internal/MuR/ModelPrivate.h
+++ b/Engine/Plugins/Mutable/Source/MutableRuntime/Internal/MuR/ModelPrivate.h
@@ -373,9 +373,6 @@ namespace mu
 		 */
 		TArray<TArray<uint16>> ParameterLists;
 
-		/** Constant Material Data */
-		TArray<TSharedPtr<const FMaterial>> ConstantMaterials;
-
 #if WITH_EDITOR
 		/** 
 		 * State of the program. True unless the streamed resources were destroyed,
@@ -410,7 +407,6 @@ namespace mu
             Arch << Parameters;
             Arch << Ranges;
             Arch << ParameterLists;
-			Arch << ConstantMaterials;
         }
 
         /** */
@@ -440,7 +436,6 @@ namespace mu
             Arch >> Parameters;
             Arch >> Ranges;
             Arch >> ParameterLists;
-            Arch >> ConstantMaterials;
         }
 
         /** Debug method that sanity-checks the program with a variety of tests. */
@@ -458,7 +453,6 @@ namespace mu
 					return ConstantImageLODsStreamed.Contains(RomIndex);
 				case uint32(ERomDataType::Mesh):
 					return ConstantMeshesStreamed.Contains(RomIndex);
-				//TODO(Max): Make materials loadable and unloadable from a rom
 				default:
 					check(false);
 					break;
@@ -502,7 +496,6 @@ namespace mu
 				}
 				break;
 			}
-			//TODO(Max): Make materials loadable and unloadable from a rom
 			default:
 				check(false);
 				break;
@@ -683,11 +676,6 @@ namespace mu
             return index;
         }
 
-		OP::ADDRESS AddConstant(TSharedPtr<const FMaterial> Material)
-		{
-			OP::ADDRESS index = OP::ADDRESS(ConstantMaterials.AddUnique(Material));
-			return index;
-		}
 
 		inline TSharedPtr<const FImage> GetImageLOD(FConstantResourceIndex Index) const
 		{
diff --git a/Engine/Plugins/Mutable/Source/MutableRuntime/Internal/MuR/Operations.h b/Engine/Plugins/Mutable/Source/MutableRuntime/Internal/MuR/Operations.h
--- a/Engine/Plugins/Mutable/Source/MutableRuntime/Internal/MuR/Operations.h
+++ b/Engine/Plugins/Mutable/Source/MutableRuntime/Internal/MuR/Operations.h
@@ -2,7 +2,7 @@
 
 #pragma once
 
-#include "MuR/ImageTypes.h"
+#include "MuR/Image.h"
 #include "MuR/Types.h"
 #include "MuR/SerialisationPrivate.h"
 #include "HAL/PlatformMath.h"
@@ -16,7 +16,6 @@
 
 namespace mu
 {
-	enum class EImageFormat : uint8;
 
     //!
     enum class EOpType : uint16
@@ -80,14 +79,6 @@ namespace mu
 		ED_SWITCH,
 		MI_SWITCH,
 
-		//! Selects a parameter value from a Material
-		SC_MATERIAL_BREAK,
-		CO_MATERIAL_BREAK,
-		IM_MATERIAL_BREAK,
-
-		//! Converts a texture of a material parameter into a texture parameter to process it at runtime.
-		IM_PARAMETER_FROM_MATERIAL,
-
         //-----------------------------------------------------------------------------------------
         // Boolean operations
         //-----------------------------------------------------------------------------------------
@@ -393,8 +384,6 @@ namespace mu
     {
         typedef ADDRESS_TYPE ADDRESS;
         typedef ADDRESS_TYPE CONSTANT_STRING_ADDRESS;
-        typedef ADDRESS_TYPE STRING_ADDRESS;
-		typedef ADDRESS_TYPE LIST_UINT32_ADDRESS;
 
         /** Arguments for every operation type. */
         struct BoolConstantArgs
@@ -442,16 +431,6 @@ namespace mu
 			ADDRESS variable;
 		};
 
-		struct MaterialParameterArgs
-		{
-			ADDRESS Variable;
-
-			LIST_UINT32_ADDRESS ScalarParameterNames;
-			LIST_UINT32_ADDRESS ColorParameterNames;
-			LIST_UINT32_ADDRESS ImageParameterNames;
-			LIST_UINT32_ADDRESS ImageParameterAddress;
-		};
-
 		struct MeshParameterArgs
 		{
 			ADDRESS variable;
@@ -1106,18 +1085,6 @@ namespace mu
 			ADDRESS Mesh;
 			uint8 LayoutIndex;
 		};
-
-		struct MaterialBreakArgs
-		{
-			ADDRESS Material;
-			STRING_ADDRESS ParameterName;
-		};
-
-		struct MaterialBreakImageParameterArgs
-		{
-			ADDRESS MaterialParameter;
-			STRING_ADDRESS ParameterName;
-		};
     };
 
 
diff --git a/Engine/Plugins/Mutable/Source/MutableRuntime/Private/MuR/CodeRunner.cpp b/Engine/Plugins/Mutable/Source/MutableRuntime/Private/MuR/CodeRunner.cpp
--- a/Engine/Plugins/Mutable/Source/MutableRuntime/Private/MuR/CodeRunner.cpp
+++ b/Engine/Plugins/Mutable/Source/MutableRuntime/Private/MuR/CodeRunner.cpp
@@ -158,7 +158,6 @@ namespace mu
 				InSettings, InSystem, InExecutionStrategy, InModel, InParams, At, InLODMask, ExecutionOptions, InImageLOD, Type, ExternalResourceProvider);
 	}
 
-
     CodeRunner::CodeRunner(FPrivateToken PrivateToken, 
 		const FSettings& InSettings,
 		FSystem::Private* InSystem,
@@ -1179,10 +1178,11 @@ namespace mu
 		case EOpType::MI_CONSTANT:
 		{
 			OP::ResourceConstantArgs Args = Program.GetOpArgs<OP::ResourceConstantArgs>(Item.At);
-			TSharedPtr<const FMaterial> SourceConst = Program.ConstantMaterials[Args.value];
-			check(SourceConst);
 
-			StoreMaterial(Item, SourceConst);
+			TSharedPtr<FMaterial> Source = MakeShared<FMaterial>();
+			Source->ReferenceID = Args.value;
+
+			StoreMaterial(Item, Source);
 			break;
 		}
 
@@ -5424,73 +5424,6 @@ namespace mu
 			break;
 		}
 
-		case EOpType::IM_MATERIAL_BREAK:
-		{
-			const OP::MaterialBreakArgs Args = Program.GetOpArgs<OP::MaterialBreakArgs>(Item.At);
-
-			switch (Item.Stage)
-			{
-			case 0:
-			{
-				AddOp(FScheduledOp(Item.At, Item, 1),
-					FScheduledOp(Args.Material, Item));
-				break;
-			}
-			case 1:
-			{
-				TSharedPtr<const FMaterial> Material = LoadMaterial(FCacheAddress(Args.Material, Item));
-
-				if (!Material)
-				{
-					StoreImage(Item, nullptr);
-					break;
-				}
-
-				check(Args.ParameterName < (uint32)Program.ConstantStrings.Num());
-				const FString& ParameterName = Program.ConstantStrings[Args.ParameterName];
-				const TSharedPtr<const FImage>* Image = Material->ImageParameters.Find(FName(ParameterName));
-				
-				if (!Image)
-				{
-					StoreImage(Item, nullptr);
-				}
-				else
-				{
-					if (Image->Get()->IsLazy())
-					{
-						FScheduledOp ItemNextStage(Item.At, Item, 2);
-						ItemNextStage.CustomState = Image->Get()->LazyAddress;
-
-						AddOp(ItemNextStage,
-							FScheduledOp(Image->Get()->LazyAddress, Item));
-					}
-					else
-					{
-						StoreImage(Item, *Image);
-					}
-				}
-				break;
-			}
-			case 2:
-			{
-				TSharedPtr<const FMaterial> Material = LoadMaterial(FCacheAddress(Args.Material, Item));
-
-				check(Args.ParameterName < (uint32)Program.ConstantStrings.Num());
-				const FString& ParameterName = Program.ConstantStrings[Args.ParameterName];
-				const TSharedPtr<const FImage>* Image = Material->ImageParameters.Find(FName(ParameterName));
-
-				TSharedPtr<const FImage> pImage = LoadImage(FCacheAddress(Image->Get()->LazyAddress, Item));
-				StoreImage(Item, pImage);
-				break;
-			}
-
-			default:
-				unimplemented();
-			}
-
-			break;
-		}
-
         default:
             if (type!=EOpType::NONE)
             {
@@ -5749,7 +5682,7 @@ namespace mu
             break;
         }
     }
-	
+
     //---------------------------------------------------------------------------------------------
     void CodeRunner::RunCode_Int(const FScheduledOp& Item, const FParameters* pParams, const FModel* InModel )
     {
@@ -5910,44 +5843,6 @@ namespace mu
             break;
         }
 
-		case EOpType::SC_MATERIAL_BREAK:
-		{
-			const OP::MaterialBreakArgs Args = Program.GetOpArgs<OP::MaterialBreakArgs>(Item.At);
-
-			switch (Item.Stage)
-			{
-			case 0:
-			{
-				AddOp(FScheduledOp(Item.At, Item, 1),
-					FScheduledOp(Args.Material, Item));
-				break;
-			}
-			case 1:
-			{
-				TSharedPtr<const FMaterial> Material = LoadMaterial(FCacheAddress(Args.Material, Item));
-
-				if (!Material)
-				{
-					StoreScalar(Item, 0.0f);
-					break;
-				}
-
-				check(Args.ParameterName < (uint32)Program.ConstantStrings.Num());
-				const FString& ParameterName = Program.ConstantStrings[Args.ParameterName];
-				const float Value = *Material->ScalarParameters.Find(FName(ParameterName));
-
-				StoreScalar(Item, Value);
-
-				break;
-			}
-
-			default:
-				unimplemented();
-			}
-
-			break;
-		}
-
         default:
             check( false );
             break;
@@ -6231,44 +6126,6 @@ namespace mu
 			break;
 		}
 
-		case EOpType::CO_MATERIAL_BREAK:
-		{
-			const OP::MaterialBreakArgs Args = Program.GetOpArgs<OP::MaterialBreakArgs>(Item.At);
-
-			switch (Item.Stage)
-			{
-			case 0:
-			{
-				AddOp(FScheduledOp(Item.At, Item, 1),
-					FScheduledOp(Args.Material, Item));
-				break;
-			}
-			case 1:
-			{
-				TSharedPtr<const FMaterial> Material = LoadMaterial(FCacheAddress(Args.Material, Item));
-
-				if (!Material)
-				{
-					StoreColor(Item, FVector4f::Zero());
-					break;
-				}
-
-				check(Args.ParameterName < (uint32)Program.ConstantStrings.Num());
-				const FString& ParameterName = Program.ConstantStrings[Args.ParameterName];
-				const FVector4f Value = *Material->ColorParameters.Find(FName(ParameterName));
-
-				StoreColor(Item, Value);
-
-				break;
-			}
-
-			default:
-				unimplemented();
-			}
-			
-			break;
-		}
-
         default:
             check( false );
             break;
@@ -6361,61 +6218,13 @@ namespace mu
 			MUTABLE_CPUPROFILER_SCOPE(RunCodeImageDesc_MI_PARAMETER);
 			check(Item.Stage == 0);
 
-			OP::MaterialParameterArgs Args = Program.GetOpArgs<OP::MaterialParameterArgs >(Item.At);
-			TSharedPtr<FRangeIndex> Index = BuildCurrentOpRangeIndex(Item, Args.Variable);
+			OP::ParameterArgs Args = Program.GetOpArgs<OP::ParameterArgs>(Item.At);
+			TSharedPtr<FRangeIndex> Index = BuildCurrentOpRangeIndex(Item, Args.variable);
 
 			TSharedPtr<FMaterial> Material = MakeShared<FMaterial>();
-			Material->Material = TStrongObjectPtr(InParams->GetMaterialValue(Args.Variable, Index.Get()));
-
-			if (Material->Material)
-			{
-				// Images
-				const TArray<uint32>& ImageParameterNames = Program.ConstantUInt32Lists[Args.ImageParameterNames];
-				const TArray<uint32>& ImageParameterOperations = Program.ConstantUInt32Lists[Args.ImageParameterAddress];
-
-				for (int32 ParameterIndex = 0; ParameterIndex < ImageParameterNames.Num(); ++ParameterIndex)
-				{
-					const FString& ParameterName = Program.ConstantStrings[ImageParameterNames[ParameterIndex]];
-
-					UTexture* Texture;
-					bool bParameterFound = Material->Material->GetTextureParameterValue(FName(ParameterName), Texture);
-
-					if (bParameterFound)
-					{
-						TSharedPtr<FImage> NewParameterImage = MakeShared<FImage>();
-						NewParameterImage->Texture = TStrongObjectPtr(Texture);
-						NewParameterImage->LazyAddress = ImageParameterOperations[ParameterIndex];
-
-						Material->ImageParameters.Add(FName(ParameterName), NewParameterImage);
-					}
-				}
-
-				// Scalars
-				const TArray<uint32>& ScalarParameterNames = Program.ConstantUInt32Lists[Args.ScalarParameterNames];
-
-				for (int32 ParameterIndex = 0; ParameterIndex < ScalarParameterNames.Num(); ++ParameterIndex)
-				{
-					const FString& ParameterName = Program.ConstantStrings[ScalarParameterNames[ParameterIndex]];
-
-					float ScalarValue;
-					Material->Material->GetScalarParameterValue(FName(ParameterName), ScalarValue);
-					Material->ScalarParameters.Add(FName(ParameterName), ScalarValue);
-				}
-
-				// Vectors
-				const TArray<uint32>& ColorParameterNames = Program.ConstantUInt32Lists[Args.ColorParameterNames];
-
-				for (int32 ParameterIndex = 0; ParameterIndex < ColorParameterNames.Num(); ++ParameterIndex)
-				{
-					const FString& ParameterName = Program.ConstantStrings[ColorParameterNames[ParameterIndex]];
-
-					FLinearColor VectorValue;
-					Material->Material->GetVectorParameterValue(FName(ParameterName), VectorValue);
-					Material->ColorParameters.Add(FName(ParameterName), VectorValue);
-				}
-
-				StoreMaterial(Item, Material);
-			}
+			Material->Material = TStrongObjectPtr(InParams->GetMaterialValue(Args.variable, Index.Get()));
+			
+			StoreMaterial(Item, Material);
 			break;
 		}
 		default: unimplemented();
@@ -7774,124 +7583,6 @@ namespace mu
 			break;
 		}
 
-		case EOpType::IM_MATERIAL_BREAK:
-		{
-			MUTABLE_CPUPROFILER_SCOPE(RunCodeImageDesc_IM_MATERIAL_BREAK);
-			const OP::MaterialBreakArgs Args = Program.GetOpArgs<OP::MaterialBreakArgs>(Item.At);
-
-			switch (Item.Stage)
-			{
-			case 0:
-			{
-				FScheduledOp MaterialOp(Args.Material, Item);
-				MaterialOp.Type = FScheduledOp::EType::Full;
-
-				AddOp(FScheduledOp(Item.At, Item, 1),
-					MaterialOp);
-				break;
-			}
-			case 1:
-			{
-				FCacheAddress CacheMaterialOp(Args.Material, Item);
-				CacheMaterialOp.Type = FScheduledOp::EType::Full;
-
-				TSharedPtr<const FMaterial> Material = LoadMaterial(CacheMaterialOp);
-
-				if (!Material)
-				{
-					ImageDescResults.FindOrAdd(Item.At);
-					StoreValidDesc(Item);
-				}
-				else
-				{
-					check(Args.ParameterName < (uint32)Program.ConstantStrings.Num());
-					const FString& ParameterName = Program.ConstantStrings[Args.ParameterName];
-					const TSharedPtr<const FImage>* Image = Material->ImageParameters.Find(FName(ParameterName));
-
-					if (!Image)
-					{
-						ImageDescResults.FindOrAdd(Item.At);
-						StoreValidDesc(Item);
-					}
-					else
-					{
-						if (Image->Get()->IsLazy())
-						{
-							FScheduledOp ItemNextStage(Item.At, Item, 2);
-							ItemNextStage.CustomState = Image->Get()->LazyAddress;
-
-							AddOp(ItemNextStage,
-								FScheduledOp(Image->Get()->LazyAddress, Item));
-						}
-					}
-				}
-
-				break;
-			}
-			case 2:
-			{
-				FCacheAddress CacheMaterialOp(Args.Material, Item);
-				CacheMaterialOp.Type = FScheduledOp::EType::Full;
-
-				TSharedPtr<const FMaterial> Material = LoadMaterial(CacheMaterialOp);
-
-				if (!Material)
-				{
-					ImageDescResults.FindOrAdd(Item.At);
-				}
-				else
-				{
-					check(Args.ParameterName < (uint32)Program.ConstantStrings.Num());
-					const FString& ParameterName = Program.ConstantStrings[Args.ParameterName];
-					const TSharedPtr<const FImage>* Image = Material->ImageParameters.Find(FName(ParameterName));
-
-					FExtendedImageDesc& Result = ImageDescResults.FindOrAdd(Item.At);
-					Result = ImageDescResults.FindChecked(Image->Get()->LazyAddress);
-				}
-
-				StoreValidDesc(Item);
-				break;
-			}
-
-			default:
-				check(false);
-			}
-
-			break;
-		}
-
-		case EOpType::IM_PARAMETER_FROM_MATERIAL:
-		{
-			check(Item.Stage == 0);
-
-			const OP::MaterialBreakImageParameterArgs Args = Program.GetOpArgs<OP::MaterialBreakImageParameterArgs>(Item.At);
-
-			TSharedPtr<FRangeIndex> Index = BuildCurrentOpRangeIndex(Item, Args.MaterialParameter);
-
-			// Get material parameter from the array of parameters
-			TSharedPtr<FMaterial> Material = MakeShared<FMaterial>();
-			Material->Material = TStrongObjectPtr(pParams->GetMaterialValue(Args.MaterialParameter, Index.Get()));
-
-			// Get the texture parameter name
-			check(Args.ParameterName < (uint32)Program.ConstantStrings.Num());
-			const FString& ParameterName = Program.ConstantStrings[Args.ParameterName];
-
-			UTexture* Image = nullptr;
-
-			// Get the parameter texture from the UMaterial
-			if (Material->Material)
-			{
-				Material->Material->GetTextureParameterValue(FName(ParameterName), Image);
-			}
-
-			FExtendedImageDesc& Result = ImageDescResults.FindOrAdd(Item.At);
-			Result = GetExternalImageDesc(Image);
-
-			StoreValidDesc(Item);
-
-			break;
-		}
-
 		default:
 			if (Type != EOpType::NONE)
 			{
diff --git a/Engine/Plugins/Mutable/Source/MutableRuntime/Private/MuR/CodeRunnerBegin.cpp b/Engine/Plugins/Mutable/Source/MutableRuntime/Private/MuR/CodeRunnerBegin.cpp
--- a/Engine/Plugins/Mutable/Source/MutableRuntime/Private/MuR/CodeRunnerBegin.cpp
+++ b/Engine/Plugins/Mutable/Source/MutableRuntime/Private/MuR/CodeRunnerBegin.cpp
@@ -384,7 +384,6 @@ namespace mu
 			case EOpType::ME_PARAMETER:
 			case EOpType::IM_PARAMETER:
 			case EOpType::MI_PARAMETER:
-			case EOpType::IM_PARAMETER_FROM_MATERIAL:
 				{
 					StoreNone(Item);
 					break;
@@ -2332,19 +2331,7 @@ namespace mu
 			case EOpType::MI_CONSTANT:
 				{
 					OP::ResourceConstantArgs Args = Program.GetOpArgs<OP::ResourceConstantArgs>(Item.At);
-					TSharedPtr<const FMaterial> SourceConst = Program.ConstantMaterials[Args.value];
-					MI_REFERENCE_ID.Add(SourceConst->ReferenceID);
-
-					StoreNone(Item);
-					break;
-				}
-
-			case EOpType::IM_MATERIAL_BREAK:
-			case EOpType::SC_MATERIAL_BREAK:
-			case EOpType::CO_MATERIAL_BREAK:
-				{
-					OP::MaterialBreakArgs Args = Program.GetOpArgs<OP::MaterialBreakArgs>(Item.At);
-					PushOp(FScheduledOpInline(Args.Material, Item));
+					MI_REFERENCE_ID.Push(Args.value);
 
 					StoreNone(Item);
 					break;
diff --git a/Engine/Plugins/Mutable/Source/MutableRuntime/Private/MuR/CodeRunnerTasked.cpp b/Engine/Plugins/Mutable/Source/MutableRuntime/Private/MuR/CodeRunnerTasked.cpp
--- a/Engine/Plugins/Mutable/Source/MutableRuntime/Private/MuR/CodeRunnerTasked.cpp
+++ b/Engine/Plugins/Mutable/Source/MutableRuntime/Private/MuR/CodeRunnerTasked.cpp
@@ -2605,41 +2605,6 @@ namespace mu
 			break;
 		}
 
-		case EOpType::IM_PARAMETER_FROM_MATERIAL:
-		{
-			const OP::MaterialBreakImageParameterArgs Args = Program.GetOpArgs<OP::MaterialBreakImageParameterArgs>(Item.At);
-			TSharedPtr<FRangeIndex> Index = BuildCurrentOpRangeIndex(Item, Args.MaterialParameter);
-
-			// Get material parameter from the array of parameters
-			TSharedPtr<FMaterial> Material = MakeShared<FMaterial>();
-			Material->Material = TStrongObjectPtr(Params->GetMaterialValue(Args.MaterialParameter, Index.Get()));
-
-			// Get the texture parameter name
-			check(Args.ParameterName < (uint32)Program.ConstantStrings.Num());
-			const FString& ParameterName = Program.ConstantStrings[Args.ParameterName];
-
-			// Get the parameter texture from the UMaterial
-			if (Material->Material)
-			{
-				UTexture* Texture;
-				bool bParameterFound = Material->Material->GetTextureParameterValue(FName(ParameterName), Texture);
-
-				if (bParameterFound && Texture)
-				{
-					check(ImageLOD < TNumericLimits<uint8>::Max() && ImageLOD >= 0);
-					check(ImageLOD + static_cast<int32>(Item.ExecutionOptions) < TNumericLimits<uint8>::Max());
-
-					const uint8 MipmapsToSkip = Item.ExecutionOptions + static_cast<uint8>(ImageLOD);
-
-					CodeRunner::FExternalResourceId FullId;
-					FullId.ImageParameter = Texture;
-					Issued = MakeShared<FImageExternalLoadTask>(Item, MipmapsToSkip, FullId);
-				}
-			}
-
-			break;
-		}
-
 		case EOpType::IM_REFERENCE:
 		{
 			OP::ResourceReferenceArgs Args = Model->GetPrivate()->Program.GetOpArgs<OP::ResourceReferenceArgs>(Item.At);
diff --git a/Engine/Plugins/Mutable/Source/MutableRuntime/Private/MuR/Image.cpp b/Engine/Plugins/Mutable/Source/MutableRuntime/Private/MuR/Image.cpp
--- a/Engine/Plugins/Mutable/Source/MutableRuntime/Private/MuR/Image.cpp
+++ b/Engine/Plugins/Mutable/Source/MutableRuntime/Private/MuR/Image.cpp
@@ -8,6 +8,7 @@
 #include "MuR/ImagePrivate.h"
 #include "MuR/MutableMath.h"
 #include "MuR/MutableTrace.h"
+
 #include "ImageCore.h"
 
 namespace mu
@@ -72,14 +73,12 @@ namespace mu
 		// Remove non-persistent flags.
 		uint8 flags = Flags & ~IF_HAS_RELEVANCY_MAP;
 		Arch << flags;
-		Arch << LazyAddress;
 	}
 
 	void FImage::Unserialise(FInputArchive& Arch)
 	{
 		Arch >> DataStorage;
 		Arch >> Flags;
-		Arch >> LazyAddress;
 	}
 
 
@@ -619,12 +618,6 @@ namespace mu
 	 }
 
 
-	bool FImage::IsLazy() const
-	{
-		return LazyAddress != 0;
-	}
-
-
     //---------------------------------------------------------------------------------------------
     void FImage::GetNonBlackRect(FImageRect& rect) const
     {            
diff --git a/Engine/Plugins/Mutable/Source/MutableRuntime/Private/MuR/Material.cpp b/Engine/Plugins/Mutable/Source/MutableRuntime/Private/MuR/Material.cpp
--- a/Engine/Plugins/Mutable/Source/MutableRuntime/Private/MuR/Material.cpp
+++ b/Engine/Plugins/Mutable/Source/MutableRuntime/Private/MuR/Material.cpp
@@ -3,68 +3,26 @@
 
 #include "MuR/Material.h"
 
+#include "Containers/UnrealString.h"
 #include "HAL/LowLevelMemTracker.h"
+#include "HAL/PlatformCrt.h"
+#include "HAL/UnrealMemory.h"
 #include "MuR/MutableTrace.h"
-#include "MuR/Serialisation.h"
-#include "MuR/SerialisationPrivate.h"
-#include "MuR/Image.h"
-
 
 namespace mu
 {
-	void FMaterial::Serialise(const FMaterial* MaterialPtr, FOutputArchive& Arch)
-	{
-		Arch << *MaterialPtr;
-	}
-
-
-	TSharedPtr<FMaterial> FMaterial::StaticUnserialise(FInputArchive& Arch)
-	{
-		MUTABLE_CPUPROFILER_SCOPE(MaterialUnserialise)
-		LLM_SCOPE_BYNAME(TEXT("MutableRuntime"));
-
-		TSharedPtr<FMaterial> Result = MakeShared<FMaterial>();
-		Arch >> *Result;
-
-		return Result;
-	}
-
-
-	void FMaterial::Serialise(FOutputArchive& Arch) const
-	{
-		Arch << ReferenceID;
-		Arch << ImageParameters;
-		Arch << ColorParameters;
-		Arch << ScalarParameters;
-	}
-
-
-	void FMaterial::Unserialise(FInputArchive& Arch)
-	{
-		Arch >>ReferenceID;
-		Arch >> ImageParameters;
-		Arch >> ColorParameters;
-		Arch >> ScalarParameters;
-	}
-
-
 	TSharedPtr<FMaterial> FMaterial::Clone() const
 	{
 		TSharedPtr<FMaterial> Result = MakeShared<FMaterial>();
 
 		Result->Material = Material;
 		Result->ReferenceID = ReferenceID;
-		Result->ImageParameters = ImageParameters;
-		Result->ColorParameters = ColorParameters;
-		Result->ScalarParameters = ScalarParameters;
 
 		return Result;
 	}
 
-
 	bool FMaterial::operator==(const FMaterial& Other) const
 	{
-		return ReferenceID == Other.ReferenceID && Material == Other.Material	&&
-			ImageParameters == Other.ImageParameters &&	ColorParameters == Other.ColorParameters && ScalarParameters == Other.ScalarParameters;
+		return ReferenceID == Other.ReferenceID && Material == Other.Material;
 	};
 }
diff --git a/Engine/Plugins/Mutable/Source/MutableRuntime/Private/MuR/Operations.cpp b/Engine/Plugins/Mutable/Source/MutableRuntime/Private/MuR/Operations.cpp
--- a/Engine/Plugins/Mutable/Source/MutableRuntime/Private/MuR/Operations.cpp
+++ b/Engine/Plugins/Mutable/Source/MutableRuntime/Private/MuR/Operations.cpp
@@ -3,7 +3,6 @@
 
 #include "MuR/Operations.h"
 
-#include "MuR/Image.h"
 #include "MuR/ModelPrivate.h"
 
 
@@ -65,11 +64,6 @@ namespace mu
 		{ EDataType::ExtensionData,	},	// ED_SWITCH
 		{ EDataType::Material,		},	// MI_SWITCH
 
-		{ EDataType::Scalar,		},	// SC_MATERIAL_BREAK
-		{ EDataType::Color,			},	// CO_MATERIAL_BREAK
-		{ EDataType::Image,			},	// IM_MATERIAL_BREAK
-		{ EDataType::Image,			},	// IM_PARAMETER_FROM_MATERIAL
-
 		{ EDataType::Bool,			},	// BO_EQUAL_SC_CONST
 		{ EDataType::Bool,			},	// BO_AND
 		{ EDataType::Bool,			},	// BO_OR
@@ -193,7 +187,6 @@ namespace mu
 		case EOpType::MI_PARAMETER:
 		case EOpType::IM_REFERENCE:
 		case EOpType::ME_REFERENCE:
-		case EOpType::IM_PARAMETER_FROM_MATERIAL:
 			break;
 
         case EOpType::SC_CURVE:
@@ -285,16 +278,6 @@ namespace mu
             break;
         }
 
-		case EOpType::SC_MATERIAL_BREAK:
-		case EOpType::CO_MATERIAL_BREAK:
-		case EOpType::IM_MATERIAL_BREAK:
-		{
-			OP::MaterialBreakArgs Args = program.GetOpArgs<OP::MaterialBreakArgs>(at);
-			f(Args.Material);
-			f(Args.ParameterName);
-			break;
-		}
-
         //-------------------------------------------------------------------------------------
         case EOpType::BO_EQUAL_INT_CONST:
         {
diff --git a/Engine/Plugins/Mutable/Source/MutableRuntime/Public/MuR/Material.h b/Engine/Plugins/Mutable/Source/MutableRuntime/Public/MuR/Material.h
--- a/Engine/Plugins/Mutable/Source/MutableRuntime/Public/MuR/Material.h
+++ b/Engine/Plugins/Mutable/Source/MutableRuntime/Public/MuR/Material.h
@@ -2,6 +2,7 @@
 
 #pragma once
 
+#include "HAL/PlatformMath.h"
 #include "MuR/Ptr.h"
 #include "MuR/RefCounted.h"
 #include "MuR/Types.h"
@@ -13,7 +14,6 @@
 
 namespace mu
 {
-	class FImage;
 
     /** Material type resource */
 	class FMaterial : public FResource
@@ -24,35 +24,20 @@ namespace mu
 		// Life cycle
 		//-----------------------------------------------------------------------------------------
 
-		//! Serialisation
-		static UE_API void Serialise(const FMaterial* MaterialPtr, FOutputArchive& Arch);
-		static UE_API TSharedPtr<FMaterial> StaticUnserialise(FInputArchive& Arch);
-
         //! Clone this material
 		UE_API TSharedPtr<FMaterial> Clone() const;
 
 		// Resource interface
-		virtual int32 GetDataSize() const override { return 0; };
+		int32 GetDataSize() const override { return 0; };
 
 		//-----------------------------------------------------------------------------------------
 		// Own interface
 		//-----------------------------------------------------------------------------------------
 		UE_API bool operator==(const FMaterial& Other) const;
 
-
-		UE_API void Serialise(FOutputArchive&) const;
-		UE_API void Unserialise(FInputArchive&);
-
-
-	public:
-
 		TStrongObjectPtr<UMaterialInterface> Material;
 
 		int32 ReferenceID = INDEX_NONE;
-
-		TMap<FName, TSharedPtr<const FImage>> ImageParameters;
-		TMap<FName, FVector4f> ColorParameters;
-		TMap<FName, float> ScalarParameters;
 	};
 }
 
diff --git a/Engine/Plugins/Mutable/Source/MutableTools/Private/MuT/ASTOpConstantResource.cpp b/Engine/Plugins/Mutable/Source/MutableTools/Private/MuT/ASTOpConstantResource.cpp
--- a/Engine/Plugins/Mutable/Source/MutableTools/Private/MuT/ASTOpConstantResource.cpp
+++ b/Engine/Plugins/Mutable/Source/MutableTools/Private/MuT/ASTOpConstantResource.cpp
@@ -232,13 +232,10 @@ namespace mu
 
 
 	//-------------------------------------------------------------------------------------------------
-	void ASTOpConstantResource::ForEachChild(const TFunctionRef<void(ASTChild&)> f)
+	//-------------------------------------------------------------------------------------------------
+	//-------------------------------------------------------------------------------------------------
+	void ASTOpConstantResource::ForEachChild(const TFunctionRef<void(ASTChild&)>)
 	{
-		/** Image Parameters are evaluated at runtime. */
-		for (TPair<FName, ASTChild >& Element : ImageOperations)
-		{
-			f(Element.Value);
-		}
 	}
 
 
@@ -250,28 +247,21 @@ namespace mu
 			const ASTOpConstantResource* Other = static_cast<const ASTOpConstantResource*>(&OtherUntyped);
 			return Type == Other->Type && ValueHash == Other->ValueHash &&
 				LoadedValue == Other->LoadedValue && Proxy == Other->Proxy
-				&& SourceDataDescriptor == Other->SourceDataDescriptor
-				&& ImageOperations == Other->ImageOperations;
+				&& SourceDataDescriptor == Other->SourceDataDescriptor;
 		}
 		return false;
 	}
 
 
 	//-------------------------------------------------------------------------------------------------
-	mu::Ptr<ASTOp> ASTOpConstantResource::Clone(MapChildFuncRef MapChild) const
+	mu::Ptr<ASTOp> ASTOpConstantResource::Clone(MapChildFuncRef) const
 	{
 		Ptr<ASTOpConstantResource> n = new ASTOpConstantResource();
 		n->Type = Type;
 		n->Proxy = Proxy;
 		n->LoadedValue = LoadedValue;
 		n->ValueHash = ValueHash;
-		n->SourceDataDescriptor = SourceDataDescriptor;
-
-		for (const TPair<FName, ASTChild>& ImageOperation : ImageOperations)
-		{
-			n->ImageOperations.Add(ImageOperation.Key, ASTChild(n, MapChild(ImageOperation.Value.child())));
-		}
-		
+		n->SourceDataDescriptor = SourceDataDescriptor; 
 		return n;
 	}
 
@@ -674,24 +664,6 @@ namespace mu
 					Args.value = Program.AddConstant(pTyped);
 					break;
 				}
-				case EOpType::MI_CONSTANT:
-				{
-					TSharedPtr<FMaterial> MaterialValue = StaticCastSharedPtr<const FMaterial>(GetValue())->Clone();
-					check(MaterialValue);
-
-					// Store linked child images into the material constant
-					for (const TPair<FName, ASTChild>& ImageOperation : ImageOperations)
-					{
-						TSharedPtr<FImage> NewImage = MakeShared<FImage>();
-						NewImage->LazyAddress = ImageOperation.Value->linkedAddress;
-
-						MaterialValue->ImageParameters.Add(FName(ImageOperation.Key), NewImage);
-					}
-
-					Args.value = Program.AddConstant(MaterialValue);
-					
-					break;
-				}
 				default:
 					check(false);
 				}
@@ -731,13 +703,6 @@ namespace mu
 			Result.m_lods = ConstImage->GetLODCount();
 			Result.m_size = ConstImage->GetSize();
 		}
-		else if (Type == EOpType::MI_CONSTANT)
-		{
-			if (TMap<FName, ASTChild>::TConstIterator Iterator = ImageOperations.CreateConstIterator())
-			{
-				Result = Iterator.Value()->GetImageDesc();
-			}
-		}
 		else
 		{
 			check(false);
@@ -968,23 +933,6 @@ namespace mu
 			LoadedValue = Resource;
 			break;
 		}
-		
-		case EOpType::MI_CONSTANT:
-		{
-			TSharedPtr<const FMaterial> Resource = StaticCastSharedPtr<const FMaterial>(Value);
-
-			FOutputHashStream stream;
-			{
-				MUTABLE_CPUPROFILER_SCOPE(Serialize);
-				FOutputArchive arch(&stream);
-				FMaterial::Serialise(Resource.Get(), arch);
-			}
-
-			ValueHash = stream.GetHash();
-			LoadedValue = Resource;
-
-			break;
-		}
 
 		default:
 			LoadedValue = Value;
diff --git a/Engine/Plugins/Mutable/Source/MutableTools/Private/MuT/ASTOpConstantResource.h b/Engine/Plugins/Mutable/Source/MutableTools/Private/MuT/ASTOpConstantResource.h
--- a/Engine/Plugins/Mutable/Source/MutableTools/Private/MuT/ASTOpConstantResource.h
+++ b/Engine/Plugins/Mutable/Source/MutableTools/Private/MuT/ASTOpConstantResource.h
@@ -38,9 +38,6 @@ namespace mu
 		/** Source data descriptor. */
 		FSourceDataDescriptor SourceDataDescriptor;
 
-		/** Image Parameters are evaluated at runtime. */
-		TMap<FName, ASTChild> ImageOperations;
-
 	public:
 
 		~ASTOpConstantResource() override;
diff --git a/Engine/Plugins/Mutable/Source/MutableTools/Private/MuT/ASTOpMaterialBreak.cpp b/Engine/Plugins/Mutable/Source/MutableTools/Private/MuT/ASTOpMaterialBreak.cpp
deleted file mode 100644
diff --git a/Engine/Plugins/Mutable/Source/MutableTools/Private/MuT/ASTOpMaterialBreak.h b/Engine/Plugins/Mutable/Source/MutableTools/Private/MuT/ASTOpMaterialBreak.h
deleted file mode 100644
diff --git a/Engine/Plugins/Mutable/Source/MutableTools/Private/MuT/ASTOpMaterialBreakImageParameter.cpp b/Engine/Plugins/Mutable/Source/MutableTools/Private/MuT/ASTOpMaterialBreakImageParameter.cpp
deleted file mode 100644
diff --git a/Engine/Plugins/Mutable/Source/MutableTools/Private/MuT/ASTOpMaterialBreakImageParameter.h b/Engine/Plugins/Mutable/Source/MutableTools/Private/MuT/ASTOpMaterialBreakImageParameter.h
deleted file mode 100644
diff --git a/Engine/Plugins/Mutable/Source/MutableTools/Private/MuT/ASTOpParameter.cpp b/Engine/Plugins/Mutable/Source/MutableTools/Private/MuT/ASTOpParameter.cpp
--- a/Engine/Plugins/Mutable/Source/MutableTools/Private/MuT/ASTOpParameter.cpp
+++ b/Engine/Plugins/Mutable/Source/MutableTools/Private/MuT/ASTOpParameter.cpp
@@ -28,11 +28,6 @@ namespace mu
 		{
 			f(c.rangeSize);
 		}
-		
-		for (TPair<FString, ASTChild>& Element : ImageOperations)
-		{
-			f(Element.Value);
-		}
 	}
 
 
@@ -46,10 +41,7 @@ namespace mu
 				LODIndex == Other->LODIndex &&
 				SectionIndex == Other->SectionIndex &&
 				MeshID == Other->MeshID &&
-				Ranges == Other->Ranges &&
-				ColorParameterNames == Other->ColorParameterNames &&
-				ScalarParameterNames == Other->ScalarParameterNames &&
-				ImageOperations == Other->ImageOperations;
+				Ranges == Other->Ranges;
 		}
 		return false;
 	}
@@ -67,15 +59,6 @@ namespace mu
 		{
 			n->Ranges.Emplace(n.get(), MapChild(c.rangeSize.child()), c.rangeName, c.rangeUID);
 		}
-
-		n->ColorParameterNames = ColorParameterNames;
-		n->ScalarParameterNames = ScalarParameterNames;
-		
-		for (const TPair<FString, ASTChild>& ImageOperation : ImageOperations)
-		{
-			n->ImageOperations.Add(ImageOperation.Key, ASTChild(n, MapChild(ImageOperation.Value.child())));
-		}
-
 		return n;
 	}
 
@@ -93,7 +76,6 @@ namespace mu
 		hash_combine(res, LODIndex);
 		hash_combine(res, SectionIndex);
 		hash_combine(res, Parameter.Name.Len());
-
 		return res;
 	}
 
@@ -129,7 +111,7 @@ namespace mu
 				Args.Section = SectionIndex;
 				Args.MeshID = MeshID;
 
-				for (const FRangeData& d : Ranges)
+				for (const auto& d : Ranges)
 				{
 					OP::ADDRESS sizeAt = 0;
 					uint16 rangeId = 0;
@@ -139,48 +121,13 @@ namespace mu
 
 				AppendCode(Program.ByteCode, Args);
 			}
-
-			else if (Type == EOpType::MI_PARAMETER)
-			{
-				OP::MaterialParameterArgs Args;
-				FMemory::Memzero(Args);
-				Args.Variable = (OP::ADDRESS)LinkedParameterIndex;
-				
-				// Material Parameters
-				Args.ScalarParameterNames = Program.AddConstant(ScalarParameterNames);
-				Args.ColorParameterNames = Program.AddConstant(ColorParameterNames);
-
-				TArray<FString> ImageParameterNames;
-				TArray<uint32> ImageParameterOperations;
-
-				for (const TPair<FString, ASTChild>& Element : ImageOperations)
-				{
-					ImageParameterNames.Add(Element.Key);
-					ImageParameterOperations.Add(Element.Value->linkedAddress);
-				}
-
-				Args.ImageParameterNames = Program.AddConstant(ImageParameterNames);
-				Args.ImageParameterAddress = Program.AddConstant(ImageParameterOperations);
-
-				// Ranges
-				for (const FRangeData& d : Ranges)
-				{
-					OP::ADDRESS sizeAt = 0;
-					uint16 rangeId = 0;
-					LinkRange(Program, d, sizeAt, rangeId);
-					Program.Parameters[LinkedParameterIndex].Ranges.Add(rangeId);
-				}
-
-				AppendCode(Program.ByteCode, Args);
-			}
-
 			else
 			{
 				OP::ParameterArgs Args;
 				FMemory::Memzero(Args);
 				Args.variable = (OP::ADDRESS)LinkedParameterIndex;
 
-				for (const FRangeData& d : Ranges)
+				for (const auto& d : Ranges)
 				{
 					OP::ADDRESS sizeAt = 0;
 					uint16 rangeId = 0;
@@ -232,17 +179,8 @@ namespace mu
 
 	FImageDesc ASTOpParameter::GetImageDesc(bool, FGetImageDescContext*) const
 	{
-		FImageDesc Result;
-
-		if (Type == EOpType::MI_PARAMETER)
-		{
-			if (TMap<FString, ASTChild>::TConstIterator Iterator = ImageOperations.CreateConstIterator())
-			{
-				Result = Iterator.Value()->GetImageDesc();
-			}
-		}
-		
-		return Result;
+		check(Type == EOpType::IM_PARAMETER);
+		return FImageDesc();
 	}
 
 
diff --git a/Engine/Plugins/Mutable/Source/MutableTools/Private/MuT/ASTOpParameter.h b/Engine/Plugins/Mutable/Source/MutableTools/Private/MuT/ASTOpParameter.h
--- a/Engine/Plugins/Mutable/Source/MutableTools/Private/MuT/ASTOpParameter.h
+++ b/Engine/Plugins/Mutable/Source/MutableTools/Private/MuT/ASTOpParameter.h
@@ -34,13 +34,6 @@ namespace mu
 		//** Ranges adding dimensions to this parameter. */
 		TArray<FRangeData> Ranges;
 
-		/** List of all the parameters that this material can be broken into. */
-		TArray<FString> ScalarParameterNames;
-		TArray<FString> ColorParameterNames;
-
-		/** Image Parameters are evaluated at runtime. */
-		TMap<FString, ASTChild> ImageOperations;
-
 		/** Index of the parameter in the final program parameter list. 
 		* This is generated at link time, because the parameters may be reordered. 
 		*/
diff --git a/Engine/Plugins/Mutable/Source/MutableTools/Private/MuT/CodeGenerator.h b/Engine/Plugins/Mutable/Source/MutableTools/Private/MuT/CodeGenerator.h
--- a/Engine/Plugins/Mutable/Source/MutableTools/Private/MuT/CodeGenerator.h
+++ b/Engine/Plugins/Mutable/Source/MutableTools/Private/MuT/CodeGenerator.h
@@ -53,19 +53,16 @@ namespace mu
 	class NodeColourTable;
 	class NodeColourVariation;
 	class NodeColorToSRGB;
-	class NodeColourMaterialBreak;
 	class NodeImageBinarise;
 	class NodeImageColourMap;
 	class NodeImageConditional;
 	class NodeImageConstant;
 	class NodeImageFormat;
-	class NodeImageFromMaterialParameter;
 	class NodeImageInterpolate;
 	class NodeImageInvert;
 	class NodeImageLayer;
 	class NodeImageLayerColour;
 	class NodeImageLuminance;
-	class NodeImageMaterialBreak;
 	class NodeImageMipmap;
 	class NodeImageMultiLayer;
 	class NodeImageNormalComposite;
@@ -79,7 +76,6 @@ namespace mu
 	class NodeImageTransform;
 	class NodeImageVariation;
 	class NodeMaterial;
-	class NodeMaterialBreak;
 	class NodeMaterialConstant;
 	class NodeMaterialTable;
 	class NodeMaterialSwitch;
@@ -108,7 +104,6 @@ namespace mu
 	class NodeScalarSwitch;
 	class NodeScalarTable;
 	class NodeScalarVariation;
-	class NodeScalarMaterialBreak;
 	class NodeStringConstant;
 	class NodeStringParameter;
 	class NodeMatrix;
@@ -429,6 +424,28 @@ namespace mu
 		void GenerateTableSwitchNode(const NODE_TABLE& node, Ptr<NODE_SWITCH_TYPE>& OutResult, F&& GenerateOption);
 
 
+		//-----------------------------------------------------------------------------------------
+		// Materials
+
+		/** Options that affect the generation of Materials. It is like list of what required data we want while parsing down the material node graph. */
+		struct FMaterialGenerationOptions : public FGenericGenerationOptions
+		{
+			
+		};
+
+		/** */
+		struct FMaterialGenerationResult
+		{
+			Ptr<ASTOp> op;
+		};
+
+		void GenerateMaterial(const FMaterialGenerationOptions&, FMaterialGenerationResult& result, const Ptr<const NodeMaterial>& node);
+		void GenerateMaterial_Constant(const FMaterialGenerationOptions&, FMaterialGenerationResult& result, const NodeMaterialConstant* node);
+		void GenerateMaterial_Switch(const FMaterialGenerationOptions&, FMaterialGenerationResult& result, const NodeMaterialSwitch* node);
+		void GenerateMaterial_Table(const FMaterialGenerationOptions&, FMaterialGenerationResult& result, const NodeMaterialTable* node);
+		void GenerateMaterial_Variation(const FMaterialGenerationOptions&, FMaterialGenerationResult& result, const NodeMaterialVariation* node);
+		void GenerateMaterial_Parameter(const FMaterialGenerationOptions&, FMaterialGenerationResult& result, const NodeMaterialParameter* node);
+
 		//-----------------------------------------------------------------------------------------
 		// Images
 		
@@ -458,9 +475,6 @@ namespace mu
 			uint64 LayoutBlockId = FLayoutBlock::InvalidBlockId;
 			TSharedPtr<const FLayout> LayoutToApply;
 
-			/**(Optional) Parameter Descriptor of the material where this image belongs. */
-			FParameterDesc MaterialParameter;
-
 			friend FORCEINLINE uint32 GetTypeHash(const FImageGenerationOptions& InKey)
 			{
 				uint32 KeyHash = GetTypeHash((FGenericGenerationOptions&)InKey);
@@ -538,8 +552,6 @@ namespace mu
 		void GenerateImage_Variation(const FImageGenerationOptions&, FImageGenerationResult&, const NodeImageVariation*);
 		void GenerateImage_NormalComposite(const FImageGenerationOptions&, FImageGenerationResult&, const NodeImageNormalComposite*);
 		void GenerateImage_Transform(const FImageGenerationOptions&, FImageGenerationResult&, const NodeImageTransform*);
-		void GenerateImage_MaterialBreak(const FImageGenerationOptions&, FImageGenerationResult&, const NodeImageMaterialBreak*);
-		void GenerateImage_FromMaterialParameter(const FImageGenerationOptions&, FImageGenerationResult&, const NodeImageFromMaterialParameter*);
 
     	void GenerateCrop(const FImageGenerationOptions& Options, FImageGenerationResult& Result, const NodeImage&);
 
@@ -570,50 +582,6 @@ namespace mu
 		/** Generate all the operations to apply the block patching on top of the BlockOp, and masking with PatchMask. */
 		Ptr<ASTOp> GenerateImageBlockPatch(Ptr<ASTOp> BlockOp, const NodeModifierSurfaceEdit::FTexture&, TSharedPtr<FImage> PatchMask, Ptr<ASTOp> ConditionOp, const FImageGenerationOptions&);
 
-		//-----------------------------------------------------------------------------------------
-		// Materials	
-
-		/** Options that affect the generation of Materials. It is like list of what required data we want while parsing down the material node graph. */
-		struct FMaterialGenerationOptions : public FGenericGenerationOptions
-		{
-			/** The id of the component that we are currently generating. */
-			int32 ComponentId = -1;
-
-			/** The LOD being generated. */
-			int32 LODIndex = -1;
-
-			/** */
-			CompilerOptions::TextureLayoutStrategy ImageLayoutStrategy = CompilerOptions::TextureLayoutStrategy::None;
-
-			/** If different than {0,0} this is the mandatory size of the image that needs to be generated. */
-			UE::Math::TIntVector2<int32> RectSize = { 0, 0 };
-
-			/** Layout block that we are trying to generate if any. */
-			uint64 LayoutBlockId = FLayoutBlock::InvalidBlockId;
-			TSharedPtr<const FLayout> LayoutToApply;
-		};
-
-		/** */
-		struct FMaterialGenerationResult
-		{
-			Ptr<ASTOp> op;
-		};
-
-		typedef TMap<FGeneratedCacheKey, FMaterialGenerationResult> FGeneratedMaterialsMap;
-		struct FGeneratedMaterials
-		{
-			UE::FMutex Mutex;
-			FGeneratedMaterialsMap Map;
-		};
-		FGeneratedMaterials GeneratedMaterials;
-
-		void GenerateMaterial(const FMaterialGenerationOptions&, FMaterialGenerationResult& result, const Ptr<const NodeMaterial>& node);
-		void GenerateMaterial_Constant(const FMaterialGenerationOptions&, FMaterialGenerationResult& result, const NodeMaterialConstant* node);
-		void GenerateMaterial_Switch(const FMaterialGenerationOptions&, FMaterialGenerationResult& result, const NodeMaterialSwitch* node);
-		void GenerateMaterial_Table(const FMaterialGenerationOptions&, FMaterialGenerationResult& result, const NodeMaterialTable* node);
-		void GenerateMaterial_Variation(const FMaterialGenerationOptions&, FMaterialGenerationResult& result, const NodeMaterialVariation* node);
-		void GenerateMaterial_Parameter(const FMaterialGenerationOptions&, FMaterialGenerationResult& result, const NodeMaterialParameter* node);
-
         //-----------------------------------------------------------------------------------------
         // Meshes
 
@@ -885,7 +853,6 @@ namespace mu
 		void GenerateScalar_Arithmetic(FScalarGenerationResult&, const FGenericGenerationOptions&, const Ptr<const NodeScalarArithmeticOperation>&);
 		void GenerateScalar_Variation(FScalarGenerationResult&, const FGenericGenerationOptions&, const Ptr<const NodeScalarVariation>&);
 		void GenerateScalar_Table(FScalarGenerationResult&, const FGenericGenerationOptions&, const Ptr<const NodeScalarTable>&);
-		void GenerateScalar_MaterialBreak(FScalarGenerationResult&, const FGenericGenerationOptions&, const Ptr<const NodeScalarMaterialBreak>&);
 		Ptr<ASTOp> GenerateMissingScalarCode(const TCHAR* strWhere, float value, const void* errorContext);
 
 		//-----------------------------------------------------------------------------------------
@@ -913,7 +880,6 @@ namespace mu
 		void GenerateColor_Variation(FColorGenerationResult&, const FGenericGenerationOptions&, const Ptr<const NodeColourVariation>&);
 		void GenerateColor_ToSRGB(FColorGenerationResult&, const FGenericGenerationOptions&, const Ptr<const NodeColorToSRGB>&);
 		void GenerateColor_Table(FColorGenerationResult&, const FGenericGenerationOptions&, const Ptr<const NodeColourTable>&);
-		void GenerateColor_MaterialBreak(FColorGenerationResult&, const FGenericGenerationOptions&, const Ptr<const NodeColourMaterialBreak>&);
 		Ptr<ASTOp> GenerateMissingColourCode(const TCHAR* strWhere, const void* errorContext);
 
 		//-----------------------------------------------------------------------------------------
diff --git a/Engine/Plugins/Mutable/Source/MutableTools/Private/MuT/CodeGenerator_Color.cpp b/Engine/Plugins/Mutable/Source/MutableTools/Private/MuT/CodeGenerator_Color.cpp
--- a/Engine/Plugins/Mutable/Source/MutableTools/Private/MuT/CodeGenerator_Color.cpp
+++ b/Engine/Plugins/Mutable/Source/MutableTools/Private/MuT/CodeGenerator_Color.cpp
@@ -21,7 +21,6 @@
 #include "MuT/ASTOpColorSampleImage.h"
 #include "MuT/ASTOpColorFromScalars.h"
 #include "MuT/ASTOpColorToSRGB.h"
-#include "MuT/ASTOpMaterialBreak.h"
 #include "MuT/CodeGenerator.h"
 #include "MuT/CodeGenerator_FirstPass.h"
 #include "MuT/ErrorLog.h"
@@ -30,7 +29,6 @@
 #include "MuT/NodeColourArithmeticOperation.h"
 #include "MuT/NodeColourConstant.h"
 #include "MuT/NodeColourFromScalars.h"
-#include "MuT/NodeColourMaterialBreak.h"
 #include "MuT/NodeColourParameter.h"
 #include "MuT/NodeColourSampleImage.h"
 #include "MuT/NodeColourSwitch.h"
@@ -120,11 +118,6 @@ namespace mu
 			const NodeColourTable* Table = static_cast<const NodeColourTable*>(Untyped.get());
 			GenerateColor_Table(Result, Options, Table);
 		}
-		else if (Untyped->GetType() == NodeColourMaterialBreak::GetStaticType())
-		{
-			const NodeColourMaterialBreak* Table = static_cast<const NodeColourMaterialBreak*>(Untyped.get());
-			GenerateColor_MaterialBreak(Result, Options, Table);
-		}
 		else
 		{
 			check(false);
@@ -547,26 +540,6 @@ namespace mu
 	}
 
 
-	//-------------------------------------------------------------------------------------------------
-	void CodeGenerator::GenerateColor_MaterialBreak(FColorGenerationResult& Result, const FGenericGenerationOptions& Options, const Ptr<const NodeColourMaterialBreak>& Typed)
-	{
-		Ptr<ASTOpMaterialBreak> MaterialOp = new ASTOpMaterialBreak();
-
-		// Set the parameter name that this break node will generate
-		MaterialOp->ParameterName = Typed->ParameterName;
-
-		// Generate Material Source
-		FMaterialGenerationOptions MaterialSourceOptions;
-		FMaterialGenerationResult MaterialSourceResult;
-
-		GenerateMaterial(MaterialSourceOptions, MaterialSourceResult, Typed->MaterialSource);
-		MaterialOp->Material = MaterialSourceResult.op;
-		MaterialOp->Type = EOpType::CO_MATERIAL_BREAK;
-
-		Result.op = MaterialOp;
-	}
-
-
 	//---------------------------------------------------------------------------------------------
 	mu::Ptr<ASTOp> CodeGenerator::GenerateMissingColourCode(const TCHAR* StrWhere, const void* ErrorContext)
 	{
diff --git a/Engine/Plugins/Mutable/Source/MutableTools/Private/MuT/CodeGenerator_Image.cpp b/Engine/Plugins/Mutable/Source/MutableTools/Private/MuT/CodeGenerator_Image.cpp
--- a/Engine/Plugins/Mutable/Source/MutableTools/Private/MuT/CodeGenerator_Image.cpp
+++ b/Engine/Plugins/Mutable/Source/MutableTools/Private/MuT/CodeGenerator_Image.cpp
@@ -38,8 +38,6 @@
 #include "MuT/ASTOpImagePlainColor.h"
 #include "MuT/ASTOpImageDisplace.h"
 #include "MuT/ASTOpImageInvert.h"
-#include "MuT/ASTOpMaterialBreak.h"
-#include "MuT/ASTOpMaterialBreakImageParameter.h"
 #include "MuT/ASTOpMeshExtractLayoutBlocks.h"
 #include "MuT/ASTOpMeshFormat.h"
 #include "MuT/ASTOpMeshProject.h"
@@ -61,13 +59,11 @@
 #include "MuT/NodeImageConditional.h"
 #include "MuT/NodeImageConstant.h"
 #include "MuT/NodeImageFormat.h"
-#include "MuT/NodeImageFromMaterialParameter.h"
 #include "MuT/NodeImageInterpolate.h"
 #include "MuT/NodeImageInvert.h"
 #include "MuT/NodeImageLayer.h"
 #include "MuT/NodeImageLayerColour.h"
 #include "MuT/NodeImageLuminance.h"
-#include "MuT/NodeImageMaterialBreak.h"
 #include "MuT/NodeImageMipmap.h"
 #include "MuT/NodeImageMultiLayer.h"
 #include "MuT/NodeImageNormalComposite.h"
@@ -138,8 +134,6 @@ namespace mu
 		case Node::EType::ImageVariation: GenerateImage_Variation(Options, Result, static_cast<const NodeImageVariation*>(Node)); break;
 		case Node::EType::ImageNormalComposite: GenerateImage_NormalComposite(Options, Result, static_cast<const NodeImageNormalComposite*>(Node)); break;
 		case Node::EType::ImageTransform: GenerateImage_Transform(Options, Result, static_cast<const NodeImageTransform*>(Node)); break;
-		case Node::EType::ImageMaterialBreak: GenerateImage_MaterialBreak(Options, Result, static_cast<const NodeImageMaterialBreak*>(Node)); break;
-		case Node::EType::ImageFromMaterialParameter: GenerateImage_FromMaterialParameter(Options, Result, static_cast<const NodeImageFromMaterialParameter*>(Node)); break;
 		default: check(false);
 		}
 
@@ -1885,42 +1879,6 @@ namespace mu
 	}
 
 
-	void CodeGenerator::GenerateImage_MaterialBreak(const FImageGenerationOptions& InOptions, FImageGenerationResult& Result, const NodeImageMaterialBreak* Node)
-	{
-		Ptr<ASTOpMaterialBreak> MaterialOp = new ASTOpMaterialBreak();
-
-		// Set the parameter name that this break node will generate
-		MaterialOp->ParameterName = Node->ParameterName;
-
-		// Generate Material Source
-		FMaterialGenerationResult MaterialSourceResult;
-		FMaterialGenerationOptions MaterialSourceOptions;
-		// Temporary fix
-		MaterialSourceOptions.ComponentId = InOptions.ComponentId;
-		MaterialSourceOptions.LODIndex = InOptions.LODIndex;
-		MaterialSourceOptions.ImageLayoutStrategy = InOptions.ImageLayoutStrategy;
-		MaterialSourceOptions.RectSize = InOptions.RectSize;
-		MaterialSourceOptions.LayoutBlockId = InOptions.LayoutBlockId;
-		MaterialSourceOptions.LayoutToApply = InOptions.LayoutToApply;
-
-		GenerateMaterial(MaterialSourceOptions, MaterialSourceResult, Node->MaterialSource);
-		MaterialOp->Material = MaterialSourceResult.op;
-		MaterialOp->Type = EOpType::IM_MATERIAL_BREAK;
-
-		Result.op = MaterialOp;
-	}
-
-
-	void CodeGenerator::GenerateImage_FromMaterialParameter(const FImageGenerationOptions& InOptions, FImageGenerationResult& Result, const NodeImageFromMaterialParameter* InNode)
-	{
-		Ptr<ASTOpImageFromMaterialParameter> ImageOp = new ASTOpImageFromMaterialParameter();
-		ImageOp->MaterialParameter = InOptions.MaterialParameter;
-		ImageOp->ParameterName = InNode->ImageParameterName;
-
-		Result.op = ImageOp;
-	}
-
-
 	TSharedPtr<FImage> CodeGenerator::GenerateMissingImage(EImageFormat Format)
 	{
 		MUTABLE_CPUPROFILER_SCOPE(GenerateMissingImage);
diff --git a/Engine/Plugins/Mutable/Source/MutableTools/Private/MuT/CodeGenerator_Material.cpp b/Engine/Plugins/Mutable/Source/MutableTools/Private/MuT/CodeGenerator_Material.cpp
--- a/Engine/Plugins/Mutable/Source/MutableTools/Private/MuT/CodeGenerator_Material.cpp
+++ b/Engine/Plugins/Mutable/Source/MutableTools/Private/MuT/CodeGenerator_Material.cpp
@@ -10,10 +10,6 @@
 #include "MuT/AST.h"
 #include "MuT/ASTOpConditional.h"
 #include "MuT/ASTOpConstantMaterial.h"
-#include "MuT/ASTOpConstantResource.h"
-#include "MuT/ASTOpImagePixelFormat.h"
-#include "MuT/ASTOpImageResize.h"
-#include "MuT/ASTOpMaterialBreak.h"
 #include "MuT/ASTOpSwitch.h"
 #include "MuT/CodeGenerator.h"
 #include "MuT/CodeGenerator_FirstPass.h"
@@ -36,98 +32,30 @@ namespace mu
 
 	void CodeGenerator::GenerateMaterial(const FMaterialGenerationOptions& InOptions, FMaterialGenerationResult& Result , const Ptr<const NodeMaterial>& InUntypedNode)
 	{
-		// Early out
 		if (!InUntypedNode)
 		{
-			Result = FMaterialGenerationResult();
 			return;
 		}
 
-		// See if it was already generated
-		FGeneratedCacheKey Key;
-		Key.Node = InUntypedNode;
-		Key.Options = InOptions;
-		{
-			UE::TUniqueLock Lock(GeneratedMaterials.Mutex);
-			FGeneratedMaterialsMap::ValueType* Found = GeneratedMaterials.Map.Find(Key);
-			if (Found)
-			{
-				Result = *Found;
-				return;
-			}
-		}
+		const NodeMaterial* Node = InUntypedNode.get();
 
 		// Generate for each different type of node
-		const NodeMaterial* Node = InUntypedNode.get();
 		switch (Node->GetType()->Type)
 		{
-		case Node::EType::MaterialConstant:
-			GenerateMaterial_Constant(InOptions, Result, static_cast<const NodeMaterialConstant*>(Node)); 
-			break;
-		
-		case Node::EType::MaterialSwitch:
-			GenerateMaterial_Switch(InOptions, Result, static_cast<const NodeMaterialSwitch*>(Node));
-			break;
-
-		case Node::EType::MaterialVariation:
-			GenerateMaterial_Variation(InOptions, Result, static_cast<const NodeMaterialVariation*>(Node));
-			break;
-
-		case Node::EType::MaterialTable:
-			GenerateMaterial_Table(InOptions, Result, static_cast<const NodeMaterialTable*>(Node));
-			break;
-
-		case Node::EType::MaterialParameter:
-			GenerateMaterial_Parameter(InOptions, Result, static_cast<const NodeMaterialParameter*>(Node));
-			break;
-
-		default: 
-			check(false);
+		case Node::EType::MaterialConstant: GenerateMaterial_Constant(InOptions, Result, static_cast<const NodeMaterialConstant*>(Node)); break;
+		case Node::EType::MaterialSwitch: GenerateMaterial_Switch(InOptions, Result, static_cast<const NodeMaterialSwitch*>(Node)); break;
+		case Node::EType::MaterialVariation: GenerateMaterial_Variation(InOptions, Result, static_cast<const NodeMaterialVariation*>(Node)); break;
+		case Node::EType::MaterialTable: GenerateMaterial_Table(InOptions, Result, static_cast<const NodeMaterialTable*>(Node)); break;
+		case Node::EType::MaterialParameter: GenerateMaterial_Parameter(InOptions, Result, static_cast<const NodeMaterialParameter*>(Node)); break;
+		default: check(false);
 		}
 	}
 
-
 	void CodeGenerator::GenerateMaterial_Constant(const FMaterialGenerationOptions& Options, FMaterialGenerationResult& Result, const NodeMaterialConstant* Node)
 	{
-		// Constant Material Resource
-		TSharedPtr<FMaterial> MaterialResource = MakeShared<FMaterial>();
-		MaterialResource->ReferenceID = Node->MaterialId;
-		MaterialResource->ColorParameters = Node->ColorValues;
-		MaterialResource->ScalarParameters = Node->ScalarValues;
-
-		// Constant Material Operation
-		Ptr<ASTOpConstantResource> ConstantMaterial = new ASTOpConstantResource();
-		ConstantMaterial->Type = EOpType::MI_CONSTANT;
-
-		//TODO: each image should have its own generation options
-		FImageGenerationOptions ImageOptions(Options.ComponentId, Options.LODIndex);
-		ImageOptions.ImageLayoutStrategy = Options.ImageLayoutStrategy;
-		ImageOptions.RectSize = Options.RectSize;
-		ImageOptions.LayoutBlockId = Options.LayoutBlockId;
-		ImageOptions.LayoutToApply = Options.LayoutToApply;
-
-		const TArray<TPair<FName, Ptr<NodeImage>>>& ImageValuesArray = Node->ImageValues.Array();
-
-		//Store the images into the operation so that we can link them later
-		for (int32 ImageIndex = 0; ImageIndex < ImageValuesArray.Num(); ++ImageIndex)
-		{
-			FName ImageParameterName = ImageValuesArray[ImageIndex].Key;
-			Ptr<NodeImage> ImageNode = ImageValuesArray[ImageIndex].Value;
-			FImageGenerationResult ImageResult;
-
-			GenerateImage(ImageOptions, ImageResult, ImageNode);
-
-			// Store lazy branches
-			ConstantMaterial->ImageOperations.Add(ImageParameterName, ASTChild(ConstantMaterial, ImageResult.op));
-		}
-
-		// Store the constant material
-		ConstantMaterial->SetValue(MaterialResource, CompilerOptions->OptimisationOptions.DiskCacheContext);
-
-		Result.op = ConstantMaterial;
+		Result.op = new ASTOpConstantMaterial(Node->MaterialId);
 	}
 
-
 	void CodeGenerator::GenerateMaterial_Switch(const FMaterialGenerationOptions& Options, FMaterialGenerationResult& Result, const NodeMaterialSwitch* Node)
 	{
 		if (Node->Options.Num() == 0)
@@ -175,7 +103,6 @@ namespace mu
 		Result.op = SwitchOp;
 	}
 
-
 	void CodeGenerator::GenerateMaterial_Variation(const FMaterialGenerationOptions& Options, FMaterialGenerationResult& Result, const NodeMaterialVariation* Node)
 	{
 		Ptr<ASTOp> LastOp;
@@ -234,7 +161,6 @@ namespace mu
 		Result.op = LastOp;
 	}
 
-
 	void CodeGenerator::GenerateMaterial_Table(const FMaterialGenerationOptions& Options, FMaterialGenerationResult& Result, const NodeMaterialTable* Node)
 	{
 		Result.op = GenerateTableSwitch<NodeMaterialTable, ETableColumnType::Material, EOpType::MI_SWITCH>(*Node,
@@ -252,7 +178,6 @@ namespace mu
 			});
 	}
 
-
 	void CodeGenerator::GenerateMaterial_Parameter(const FMaterialGenerationOptions& InOptions, FMaterialGenerationResult& Result, const NodeMaterialParameter* InNode)
 	{
 		{
@@ -280,27 +205,6 @@ namespace mu
 			GenerateRange(RangeResult, InOptions, RangeNode);
 			MaterialOp->Ranges.Emplace(MaterialOp.get(), RangeResult.sizeOp, RangeResult.rangeName, RangeResult.rangeUID);
 		}
-		
-		// Copy Scalar and Color Names
-		MaterialOp->ColorParameterNames = InNode->ColorParameterNames;
-		MaterialOp->ScalarParameterNames = InNode->ScalarParameterNames;
-
-		// Generate Image Nodes
-		FImageGenerationResult ImageResult;
-		FImageGenerationOptions ImageOptions(InOptions.ComponentId, InOptions.LODIndex);
-		ImageOptions.ImageLayoutStrategy = InOptions.ImageLayoutStrategy;
-		ImageOptions.RectSize = InOptions.RectSize;
-		ImageOptions.LayoutBlockId = InOptions.LayoutBlockId;
-		ImageOptions.LayoutToApply = InOptions.LayoutToApply;
-		ImageOptions.MaterialParameter = MaterialOp->Parameter;
-
-		for (const TPair<FString, Ptr<NodeImage>>& ImageNode : InNode->ImageParameterNodes)
-		{
-			GenerateImage(ImageOptions, ImageResult, ImageNode.Value);
-
-			// Store lazy branches
-			MaterialOp->ImageOperations.Add(ImageNode.Key, ASTChild(MaterialOp, ImageResult.op));
-		}
 
 		Result.op = MaterialOp;
 
diff --git a/Engine/Plugins/Mutable/Source/MutableTools/Private/MuT/CodeGenerator_Scalar.cpp b/Engine/Plugins/Mutable/Source/MutableTools/Private/MuT/CodeGenerator_Scalar.cpp
--- a/Engine/Plugins/Mutable/Source/MutableTools/Private/MuT/CodeGenerator_Scalar.cpp
+++ b/Engine/Plugins/Mutable/Source/MutableTools/Private/MuT/CodeGenerator_Scalar.cpp
@@ -17,7 +17,6 @@
 #include "MuT/ASTOpScalarCurve.h"
 #include "MuT/ASTOpScalarArithmetic.h"
 #include "MuT/ASTOpConstantScalar.h"
-#include "MuT/ASTOpMaterialBreak.h"
 #include "MuT/ASTOpSwitch.h"
 #include "MuT/CodeGenerator.h"
 #include "MuT/CodeGenerator_FirstPass.h"
@@ -30,7 +29,6 @@
 #include "MuT/NodeScalarConstant.h"
 #include "MuT/NodeScalarCurve.h"
 #include "MuT/NodeScalarEnumParameter.h"
-#include "MuT/NodeScalarMaterialBreak.h"
 #include "MuT/NodeScalarParameter.h"
 #include "MuT/NodeScalarSwitch.h"
 #include "MuT/NodeScalarTable.h"
@@ -105,11 +103,6 @@ namespace mu
 			const NodeScalarTable* Table = static_cast<const NodeScalarTable*>(Untyped.get());
 			GenerateScalar_Table(result, Options, Table);
 		}
-		else if (Untyped->GetType() == NodeScalarMaterialBreak::GetStaticType())
-		{
-			const NodeScalarMaterialBreak* Table = static_cast<const NodeScalarMaterialBreak*>(Untyped.get());
-			GenerateScalar_MaterialBreak(result, Options, Table);
-		}
 		else
 		{
 			check(false);
@@ -429,25 +422,6 @@ namespace mu
 	}
 
 
-	void CodeGenerator::GenerateScalar_MaterialBreak(FScalarGenerationResult& Result, const FGenericGenerationOptions& Options, const Ptr<const NodeScalarMaterialBreak>& Node)
-	{
-		Ptr<ASTOpMaterialBreak> MaterialOp = new ASTOpMaterialBreak();
-
-		// Set the parameter name that this break node will generate
-		MaterialOp->ParameterName = Node->ParameterName;
-
-		// Generate Material Source
-		FMaterialGenerationOptions MaterialSourceOptions;
-		FMaterialGenerationResult MaterialSourceResult;
-
-		GenerateMaterial(MaterialSourceOptions, MaterialSourceResult, Node->MaterialSource);
-		MaterialOp->Material = MaterialSourceResult.op;
-		MaterialOp->Type = EOpType::SC_MATERIAL_BREAK;
-
-		Result.op = MaterialOp;
-	}
-
-
 	mu::Ptr<ASTOp> CodeGenerator::GenerateMissingScalarCode(const TCHAR* strWhere, float value, const void* errorContext)
 	{
 		// Log a warning
diff --git a/Engine/Plugins/Mutable/Source/MutableTools/Private/MuT/CodeOptimiser.cpp b/Engine/Plugins/Mutable/Source/MutableTools/Private/MuT/CodeOptimiser.cpp
--- a/Engine/Plugins/Mutable/Source/MutableTools/Private/MuT/CodeOptimiser.cpp
+++ b/Engine/Plugins/Mutable/Source/MutableTools/Private/MuT/CodeOptimiser.cpp
@@ -550,8 +550,6 @@ namespace mu
 					case EOpType::IM_PARAMETER:
 					case EOpType::ME_PARAMETER:
 					case EOpType::MA_PARAMETER:
-					case EOpType::MI_PARAMETER:
-					case EOpType::IM_PARAMETER_FROM_MATERIAL:
 						bIsConstantSubgraph = false;
 						break;
 					default:
diff --git a/Engine/Plugins/Mutable/Source/MutableTools/Private/MuT/Compiler.cpp b/Engine/Plugins/Mutable/Source/MutableTools/Private/MuT/Compiler.cpp
--- a/Engine/Plugins/Mutable/Source/MutableTools/Private/MuT/Compiler.cpp
+++ b/Engine/Plugins/Mutable/Source/MutableTools/Private/MuT/Compiler.cpp
@@ -103,11 +103,6 @@ namespace mu
 		{ false,	{ 0, 0, 0, 0, 0, 0, 0, 0, 0, 0, 0, 0, 0, 0, 0, 0, 0, 0, 0, 0, 0, 0, 0, 0, 0, 0, 0, 0, 0 }		},	// ED_SWITCH
 		{ false,	{ 0, 0, 0, 0, 0, 0, 0, 0, 0, 0, 0, 0, 0, 0, 0, 0, 0, 0, 0, 0, 0, 0, 0, 0, 0, 0, 0, 0, 0 }		},	// MI_SWITCH
 
-		{ false,	{ 0, 0, 0, 0, 0, 0, 0, 0, 0, 0, 0, 0, 0, 0, 0, 0, 0, 0, 0, 0, 0, 0, 0, 0, 0, 0, 0, 0, 0 }		},	// SC_MATERIAL_BREAK
-		{ false,	{ 0, 0, 0, 0, 0, 0, 0, 0, 0, 0, 0, 0, 0, 0, 0, 0, 0, 0, 0, 0, 0, 0, 0, 0, 0, 0, 0, 0, 0 }		},	// CO_MATERIAL_BREAK
-		{ false,	{ 0, 0, 0, 0, 0, 0, 0, 0, 0, 0, 0, 0, 0, 0, 0, 0, 0, 0, 0, 0, 0, 0, 0, 0, 0, 0, 0, 0, 0 }		},	// IM_MATERIAL_BREAK
-		{ false,	{ 0, 0, 0, 0, 0, 0, 0, 0, 0, 0, 0, 0, 0, 0, 0, 0, 0, 0, 0, 0, 0, 0, 0, 0, 0, 0, 0, 0, 0 }		},	// IM_PARAMETER_FROM_MATERIAL
-
 		{ false,	{ 0, 0, 0, 0, 0, 0, 0, 0, 0, 0, 0, 0, 0, 0, 0, 0, 0, 0, 0, 0, 0, 0, 0, 0, 0, 0, 0, 0, 0 }		},	// BO_EQUAL_SC_CONST
 		{ false,	{ 0, 0, 0, 0, 0, 0, 0, 0, 0, 0, 0, 0, 0, 0, 0, 0, 0, 0, 0, 0, 0, 0, 0, 0, 0, 0, 0, 0, 0 }		},	// BO_AND
 		{ false,	{ 0, 0, 0, 0, 0, 0, 0, 0, 0, 0, 0, 0, 0, 0, 0, 0, 0, 0, 0, 0, 0, 0, 0, 0, 0, 0, 0, 0, 0 }		},	// BO_OR
diff --git a/Engine/Plugins/Mutable/Source/MutableTools/Private/MuT/ErrorLog.cpp b/Engine/Plugins/Mutable/Source/MutableTools/Private/MuT/ErrorLog.cpp
--- a/Engine/Plugins/Mutable/Source/MutableTools/Private/MuT/ErrorLog.cpp
+++ b/Engine/Plugins/Mutable/Source/MutableTools/Private/MuT/ErrorLog.cpp
@@ -248,11 +248,6 @@ namespace mu
 		TEXT("ED_SWITCH        "),
 		TEXT("MI_SWITCH        "),
 
-		TEXT("SC_MATERIAL_BREAK"),
-		TEXT("CO_MATERIAL_BREAK"),
-		TEXT("IM_MATERIAL_BREAK"),
-		TEXT("IM_PARAMETER_FROM_MATERIAL"),
-
 		TEXT("BO_EQUAL_SC_CONST"),
 		TEXT("BO_AND           "),
 		TEXT("BO_OR            "),
diff --git a/Engine/Plugins/Mutable/Source/MutableTools/Private/MuT/Node.cpp b/Engine/Plugins/Mutable/Source/MutableTools/Private/MuT/Node.cpp
--- a/Engine/Plugins/Mutable/Source/MutableTools/Private/MuT/Node.cpp
+++ b/Engine/Plugins/Mutable/Source/MutableTools/Private/MuT/Node.cpp
@@ -25,7 +25,6 @@
 #include "MuT/NodeImageFormat.h"
 #include "MuT/NodeImageBinarise.h"
 #include "MuT/NodeImageConditional.h"
-#include "MuT/NodeImageFromMaterialParameter.h"
 #include "MuT/NodeImageInterpolate.h"
 #include "MuT/NodeImageInvert.h"
 #include "MuT/NodeImageLayer.h"
@@ -54,7 +53,6 @@
 #include "MuT/NodeColourFromScalars.h"
 #include "MuT/NodeColourTable.h"
 #include "MuT/NodeColorToSRGB.h"
-#include "MuT/NodeColourMaterialBreak.h"
 #include "MuT/NodeScalarSwitch.h"
 #include "MuT/NodeScalarConstant.h"
 #include "MuT/NodeScalarVariation.h"
@@ -63,7 +61,6 @@
 #include "MuT/NodeScalarEnumParameter.h"
 #include "MuT/NodeScalarTable.h"
 #include "MuT/NodeScalarCurve.h"
-#include "MuT/NodeScalarMaterialBreak.h"
 #include "MuT/NodeMeshConstant.h"
 #include "MuT/NodeMeshFragment.h"
 #include "MuT/NodeMeshClipMorphPlane.h"
@@ -97,7 +94,6 @@
 #include "MuT/NodeRange.h"
 #include "MuT/NodeRangeFromScalar.h"
 #include "MuT/NodeMaterial.h"
-#include "MuT/NodeImageMaterialBreak.h"
 #include "MuT/NodeMaterialConstant.h"
 #include "MuT/NodeMaterialTable.h"
 #include "MuT/NodeMaterialSwitch.h"
@@ -170,7 +166,6 @@ namespace mu
 	FNodeType NodeImageTransform::StaticType = FNodeType(Node::EType::ImageTransform, NodeImage::GetStaticType());
 	FNodeType NodeImageVariation::StaticType = FNodeType(Node::EType::ImageVariation, NodeImage::GetStaticType());
 	FNodeType NodeImageColourMap::StaticType = FNodeType(Node::EType::ImageColorMap, NodeImage::GetStaticType());
-	FNodeType NodeImageFromMaterialParameter::StaticType = FNodeType(Node::EType::ImageFromMaterialParameter, NodeImage::GetStaticType());
 
 	FNodeType NodeColour::StaticType = FNodeType(Node::EType::Color, Node::GetStaticType());
 	FNodeType NodeColourConstant::StaticType = FNodeType(Node::EType::ColorConstant, NodeColour::GetStaticType());
@@ -229,10 +224,6 @@ namespace mu
 	FNodeType NodeMaterialSwitch::StaticType = FNodeType(Node::EType::MaterialSwitch, NodeMaterial::GetStaticType());
 	FNodeType NodeMaterialVariation::StaticType = FNodeType(Node::EType::MaterialVariation, NodeMaterial::GetStaticType());
 	FNodeType NodeMaterialParameter::StaticType = FNodeType(Node::EType::MaterialParameter, NodeMaterial::GetStaticType());
-	
-	FNodeType NodeImageMaterialBreak::StaticType = FNodeType(Node::EType::ImageMaterialBreak, NodeImage::GetStaticType());
-	FNodeType NodeScalarMaterialBreak::StaticType = FNodeType(Node::EType::ScalarMaterialBreak, NodeScalar::GetStaticType());
-	FNodeType NodeColourMaterialBreak::StaticType = FNodeType(Node::EType::ColourMaterialBreak, NodeColour::GetStaticType());
 
 
 	FNodeType::FNodeType()
diff --git a/Engine/Plugins/Mutable/Source/MutableTools/Public/MuT/Node.h b/Engine/Plugins/Mutable/Source/MutableTools/Public/MuT/Node.h
--- a/Engine/Plugins/Mutable/Source/MutableTools/Public/MuT/Node.h
+++ b/Engine/Plugins/Mutable/Source/MutableTools/Public/MuT/Node.h
@@ -75,7 +75,6 @@ namespace mu
 			ImageVariation,
 			ImageNormalComposite,
 			ImageTransform,
-			ImageFromMaterialParameter,
 
 			Bool,
 			BoolConstant,
@@ -161,10 +160,6 @@ namespace mu
 			MaterialSwitch,
 			MaterialVariation,
 			MaterialParameter,
-			
-			ImageMaterialBreak,
-			ScalarMaterialBreak,
-			ColourMaterialBreak,
 
 			Count
 		};
diff --git a/Engine/Plugins/Mutable/Source/MutableTools/Public/MuT/NodeColourMaterialBreak.h b/Engine/Plugins/Mutable/Source/MutableTools/Public/MuT/NodeColourMaterialBreak.h
deleted file mode 100644
diff --git a/Engine/Plugins/Mutable/Source/MutableTools/Public/MuT/NodeImageFromMaterialParameter.h b/Engine/Plugins/Mutable/Source/MutableTools/Public/MuT/NodeImageFromMaterialParameter.h
deleted file mode 100644
diff --git a/Engine/Plugins/Mutable/Source/MutableTools/Public/MuT/NodeImageMaterialBreak.h b/Engine/Plugins/Mutable/Source/MutableTools/Public/MuT/NodeImageMaterialBreak.h
deleted file mode 100644
diff --git a/Engine/Plugins/Mutable/Source/MutableTools/Public/MuT/NodeMaterialConstant.h b/Engine/Plugins/Mutable/Source/MutableTools/Public/MuT/NodeMaterialConstant.h
--- a/Engine/Plugins/Mutable/Source/MutableTools/Public/MuT/NodeMaterialConstant.h
+++ b/Engine/Plugins/Mutable/Source/MutableTools/Public/MuT/NodeMaterialConstant.h
@@ -6,15 +6,13 @@
 #include "MuR/Ptr.h"
 #include "MuR/RefCounted.h"
 #include "MuT/NodeMaterial.h"
-#include "MuT/NodeImage.h"
-#include "MuR/Types.h"
-#include "Math/Vector4.h"
 
 #define UE_API MUTABLETOOLS_API
 
 namespace mu
 {
-	/** NodeMaterialConstant class. */
+
+    /** NodeMaterialConstant class. */
 	class NodeMaterialConstant : public NodeMaterial
 	{
 	public:
@@ -25,10 +23,6 @@ namespace mu
 
 		int32 MaterialId = INDEX_NONE;
 
-		TMap<FName, FVector4f> ColorValues;
-		TMap<FName, float> ScalarValues;
-		TMap<FName, Ptr<NodeImage>> ImageValues;
-
 	protected:
 
 		// Forbidden. Manage with the Ptr<> template.
@@ -37,7 +31,10 @@ namespace mu
 	private:
 
 		static UE_API FNodeType StaticType;
+
 	};
+
+
 }
 
 #undef UE_API
diff --git a/Engine/Plugins/Mutable/Source/MutableTools/Public/MuT/NodeMaterialParameter.h b/Engine/Plugins/Mutable/Source/MutableTools/Public/MuT/NodeMaterialParameter.h
--- a/Engine/Plugins/Mutable/Source/MutableTools/Public/MuT/NodeMaterialParameter.h
+++ b/Engine/Plugins/Mutable/Source/MutableTools/Public/MuT/NodeMaterialParameter.h
@@ -2,14 +2,11 @@
 
 #pragma once
 
-#include "Containers/Map.h"
 #include "HAL/Platform.h"
 #include "MuR/Ptr.h"
 #include "MuR/RefCounted.h"
 #include "MuT/NodeRange.h"
 #include "MuT/NodeMaterial.h"
-#include "MuT/NodeImage.h"
-
 
 #define UE_API MUTABLETOOLS_API
 
@@ -30,10 +27,6 @@ namespace mu
 		/** Ranges for the parameter in case it is multidimensional. */
 		TArray<Ptr<NodeRange>> Ranges;
 
-		TArray<FString> ColorParameterNames;
-		TArray<FString> ScalarParameterNames;
-		TMap<FString, Ptr<NodeImage>> ImageParameterNodes;
-
 	public:
 
 		// Node Interface
diff --git a/Engine/Plugins/Mutable/Source/MutableTools/Public/MuT/NodeScalarMaterialBreak.h b/Engine/Plugins/Mutable/Source/MutableTools/Public/MuT/NodeScalarMaterialBreak.h
deleted file mode 100644
