[Mutable] Add node to transform vertices from a mesh, filtering vertices skinned to a given bone hierarchy

#rnx

diff --git a/Engine/Plugins/Mutable/Source/CustomizableObject/Internal/MuCO/CustomizableObjectPrivate.h b/Engine/Plugins/Mutable/Source/CustomizableObject/Internal/MuCO/CustomizableObjectPrivate.h
--- a/Engine/Plugins/Mutable/Source/CustomizableObject/Internal/MuCO/CustomizableObjectPrivate.h
+++ b/Engine/Plugins/Mutable/Source/CustomizableObject/Internal/MuCO/CustomizableObjectPrivate.h
@@ -1693,6 +1693,8 @@ enum class ECustomizableObjectVersions : int32
 
 		BackoutMaterialBreak,
 
+		AddNodeMeshTransformWithBoneHierarchy,
+
 	// -----<new versions can be added above this line>--------
 	LastCustomizableObjectVersion
 };
diff --git a/Engine/Plugins/Mutable/Source/CustomizableObjectEditor/Private/MuCOE/CustomizableObjectEditor.cpp b/Engine/Plugins/Mutable/Source/CustomizableObjectEditor/Private/MuCOE/CustomizableObjectEditor.cpp
--- a/Engine/Plugins/Mutable/Source/CustomizableObjectEditor/Private/MuCOE/CustomizableObjectEditor.cpp
+++ b/Engine/Plugins/Mutable/Source/CustomizableObjectEditor/Private/MuCOE/CustomizableObjectEditor.cpp
@@ -44,6 +44,7 @@
 #include "MuCOE/Nodes/CustomizableObjectNodeSkeletalMesh.h"
 #include "MuCOE/Nodes/CustomizableObjectNodeStaticMesh.h"
 #include "MuCOE/Nodes/CustomizableObjectNodeTable.h"
+#include "MuCOE/Nodes/CONodeModifierTransformWithBone.h"
 #include "MuCOE/SCustomizableObjectEditorAdvancedPreviewSettings.h"
 #include "MuCOE/SCustomizableObjectEditorTagExplorer.h"
 #include "MuCOE/SCustomizableObjectEditorTextureAnalyzer.h"
@@ -920,7 +921,8 @@ void FCustomizableObjectEditor::HideGizmoClipMesh()
 	{
 		const UObject* Node = *NodeIt;
 		if (Node->IsA<UCustomizableObjectNodeModifierClipWithMesh>() ||
-			Node->IsA<UCustomizableObjectNodeModifierTransformInMesh>())
+			Node->IsA<UCustomizableObjectNodeModifierTransformInMesh>() ||
+			Node->IsA<UCONodeModifierTransformWithBone>())
 		{
 			GraphEditor->ClearSelectionSet();
 			break;
diff --git a/Engine/Plugins/Mutable/Source/CustomizableObjectEditor/Private/MuCOE/CustomizableObjectEditorModule.cpp b/Engine/Plugins/Mutable/Source/CustomizableObjectEditor/Private/MuCOE/CustomizableObjectEditorModule.cpp
--- a/Engine/Plugins/Mutable/Source/CustomizableObjectEditor/Private/MuCOE/CustomizableObjectEditorModule.cpp
+++ b/Engine/Plugins/Mutable/Source/CustomizableObjectEditor/Private/MuCOE/CustomizableObjectEditorModule.cpp
@@ -58,6 +58,9 @@
 #include "MuCOE/Nodes/CustomizableObjectNodeModifierClipWithUVMask.h"
 #include "MuCOE/Nodes/CustomizableObjectNodeModifierClipDeform.h"
 #include "MuCOE/Nodes/CustomizableObjectNodeModifierTransformInMesh.h"
+#include "MuCOE/Nodes/CustomizableObjectNodeModifierTransformInMeshDetails.h"
+#include "MuCOE/Nodes/CONodeModifierTransformWithBone.h"
+#include "MuCOE/Nodes/CONodeModifierTransformWithBoneDetails.h"
 #include "MuCOE/Nodes/CustomizableObjectNodeSkeletalMeshParameter.h"
 #include "MuCOE/Nodes/CustomizableObjectNodeMeshMorph.h"
 #include "MuCOE/Nodes/CustomizableObjectNodeMeshMorphDetails.h"
@@ -104,7 +107,6 @@
 #include "MuCOE/CompileRequest.h"
 #include "MuCOE/CustomizableObjectGraph.h"
 #include "Nodes/CustomizableObjectNodeComponentMeshDetails.h"
-#include "Nodes/CustomizableObjectNodeModifierTransformInMeshDetails.h"
 #include "MuCOE/GenerateMutableSource/GenerateMutableSourceTable.h"
 #include "Nodes/CustomizableObjectNodeTextureParameter.h"
 #include "PhysicsEngine/PhysicsAsset.h"
@@ -261,6 +263,7 @@ void FCustomizableObjectEditorModule::StartupModule()
 	RegisterCustomDetails(PropertyModule, UCustomizableObjectNodeModifierClipWithUVMask::StaticClass(), FOnGetDetailCustomizationInstance::CreateStatic(&FCustomizableObjectNodeModifierBaseDetails::MakeInstance));
 	RegisterCustomDetails(PropertyModule, UCustomizableObjectNodeModifierClipDeform::StaticClass(), FOnGetDetailCustomizationInstance::CreateStatic(&FCustomizableObjectNodeModifierBaseDetails::MakeInstance));
 	RegisterCustomDetails(PropertyModule, UCustomizableObjectNodeModifierTransformInMesh::StaticClass(), FOnGetDetailCustomizationInstance::CreateStatic(&FCustomizableObjectNodeModifierTransformInMeshDetails::MakeInstance));
+	RegisterCustomDetails(PropertyModule, UCONodeModifierTransformWithBone::StaticClass(), FOnGetDetailCustomizationInstance::CreateStatic(&FCONodeModifierTransformWithBoneDetails::MakeInstance));
 	RegisterCustomDetails(PropertyModule, UCustomizableObjectNodeObject::StaticClass(), FOnGetDetailCustomizationInstance::CreateStatic(&FCustomizableObjectNodeObjectDetails::MakeInstance));
 	RegisterCustomDetails(PropertyModule, UCustomizableObjectNodeObjectGroup::StaticClass(), FOnGetDetailCustomizationInstance::CreateStatic(&FCustomizableObjectNodeObjectGroupDetails::MakeInstance));
 	RegisterCustomDetails(PropertyModule, UCustomizableObjectNodeProjectorParameter::StaticClass(), FOnGetDetailCustomizationInstance::CreateStatic(&FCustomizableObjectNodeProjectorParameterDetails::MakeInstance));
diff --git a/Engine/Plugins/Mutable/Source/CustomizableObjectEditor/Private/MuCOE/CustomizableObjectEditorViewportClient.cpp b/Engine/Plugins/Mutable/Source/CustomizableObjectEditor/Private/MuCOE/CustomizableObjectEditorViewportClient.cpp
--- a/Engine/Plugins/Mutable/Source/CustomizableObjectEditor/Private/MuCOE/CustomizableObjectEditorViewportClient.cpp
+++ b/Engine/Plugins/Mutable/Source/CustomizableObjectEditor/Private/MuCOE/CustomizableObjectEditorViewportClient.cpp
@@ -36,6 +36,7 @@
 #include "MuCOE/Nodes/CustomizableObjectNodeModifierTransformInMesh.h"
 #include "MuCOE/Nodes/CustomizableObjectNodeProjectorConstant.h"
 #include "MuCOE/Nodes/CustomizableObjectNodeProjectorParameter.h"
+#include "MuCOE/Nodes/CONodeModifierTransformWithBone.h"
 #include "MuCOE/UnrealEditorPortabilityHelpers.h"
 #include "Preferences/PersonaOptions.h"
 #include "SceneView.h"
@@ -642,6 +643,10 @@ void FCustomizableObjectEditorViewportClient::ShowGizmoClipMesh(UCustomizableObj
 	{
 		TransformExternallyChangedDelegateHandle = NodeModifierTransformInMesh->TransformChangedDelegate.AddSP(this, &FCustomizableObjectEditorViewportClient::UpdateGizmoClipMeshTransform);
 	}
+	else if (UCONodeModifierTransformWithBone* NodeModifierTransformWithBone = Cast<UCONodeModifierTransformWithBone>(ClipMeshNode.Get()))
+	{
+		TransformExternallyChangedDelegateHandle = NodeModifierTransformWithBone->TransformChangedDelegate.AddSP(this, &FCustomizableObjectEditorViewportClient::UpdateGizmoClipMeshTransform);
+	}
 	else
 	{
 		unimplemented();
@@ -704,6 +709,10 @@ void FCustomizableObjectEditorViewportClient::HideGizmoClipMesh()
 			{
 				NodeModifierTransformInMesh->TransformChangedDelegate.Remove(TransformExternallyChangedDelegateHandle);
 			}
+			else if (UCONodeModifierTransformWithBone* NodeModifierTransformWithBone = Cast<UCONodeModifierTransformWithBone>(ClipMeshNode.Get()))
+			{
+				NodeModifierTransformWithBone->TransformChangedDelegate.Remove(TransformExternallyChangedDelegateHandle);
+			}
 			else
 			{
 				unimplemented();
diff --git a/Engine/Plugins/Mutable/Source/CustomizableObjectEditor/Private/MuCOE/EdGraphSchema_CustomizableObject.cpp b/Engine/Plugins/Mutable/Source/CustomizableObjectEditor/Private/MuCOE/EdGraphSchema_CustomizableObject.cpp
--- a/Engine/Plugins/Mutable/Source/CustomizableObjectEditor/Private/MuCOE/EdGraphSchema_CustomizableObject.cpp
+++ b/Engine/Plugins/Mutable/Source/CustomizableObjectEditor/Private/MuCOE/EdGraphSchema_CustomizableObject.cpp
@@ -20,6 +20,7 @@
 #include "MuCOE/Nodes/CONodeMaterialConstant.h"
 #include "MuCOE/Nodes/CONodeMaterialVariation.h"
 #include "MuCOE/Nodes/CONodeMaterialSwitch.h"
+#include "MuCOE/Nodes/CONodeModifierTransformWithBone.h"
 #include "MuCOE/Nodes/CustomizableObjectNodeAnimationPose.h"
 #include "MuCOE/Nodes/CustomizableObjectNodeColorArithmeticOp.h"
 #include "MuCOE/Nodes/CustomizableObjectNodeColorConstant.h"
@@ -619,6 +620,7 @@ void UEdGraphSchema_CustomizableObject::GetGraphContextActions(FGraphContextMenu
 			ContextMenuBuilder.CreateTemplateNode<UCustomizableObjectNodeMeshReshape>(),
 			ContextMenuBuilder.CreateTemplateNode<UCustomizableObjectNodeModifierClipDeform>(),
 			ContextMenuBuilder.CreateTemplateNode<UCustomizableObjectNodeModifierTransformInMesh>(),
+			ContextMenuBuilder.CreateTemplateNode<UCONodeModifierTransformWithBone>(),
 			ContextMenuBuilder.CreateTemplateNode<UCustomizableObjectNodeTextureParameter>(),
 			ContextMenuBuilder.CreateTemplateNode<UCustomizableObjectNodeSkeletalMeshParameter>(),
 			ContextMenuBuilder.CreateTemplateNode<UCustomizableObjectNodeMaterialParameter>(),
diff --git a/Engine/Plugins/Mutable/Source/CustomizableObjectEditor/Private/MuCOE/GenerateMutableSource/GenerateMutableSourceModifier.cpp b/Engine/Plugins/Mutable/Source/CustomizableObjectEditor/Private/MuCOE/GenerateMutableSource/GenerateMutableSourceModifier.cpp
--- a/Engine/Plugins/Mutable/Source/CustomizableObjectEditor/Private/MuCOE/GenerateMutableSource/GenerateMutableSourceModifier.cpp
+++ b/Engine/Plugins/Mutable/Source/CustomizableObjectEditor/Private/MuCOE/GenerateMutableSource/GenerateMutableSourceModifier.cpp
@@ -15,6 +15,7 @@
 #include "MuCOE/GenerateMutableSource/GenerateMutableSourceMacro.h"
 #include "MuCOE/GraphTraversal.h"
 #include "MuCOE/MutableUtils.h"
+#include "MuCOE/Nodes/CONodeModifierTransformWithBone.h"
 #include "MuCOE/Nodes/CustomizableObjectNodeMacroInstance.h"
 #include "MuCOE/Nodes/CustomizableObjectNodeModifierClipDeform.h"
 #include "MuCOE/Nodes/CustomizableObjectNodeModifierClipMorph.h"
@@ -39,6 +40,7 @@
 #include "MuT/NodeModifierMeshClipMorphPlane.h"
 #include "MuT/NodeModifierMeshClipWithUVMask.h"
 #include "MuT/NodeModifierMeshTransformInMesh.h"
+#include "MuT/NodeModifierMeshTransformWithBone.h"
 #include "MuT/NodeModifierSurfaceEdit.h"
 #include "MuT/NodeMeshConstant.h"
 #include "MuT/NodeMeshFragment.h"
@@ -643,6 +645,25 @@ mu::Ptr<mu::NodeModifier> GenerateMutableSourceModifier(const UEdGraphPin * Pin,
 		GenerationContext.MeshGenerationFlags.Pop();
 	}
 
+	else if (const UCONodeModifierTransformWithBone* TypedNodeTransformWithBone = Cast<UCONodeModifierTransformWithBone>(Node))
+	{
+		bCacheNode = false;
+
+		mu::Ptr<mu::NodeModifierMeshTransformWithBone> TransformNode = new mu::NodeModifierMeshTransformWithBone();
+		Result = TransformNode;
+
+		TransformNode->BoneName = GenerationContext.CompilationContext->GetBoneUnique(*TypedNodeTransformWithBone->BoneName);
+		TransformNode->ThresholdFactor = TypedNodeTransformWithBone->ThresholdFactor;
+
+		if (const UEdGraphPin* ConnectedPin = FollowInputPin(*TypedNodeTransformWithBone->GetTransformPin()))
+		{
+			TransformNode->MatrixNode = GenerateMutableSourceTransform(ConnectedPin, GenerationContext);
+		}
+
+		TransformNode->MultipleTagsPolicy = TypedNodeTransformWithBone->MultipleTagPolicy;
+		TransformNode->RequiredTags = TypedNodeTransformWithBone->GetNodeRequiredTags(&GenerationContext.MacroNodesStack);
+	}
+
 	else if (const UCustomizableObjectNodeMacroInstance* TypedNodeMacro = Cast<UCustomizableObjectNodeMacroInstance>(Node))
 	{
 		bCacheNode = false;
diff --git a/Engine/Plugins/Mutable/Source/CustomizableObjectEditor/Private/MuCOE/Nodes/CONodeModifierTransformWithBone.cpp b/Engine/Plugins/Mutable/Source/CustomizableObjectEditor/Private/MuCOE/Nodes/CONodeModifierTransformWithBone.cpp
new file mode 100644
--- /dev/null
+++ b/Engine/Plugins/Mutable/Source/CustomizableObjectEditor/Private/MuCOE/Nodes/CONodeModifierTransformWithBone.cpp
@@ -0,0 +1,69 @@
+// Copyright Epic Games, Inc. All Rights Reserved.
+
+#include "MuCOE/Nodes/CONodeModifierTransformWithBone.h"
+
+#include "MuCOE/EdGraphSchema_CustomizableObject.h"
+
+#include UE_INLINE_GENERATED_CPP_BY_NAME(CONodeModifierTransformWithBone)
+
+#define LOCTEXT_NAMESPACE "CustomizableObjectEditor"
+
+FText UCONodeModifierTransformWithBone::GetNodeTitle(ENodeTitleType::Type TitleType) const
+{
+	return LOCTEXT("Transform_Mesh_In_Bone_Hierarchy", "Transform Mesh In Bone Hierarchy");
+}
+
+FText UCONodeModifierTransformWithBone::GetTooltipText() const
+{
+	return LOCTEXT("Transform_Mesh_In_Bone_Hierarchy_Tooltip", "Applies a transform to the vertices of a mesh that is skinned to the target bone or its child bones");
+}
+
+void UCONodeModifierTransformWithBone::AllocateDefaultPins(UCustomizableObjectNodeRemapPins* RemapPins)
+{
+	const UEdGraphSchema_CustomizableObject* Schema = GetDefault<UEdGraphSchema_CustomizableObject>();
+
+	UEdGraphPin* NewTransformPin = CustomCreatePin(EGPD_Input, Schema->PC_Transform, FName("Transform"));
+	NewTransformPin->bDefaultValueIsIgnored = true;
+	TransformPin = NewTransformPin;
+
+	UEdGraphPin* OutPin = CustomCreatePin(EGPD_Output, Schema->PC_Modifier, FName("Modifier"));
+	OutPin->bDefaultValueIsIgnored = true;
+	OutputPin = OutPin;
+
+	//Create Node Modifier Common Pins
+	Super::AllocateDefaultPins(RemapPins);
+}
+
+
+void UCONodeModifierTransformWithBone::PostEditChangeProperty(struct FPropertyChangedEvent& PropertyChangedEvent)
+{
+	Super::PostEditChangeProperty(PropertyChangedEvent);
+
+	if (PropertyChangedEvent.Property == nullptr)
+	{
+		return;
+	}
+	
+	if (PropertyChangedEvent.Property->GetFName() == GET_MEMBER_NAME_CHECKED(UCONodeModifierTransformWithBone, BoundingMeshTransform)
+		|| PropertyChangedEvent.Property->GetFName() == GET_MEMBER_NAME_CHECKED(UCONodeModifierTransformWithBone, BoneName))
+	{
+		if (TransformChangedDelegate.IsBound())
+		{
+			TransformChangedDelegate.Broadcast(BoundingMeshTransform);
+		}
+	}
+}
+
+
+UEdGraphPin* UCONodeModifierTransformWithBone::GetOutputPin() const
+{
+	return OutputPin.Get();
+}
+
+
+UEdGraphPin* UCONodeModifierTransformWithBone::GetTransformPin() const
+{
+	return TransformPin.Get();
+}
+
+#undef LOCTEXT_NAMESPACE
diff --git a/Engine/Plugins/Mutable/Source/CustomizableObjectEditor/Private/MuCOE/Nodes/CONodeModifierTransformWithBone.h b/Engine/Plugins/Mutable/Source/CustomizableObjectEditor/Private/MuCOE/Nodes/CONodeModifierTransformWithBone.h
new file mode 100644
--- /dev/null
+++ b/Engine/Plugins/Mutable/Source/CustomizableObjectEditor/Private/MuCOE/Nodes/CONodeModifierTransformWithBone.h
@@ -0,0 +1,70 @@
+// Copyright Epic Games, Inc. All Rights Reserved.
+
+#pragma once
+
+#include "MuCOE/Nodes/CustomizableObjectNodeModifierBase.h"
+#include "MuT/NodeModifier.h"
+
+#include "CONodeModifierTransformWithBone.generated.h"
+
+#define UE_API CUSTOMIZABLEOBJECTEDITOR_API
+
+namespace ENodeTitleType { enum Type : int; }
+
+class UCustomizableObject;
+class UCustomizableObjectNodeRemapPins;
+class UEdGraphPin;
+
+DECLARE_MULTICAST_DELEGATE_OneParam(OnTransformChanged, const FTransform&);
+
+UCLASS(MinimalAPI)
+class UCONodeModifierTransformWithBone : public UCustomizableObjectNodeModifierBase
+{
+	GENERATED_BODY()
+
+public:
+	/** Transform to apply to the bounding mesh before selecting for vertices to transform. */
+	UPROPERTY(EditAnywhere, Category = Transform)
+	FTransform BoundingMeshTransform = FTransform::Identity;
+
+	/** Root bone used to filter vertices based on skinning weights. Vertices skinned to this
+	  * bone or below will me transformed. */
+	UPROPERTY(EditAnywhere, Category = Transform)
+	FString BoneName;
+
+	/** Vertex influences under this threshold will not be considered when filtering vertices.
+	  * Percentage of total skin weight [0.0 .. 1.0]. */
+	UPROPERTY(EditAnywhere, Category = Transform, meta = (ClampMin = 0.0f, ClampMax = 1.0f))
+	float ThresholdFactor = 0.05f;
+
+	/**
+	 * Delegate invoked each time the BoundingMeshTransform value gets modified
+	 */
+	OnTransformChanged TransformChangedDelegate;
+	
+public:
+	// EdGraphNode interface
+	UE_API virtual FText GetNodeTitle(ENodeTitleType::Type TitleType) const override;
+	UE_API virtual FText GetTooltipText() const override;
+
+	// UCustomizableObjectNode interface
+	UE_API virtual void AllocateDefaultPins(UCustomizableObjectNodeRemapPins* RemapPins) override;
+	virtual bool IsExperimental() const override { return true; }
+	UE_API virtual void PostEditChangeProperty( struct FPropertyChangedEvent& PropertyChangedEvent) override;
+
+	// UCustomizableObjectNodeModifierBase interface
+	UE_API virtual UEdGraphPin* GetOutputPin() const override;
+
+	// Own interface
+	UE_API UEdGraphPin* GetTransformPin() const;
+	
+private:
+
+	UPROPERTY()
+	FEdGraphPinReference TransformPin;
+
+	UPROPERTY()
+	FEdGraphPinReference OutputPin;
+};
+
+#undef UE_API
diff --git a/Engine/Plugins/Mutable/Source/CustomizableObjectEditor/Private/MuCOE/Nodes/CONodeModifierTransformWithBoneDetails.cpp b/Engine/Plugins/Mutable/Source/CustomizableObjectEditor/Private/MuCOE/Nodes/CONodeModifierTransformWithBoneDetails.cpp
new file mode 100644
--- /dev/null
+++ b/Engine/Plugins/Mutable/Source/CustomizableObjectEditor/Private/MuCOE/Nodes/CONodeModifierTransformWithBoneDetails.cpp
@@ -0,0 +1,46 @@
+// Copyright Epic Games, Inc. All Rights Reserved.
+
+#include "MuCOE/Nodes/CONodeModifierTransformWithBoneDetails.h"
+
+#include "MuCOE/Nodes/CONodeModifierTransformWithBone.h"
+
+#include "DetailCategoryBuilder.h"
+#include "DetailLayoutBuilder.h"
+#include "DetailWidgetRow.h"
+
+#define LOCTEXT_NAMESPACE "CONodeModifierTransformWithMeshDetails"
+
+
+TSharedRef<IDetailCustomization> FCONodeModifierTransformWithBoneDetails::MakeInstance()
+{
+	return MakeShareable(new FCONodeModifierTransformWithBoneDetails);
+}
+
+
+void FCONodeModifierTransformWithBoneDetails::CustomizeDetails(IDetailLayoutBuilder& DetailBuilder)
+{
+	FCustomizableObjectNodeModifierBaseDetails::CustomizeDetails(DetailBuilder);
+
+	Node = nullptr;
+	DetailBuilderPtr = &DetailBuilder;
+
+	TSharedPtr<const IDetailsView> DetailsView = DetailBuilder.GetDetailsViewSharedPtr();
+	if (DetailsView && DetailsView->GetSelectedObjects().Num())
+	{
+		Node = Cast<UCONodeModifierTransformWithBone>(DetailsView->GetSelectedObjects()[0].Get());
+	}
+
+	IDetailCategoryBuilder& CustomizableObjectToClipCategory = DetailBuilder.EditCategory("Transform");
+
+	// TODO Remove
+	if(!Node)
+	{
+		CustomizableObjectToClipCategory.AddCustomRow(LOCTEXT("Node", "Node"))
+			[
+				SNew(STextBlock)
+				.Text(LOCTEXT("Node not found", "Node not found"))
+			];
+	}
+}
+
+#undef LOCTEXT_NAMESPACE
diff --git a/Engine/Plugins/Mutable/Source/CustomizableObjectEditor/Private/MuCOE/Nodes/CONodeModifierTransformWithBoneDetails.h b/Engine/Plugins/Mutable/Source/CustomizableObjectEditor/Private/MuCOE/Nodes/CONodeModifierTransformWithBoneDetails.h
new file mode 100644
--- /dev/null
+++ b/Engine/Plugins/Mutable/Source/CustomizableObjectEditor/Private/MuCOE/Nodes/CONodeModifierTransformWithBoneDetails.h
@@ -0,0 +1,29 @@
+// Copyright Epic Games, Inc. All Rights Reserved.
+
+#pragma once
+
+#include "CustomizableObjectNodeModifierBaseDetails.h"
+#include "IDetailCustomization.h"
+
+class IDetailLayoutBuilder;
+class UCONodeModifierTransformWithBone;
+
+
+class FCONodeModifierTransformWithBoneDetails : public FCustomizableObjectNodeModifierBaseDetails
+{
+public:
+	/** Makes a new instance of this detail layout class for a specific detail view requesting it */
+	static TSharedRef<IDetailCustomization> MakeInstance();
+
+	/** ILayoutDetails interface */
+	virtual void CustomizeDetails( IDetailLayoutBuilder& DetailBuilder ) override;
+
+private:
+
+	/** The node for which details are being customized */
+	UCONodeModifierTransformWithBone* Node = nullptr;
+
+	/** Pointer to the builder passed by CustomizeDetails method */
+	IDetailLayoutBuilder* DetailBuilderPtr = nullptr;
+
+};
diff --git a/Engine/Plugins/Mutable/Source/CustomizableObjectEditor/Private/MuCOE/SMutableCodeViewer.cpp b/Engine/Plugins/Mutable/Source/CustomizableObjectEditor/Private/MuCOE/SMutableCodeViewer.cpp
--- a/Engine/Plugins/Mutable/Source/CustomizableObjectEditor/Private/MuCOE/SMutableCodeViewer.cpp
+++ b/Engine/Plugins/Mutable/Source/CustomizableObjectEditor/Private/MuCOE/SMutableCodeViewer.cpp
@@ -2721,6 +2721,10 @@ bool SMutableCodeViewer::IsConstantResourceUsedByOperation(const int32 IndexOnCo
 				{
 					bResourceLocated = IndexOnConstantsArray == InProgram.GetOpArgs<mu::OP::MeshTransformWithinMeshArgs>(OperationAddress).matrix;
 				}
+				else if (OperationType == mu::EOpType::ME_TRANSFORMWITHBONE)
+				{
+					bResourceLocated = IndexOnConstantsArray == InProgram.GetOpArgs<mu::OP::MeshTransformWithBoneArgs>(OperationAddress).Matrix;
+				}
 				else if (OperationType == mu::EOpType::MA_CONSTANT)
 				{
 					bResourceLocated = IndexOnConstantsArray == InProgram.GetOpArgs<mu::OP::MatrixConstantArgs>(OperationAddress).value;
diff --git a/Engine/Plugins/Mutable/Source/CustomizableObjectEditor/Private/MuCOE/SMutableGraphViewer.cpp b/Engine/Plugins/Mutable/Source/CustomizableObjectEditor/Private/MuCOE/SMutableGraphViewer.cpp
--- a/Engine/Plugins/Mutable/Source/CustomizableObjectEditor/Private/MuCOE/SMutableGraphViewer.cpp
+++ b/Engine/Plugins/Mutable/Source/CustomizableObjectEditor/Private/MuCOE/SMutableGraphViewer.cpp
@@ -60,6 +60,7 @@
 #include "MuT/NodeScalarTable.h"
 #include "MuT/NodeObjectNew.h"
 #include "MuT/NodeModifierMeshTransformInMesh.h"
+#include "MuT/NodeModifierMeshTransformWithBone.h"
 #include "Widgets/MutableExpanderArrow.h"
 
 class FExtender;
@@ -196,6 +197,7 @@ static const char* MutableNodeNames[] =
 	"ModifierMeshClipWithUVMask",
 	"ModifierSurfaceEdit",
 	"ModifierTransformInMesh",
+	"ModifierTransformWithBone",
 
 	"ExtensionData",
 	"ExtensionDataConstant",
@@ -612,6 +614,12 @@ void SMutableGraphViewer::GetChildrenForInfo(TSharedPtr<FMutableGraphTreeElement
 		AddChildFunc(ModifierMeshTransformInMeshVar->MatrixNode.get(), FString::Printf(TEXT("MESH TRANSFORM")));
 	}
 
+	else if (ParentNode->GetType() == mu::NodeModifierMeshTransformWithBone::GetStaticType())
+	{
+		mu::NodeModifierMeshTransformWithBone* ModifierMeshTransformWithBoneVar = StaticCast<mu::NodeModifierMeshTransformWithBone*>(ParentNode);
+		AddChildFunc(ModifierMeshTransformWithBoneVar->MatrixNode.get(), FString::Printf(TEXT("MESH TRANSFORM")));
+	}
+
 	else if (ParentNode->GetType() == mu::NodeImageSwitch::GetStaticType())
 	{
 		mu::NodeImageSwitch* ImageSwitchVar = StaticCast<mu::NodeImageSwitch*>(ParentNode);
diff --git a/Engine/Plugins/Mutable/Source/MutableRuntime/Internal/MuR/Operations.h b/Engine/Plugins/Mutable/Source/MutableRuntime/Internal/MuR/Operations.h
--- a/Engine/Plugins/Mutable/Source/MutableRuntime/Internal/MuR/Operations.h
+++ b/Engine/Plugins/Mutable/Source/MutableRuntime/Internal/MuR/Operations.h
@@ -285,6 +285,9 @@ namespace mu
     	//! Transform with a 4x4 matrix the geometry channels of a mesh that are bounded by another mesh
     	ME_TRANSFORMWITHMESH,
 
+		//! Transform with a 4x4 matrix the geometry channels of a mesh that are skinned to a bone or hierarchy
+		ME_TRANSFORMWITHBONE,
+
         //-----------------------------------------------------------------------------------------
         // Instance operations
         //-----------------------------------------------------------------------------------------
@@ -973,6 +976,14 @@ namespace mu
     		ADDRESS matrix;
     	};
 
+		struct MeshTransformWithBoneArgs
+		{
+			ADDRESS SourceMesh;
+			ADDRESS Matrix;
+			uint32 BoneId;
+			float ThresholdFactor;
+		};
+
 		struct MeshOptimizeSkinningArgs
 		{
 			ADDRESS source;
diff --git a/Engine/Plugins/Mutable/Source/MutableRuntime/Private/MuR/CodeRunner.cpp b/Engine/Plugins/Mutable/Source/MutableRuntime/Private/MuR/CodeRunner.cpp
--- a/Engine/Plugins/Mutable/Source/MutableRuntime/Private/MuR/CodeRunner.cpp
+++ b/Engine/Plugins/Mutable/Source/MutableRuntime/Private/MuR/CodeRunner.cpp
@@ -2,7 +2,6 @@
 
 #include "MuR/CodeRunner.h"
 
-#include "OpMeshTransformWithMesh.h"
 #include "GenericPlatform/GenericPlatformMath.h"
 #include "HAL/UnrealMemory.h"
 #include "Logging/LogCategory.h"
@@ -50,6 +49,8 @@
 #include "MuR/OpMeshRemove.h"
 #include "MuR/OpMeshReshape.h"
 #include "MuR/OpMeshTransform.h"
+#include "MuR/OpMeshTransformWithMesh.h"
+#include "MuR/OpMeshTransformWithBone.h"
 #include "MuR/OpMeshOptimizeSkinning.h"
 #include "MuR/Operations.h"
 #include "MuR/Parameters.h"
@@ -3381,6 +3382,69 @@ namespace mu
 			break;
 		}
 
+		case EOpType::ME_TRANSFORMWITHBONE:
+		{
+			OP::MeshTransformWithBoneArgs Args = Program.GetOpArgs<OP::MeshTransformWithBoneArgs>(Item.At);
+			switch (Item.Stage)
+			{
+			case 0:
+			{
+				if (Args.SourceMesh)
+				{
+					if (!EnumHasAnyFlags(EMeshContentFlags(Item.ExecutionOptions), EMeshContentFlags::GeometryData))
+					{
+						AddOp(FScheduledOp(Item.At, Item, 1),
+							FScheduledOp(Args.SourceMesh, Item));
+					}
+					else
+					{
+						AddOp(FScheduledOp(Item.At, Item, 1),
+							FScheduledOp(Args.SourceMesh, Item),
+							FScheduledOp(Args.Matrix, Item));
+					}
+				}
+				else
+				{
+					StoreMesh(Item, nullptr);
+				}
+				break;
+			}
+			case 1:
+			{
+				MUTABLE_CPUPROFILER_SCOPE(ME_TRANSFORMWITHBONE_1)
+
+					if (!EnumHasAnyFlags(EMeshContentFlags(Item.ExecutionOptions), EMeshContentFlags::GeometryData))
+					{
+						TSharedPtr<const FMesh> SourceMesh = LoadMesh(FCacheAddress(Args.SourceMesh, Item));
+						StoreMesh(Item, SourceMesh);
+					}
+					else
+					{
+						TSharedPtr<const FMesh> SourceMesh = LoadMesh(FCacheAddress(Args.SourceMesh, Item));
+						const FMatrix44f& Transform = LoadMatrix(FCacheAddress(Args.Matrix, Item));
+						const FBoneName BoneName = FBoneName(Args.BoneId);
+
+						if (SourceMesh)
+						{
+							TSharedPtr<FMesh> Result = CloneOrTakeOver(SourceMesh);
+
+							MeshTransformWithBoneInline(Result.Get(), Transform, FBoneName(Args.BoneId), Args.ThresholdFactor);
+							StoreMesh(Item, Result);
+						}
+						else
+						{
+							StoreMesh(Item, SourceMesh);
+						}
+					}
+				break;
+			}
+
+			default:
+				check(false);
+			}
+			break;
+		}
+
 		default:
 			if (type!=EOpType::NONE)
 			{
diff --git a/Engine/Plugins/Mutable/Source/MutableRuntime/Private/MuR/CodeRunnerBegin.cpp b/Engine/Plugins/Mutable/Source/MutableRuntime/Private/MuR/CodeRunnerBegin.cpp
--- a/Engine/Plugins/Mutable/Source/MutableRuntime/Private/MuR/CodeRunnerBegin.cpp
+++ b/Engine/Plugins/Mutable/Source/MutableRuntime/Private/MuR/CodeRunnerBegin.cpp
@@ -1981,6 +1981,15 @@ namespace mu
 					StoreNone(Item);
 					break;
 				}
+
+			case EOpType::ME_TRANSFORMWITHBONE:
+				{
+					OP::MeshTransformWithBoneArgs Args = Program.GetOpArgs<OP::MeshTransformWithBoneArgs>(Item.At);
+					PushOp(FScheduledOpInline(Args.SourceMesh, Item));
+
+					StoreNone(Item);
+					break;
+				}
 			
 			case EOpType::IM_REFERENCE:
 				{
diff --git a/Engine/Plugins/Mutable/Source/MutableRuntime/Private/MuR/OpMeshTransformWithBone.cpp b/Engine/Plugins/Mutable/Source/MutableRuntime/Private/MuR/OpMeshTransformWithBone.cpp
new file mode 100644
--- /dev/null
+++ b/Engine/Plugins/Mutable/Source/MutableRuntime/Private/MuR/OpMeshTransformWithBone.cpp
@@ -0,0 +1,330 @@
+﻿// Copyright Epic Games, Inc. All Rights Reserved.
+
+#include "MuR/OpMeshTransformWithBone.h"
+
+#include "OpMeshClipWithMesh.h"
+#include "MuR/Mesh.h"
+#include "MuR/MeshPrivate.h"
+#include "MuR/MutableTrace.h"
+
+#include "MuR/ParallelExecutionUtils.h"
+#include "PackedNormal.h"
+
+namespace mu
+{
+
+
+void MeshTransformWithBoneInline(FMesh* Mesh, const FMatrix44f& Transform, const FBoneName& InBoneName, const float InThresholdFactor)
+{
+	MUTABLE_CPUPROFILER_SCOPE(MeshTransformWithBone);
+
+	if (!Mesh || Mesh->GetVertexCount() == 0 || Mesh->IsReference())
+	{
+		return;
+	}
+
+
+	TSharedPtr<const FSkeleton> BaseSkeleton = Mesh->GetSkeleton();
+
+	const int32 BaseBoneIndex = BaseSkeleton ? BaseSkeleton->FindBone(InBoneName.Id) : INDEX_NONE;
+	if (BaseBoneIndex == INDEX_NONE)
+	{
+		return;
+	}
+
+	TBitArray<> AffectedBoneMapIndices;
+
+	// Find affected bones
+	{
+		const int32 NumBonesBoneMap = Mesh->BoneMap.Num();
+		const TArray<FBoneName>& BoneMap = Mesh->BoneMap;
+		AffectedBoneMapIndices.SetNum(NumBonesBoneMap, false);
+
+		const int32 BoneCount = BaseSkeleton->GetBoneCount();
+		TBitArray<> AffectedSkeletonBones;
+		AffectedSkeletonBones.SetNum(BoneCount, false);
+
+		for (int32 BoneIndex = 0; BoneIndex < BoneCount; ++BoneIndex)
+		{
+			const int32 ParentBoneIndex = BaseSkeleton->GetBoneParent(BoneIndex);
+			check(ParentBoneIndex < BoneIndex);
+
+			const bool bIsBoneAffected = (AffectedSkeletonBones.IsValidIndex(ParentBoneIndex) && AffectedSkeletonBones[ParentBoneIndex])
+				|| BoneIndex == BaseBoneIndex;
+
+			AffectedSkeletonBones[BoneIndex] = bIsBoneAffected;
+		}
+
+		for (int32 BoneIndex = 0; BoneIndex < NumBonesBoneMap; ++BoneIndex)
+		{
+			const FBoneName& BoneName = BoneMap[BoneIndex];
+			
+			const int32 SkeletonIndex = BaseSkeleton->FindBone(BoneName);
+			if (SkeletonIndex != INDEX_NONE && AffectedSkeletonBones[SkeletonIndex])
+			{
+				AffectedBoneMapIndices[BoneIndex] = true;
+			}
+		}
+	}
+
+	// Get pointers to skinning data
+	UntypedMeshBufferIteratorConst BoneIndicesIterBegin(Mesh->VertexBuffers, EMeshBufferSemantic::BoneIndices, 0);
+	UntypedMeshBufferIteratorConst BoneWeightsIterBegin(Mesh->VertexBuffers, EMeshBufferSemantic::BoneWeights, 0);
+	if (!BoneIndicesIterBegin.ptr() || !BoneWeightsIterBegin.ptr())
+	{
+		// No skinning data
+		return;
+	}
+
+	const int32 NumVertices = Mesh->GetVertexCount();
+	const int32 NumWeights = BoneWeightsIterBegin.GetComponents();
+
+	const uint16 MinWeightThreshold = (BoneWeightsIterBegin.GetFormat() == EMeshBufferFormat::NUInt8 ? MAX_uint8 : MAX_uint16) * InThresholdFactor;
+
+	// Classify which vertices in the SourceMesh are completely bounded by the BoundingMesh geometry.
+	// If no BoundingMesh is provided, this defaults to act exactly like mu::MeshTransform
+	TBitArray<> VertexInBoneHierarchy;
+	VertexInBoneHierarchy.Init(false, NumVertices);
+	{
+		for (const FMeshSurface& Surface : Mesh->Surfaces)
+		{
+			const int32 AffectedBoneMapIndex = AffectedBoneMapIndices.FindFrom(true, int32(Surface.BoneMapIndex), int32(Surface.BoneMapIndex + Surface.BoneMapCount));
+			if (AffectedBoneMapIndex == INDEX_NONE)
+			{
+				// Skip. No bones affected in this surface, vertices won't be affected either.
+				continue;
+			}
+
+			const int32 FirstBoneIndexInBoneMap = Surface.BoneMapIndex;
+
+			const int32 VertexStart = Surface.SubMeshes[0].VertexBegin;
+			const int32 VertexEnd = Surface.SubMeshes.Last().VertexEnd;
+
+			for (int32 VertexIndex = VertexStart; VertexIndex < VertexEnd; ++VertexIndex)
+			{
+				int8 NumValidWeights = 0;
+
+				switch (BoneWeightsIterBegin.GetFormat())
+				{
+				case EMeshBufferFormat::NUInt8:
+				{
+					// Skin weights are ordered from highest to lowest. Break once the weight value falls below the threshold.
+					const uint8* Weights = (BoneWeightsIterBegin + VertexIndex).ptr();
+					for (; NumValidWeights < NumWeights && *Weights >= (uint8)MinWeightThreshold; ++NumValidWeights, ++Weights);
+					break;
+				}
+
+				case EMeshBufferFormat::NUInt16:
+				{
+					// Skin weights are ordered from highest to lowest. Break once the weight value falls below the threshold.
+					const uint16* Weights = reinterpret_cast<const uint16*>((BoneWeightsIterBegin + VertexIndex).ptr());
+					for (; NumValidWeights < NumWeights && *Weights >= MinWeightThreshold; ++NumValidWeights, ++Weights);
+					break;
+				}
+				default: unimplemented();
+				}
+
+				switch (BoneIndicesIterBegin.GetFormat())
+				{
+				case EMeshBufferFormat::UInt8:
+				{
+					const uint8* Indices = (BoneIndicesIterBegin + VertexIndex).ptr();
+					for (int32 Index = 0; Index < NumValidWeights; ++Index)
+					{
+						if (AffectedBoneMapIndices[FirstBoneIndexInBoneMap + *Indices])
+						{
+							VertexInBoneHierarchy[VertexIndex] = true;
+							break;
+						}
+
+						++Indices;
+					}
+					break;
+				}
+
+				case EMeshBufferFormat::UInt16:
+				{
+					const uint16* Indices = reinterpret_cast<const uint16*>((BoneIndicesIterBegin + VertexIndex).ptr());
+					for (int32 Index = 0; Index < NumValidWeights; ++Index)
+					{
+						if (AffectedBoneMapIndices[FirstBoneIndexInBoneMap + *Indices])
+						{
+							VertexInBoneHierarchy[VertexIndex] = true;
+							break;
+						}
+
+						++Indices;
+					}
+					break;
+				}
+				default: unimplemented();
+				}
+			}
+		}
+	}
+
+	// Get pointers to vertex position data
+	const MeshBufferIterator<EMeshBufferFormat::Float32, float, 3> PositionIterBegin(Mesh->VertexBuffers, EMeshBufferSemantic::Position, 0);
+
+	check(PositionIterBegin.ptr());
+	check(PositionIterBegin.GetFormat() == EMeshBufferFormat::Float32 && PositionIterBegin.GetComponents() == 3);
+
+	const FMatrix44f TransformInvT = Transform.Inverse().GetTransposed();
+
+	constexpr int32 NumVerticesPerBatch = 1 << 13;
+
+	// Tangent frame buffers are optional.
+	const UntypedMeshBufferIterator NormalIterBegin(Mesh->VertexBuffers, EMeshBufferSemantic::Normal, 0);
+	const UntypedMeshBufferIterator TangentIterBegin(Mesh->VertexBuffers, EMeshBufferSemantic::Tangent, 0);
+	const UntypedMeshBufferIterator BiNormalIterBegin(Mesh->VertexBuffers, EMeshBufferSemantic::Binormal, 0);
+
+	auto ProcessVertexBatch =
+		[
+			PositionIterBegin,
+				NormalIterBegin,
+				TangentIterBegin,
+				BiNormalIterBegin,
+				NumVertices,
+				NumVerticesPerBatch,
+				&VertexInBoneHierarchy,
+				Transform,
+				TransformInvT
+		](int32 BatchId)
+		{
+			const uint8 NumNormalComponents = FMath::Min(NormalIterBegin.GetComponents(), 3); // Due to quantization, the serialized component W may not be zero. Must be zero to avoid being affected by the transform position.
+			const uint8 NumTangentComponents = FMath::Min(TangentIterBegin.GetComponents(), 3); // Due to quantization, the serialized component W may not be zero. Must be zero to avoid being affected by the transform position.
+			const uint8 NumBiNormalComponents = FMath::Min(BiNormalIterBegin.GetComponents(), 3); // Due to quantization, the serialized component W may not be zero. Must be zero to avoid being affected by the transform position.
+
+			const EMeshBufferFormat NormalFormat = NormalIterBegin.GetFormat();
+			const EMeshBufferFormat TangentFormat = TangentIterBegin.GetFormat();
+
+			const int32 BatchBeginVertIndex = BatchId * NumVerticesPerBatch;
+			const int32 BatchEndVertIndex = FMath::Min(BatchBeginVertIndex + NumVerticesPerBatch, NumVertices);
+
+			const bool bHasOptimizedBuffers = NormalFormat == EMeshBufferFormat::PackedDirS8_W_TangentSign && TangentFormat == EMeshBufferFormat::PackedDirS8;
+
+			for (int32 VertexIndex = VertexInBoneHierarchy.FindFrom(true, BatchBeginVertIndex, BatchEndVertIndex); VertexIndex >= 0;)
+			{
+				const int32 AffectedSpanBegin = VertexIndex;
+				VertexIndex = VertexInBoneHierarchy.FindFrom(false, VertexIndex, BatchEndVertIndex);
+
+				// At the end of the buffer we may not find a false element, in that case
+				// FindForm returns INDEX_NONE, set the vertex at the range end. 
+				VertexIndex = VertexIndex >= 0 && VertexIndex < BatchEndVertIndex ? VertexIndex : BatchEndVertIndex;
+				const int32 AffectedSpanEnd = VertexIndex;
+
+				// VertexIndex may be one past the end of the array, VertexIndex will become INDEX_NONE
+				// and the loop will finish.
+				VertexIndex = VertexInBoneHierarchy.FindFrom(true, VertexIndex, BatchEndVertIndex);
+				VertexIndex = VertexIndex < BatchEndVertIndex ? VertexIndex : INDEX_NONE;
+
+				MeshBufferIterator<EMeshBufferFormat::Float32, float, 3> PositionIter = (PositionIterBegin + AffectedSpanBegin);
+				for (int32 Index = AffectedSpanBegin; Index < AffectedSpanEnd; ++Index)
+				{
+					FVector3f* Position = reinterpret_cast<FVector3f*>(PositionIter.ptr());
+					*Position = Transform.TransformFVector4(*Position);
+					PositionIter++;
+				}
+
+				if (bHasOptimizedBuffers)
+				{
+					UntypedMeshBufferIterator TangentIter = TangentIterBegin + AffectedSpanBegin;
+					UntypedMeshBufferIterator NormalIter = NormalIterBegin + AffectedSpanBegin;
+
+					for (int32 Index = AffectedSpanBegin; Index < AffectedSpanEnd; ++Index)
+					{
+						// Tangents
+						FPackedNormal* PackedTangent = reinterpret_cast<FPackedNormal*>(TangentIter.ptr());
+						FVector4f Tangent = TransformInvT.TransformFVector4(PackedTangent->ToFVector3f());
+						*PackedTangent = *reinterpret_cast<FVector3f*>(&Tangent);
+						TangentIter++;
+
+						// Normals
+						FPackedNormal* PackedNormal = reinterpret_cast<FPackedNormal*>(NormalIter.ptr());
+						int8 W = PackedNormal->Vector.W;
+						FVector4f Normal = TransformInvT.TransformFVector4(PackedNormal->ToFVector3f());
+						*PackedNormal = *reinterpret_cast<FVector3f*>(&Normal);
+						PackedNormal->Vector.W = W;
+						NormalIter++;
+					}
+				}
+				else
+				{
+					MUTABLE_CPUPROFILER_SCOPE(MeshTransform_Vertices_SlowPath);
+					if (NormalIterBegin.ptr())
+					{
+						UntypedMeshBufferIterator it = NormalIterBegin + AffectedSpanBegin;
+						for (int32 Index = AffectedSpanBegin; Index < AffectedSpanEnd; ++Index)
+						{
+							FVector4f value(0.0f, 0.0f, 0.0f, 0.0f);
+
+							for (uint8 i = 0; i < NumNormalComponents; ++i)
+							{
+								ConvertData(i, &value[0], EMeshBufferFormat::Float32, it.ptr(), it.GetFormat());
+							}
+
+							value = TransformInvT.TransformFVector4(value);
+
+							// Notice that 4th component is not modified.
+							for (uint8 i = 0; i < NumNormalComponents; ++i)
+							{
+								ConvertData(i, it.ptr(), it.GetFormat(), &value[0], EMeshBufferFormat::Float32);
+							}
+							it++;
+						}
+					}
+
+					if (TangentIterBegin.ptr())
+					{
+						UntypedMeshBufferIterator it = TangentIterBegin + AffectedSpanBegin;
+						for (int32 Index = AffectedSpanBegin; Index < AffectedSpanEnd; ++Index)
+						{
+							FVector4f value(0.0f, 0.0f, 0.0f, 0.0f);
+
+							for (uint8 i = 0; i < NumTangentComponents; ++i)
+							{
+								ConvertData(i, &value[0], EMeshBufferFormat::Float32, it.ptr(), it.GetFormat());
+							}
+
+							value = TransformInvT.TransformFVector4(value);
+
+							// Notice that 4th component is not modified.
+							for (uint8 i = 0; i < NumTangentComponents; ++i)
+							{
+								ConvertData(i, it.ptr(), it.GetFormat(), &value[0], EMeshBufferFormat::Float32);
+							}
+							it++;
+						}
+					}
+
+					if (BiNormalIterBegin.ptr())
+					{
+						UntypedMeshBufferIterator it = BiNormalIterBegin + AffectedSpanBegin;
+						for (int32 Index = AffectedSpanBegin; Index < AffectedSpanEnd; ++Index)
+						{
+							FVector4f value(0.0f, 0.0f, 0.0f, 0.0f);
+
+							for (uint8 i = 0; i < NumBiNormalComponents; ++i)
+							{
+								ConvertData(i, &value[0], EMeshBufferFormat::Float32, it.ptr(), it.GetFormat());
+							}
+
+							value = TransformInvT.TransformFVector4(value);
+
+							// Notice that 4th component is not modified.
+							for (uint8 i = 0; i < NumBiNormalComponents; ++i)
+							{
+								ConvertData(i, it.ptr(), it.GetFormat(), &value[0], EMeshBufferFormat::Float32);
+							}
+							it++;
+						}
+					}
+				}
+			}
+		};
+
+	const int32 NumBatches = FMath::DivideAndRoundUp<int32>(NumVertices, NumVerticesPerBatch);
+	ParallelExecutionUtils::InvokeBatchParallelFor(NumBatches, ProcessVertexBatch);
+}
+
+}
diff --git a/Engine/Plugins/Mutable/Source/MutableRuntime/Private/MuR/OpMeshTransformWithBone.h b/Engine/Plugins/Mutable/Source/MutableRuntime/Private/MuR/OpMeshTransformWithBone.h
new file mode 100644
--- /dev/null
+++ b/Engine/Plugins/Mutable/Source/MutableRuntime/Private/MuR/OpMeshTransformWithBone.h
@@ -0,0 +1,15 @@
+﻿// Copyright Epic Games, Inc. All Rights Reserved.
+
+#pragma once
+
+#include "Math/MathFwd.h"
+
+namespace mu
+{
+class FMesh;
+struct FBoneName;
+
+/**  */
+extern void MeshTransformWithBoneInline(FMesh* Mesh, const FMatrix44f& Transform, const FBoneName& BoneName, const float Threshold);
+
+}
diff --git a/Engine/Plugins/Mutable/Source/MutableRuntime/Private/MuR/Operations.cpp b/Engine/Plugins/Mutable/Source/MutableRuntime/Private/MuR/Operations.cpp
--- a/Engine/Plugins/Mutable/Source/MutableRuntime/Private/MuR/Operations.cpp
+++ b/Engine/Plugins/Mutable/Source/MutableRuntime/Private/MuR/Operations.cpp
@@ -128,6 +128,7 @@ namespace mu
 		{ EDataType::Mesh,			},	// ME_OPTIMIZESKINNING
 		{ EDataType::Mesh,			},	// ME_ADDMETADATA
 		{ EDataType::Mesh,			},	// ME_TRANSFORMWITHMESH
+		{ EDataType::Mesh,			},	// ME_TRANSFORMWITHBONE
 
 		{ EDataType::Instance,		},	// IN_ADDMESH
 		{ EDataType::Instance,		},	// IN_ADDIMAGE
@@ -793,6 +794,14 @@ namespace mu
         	f(args.matrix);
         	break;
         }
+
+		case EOpType::ME_TRANSFORMWITHBONE:
+		{
+			OP::MeshTransformWithBoneArgs args = program.GetOpArgs<OP::MeshTransformWithBoneArgs>(at);
+			f(args.SourceMesh);
+			f(args.Matrix);
+			break;
+		}
         	
         //-------------------------------------------------------------------------------------
         case EOpType::IN_ADDMESH:
diff --git a/Engine/Plugins/Mutable/Source/MutableTools/Private/MuT/ASTOpMeshFormat.cpp b/Engine/Plugins/Mutable/Source/MutableTools/Private/MuT/ASTOpMeshFormat.cpp
--- a/Engine/Plugins/Mutable/Source/MutableTools/Private/MuT/ASTOpMeshFormat.cpp
+++ b/Engine/Plugins/Mutable/Source/MutableTools/Private/MuT/ASTOpMeshFormat.cpp
@@ -15,6 +15,7 @@
 #include "MuT/ASTOpConstantResource.h"
 #include "MuT/ASTOpMeshClipMorphPlane.h"
 #include "MuT/ASTOpMeshTransformWithBoundingMesh.h"
+#include "MuT/ASTOpMeshTransformWithBone.h"
 #include "MuT/ASTOpMeshRemoveMask.h"
 #include "MuT/ASTOpMeshMorph.h"
 #include "MuT/ASTOpMeshAddMetadata.h"
@@ -337,6 +338,14 @@ mu::Ptr<ASTOp> Sink_MeshFormatAST::Visit(const mu::Ptr<ASTOp>& at, const ASTOpMe
 		break;
 	}
 
+	case EOpType::ME_TRANSFORMWITHBONE:
+	{
+		Ptr<ASTOpMeshTransformWithBone> NewOp = mu::Clone<ASTOpMeshTransformWithBone>(at);
+		NewOp->SourceMesh = Visit(NewOp->SourceMesh.child(), currentFormatOp);
+		newAt = NewOp;
+		break;
+	}
+
 	case EOpType::ME_MORPH:
 	{
 		// Move the format down the base of the morph
diff --git a/Engine/Plugins/Mutable/Source/MutableTools/Private/MuT/ASTOpMeshTransformWithBone.cpp b/Engine/Plugins/Mutable/Source/MutableTools/Private/MuT/ASTOpMeshTransformWithBone.cpp
new file mode 100644
--- /dev/null
+++ b/Engine/Plugins/Mutable/Source/MutableTools/Private/MuT/ASTOpMeshTransformWithBone.cpp
@@ -0,0 +1,94 @@
+﻿// Copyright Epic Games, Inc. All Rights Reserved.
+
+#include "MuT/ASTOpMeshTransformWithBone.h"
+
+#include "ASTOpImageTransform.h"
+
+
+mu::ASTOpMeshTransformWithBone::ASTOpMeshTransformWithBone()
+	: SourceMesh(this)
+	, Matrix(this)
+{
+}
+
+mu::ASTOpMeshTransformWithBone::~ASTOpMeshTransformWithBone()
+{
+	// Explicit call needed to avoid recursive destruction
+	ASTOp::RemoveChildren();
+}
+
+uint64 mu::ASTOpMeshTransformWithBone::Hash() const
+{
+	uint64 res = std::hash<EOpType>()(GetOpType());
+	hash_combine(res, SourceMesh.child());
+	hash_combine(res, Matrix.child());
+	hash_combine(res, BoneName.Id);
+	hash_combine(res, ThresholdFactor);
+
+	return res;
+}
+
+bool mu::ASTOpMeshTransformWithBone::IsEqual(const ASTOp& OtherUntyped) const
+{
+	if (GetOpType() == OtherUntyped.GetOpType())
+	{
+		const ASTOpMeshTransformWithBone& Other = static_cast<const ASTOpMeshTransformWithBone&>(OtherUntyped);
+		return SourceMesh == Other.SourceMesh
+			&& BoneName == Other.BoneName
+			&& Matrix == Other.Matrix
+			&& ThresholdFactor == Other.ThresholdFactor;
+	}
+	return false;
+}
+
+mu::Ptr<mu::ASTOp> mu::ASTOpMeshTransformWithBone::Clone(MapChildFuncRef MapChild) const
+{
+	Ptr<ASTOpMeshTransformWithBone> NewASTOp = new ASTOpMeshTransformWithBone();
+	NewASTOp->SourceMesh = MapChild(SourceMesh.child());
+	NewASTOp->Matrix = MapChild(Matrix.child());
+	NewASTOp->BoneName = BoneName;
+	NewASTOp->ThresholdFactor = ThresholdFactor;
+	return NewASTOp;
+}
+
+void mu::ASTOpMeshTransformWithBone::ForEachChild(const TFunctionRef<void(ASTChild&)> Function)
+{
+	Function(SourceMesh);
+	Function(Matrix);
+}
+
+void mu::ASTOpMeshTransformWithBone::Link(FProgram& Program, FLinkerOptions* Options)
+{
+	if (!linkedAddress)
+	{
+		OP::MeshTransformWithBoneArgs Args;
+		FMemory::Memzero(Args);
+		if (SourceMesh)
+		{
+			Args.SourceMesh = SourceMesh->linkedAddress;
+		}
+
+		if (Matrix)
+		{
+			Args.Matrix = Matrix->linkedAddress;
+		}
+
+		Args.BoneId = BoneName.Id;
+		Args.ThresholdFactor = ThresholdFactor;
+
+		linkedAddress = static_cast<OP::ADDRESS>(Program.OpAddress.Num());
+		Program.OpAddress.Add(Program.ByteCode.Num());
+		AppendCode(Program.ByteCode, EOpType::ME_TRANSFORMWITHBONE);
+		AppendCode(Program.ByteCode, Args);
+	}
+}
+
+mu::FSourceDataDescriptor mu::ASTOpMeshTransformWithBone::GetSourceDataDescriptor(FGetSourceDataDescriptorContext* Context) const
+{
+	if (SourceMesh)
+	{
+		return SourceMesh->GetSourceDataDescriptor(Context);
+	}
+
+	return {};
+}
diff --git a/Engine/Plugins/Mutable/Source/MutableTools/Private/MuT/ASTOpMeshTransformWithBone.h b/Engine/Plugins/Mutable/Source/MutableTools/Private/MuT/ASTOpMeshTransformWithBone.h
new file mode 100644
--- /dev/null
+++ b/Engine/Plugins/Mutable/Source/MutableTools/Private/MuT/ASTOpMeshTransformWithBone.h
@@ -0,0 +1,44 @@
+﻿// Copyright Epic Games, Inc. All Rights Reserved.
+
+#pragma once
+
+#include "HAL/Platform.h"
+#include "MuR/Operations.h"
+#include "MuR/Ptr.h"
+#include "MuR/Skeleton.h"
+#include "MuT/AST.h"
+
+
+namespace mu
+{
+struct FProgram;
+
+
+//---------------------------------------------------------------------------------------------
+//!
+//---------------------------------------------------------------------------------------------
+class ASTOpMeshTransformWithBone final : public ASTOp
+{
+public:
+	ASTChild SourceMesh;
+	ASTChild Matrix;
+	FBoneName BoneName;
+	float ThresholdFactor;
+public:
+
+	ASTOpMeshTransformWithBone();
+	ASTOpMeshTransformWithBone(const ASTOpMeshTransformWithBone&) = delete;
+	~ASTOpMeshTransformWithBone() override;
+
+	EOpType GetOpType() const override { return EOpType::ME_TRANSFORMWITHBONE; }
+	uint64 Hash() const override;
+	bool IsEqual(const ASTOp& otherUntyped) const override;
+	Ptr<ASTOp> Clone(MapChildFuncRef mapChild) const override;
+	void ForEachChild(const TFunctionRef<void(ASTChild&)>) override;
+	void Link(FProgram& program, FLinkerOptions* Options) override;
+	virtual FSourceDataDescriptor GetSourceDataDescriptor(FGetSourceDataDescriptorContext*) const override;
+
+};
+
+}
+
diff --git a/Engine/Plugins/Mutable/Source/MutableTools/Private/MuT/CodeGenerator.cpp b/Engine/Plugins/Mutable/Source/MutableTools/Private/MuT/CodeGenerator.cpp
--- a/Engine/Plugins/Mutable/Source/MutableTools/Private/MuT/CodeGenerator.cpp
+++ b/Engine/Plugins/Mutable/Source/MutableTools/Private/MuT/CodeGenerator.cpp
@@ -2,7 +2,6 @@
 
 #include "MuT/CodeGenerator.h"
 
-#include "ASTOpMeshTransformWithBoundingMesh.h"
 #include "Containers/Array.h"
 #include "Logging/LogCategory.h"
 #include "Logging/LogMacros.h"
@@ -45,6 +44,8 @@
 #include "MuT/ASTOpMeshApplyLayout.h"
 #include "MuT/ASTOpMeshMerge.h"
 #include "MuT/ASTOpMeshMaskDiff.h"
+#include "MuT/ASTOpMeshTransformWithBoundingMesh.h"
+#include "MuT/ASTOpMeshTransformWithBone.h"
 #include "MuT/ASTOpParameter.h"
 #include "MuT/ASTOpLayoutRemoveBlocks.h"
 #include "MuT/ASTOpLayoutFromMesh.h"
@@ -80,6 +81,7 @@
 #include "MuT/NodeModifierMeshClipWithMesh.h"
 #include "MuT/NodeModifierMeshClipWithUVMask.h"
 #include "MuT/NodeModifierMeshTransformInMesh.h"
+#include "MuT/NodeModifierMeshTransformWithBone.h"
 #include "MuT/NodeModifierSurfaceEdit.h"
 #include "MuT/NodeObject.h"
 #include "MuT/NodeObjectNew.h"
@@ -2925,7 +2927,7 @@ namespace mu
 					BoundingMeshDynamicOptions.bLayouts = false;
 					FMeshTask BoundingMeshTask = GenerateMesh(ClipStaticOptions, UE::Tasks::MakeCompletedTask<FMeshGenerationDynamicOptions>(BoundingMeshDynamicOptions), TypedTransformNode->BoundingMesh);
 
-					LastMeshTask = UE::Tasks::Launch(TEXT("MutableMeshMergeModifier"),
+					LastMeshTask = UE::Tasks::Launch(TEXT("MutableMeshTransformWithBoundingMeshModifier"),
 						[BoundingMeshTask, LastMeshTask, Modifier, MatrixOp, ErrorLog = ErrorLog, ErrorContext]() mutable
 						{
 							FMeshGenerationResult BoundingResults = BoundingMeshTask.GetResult();
@@ -2966,6 +2968,58 @@ namespace mu
 					);
 				}
 			}
+			else if (Modifier.Node->GetType() == NodeModifierMeshTransformWithBone::GetStaticType())
+			{
+				const NodeModifierMeshTransformWithBone* TypedTransformNode = static_cast<const NodeModifierMeshTransformWithBone*>(Modifier.Node);
+
+				// If a matrix node is not connected, the op won't do anything, so let's not create it at all.
+				if (TypedTransformNode->MatrixNode)
+				{
+					if (!TypedTransformNode->MatrixNode)
+					{
+						ErrorLog->Add("Matrix have not been generated", ELMT_ERROR, ErrorContext);
+						continue;
+					}
+
+					// Transform matrix.
+					Ptr<ASTOp> MatrixOp;
+					FMatrixGenerationResult ChildResult;
+					GenerateMatrix(ChildResult, StaticOptions, TypedTransformNode->MatrixNode);
+					MatrixOp = ChildResult.op;
+
+
+					LastMeshTask = UE::Tasks::Launch(TEXT("MutableMeshTransformWithBoneModifier"),
+						[LastMeshTask, Modifier, MatrixOp, BoneName = TypedTransformNode->BoneName, ThresholdFactor = TypedTransformNode->ThresholdFactor, ErrorLog = ErrorLog, ErrorContext]() mutable
+						{
+							FMeshGenerationResult LastMeshResults = LastMeshTask.GetResult();
+
+							Ptr<ASTOpMeshTransformWithBone> TransformOp = new ASTOpMeshTransformWithBone();
+							TransformOp->SourceMesh = LastMeshResults.MeshOp;
+							TransformOp->Matrix = MatrixOp;
+							TransformOp->BoneName = BoneName;
+							TransformOp->ThresholdFactor = ThresholdFactor;
+
+							// Condition to apply the transform op
+							if (Modifier.FinalCondition)
+							{
+								Ptr<ASTOpConditional> ConditionalOp = new ASTOpConditional();
+								ConditionalOp->type = EOpType::ME_CONDITIONAL;
+								ConditionalOp->no = LastMeshResults.MeshOp;
+								ConditionalOp->yes = TransformOp;
+								ConditionalOp->condition = Modifier.FinalCondition;
+								LastMeshResults.MeshOp = ConditionalOp;
+							}
+							else
+							{
+								LastMeshResults.MeshOp = TransformOp;
+							}
+
+							return LastMeshResults;
+						},
+						UE::Tasks::Prerequisites(LastMeshTask)
+					);
+				}
+			}
 		}
 		
 
diff --git a/Engine/Plugins/Mutable/Source/MutableTools/Private/MuT/Compiler.cpp b/Engine/Plugins/Mutable/Source/MutableTools/Private/MuT/Compiler.cpp
--- a/Engine/Plugins/Mutable/Source/MutableTools/Private/MuT/Compiler.cpp
+++ b/Engine/Plugins/Mutable/Source/MutableTools/Private/MuT/Compiler.cpp
@@ -167,6 +167,7 @@ namespace mu
 		{ true,		{ 0, 0, 0, 0, 0, 0, 0, 0, 0, 0, 0, 0, 0, 0, 0, 0, 0, 0, 0, 0, 0, 0, 0, 0, 0, 0, 0, 0, 0 }		},	// ME_OPTIMIZESKINNING
 		{ true,		{ 0, 0, 0, 0, 0, 0, 0, 0, 0, 0, 0, 0, 0, 0, 0, 0, 0, 0, 0, 0, 0, 0, 0, 0, 0, 0, 0, 0, 0 }		},	// ME_ADDMETADATA
 		{ true,		{ 0, 0, 0, 0, 0, 0, 0, 0, 0, 0, 0, 0, 0, 0, 0, 0, 0, 0, 0, 0, 0, 0, 0, 0, 0, 0, 0, 0, 0 }		},	// ME_TRANSFORMWITHMESH
+		{ true,		{ 0, 0, 0, 0, 0, 0, 0, 0, 0, 0, 0, 0, 0, 0, 0, 0, 0, 0, 0, 0, 0, 0, 0, 0, 0, 0, 0, 0, 0 }		},	// ME_TRANSFORMWITHBONE
 
 		{ false,	{ 0, 0, 0, 0, 0, 0, 0, 0, 0, 0, 0, 0, 0, 0, 0, 0, 0, 0, 0, 0, 0, 0, 0, 0, 0, 0, 0, 0, 0 }		},	// IN_ADDMESH
 		{ false,	{ 0, 0, 0, 0, 0, 0, 0, 0, 0, 0, 0, 0, 0, 0, 0, 0, 0, 0, 0, 0, 0, 0, 0, 0, 0, 0, 0, 0, 0 }		},	// IN_ADDIMAGE
diff --git a/Engine/Plugins/Mutable/Source/MutableTools/Private/MuT/ErrorLog.cpp b/Engine/Plugins/Mutable/Source/MutableTools/Private/MuT/ErrorLog.cpp
--- a/Engine/Plugins/Mutable/Source/MutableTools/Private/MuT/ErrorLog.cpp
+++ b/Engine/Plugins/Mutable/Source/MutableTools/Private/MuT/ErrorLog.cpp
@@ -312,6 +312,7 @@ namespace mu
 		TEXT("ME_OPTIMIZESKIN  "),
 		TEXT("ME_ADDMETADATA   "),
 		TEXT("ME_TRANSFORMWITHMESH"),
+		TEXT("ME_TRANSFORMWITHBONE"),
 
 		TEXT("IN_ADDMESH       "),
 		TEXT("IN_ADDIMAGE      "),
diff --git a/Engine/Plugins/Mutable/Source/MutableTools/Private/MuT/Node.cpp b/Engine/Plugins/Mutable/Source/MutableTools/Private/MuT/Node.cpp
--- a/Engine/Plugins/Mutable/Source/MutableTools/Private/MuT/Node.cpp
+++ b/Engine/Plugins/Mutable/Source/MutableTools/Private/MuT/Node.cpp
@@ -82,6 +82,7 @@
 #include "MuT/NodeModifierMeshClipWithMesh.h"
 #include "MuT/NodeModifierMeshClipWithUVMask.h"
 #include "MuT/NodeModifierMeshTransformInMesh.h"
+#include "MuT/NodeModifierMeshTransformWithBone.h"
 #include "MuT/NodeModifierSurfaceEdit.h"
 #include "MuT/NodeMatrix.h"
 #include "MuT/NodeMatrixConstant.h"
@@ -201,6 +202,7 @@ namespace mu
 	FNodeType NodeModifierMeshClipWithMesh::StaticType = FNodeType(Node::EType::ModifierMeshClipWithMesh, NodeModifier::GetStaticType());
 	FNodeType NodeModifierMeshClipWithUVMask::StaticType = FNodeType(Node::EType::ModifierMeshClipWithUVMask, NodeModifier::GetStaticType());
 	FNodeType NodeModifierMeshTransformInMesh::StaticType = FNodeType(Node::EType::ModifierTransformInMesh, NodeModifier::GetStaticType());
+	FNodeType NodeModifierMeshTransformWithBone::StaticType = FNodeType(Node::EType::ModifierTransformWithBone, NodeModifier::GetStaticType());
 	FNodeType NodeModifierSurfaceEdit::StaticType = FNodeType(Node::EType::ModifierSurfaceEdit, NodeSurface::GetStaticType());
 
 	FNodeType NodeMatrix::StaticType = FNodeType(Node::EType::Matrix, Node::GetStaticType());
diff --git a/Engine/Plugins/Mutable/Source/MutableTools/Public/MuT/Node.h b/Engine/Plugins/Mutable/Source/MutableTools/Public/MuT/Node.h
--- a/Engine/Plugins/Mutable/Source/MutableTools/Public/MuT/Node.h
+++ b/Engine/Plugins/Mutable/Source/MutableTools/Public/MuT/Node.h
@@ -144,6 +144,7 @@ namespace mu
 			ModifierMeshClipWithUVMask,
 			ModifierSurfaceEdit,
 			ModifierTransformInMesh,
+			ModifierTransformWithBone,
 
 			ExtensionData,
 			ExtensionDataConstant,
diff --git a/Engine/Plugins/Mutable/Source/MutableTools/Public/MuT/NodeModifierMeshTransformWithBone.h b/Engine/Plugins/Mutable/Source/MutableTools/Public/MuT/NodeModifierMeshTransformWithBone.h
new file mode 100644
--- /dev/null
+++ b/Engine/Plugins/Mutable/Source/MutableTools/Public/MuT/NodeModifierMeshTransformWithBone.h
@@ -0,0 +1,49 @@
+// Copyright Epic Games, Inc. All Rights Reserved.
+
+#pragma once
+
+#include "MuR/Ptr.h"
+#include "MuR/Types.h"
+#include "MuR/Skeleton.h"
+#include "MuT/Node.h"
+#include "MuT/NodeModifier.h"
+#include "MuT/NodeMatrix.h"
+
+#define UE_API MUTABLETOOLS_API
+
+namespace mu
+{
+
+/** */
+class NodeModifierMeshTransformWithBone : public NodeModifier
+{
+public:
+
+	/* */
+	FBoneName BoneName;
+	float ThresholdFactor;
+
+	/** */
+	Ptr<NodeMatrix> MatrixNode;
+
+
+public:
+
+	// Node interface
+	virtual const FNodeType* GetType() const override { return &StaticType; }
+	static const FNodeType* GetStaticType() { return &StaticType; }
+
+protected:
+
+	/** Forbidden. Manage with the Ptr<> template. */
+	virtual ~NodeModifierMeshTransformWithBone() override {}
+
+private:
+
+	static UE_API FNodeType StaticType;
+
+};
+
+}
+
+#undef UE_API
