[Mutable] Add Passthrough to Node Schema.

- Fixed color picker closing automatically.

diff --git a/Engine/Plugins/Mutable/Source/CustomizableObjectEditor/Private/MuCOE/EdGraphSchema_CustomizableObject.cpp b/Engine/Plugins/Mutable/Source/CustomizableObjectEditor/Private/MuCOE/EdGraphSchema_CustomizableObject.cpp
--- a/Engine/Plugins/Mutable/Source/CustomizableObjectEditor/Private/MuCOE/EdGraphSchema_CustomizableObject.cpp
+++ b/Engine/Plugins/Mutable/Source/CustomizableObjectEditor/Private/MuCOE/EdGraphSchema_CustomizableObject.cpp
@@ -1596,4 +1596,11 @@ FText UEdGraphSchema_CustomizableObject::GetPinCategoryFriendlyName(const FName&
 }
 
 
+bool UEdGraphSchema_CustomizableObject::IsPassthrough(const FName& PinCategory)
+{
+	return PinCategory == PC_PassThroughImage ||
+		PinCategory == PC_PassThroughMesh;
+}
+
+
 #undef LOCTEXT_NAMESPACE
diff --git a/Engine/Plugins/Mutable/Source/CustomizableObjectEditor/Private/MuCOE/Nodes/CONodeSchema.cpp b/Engine/Plugins/Mutable/Source/CustomizableObjectEditor/Private/MuCOE/Nodes/CONodeSchema.cpp
--- a/Engine/Plugins/Mutable/Source/CustomizableObjectEditor/Private/MuCOE/Nodes/CONodeSchema.cpp
+++ b/Engine/Plugins/Mutable/Source/CustomizableObjectEditor/Private/MuCOE/Nodes/CONodeSchema.cpp
@@ -2,6 +2,7 @@
 
 #include "MuCOE/Nodes/CONodeSchema.h"
 
+#include "MuCOE/CustomizableObjectGraph.h"
 #include "MuCOE/RemapPins/CustomizableObjectNodeRemapPinsByPosition.h"
 
 
@@ -31,6 +32,7 @@ void UCONodeSchema::AllocateDefaultPins(UCustomizableObjectNodeRemapPins* RemapP
 
 		UCONodeSchemaPinData* PinData = NewObject<UCONodeSchemaPinData>(this);
 		PinData->Color = PinSchema.Color;
+		PinData->bIsPassthrough = PinSchema.bIsPassthrough;
 		
 		AddPinData(*Pin, *PinData);
 	}
@@ -43,6 +45,12 @@ FLinearColor UCONodeSchema::GetPinColor(const UEdGraphPin& Pin) const
 }
 
 
+bool UCONodeSchema::IsPassthrough(const UEdGraphPin& Pin) const
+{
+	return GetPinData<UCONodeSchemaPinData>(Pin).bIsPassthrough;
+}
+
+
 UCustomizableObjectNodeRemapPins* UCONodeSchema::CreateRemapPinsDefault() const
 {
 	return NewObject<UCustomizableObjectNodeRemapPinsByPosition>();
@@ -55,6 +63,10 @@ void UCONodeSchema::PostEditChangeProperty(FPropertyChangedEvent& PropertyChange
 
 	if (PropertyChangedEvent.Property)
 	{
-		ReconstructNode();
+		if (PropertyChangedEvent.Property->GetFName() != GET_MEMBER_NAME_CHECKED(FPinSchema, Color) &&
+			PropertyChangedEvent.Property->GetFName() != GET_MEMBER_NAME_CHECKED(UCONodeSchema, TitleColor))
+		{
+			ReconstructNode();
+		}
 	}
 }
diff --git a/Engine/Plugins/Mutable/Source/CustomizableObjectEditor/Private/MuCOE/Nodes/CONodeSchema.h b/Engine/Plugins/Mutable/Source/CustomizableObjectEditor/Private/MuCOE/Nodes/CONodeSchema.h
--- a/Engine/Plugins/Mutable/Source/CustomizableObjectEditor/Private/MuCOE/Nodes/CONodeSchema.h
+++ b/Engine/Plugins/Mutable/Source/CustomizableObjectEditor/Private/MuCOE/Nodes/CONodeSchema.h
@@ -18,6 +18,9 @@ class UCONodeSchemaPinData : public UCustomizableObjectNodePinData
 public:	
 	UPROPERTY()
 	FLinearColor Color;
+
+	UPROPERTY()
+	bool bIsPassthrough = false;
 };
 
 
@@ -34,9 +37,9 @@ struct FPinSchema
 
 	UPROPERTY(EditAnywhere, Category = Schema)
 	bool bIsArray = false;
-
+	
 	UPROPERTY(EditAnywhere, Category = Schema)
-	bool bEdit = false;
+	bool bIsPassthrough = false;
 	
 	UPROPERTY(EditAnywhere, Category = Schema)
 	FLinearColor Color = FLinearColor::White;
@@ -59,6 +62,7 @@ public:
 	// UCustomizableObjectNode interface
 	virtual void AllocateDefaultPins(UCustomizableObjectNodeRemapPins* RemapPins) override;
 	virtual FLinearColor GetPinColor(const UEdGraphPin &Pin) const override;
+	virtual bool IsPassthrough(const UEdGraphPin& Pin) const override;
 
 protected:
 	virtual UCustomizableObjectNodeRemapPins* CreateRemapPinsDefault() const override;
diff --git a/Engine/Plugins/Mutable/Source/CustomizableObjectEditor/Private/MuCOE/Nodes/CustomizableObjectNode.cpp b/Engine/Plugins/Mutable/Source/CustomizableObjectEditor/Private/MuCOE/Nodes/CustomizableObjectNode.cpp
--- a/Engine/Plugins/Mutable/Source/CustomizableObjectEditor/Private/MuCOE/Nodes/CustomizableObjectNode.cpp
+++ b/Engine/Plugins/Mutable/Source/CustomizableObjectEditor/Private/MuCOE/Nodes/CustomizableObjectNode.cpp
@@ -761,6 +761,13 @@ FLinearColor UCustomizableObjectNode::GetPinColor(const UEdGraphPin& Pin) const
 }
 
 
+bool UCustomizableObjectNode::IsPassthrough(const UEdGraphPin& Pin) const
+{
+	const UEdGraphSchema_CustomizableObject* Schema = GetDefault<UEdGraphSchema_CustomizableObject>();
+	return Schema->IsPassthrough(Pin.PinType.PinCategory);
+}
+
+
 bool UCustomizableObjectNode::HasPinViewer() const
 {
 	return false;
diff --git a/Engine/Plugins/Mutable/Source/CustomizableObjectEditor/Private/MuCOE/Nodes/SCustomizableObjectNodePin.cpp b/Engine/Plugins/Mutable/Source/CustomizableObjectEditor/Private/MuCOE/Nodes/SCustomizableObjectNodePin.cpp
--- a/Engine/Plugins/Mutable/Source/CustomizableObjectEditor/Private/MuCOE/Nodes/SCustomizableObjectNodePin.cpp
+++ b/Engine/Plugins/Mutable/Source/CustomizableObjectEditor/Private/MuCOE/Nodes/SCustomizableObjectNodePin.cpp
@@ -19,16 +19,27 @@ void SCustomizableObjectNodePin::Construct(const FArguments& InArgs, UEdGraphPin
 
 const FSlateBrush* SCustomizableObjectNodePin::GetPinIcon() const
 {
-	if (GraphPinObj->PinType.PinCategory == UEdGraphSchema_CustomizableObject::PC_PassThroughImage)
-	{
-		return GraphPinObj->LinkedTo.Num() ?
-			PassThroughImageConnected :
-			PassThroughImageDisconnected;
-	}
-	else
+	UEdGraphPin* GraphPin = GetPinObj();
+	if (GraphPin && !GraphPin->IsPendingKill())
 	{
-		return SGraphPin::GetPinIcon();
+		UCustomizableObjectNode* Node = CastChecked<UCustomizableObjectNode>(GetPinObj()->GetOwningNode());
+		if (Node)
+		{
+			if (Node->IsPassthrough(*GraphPin))
+			{
+				if (GraphPinObj->LinkedTo.IsEmpty())
+				{
+					return PassThroughImageDisconnected;
+				}
+				else
+				{
+					return PassThroughImageConnected;
+				}
+			}
+		}
 	}
+	
+	return SGraphPin::GetPinIcon();
 }
 
 
@@ -47,15 +58,17 @@ FSlateColor SCustomizableObjectNodePin::GetPinColor() const
 		}
 		
 		UCustomizableObjectNode* Node = CastChecked<UCustomizableObjectNode>(GetPinObj()->GetOwningNode());
-
-		const FLinearColor Color = Node->GetPinColor(*GraphPin);
-			
-		if (!Node->IsNodeEnabled() || Node->IsDisplayAsDisabledForced() || !IsEditingEnabled() || Node->IsNodeUnrelated())
+		if (Node)
 		{
-			return Color * FLinearColor(1.0f, 1.0f, 1.0f, 0.5f);
-		}
+			const FLinearColor Color = Node->GetPinColor(*GraphPin);
+			
+			if (!Node->IsNodeEnabled() || Node->IsDisplayAsDisabledForced() || !IsEditingEnabled() || Node->IsNodeUnrelated())
+			{
+				return Color * FLinearColor(1.0f, 1.0f, 1.0f, 0.5f);
+			}
 
-		return Color * PinColorModifier;
+			return Color * PinColorModifier;
+		}
 	}
 
 	return FLinearColor::White;
diff --git a/Engine/Plugins/Mutable/Source/CustomizableObjectEditor/Public/MuCOE/EdGraphSchema_CustomizableObject.h b/Engine/Plugins/Mutable/Source/CustomizableObjectEditor/Public/MuCOE/EdGraphSchema_CustomizableObject.h
--- a/Engine/Plugins/Mutable/Source/CustomizableObjectEditor/Public/MuCOE/EdGraphSchema_CustomizableObject.h
+++ b/Engine/Plugins/Mutable/Source/CustomizableObjectEditor/Public/MuCOE/EdGraphSchema_CustomizableObject.h
@@ -76,6 +76,7 @@ public:
 	/** Given a pin category name, get its friendly name (user readable). */
 	static UE_API FName GetPinCategoryName(const FName& PinCategory);
 	static UE_API FText GetPinCategoryFriendlyName(const FName& PinCategory);
+	static UE_API bool IsPassthrough(const FName& PinCategory);
 };
 
 #undef UE_API
diff --git a/Engine/Plugins/Mutable/Source/CustomizableObjectEditor/Public/MuCOE/Nodes/CustomizableObjectNode.h b/Engine/Plugins/Mutable/Source/CustomizableObjectEditor/Public/MuCOE/Nodes/CustomizableObjectNode.h
--- a/Engine/Plugins/Mutable/Source/CustomizableObjectEditor/Public/MuCOE/Nodes/CustomizableObjectNode.h
+++ b/Engine/Plugins/Mutable/Source/CustomizableObjectEditor/Public/MuCOE/Nodes/CustomizableObjectNode.h
@@ -271,6 +271,8 @@ public:
 
 	UE_API virtual FLinearColor GetPinColor(const UEdGraphPin &Pin) const;
 	
+	UE_API virtual bool IsPassthrough(const UEdGraphPin &Pin) const;
+	
 	/** @return true if the details should spawn a Pin Viewer. */
 	UE_API virtual bool HasPinViewer() const;
 
