[Mutable] Initialize ThresholdFactor in ASTOpMeshTransformWithBone to avoid indeterminism issues.
Remove incorrect condition check. The condition was always false.

#rnx

diff --git a/Engine/Plugins/Mutable/Source/CustomizableObjectEditor/Private/MuCOE/Nodes/CONodeModifierTransformWithBone.h b/Engine/Plugins/Mutable/Source/CustomizableObjectEditor/Private/MuCOE/Nodes/CONodeModifierTransformWithBone.h
--- a/Engine/Plugins/Mutable/Source/CustomizableObjectEditor/Private/MuCOE/Nodes/CONodeModifierTransformWithBone.h
+++ b/Engine/Plugins/Mutable/Source/CustomizableObjectEditor/Private/MuCOE/Nodes/CONodeModifierTransformWithBone.h
@@ -23,7 +23,7 @@ class UCONodeModifierTransformWithBone : public UCustomizableObjectNodeModifierB
 	GENERATED_BODY()
 
 public:
-	/** Transform to apply to the bounding mesh before selecting for vertices to transform. */
+	/** Transform to apply to the vertices. */
 	UPROPERTY(EditAnywhere, Category = Transform)
 	FTransform BoundingMeshTransform = FTransform::Identity;
 
diff --git a/Engine/Plugins/Mutable/Source/CustomizableObjectEditor/Private/MuCOE/Nodes/CONodeModifierTransformWithBoneDetails.h b/Engine/Plugins/Mutable/Source/CustomizableObjectEditor/Private/MuCOE/Nodes/CONodeModifierTransformWithBoneDetails.h
--- a/Engine/Plugins/Mutable/Source/CustomizableObjectEditor/Private/MuCOE/Nodes/CONodeModifierTransformWithBoneDetails.h
+++ b/Engine/Plugins/Mutable/Source/CustomizableObjectEditor/Private/MuCOE/Nodes/CONodeModifierTransformWithBoneDetails.h
@@ -15,7 +15,7 @@ public:
 	/** Makes a new instance of this detail layout class for a specific detail view requesting it */
 	static TSharedRef<IDetailCustomization> MakeInstance();
 
-	/** ILayoutDetails interface */
+	/** IDetailCustomization interface */
 	virtual void CustomizeDetails( IDetailLayoutBuilder& DetailBuilder ) override;
 
 private:
diff --git a/Engine/Plugins/Mutable/Source/MutableTools/Private/MuT/ASTOpMeshTransformWithBone.h b/Engine/Plugins/Mutable/Source/MutableTools/Private/MuT/ASTOpMeshTransformWithBone.h
--- a/Engine/Plugins/Mutable/Source/MutableTools/Private/MuT/ASTOpMeshTransformWithBone.h
+++ b/Engine/Plugins/Mutable/Source/MutableTools/Private/MuT/ASTOpMeshTransformWithBone.h
@@ -23,7 +23,7 @@ public:
 	ASTChild SourceMesh;
 	ASTChild Matrix;
 	FBoneName BoneName;
-	float ThresholdFactor;
+	float ThresholdFactor = 0.05f;
 public:
 
 	ASTOpMeshTransformWithBone();
diff --git a/Engine/Plugins/Mutable/Source/MutableTools/Private/MuT/CodeGenerator.cpp b/Engine/Plugins/Mutable/Source/MutableTools/Private/MuT/CodeGenerator.cpp
--- a/Engine/Plugins/Mutable/Source/MutableTools/Private/MuT/CodeGenerator.cpp
+++ b/Engine/Plugins/Mutable/Source/MutableTools/Private/MuT/CodeGenerator.cpp
@@ -2975,12 +2975,6 @@ namespace mu
 				// If a matrix node is not connected, the op won't do anything, so let's not create it at all.
 				if (TypedTransformNode->MatrixNode)
 				{
-					if (!TypedTransformNode->MatrixNode)
-					{
-						ErrorLog->Add("Matrix have not been generated", ELMT_ERROR, ErrorContext);
-						continue;
-					}
-
 					// Transform matrix.
 					Ptr<ASTOp> MatrixOp;
 					FMatrixGenerationResult ChildResult;
