[mutable] Baking: Added auto-compilation of the CO whose COI is going to get baked.
- This prevents the baking from failing if the CO was not compiled before.
- Removed code that handled the cancellation of the baking operation if the object was not compiled.

#rnx

diff --git a/Engine/Plugins/Mutable/Source/CustomizableObjectEditor/Internal/MuCOE/CustomizableObjectInstanceBakingUtils.h b/Engine/Plugins/Mutable/Source/CustomizableObjectEditor/Internal/MuCOE/CustomizableObjectInstanceBakingUtils.h
--- a/Engine/Plugins/Mutable/Source/CustomizableObjectEditor/Internal/MuCOE/CustomizableObjectInstanceBakingUtils.h
+++ b/Engine/Plugins/Mutable/Source/CustomizableObjectEditor/Internal/MuCOE/CustomizableObjectInstanceBakingUtils.h
@@ -1,8 +1,21 @@
 ﻿// Copyright Epic Games, Inc. All Rights Reserved.
 #pragma once
 
+#include "MuCO/CustomizableObject.h"
+#include "MuCO/CustomizableObjectCompilerTypes.h"
 #include "MuCO/CustomizableObjectInstance.h"
 
+
+/**
+ * Schedules the async compilation of a given Customizable object and, if a Customizable Object Instance is provided, tries to compile only the data required
+ * for that given COI. Use the delegate to know when the operation has completed.
+ * @param InCustomizableObject The Customizable Object we want to compile.
+ * @param InTargetInstance The Customizable Object Instance generated from the provided COI that we want to use to limit the scope of the CO compilation.
+ * @param InCustomizableObjectCompilationDelegate Delegate invoked once the compilation has completed. 
+ */
+CUSTOMIZABLEOBJECTEDITOR_API void ScheduleCOCompilationForBaking(UCustomizableObjectInstance& InTargetInstance, const FCompileNativeDelegate& InCustomizableObjectCompilationDelegate);
+
+
 /**
  * Schedules the async update of the target instance. It will clear previous update data.
  * This method will also take care of setting and later resetting the state of the Customizable Object System.
diff --git a/Engine/Plugins/Mutable/Source/CustomizableObjectEditor/Private/MuCOE/CustomizableObjectEditorViewportClient.cpp b/Engine/Plugins/Mutable/Source/CustomizableObjectEditor/Private/MuCOE/CustomizableObjectEditorViewportClient.cpp
--- a/Engine/Plugins/Mutable/Source/CustomizableObjectEditor/Private/MuCOE/CustomizableObjectEditorViewportClient.cpp
+++ b/Engine/Plugins/Mutable/Source/CustomizableObjectEditor/Private/MuCOE/CustomizableObjectEditorViewportClient.cpp
@@ -1606,6 +1606,8 @@ protected:
 
 void FCustomizableObjectEditorViewportClient::BakeInstance()
 {
+	// TODO: UE-314349 Try merging both code paths that end up calling the "CustomizableObjectInstanceBakingUtils" file to remove COI and CO defensive code
+	
 	// Early exit if no instance is set in the editor 
 	UCustomizableObjectInstance* Instance = CustomizableObjectEditorPtr.Pin()->GetPreviewInstance();
 	if (!Instance)
@@ -1614,67 +1616,88 @@ void FCustomizableObjectEditorViewportClient::BakeInstance()
 		return;
 	}
 	
-	BakeTempInstance = TStrongObjectPtr<UCustomizableObjectInstance>(Instance->Clone());
+	// Call the instance CO compilation async method
+	TObjectPtr<UCustomizableObject> InstanceToBakeCO = Instance->GetCustomizableObject();
+	if (!InstanceToBakeCO)
+	{
+		UE_LOG(LogMutable, Error, TEXT("The Customizable Object instance provided for the baking does not have a Customizable Object set."));
+		return;
+	}
 
-	// Call the instance update async method
-	FInstanceUpdateNativeDelegate UpdateDelegate;
-	UpdateDelegate.AddRaw(this, &FCustomizableObjectEditorViewportClient::OnInstanceForBakingUpdate);
-	ScheduleInstanceUpdateForBaking(*BakeTempInstance, UpdateDelegate);
+	// Store a clone of the initial COI so it is kept alive during the runtime of the bake
+	BakeTempInstance = TStrongObjectPtr<UCustomizableObjectInstance>(Instance->Clone());
+	
+	// Request the async compilation of the CO so we can later perform the update safely
+	if (ensure(BakeTempInstance))
+	{
+		FCompileNativeDelegate OnCompilationCallbackNative;
+		OnCompilationCallbackNative.BindRaw(this, &FCustomizableObjectEditorViewportClient::OnCOForBakingCompilationEnd);
+		ScheduleCOCompilationForBaking(*BakeTempInstance, OnCompilationCallbackNative );
+	}
 }
 
 
-void FCustomizableObjectEditorViewportClient::OnInstanceForBakingUpdate(const FUpdateContext& Result)
+void FCustomizableObjectEditorViewportClient::OnCOForBakingCompilationEnd(const FCompileCallbackParams& CompileCallbackParams)
 {
-	// Early exit if no instance was provided
-	if (!BakeTempInstance)
+	if (CompileCallbackParams.bRequestFailed)
 	{
-		UE_LOG(LogMutable, Error, TEXT("No Mutable Customizable Object instance was provided for the baking."));
+		UE_LOG(LogMutable, Warning ,TEXT("The Customizable Object failed the compilation. Skipping instance baking"));
 		return;
 	}
+	
+	if (ensure(BakeTempInstance))
+	{
+		// Call the instance update async method
+    	FInstanceUpdateNativeDelegate UpdateDelegate;
+    	UpdateDelegate.AddRaw(this, &FCustomizableObjectEditorViewportClient::OnInstanceForBakingUpdate);
+    	ScheduleInstanceUpdateForBaking(*BakeTempInstance, UpdateDelegate);
+	}
+}
 
+
+void FCustomizableObjectEditorViewportClient::OnInstanceForBakingUpdate(const FUpdateContext& Result)
+{
 	// Early exit if update result is not success
-	if ( !UCustomizableObjectSystem::IsUpdateResultValid(Result.UpdateResult) )
+	if (!UCustomizableObjectSystem::IsUpdateResultValid(Result.UpdateResult))
 	{
 		UE_LOG(LogMutable, Warning ,TEXT("Instance finished update with an error state : %s. Skipping instance baking"), *UEnum::GetValueAsString(Result.UpdateResult));
 		BakeTempInstance = nullptr;
 		return;
 	}
 	
-	const UCustomizableObject* CO = BakeTempInstance->GetCustomizableObject();
-
-	if (!CO)
-	{
-		UE_LOG(LogMutable, Error, TEXT("The Customizable Object instance provided for the baking does not have a Customizable Object."));
-		return;
-	}
-	
 	// Let the user set some configurations at the editor level
-	const FText DefaultFileName = FText::Format(LOCTEXT("DefaultFileNameForBakeInstance", "{0}"), FText::AsCultureInvariant(CO->GetName()));
+	if (ensure(BakeTempInstance))
+	{
+		const UCustomizableObject* CO = BakeTempInstance->GetCustomizableObject();
+		check(CO)
+		const FText DefaultFileName = FText::Format(LOCTEXT("DefaultFileNameForBakeInstance", "{0}"), CO ? FText::AsCultureInvariant(CO->GetName()) : FText::GetEmpty());
 	
-	TSharedRef<SMutableSelectFolderDlg> FolderDlg =
-		SNew(SMutableSelectFolderDlg)
-		.DefaultAssetPath(FText())
-		.DefaultFileName(DefaultFileName);
+		TSharedRef<SMutableSelectFolderDlg> FolderDlg =
+			SNew(SMutableSelectFolderDlg)
+			.DefaultAssetPath(FText())
+			.DefaultFileName(DefaultFileName);
 	
-	if (FolderDlg->ShowModal() != EAppReturnType::Cancel)
-	{
-		FBakingConfiguration Configuration;
-		Configuration.OutputFilesBaseName = FolderDlg->GetFileName();
-		Configuration.OutputPath = FolderDlg->GetAssetPath();
-		Configuration.bExportAllResourcesOnBake = FolderDlg->GetExportAllResources();
-		Configuration.bGenerateConstantMaterialInstancesOnBake = FolderDlg->GetGenerateConstantMaterialInstances();
+		if (FolderDlg->ShowModal() != EAppReturnType::Cancel)
+		{
+			FBakingConfiguration Configuration;
+			Configuration.OutputFilesBaseName = FolderDlg->GetFileName();
+			Configuration.OutputPath = FolderDlg->GetAssetPath();
+			Configuration.bExportAllResourcesOnBake = FolderDlg->GetExportAllResources();
+			Configuration.bGenerateConstantMaterialInstancesOnBake = FolderDlg->GetGenerateConstantMaterialInstances();
 
-		// TODO: UE-266716 Add a way to setup the prefixes from the editor dialog
+			// TODO: UE-266716 Add a way to setup the prefixes from the editor dialog
 		
-		TMap<UPackage*, EPackageSaveResolutionType> SavedPackages;
-		BakeCustomizableObjectInstance(
-			*BakeTempInstance,
-			Configuration,
-			false,
-			SavedPackages);
-	}
+			TMap<UPackage*, EPackageSaveResolutionType> SavedPackages;
+			BakeCustomizableObjectInstance(
+				*BakeTempInstance,
+				Configuration,
+				false,
+				SavedPackages);
+		}
 	
-	BakeTempInstance = nullptr;
+		BakeTempInstance = nullptr;
+	}
+
 }
 
 
diff --git a/Engine/Plugins/Mutable/Source/CustomizableObjectEditor/Private/MuCOE/CustomizableObjectEditorViewportClient.h b/Engine/Plugins/Mutable/Source/CustomizableObjectEditor/Private/MuCOE/CustomizableObjectEditorViewportClient.h
--- a/Engine/Plugins/Mutable/Source/CustomizableObjectEditor/Private/MuCOE/CustomizableObjectEditorViewportClient.h
+++ b/Engine/Plugins/Mutable/Source/CustomizableObjectEditor/Private/MuCOE/CustomizableObjectEditorViewportClient.h
@@ -3,6 +3,7 @@
 #pragma once
 
 #include "EditorViewportClient.h"
+#include "MuCO/CustomizableObject.h"
 #include "MuCO/CustomizableObjectInstance.h"
 #include "MuCO/CustomizableObjectParameterTypeDefinitions.h"
 
@@ -129,16 +130,20 @@ public:
 	void DrawUVs(FViewport* InViewport, FCanvas* InCanvas, int32 InTextYPos);
 
 	/**
-	 * Bake the instance currently present in the editor. Internally will schedule the update of the instance before baking it's resources
-	 * Note : It updates the UCustomizableObjectSystem so it is in the desired config for baking
+	 * Bake the instance currently present in the editor. Internally will schedule the compilation and later the update of the instance before baking
+	 * its resources.
 	 */
 	void BakeInstance();
 	
 private:
+	/**
+	 * Callback executed after the Customizable Object instance for baking had it's customizable object compiled
+	 * @param CompileCallbackParams Data returned on compilation that includes the end state of the compilation
+	 */
+	void OnCOForBakingCompilationEnd(const FCompileCallbackParams& CompileCallbackParams);
 	
 	/**
 	 * Callback executed after the instance in the editor gets updated for baking its contents.
-	 * Note : It updates the UCustomizableObjectSystem so it resets it to the state it had prior to the invocation of BakeEditorInstance
 	 */
 	void OnInstanceForBakingUpdate(const FUpdateContext& Result);
 	
diff --git a/Engine/Plugins/Mutable/Source/CustomizableObjectEditor/Private/MuCOE/CustomizableObjectInstanceBaker.cpp b/Engine/Plugins/Mutable/Source/CustomizableObjectEditor/Private/MuCOE/CustomizableObjectInstanceBaker.cpp
--- a/Engine/Plugins/Mutable/Source/CustomizableObjectEditor/Private/MuCOE/CustomizableObjectInstanceBaker.cpp
+++ b/Engine/Plugins/Mutable/Source/CustomizableObjectEditor/Private/MuCOE/CustomizableObjectInstanceBaker.cpp
@@ -12,13 +12,20 @@
 
 void UCustomizableObjectInstanceBaker::BakeInstance(UCustomizableObjectInstance* InTargetInstance, const FBakingConfiguration& InBakingConfig, const TSharedPtr<const FOnBakerFinishedWork> InOnBakerFinishedWork)
 {
+	// Cache a delegate that we will later invoke once the baking operation has finished all the work (so we can safely delete this object )
+	if (InOnBakerFinishedWork)
+	{
+		OnWorkFinished = InOnBakerFinishedWork;
+	}
+	
 	// Cache the configuration and the callback used to report the baking results
 	Configuration = InBakingConfig;
 
-	// Cache a delegate that we will later invoke once the baking operation has finished all the work (so we can delete this object safelly)
-	if (InOnBakerFinishedWork)
+	if (Configuration.OutputPath.IsEmpty())
 	{
-		OnWorkFinished = InOnBakerFinishedWork;
+		UE_LOG(LogMutable,Error,TEXT("No target save path has been provided."))
+		FinishBakingOperation(false, nullptr);
+		return;
 	}
 	
 	if (!InTargetInstance)
@@ -27,19 +34,40 @@ void UCustomizableObjectInstanceBaker::BakeInstance(UCustomizableObjectInstance*
 		FinishBakingOperation(false, nullptr);
 		return;
 	}
-
-	// We clone the InTargetInstance to ensure that the baked it is not open in a COI editor.
-	InstanceToBake = InTargetInstance->Clone();
 	
-	if (Configuration.OutputPath.IsEmpty())
+	if (!InTargetInstance->GetCustomizableObject())
 	{
-		UE_LOG(LogMutable,Error,TEXT("No target save path has been provided."))
+		UE_LOG(LogMutable,Error,TEXT("The provided instance does not have a Customizable Object defined."))
 		FinishBakingOperation(false, nullptr);
 		return;
 	}
+	
+	// We clone the InTargetInstance to ensure that the baked it is not open in a COI editor.
+	InstanceToBake = InTargetInstance->Clone();
 
-	// If no name is provided then compose one (allow not having a name)
-	// Configuration.OutputFilesBaseName = !Configuration.OutputFilesBaseName.IsEmpty() ? Configuration.OutputFilesBaseName : InstanceToBake.GetName();	
+	if (ensure(InstanceToBake))
+	{
+		// Request the async compilation of the CO so we can later perform the update safely
+		FCompileNativeDelegate OnCompilationCallbackNative;
+		OnCompilationCallbackNative.BindUObject(this, &UCustomizableObjectInstanceBaker::OnCompilationEnd);
+		ScheduleCOCompilationForBaking(*InstanceToBake, OnCompilationCallbackNative );
+	}
+	else
+	{
+		FinishBakingOperation(false, nullptr);
+	}
+}
+
+
+void UCustomizableObjectInstanceBaker::OnCompilationEnd(const FCompileCallbackParams& InCompileCallbackParams)
+{
+	if (InCompileCallbackParams.bRequestFailed)
+	{
+		// Exit if the compilation failed
+		UE_LOG(LogMutable, Display, TEXT("The Customizable Object compilation for baking failed. Aborting instance baking operation."));
+		FinishBakingOperation(false, nullptr);
+		return;
+	}
 
 	// Schedule update of the instance for baking
 	UpdateInstance();
@@ -50,19 +78,26 @@ void UCustomizableObjectInstanceBaker::UpdateInstance()
 {
 	// Call the instance update async method
 	FInstanceUpdateNativeDelegate UpdateDelegate;
-	UpdateDelegate.AddUObject(this, &UCustomizableObjectInstanceBaker::OnInstanceUpdate);
+	UpdateDelegate.AddUObject(this, &UCustomizableObjectInstanceBaker::OnInstanceUpdated);
 	
 	// Once the update finishes successfully the baking operation will engage
-	ScheduleInstanceUpdateForBaking(*InstanceToBake, UpdateDelegate);
+	if (ensure(InstanceToBake))
+	{
+		ScheduleInstanceUpdateForBaking(*InstanceToBake, UpdateDelegate);
+	}
+	else
+	{
+		FinishBakingOperation(false, nullptr);
+	}
 }
 
 
-void UCustomizableObjectInstanceBaker::OnInstanceUpdate(const FUpdateContext& Result) const
+void UCustomizableObjectInstanceBaker::OnInstanceUpdated(const FUpdateContext& Result) const
 {
 	const EUpdateResult InstanceUpdateResult = Result.UpdateResult;
 	UE_LOG(LogMutable, Display, TEXT("Instance finished update with state : %s."), *UEnum::GetValueAsString(InstanceUpdateResult));
 	
-	if (InstanceToBake->GetPrivate()->SkeletalMeshStatus == ESkeletalMeshStatus::Success)
+	if (ensure(InstanceToBake) && InstanceToBake->GetPrivate()->SkeletalMeshStatus == ESkeletalMeshStatus::Success)
 	{
 		Bake();
 	}
@@ -83,11 +118,19 @@ void UCustomizableObjectInstanceBaker::Bake() const
 	
 	// Array with all the packages that did get saved. Packages that failed the saving procedure will not appear in this array.
 	TMap<UPackage*, EPackageSaveResolutionType> SavedPackages;
-
-	// Ensures no interaction from the user will be required (alongside GIsRunningUnattendedScript)
-	constexpr bool bIsUnhandledExecution = true;
 	
-	const bool bBakeWasSuccessful = BakeCustomizableObjectInstance(*InstanceToBake, Configuration, bIsUnhandledExecution, SavedPackages);
+	bool bBakeWasSuccessful = false;
+	if (ensure(InstanceToBake))
+	{
+		// Ensures no interaction from the user will be required (alongside GIsRunningUnattendedScript)
+		constexpr bool bIsUnhandledExecution = true;		// TODO: UE-314349
+		bBakeWasSuccessful = BakeCustomizableObjectInstance(*InstanceToBake, Configuration, bIsUnhandledExecution, SavedPackages);
+	}
+	else
+	{
+		FinishBakingOperation(false, nullptr);
+		return;
+	}
 	
 	// Generate the array of structs to output based on the packages reported as saved
 	TArray<FBakedResourceData> SavedPackagesData;
diff --git a/Engine/Plugins/Mutable/Source/CustomizableObjectEditor/Private/MuCOE/CustomizableObjectInstanceBaker.h b/Engine/Plugins/Mutable/Source/CustomizableObjectEditor/Private/MuCOE/CustomizableObjectInstanceBaker.h
--- a/Engine/Plugins/Mutable/Source/CustomizableObjectEditor/Private/MuCOE/CustomizableObjectInstanceBaker.h
+++ b/Engine/Plugins/Mutable/Source/CustomizableObjectEditor/Private/MuCOE/CustomizableObjectInstanceBaker.h
@@ -2,6 +2,7 @@
 
 #pragma once
 
+#include "MuCO/CustomizableObject.h"
 #include "MuCO/CustomizableObjectInstance.h"
 
 #include "CustomizableObjectInstanceBaker.generated.h"
@@ -19,6 +20,7 @@ DECLARE_DELEGATE(FOnBakerFinishedWork);
  * Utility class designed to allow tbe baking of the resources of a target Mutable Customizable Object instance onto disk.
  *
  * The actual baking operation does require some work done before the baking itself.
+ *		- Compilation of the instance's Customizable Object : The mutable CO should be compiled to allow for the update of the instance.
  *		- Updating of the instance (async) : The mutable COInstance gets updated prior to the bake to make sure the instance resources to bake are the latest
  *		- Baking of the instance resources (sync) : The resources of the instance get baked onto disk.
  *
@@ -45,6 +47,12 @@ public:
 	
 private:
 	
+	/**
+	 * Callback executed once the compilation of the CO has finished.
+	 * @param InCompileCallbackParams The compilation parameters that the compilation callback invocation provides to the caller.
+	 */
+	UE_API void OnCompilationEnd(const FCompileCallbackParams& InCompileCallbackParams);
+	
 	/**
 	 * Takes care of updating and generating the instance resources so we can later bake them.
 	 */
@@ -54,7 +62,7 @@ private:
 	 * Callback executed when the instance finishes its update. It does not matter if it is successful or not
 	 * @param Result The result returned by the instance updating operation once the operation is completed.
 	 */
-	UE_API void OnInstanceUpdate(const FUpdateContext& Result) const;
+	UE_API void OnInstanceUpdated(const FUpdateContext& Result) const;
 
 	/**
 	 * Actual baking method. It will take care of baking the already updated instance and saving it's resources to disk. It will also log the paths
diff --git a/Engine/Plugins/Mutable/Source/CustomizableObjectEditor/Private/MuCOE/CustomizableObjectInstanceBakingUtils.cpp b/Engine/Plugins/Mutable/Source/CustomizableObjectEditor/Private/MuCOE/CustomizableObjectInstanceBakingUtils.cpp
--- a/Engine/Plugins/Mutable/Source/CustomizableObjectEditor/Private/MuCOE/CustomizableObjectInstanceBakingUtils.cpp
+++ b/Engine/Plugins/Mutable/Source/CustomizableObjectEditor/Private/MuCOE/CustomizableObjectInstanceBakingUtils.cpp
@@ -17,6 +17,7 @@
 #include "PhysicsEngine/PhysicsAsset.h"
 
 #include "MuCO/CustomizableObject.h"
+#include "MuCO/CustomizableObjectCustomVersion.h"
 #include "MuCO/CustomizableObjectInstanceAssetUserData.h"
 #include "MuCO/CustomizableObjectInstancePrivate.h"
 #include "MuCO/CustomizableObjectMipDataProvider.h"
@@ -295,23 +296,40 @@ void OnInstanceUpdateFinish(const FUpdateContext& Result)
 }
 
 
-void ScheduleInstanceUpdateForBaking(UCustomizableObjectInstance& InInstance, FInstanceUpdateNativeDelegate& InInstanceUpdateDelegate)
+void ScheduleCOCompilationForBaking(UCustomizableObjectInstance& InTargetInstance, const FCompileNativeDelegate& InCustomizableObjectCompilationDelegate)
 {
 	// This prevents the queue of updates to have more than one instance for baking.
 	if (bIsBakeOperationAlreadyScheduled)
 	{
-		UE_LOG(LogMutable, Error, TEXT("The COInstance update for baking could not be scheduled. Another instance is being updated for baking."));
-		InInstanceUpdateDelegate.Broadcast({EUpdateResult::Error});
+		UE_LOG(LogMutable, Error, TEXT("The CO compilation for baking could not be scheduled. Another instance is being processed at this time."));
+		InCustomizableObjectCompilationDelegate.ExecuteIfBound({true,false,false,false, false});
 		return;
 	}
 
-	// TODO UE-262123: Compile Instance with correct settings instead
-	if (!InInstance.GetCustomizableObject() || !InInstance.GetCustomizableObject()->IsCompiled())
+	TObjectPtr<UCustomizableObject> CustomizableObject = InTargetInstance.GetCustomizableObject();
+	if (!CustomizableObject)
 	{
-		UE_LOG(LogMutable, Error, TEXT("The COInstance update for baking could not be scheduled. Instance not compiled."));
-		InInstanceUpdateDelegate.Broadcast({ EUpdateResult::Error });
+		UE_LOG(LogMutable, Error, TEXT("The CO compilation for baking could not be scheduled. The provided instance does not have a CO set."));
+		InCustomizableObjectCompilationDelegate.ExecuteIfBound({true,false,false,false, false});
 		return;
 	}
+	
+	// Tell the baking system that an instance for baking will be processed
+	bIsBakeOperationAlreadyScheduled = true;
+
+	FCompileParams CompileParams;
+	CompileParams.bAsync = true;
+	CompileParams.CallbackNative = InCustomizableObjectCompilationDelegate;
+	CompileParams.CompileOnlySelectedInstance = &InTargetInstance;
+	
+	CustomizableObject->Compile(CompileParams);
+}
+
+
+void ScheduleInstanceUpdateForBaking(UCustomizableObjectInstance& InInstance, FInstanceUpdateNativeDelegate& InInstanceUpdateDelegate)
+{
+	check(InInstance.GetCustomizableObject() && InInstance.GetCustomizableObject()->IsCompiled());
+	ensure(bIsBakeOperationAlreadyScheduled);	
 
 	InInstanceUpdateDelegate.AddStatic(&OnInstanceUpdateFinish);
 	
@@ -329,8 +347,6 @@ void ScheduleInstanceUpdateForBaking(UCustomizableObjectInstance& InInstance, FI
 	UCustomizableObjectSystemPrivate* SystemPrivate = UCustomizableObjectSystem::GetInstance()->GetPrivate();
 	SystemPrivate->EnqueueUpdateSkeletalMesh(Context);
 	
-	// Tell the baking system that an instance for baking will be processed
-	bIsBakeOperationAlreadyScheduled = true;
 
 	FCustomizableObjectEditorLogger::CreateLog(
 	LOCTEXT("CustomizableObjectInstanceBakingUtils_UpdateScheduled", "The COInstance Update operation for baking has been scheduled. Please hold."))
