[Mutable] New Skeletal Mesh types.

- Types now use the colors from the assets.
- Renamed some types.
- Removed Layout type.

diff --git a/Engine/Plugins/Mutable/Source/CustomizableObjectEditor/Private/MuCOE/CustomizableObjectConnectionDrawingPolicy.cpp b/Engine/Plugins/Mutable/Source/CustomizableObjectEditor/Private/MuCOE/CustomizableObjectConnectionDrawingPolicy.cpp
--- a/Engine/Plugins/Mutable/Source/CustomizableObjectEditor/Private/MuCOE/CustomizableObjectConnectionDrawingPolicy.cpp
+++ b/Engine/Plugins/Mutable/Source/CustomizableObjectEditor/Private/MuCOE/CustomizableObjectConnectionDrawingPolicy.cpp
@@ -9,7 +9,7 @@ void FCustomizableObjectConnectionDrawingPolicy::DetermineWiringStyle(UEdGraphPi
 {
 	FConnectionDrawingPolicy::DetermineWiringStyle(OutputPin, InputPin, Params);
 
-	Params.WireColor = UEdGraphSchema_CustomizableObject::GetPinTypeColor(OutputPin->PinType.PinCategory); // Color.
+	Params.WireColor = GetDefault<UEdGraphSchema_CustomizableObject>()->GetPinTypeColor(OutputPin->PinType.PinCategory); // Color.
 
 	if (HoveredPins.Num() > 0)
 	{
diff --git a/Engine/Plugins/Mutable/Source/CustomizableObjectEditor/Private/MuCOE/CustomizableObjectMacroLibrary/CustomizableObjectMacroDetails.cpp b/Engine/Plugins/Mutable/Source/CustomizableObjectEditor/Private/MuCOE/CustomizableObjectMacroLibrary/CustomizableObjectMacroDetails.cpp
--- a/Engine/Plugins/Mutable/Source/CustomizableObjectEditor/Private/MuCOE/CustomizableObjectMacroLibrary/CustomizableObjectMacroDetails.cpp
+++ b/Engine/Plugins/Mutable/Source/CustomizableObjectEditor/Private/MuCOE/CustomizableObjectMacroLibrary/CustomizableObjectMacroDetails.cpp
@@ -290,7 +290,7 @@ TSharedRef<SWidget> SCOMacroPinTypeSelector::OnGenerateRow(TSharedPtr<FPinNameRo
 	[
 		SNew(SImage)
 		.Image(IconBrush)
-		.ColorAndOpacity(UEdGraphSchema_CustomizableObject::GetPinTypeColor(Option->PinCategory))
+		.ColorAndOpacity(GetDefault<UEdGraphSchema_CustomizableObject>()->GetPinTypeColor(Option->PinCategory))
 	]
 	+ SHorizontalBox::Slot()
 	.AutoWidth()
@@ -317,7 +317,7 @@ TSharedRef<SWidget> SCOMacroPinTypeSelector::GenerateSelectedContent() const
 		.Image(IconBrush)
 		.ColorAndOpacity_Lambda([this]() 
 		{
-			return UEdGraphSchema_CustomizableObject::GetPinTypeColor(Variable->PinCategoryType); 
+			return GetDefault<UEdGraphSchema_CustomizableObject>()->GetPinTypeColor(Variable->PinCategoryType); 
 		})
 	]
 	+ SHorizontalBox::Slot()
diff --git a/Engine/Plugins/Mutable/Source/CustomizableObjectEditor/Private/MuCOE/EdGraphSchema_CustomizableObject.cpp b/Engine/Plugins/Mutable/Source/CustomizableObjectEditor/Private/MuCOE/EdGraphSchema_CustomizableObject.cpp
--- a/Engine/Plugins/Mutable/Source/CustomizableObjectEditor/Private/MuCOE/EdGraphSchema_CustomizableObject.cpp
+++ b/Engine/Plugins/Mutable/Source/CustomizableObjectEditor/Private/MuCOE/EdGraphSchema_CustomizableObject.cpp
@@ -1,6 +1,8 @@
 // Copyright Epic Games, Inc. All Rights Reserved.
 
 #include "MuCOE/EdGraphSchema_CustomizableObject.h"
+
+#include "AssetDefinitionRegistry.h"
 #include "MuCOE/CustomizableObjectSchemaActions.h"
 
 #include "CustomizableObjectConnectionDrawingPolicy.h"
@@ -98,6 +100,7 @@
 #include "MuCOE/Nodes/CustomizableObjectNodeReroute.h"
 #include "Toolkits/ToolkitManager.h"
 #include "PropertyEditorModule.h"
+#include "Animation/PoseAsset.h"
 #include "MuCOE/GraphTraversal.h"
 #include "MuCOE/Nodes/CustomizableObjectNodeComponentMeshAddTo.h"
 #include "MuCOE/Nodes/CustomizableObjectNodeTransformConstant.h"
@@ -278,10 +281,10 @@ const FName UEdGraphSchema_CustomizableObject::PC_Component("component");
 const FName UEdGraphSchema_CustomizableObject::PC_MeshSection("material");
 const FName UEdGraphSchema_CustomizableObject::PC_Modifier("modifier");
 const FName UEdGraphSchema_CustomizableObject::PC_Mesh("mesh");
-const FName UEdGraphSchema_CustomizableObject::PC_PassThroughMesh("passThroughMesh");
-const FName UEdGraphSchema_CustomizableObject::PC_Layout("layout");
-const FName UEdGraphSchema_CustomizableObject::PC_Image("image");
-const FName UEdGraphSchema_CustomizableObject::PC_PassThroughImage("passThroughImage");
+const FName UEdGraphSchema_CustomizableObject::PC_SkeletalMesh("skeletalMesh");
+const FName UEdGraphSchema_CustomizableObject::PC_PassthroughSkeletalMesh("passThroughSkeletalMesh");
+const FName UEdGraphSchema_CustomizableObject::PC_Texture("image");
+const FName UEdGraphSchema_CustomizableObject::PC_PassthroughTexture("passThroughImage");
 const FName UEdGraphSchema_CustomizableObject::PC_Projector("projector");
 const FName UEdGraphSchema_CustomizableObject::PC_GroupProjector("groupProjector");
 const FName UEdGraphSchema_CustomizableObject::PC_Color("color");
@@ -296,8 +299,8 @@ const FName UEdGraphSchema_CustomizableObject::PC_Transform("transform");
 const FName UEdGraphSchema_CustomizableObject::PC_String("string");
 
 // Add more pin types to this array if needed
-const TArray<FName> UEdGraphSchema_CustomizableObject::SupportedMacroPinTypes = { PC_Object, PC_Component, PC_MeshSection, PC_Modifier, PC_Mesh, PC_Image,
-PC_PassThroughImage, PC_Projector, PC_GroupProjector, PC_Color, PC_Float, PC_Enum, PC_Stack, PC_Material, PC_PoseAsset, PC_Transform, PC_String };
+const TArray<FName> UEdGraphSchema_CustomizableObject::SupportedMacroPinTypes = { PC_Object, PC_Component, PC_MeshSection, PC_Modifier, PC_Mesh, PC_Texture,
+PC_PassthroughTexture, PC_Projector, PC_GroupProjector, PC_Color, PC_Float, PC_Enum, PC_Stack, PC_Material, PC_PoseAsset, PC_Transform, PC_String };
 
 // Node categories
 const FText UEdGraphSchema_CustomizableObject::NC_Experimental(FText::FromString("Experimental"));
@@ -642,7 +645,7 @@ void UEdGraphSchema_CustomizableObject::GetGraphContextActions(FGraphContextMenu
 
 	{
 		// External Pin Nodes
-		TArray<FName> PinTypes({ PC_MeshSection, PC_Modifier, PC_Mesh, PC_Image, PC_Projector, PC_GroupProjector, PC_Color, PC_Float, PC_Bool, PC_Enum, PC_Transform, PC_Stack, PC_PassThroughImage, PC_Material, PC_PoseAsset, PC_Component });
+		TArray<FName> PinTypes({ PC_MeshSection, PC_Modifier, PC_Mesh, PC_Texture, PC_Projector, PC_GroupProjector, PC_Color, PC_Float, PC_Bool, PC_Enum, PC_Transform, PC_Stack, PC_PassthroughTexture, PC_Material, PC_PoseAsset, PC_Component });
 
 		// Add pin types from extensions
 		for (const FRegisteredCustomizableObjectPinType& PinType : ICustomizableObjectModule::Get().GetExtendedPinTypes())
@@ -792,83 +795,56 @@ const FPinConnectionResponse UEdGraphSchema_CustomizableObject::CanCreateConnect
 }
 
 
-FLinearColor UEdGraphSchema_CustomizableObject::GetPinTypeColor(const FName& TypeString)
+FLinearColor UEdGraphSchema_CustomizableObject::GetPinTypeColor(const FName& TypeString) const
 {
-	if (TypeString == PC_Enum)
-	{
-		return FLinearColor(0.004f, 0.42f, 0.384f, 1.000000f); // Light green
-	}
-	else if (TypeString == PC_Float)
-	{
-		return FLinearColor(0.357667f, 1.000000f, 0.060000f, 1.000000f); // Green
-	}
-	else if (TypeString == PC_Color)
-	{
-		return FLinearColor(1.000000f, 0.591255f, 0.016512f, 1.000000f); // Yellow
-	}
-	else if (TypeString == PC_Bool)
-	{
-		return FLinearColor(0.470000f, 0.0f, 0.000000f, 1.000000f); // Red
-	}
-	else if (TypeString == PC_Projector)
-	{
-		return FLinearColor(FColorList::Aquamarine);
-	}
-	else if (TypeString == PC_GroupProjector)
-	{
-		return FLinearColor(FColorList::DarkTan);
-	}
-	else if (TypeString == PC_Mesh)
-	{
-		return FLinearColor(FColorList::MediumOrchid); 
-	}
-	else if (TypeString == PC_Layout)
-	{
-		return FLinearColor(0.500000f, 0.500000f, 0.100000f, 1.000000f); // Light yellow
-	}
-	else if (TypeString == PC_Image || TypeString == PC_PassThroughImage)
-	{
-		return FLinearColor(0.75294f, 0.08098f, 0.08098f, 1.000000f); // Soft Red
-	}
-	else if (TypeString == PC_MeshSection)
-	{
-		return FLinearColor(0.000000f, 0.100000f, 0.600000f, 1.000000f); // Blue
-	}
-	else if (TypeString == PC_Modifier)
-	{
-		return FLinearColor(FColorList::LightGrey); 
-	}
-	else if (TypeString == PC_Object)
-	{
-		return FLinearColor(0.000000f, 0.400000f, 0.910000f, 1.000000f); // Light blue
-	}
-	else if (TypeString == PC_Component)
-	{
-		return FLinearColor(FColorList::DarkOrchid); 
-	}
-	else if (TypeString == PC_Stack)
+	// Cache colors
+	if (PinCategoryColors.IsEmpty())
 	{
-		return FLinearColor(1.000000f, 0.000000f, 0.800000f, 1.000000f); // Pink
-	}
-	else if (TypeString == PC_Material)
-	{
-		return FLinearColor(0.000000f, 1.000000f, 0.1000f, 1.000000f); // Cian
-	}
-	else if (TypeString == PC_Wildcard)
-	{
-		return FLinearColor(1.000000f, 1.000000f, 1.000000f, 1.000000f); // White
-	}
-	else if (TypeString == PC_PoseAsset)
-	{
-		return FLinearColor(0.700000f, 0.000000f, 0.000000f, 1.000000f); // Dark Red
-	}
-	else if (TypeString == PC_Transform)
-	{
-		return FLinearColor(FColorList::Orange); 
+		UAssetDefinitionRegistry* AssetDefinitionRegistry = UAssetDefinitionRegistry::Get();
+		
+		PinCategoryColors.Add(PC_Enum, FLinearColor(0.004f, 0.42f, 0.384f, 1.000000f));
+		PinCategoryColors.Add(PC_Float, FLinearColor(0.357667f, 1.000000f, 0.060000f, 1.000000f));
+		PinCategoryColors.Add(PC_Color, FLinearColor(1.000000f, 0.591255f, 0.016512f, 1.000000f));
+		PinCategoryColors.Add(PC_Bool, FLinearColor(0.470000f, 0.0f, 0.000000f, 1.000000f));
+		PinCategoryColors.Add(PC_Projector, FLinearColor(FColorList::Aquamarine));
+		PinCategoryColors.Add(PC_GroupProjector, FLinearColor(FColorList::DarkTan));
+		PinCategoryColors.Add(PC_Mesh, FLinearColor(FColorList::MediumOrchid)); 
+		PinCategoryColors.Add(PC_MeshSection, FLinearColor(0.000000f, 0.100000f, 0.600000f, 1.000000f));
+		PinCategoryColors.Add(PC_Modifier, FLinearColor(FColorList::LightGrey)); 
+		PinCategoryColors.Add(PC_Object, FLinearColor(0.000000f, 0.400000f, 0.910000f, 1.000000f));
+		PinCategoryColors.Add(PC_Component, FLinearColor(FColorList::DarkOrchid)); 
+		PinCategoryColors.Add(PC_Stack, FLinearColor(1.000000f, 0.000000f, 0.800000f, 1.000000f));
+		PinCategoryColors.Add(PC_Wildcard, FLinearColor(1.000000f, 1.000000f, 1.000000f, 1.000000f));
+		PinCategoryColors.Add(PC_Transform, FLinearColor(FColorList::Orange)); 
+		PinCategoryColors.Add(PC_String, FLinearColor(0.700000f, 0.010000f, 0.660000f, 1.000000f));
+
+		{
+			const FLinearColor Color = AssetDefinitionRegistry->GetAssetDefinitionForClass(UMaterialInterface::StaticClass())->GetAssetColor();
+			PinCategoryColors.Add(PC_Material, Color);
+		}
+		
+		{
+			const FLinearColor Color = AssetDefinitionRegistry->GetAssetDefinitionForClass(USkeletalMesh::StaticClass())->GetAssetColor();
+			PinCategoryColors.Add(PC_SkeletalMesh, Color);
+			PinCategoryColors.Add(PC_PassthroughSkeletalMesh, Color); 	
+			
+		}
+
+		{
+			const FLinearColor Color = AssetDefinitionRegistry->GetAssetDefinitionForClass(UTexture::StaticClass())->GetAssetColor();
+			PinCategoryColors.Add(PC_Texture, Color);
+			PinCategoryColors.Add(PC_PassthroughTexture, Color);	
+		}
+
+		{
+			const FLinearColor Color = AssetDefinitionRegistry->GetAssetDefinitionForClass(UPoseAsset::StaticClass())->GetAssetColor();
+			PinCategoryColors.Add(PC_PoseAsset, Color);
+		}
 	}
-	else if (TypeString == PC_String)
+		
+	if (const FLinearColor* Color = PinCategoryColors.Find(TypeString))
 	{
-		return FLinearColor(0.700000f, 0.010000f, 0.660000f, 1.000000f); // Dark Pink
+		return *Color;		
 	}
 
 	for (const FRegisteredCustomizableObjectPinType& PinType : ICustomizableObjectModule::Get().GetExtendedPinTypes())
@@ -1300,7 +1276,7 @@ FConnectionDrawingPolicy* UEdGraphSchema_CustomizableObject::CreateConnectionDra
 
 
 void UEdGraphSchema_CustomizableObject::GetAssetsGraphHoverMessage(const TArray<FAssetData>& Assets,
-	const UEdGraph* HoverGraph, FString& OutTooltipText, bool& OutOkIcon) const
+                                                                   const UEdGraph* HoverGraph, FString& OutTooltipText, bool& OutOkIcon) const
 {
 	Super::GetAssetsGraphHoverMessage(Assets, HoverGraph, OutTooltipText, OutOkIcon);
 
@@ -1400,83 +1376,87 @@ FPinConnectionResponse UEdGraphSchema_CustomizableObject::MovePinLinks(UEdGraphP
 
 FName UEdGraphSchema_CustomizableObject::GetPinCategoryName(const FName& PinCategory)
 {
-	if (PinCategory == UEdGraphSchema_CustomizableObject::PC_Object)
+	if (PinCategory == PC_Object)
 	{
 		return TEXT("Object");
 	}
-	else if (PinCategory == UEdGraphSchema_CustomizableObject::PC_Component)
+	else if (PinCategory == PC_Component)
 	{
 		return TEXT("Component");
 	}
-	else if (PinCategory == UEdGraphSchema_CustomizableObject::PC_MeshSection)
+	else if (PinCategory == PC_MeshSection)
 	{
 		return TEXT("Material");
 	}
-	else if (PinCategory == UEdGraphSchema_CustomizableObject::PC_Modifier)
+	else if (PinCategory == PC_Modifier)
 	{
 		return TEXT("Modifier");
 	}
-	else if (PinCategory == UEdGraphSchema_CustomizableObject::PC_Mesh)
+	else if (PinCategory == PC_Mesh)
 	{
 		return TEXT("Mesh");
 	}
-	else if (PinCategory == UEdGraphSchema_CustomizableObject::PC_Layout)
+	else if (PinCategory == PC_SkeletalMesh)
+	{
+		return TEXT("Skeletal Mesh");
+	}
+	else if (PinCategory == PC_PassthroughSkeletalMesh)
 	{
-		return TEXT("Layout");
+		return TEXT("Passthrough Skeletal Mesh");
 	}
-	else if (PinCategory == UEdGraphSchema_CustomizableObject::PC_Image)
+	else if (PinCategory == PC_Texture)
 	{
 		return TEXT("Texture");
 	}
-	else if (PinCategory == UEdGraphSchema_CustomizableObject::PC_PassThroughImage)
+	else if (PinCategory == PC_PassthroughTexture)
 	{
 		return TEXT("PassThrough Texture");
 	}
-	else if (PinCategory == UEdGraphSchema_CustomizableObject::PC_Projector)
+	else if (PinCategory == PC_Projector)
 	{
 		return TEXT("Projector");
 	}
-	else if (PinCategory == UEdGraphSchema_CustomizableObject::PC_GroupProjector)
+	else if (PinCategory == PC_GroupProjector)
 	{
 		return TEXT("Group Projector");
 	}
-	else if (PinCategory == UEdGraphSchema_CustomizableObject::PC_Color)
+	else if (PinCategory == PC_Color)
 	{
 		return TEXT("Color");
 	}
-	else if (PinCategory == UEdGraphSchema_CustomizableObject::PC_Float)
+	else if (PinCategory == PC_Float)
 	{
 		return TEXT("Float");
 	}
-	else if (PinCategory == UEdGraphSchema_CustomizableObject::PC_Bool)
+	else if (PinCategory == PC_Bool)
 	{
 		return TEXT("Bool");
 	}
-	else if (PinCategory == UEdGraphSchema_CustomizableObject::PC_Enum)
+	else if (PinCategory == PC_Enum)
 	{
 		return TEXT("Enum");
 	}
-	else if (PinCategory == UEdGraphSchema_CustomizableObject::PC_Stack)
+	else if (PinCategory == PC_Stack)
 	{
 		return TEXT("Stack");
 	}
-	else if (PinCategory == UEdGraphSchema_CustomizableObject::PC_Material)
+	else if (PinCategory == PC_Material)
 	{
 		return TEXT("Material");
 	}
-	else if (PinCategory == UEdGraphSchema_CustomizableObject::PC_Wildcard)
+	else if (PinCategory == PC_Wildcard)
 	{
 		return TEXT("Wildcard");
 	}
-	else if (PinCategory == UEdGraphSchema_CustomizableObject::PC_PoseAsset)
+	else if (PinCategory == PC_PoseAsset)
 	{
 		return TEXT("PoseAsset");
 	}
-	else if (PinCategory == UEdGraphSchema_CustomizableObject::PC_Transform)
+	else if (PinCategory == PC_Transform)
 	{
 		return TEXT("Transform");
 	}
-	else if (PinCategory == UEdGraphSchema_CustomizableObject::PC_String)
+	else if (PinCategory == PC_String)
 	{
 		return TEXT("String");
 	}
@@ -1499,83 +1479,87 @@ FName UEdGraphSchema_CustomizableObject::GetPinCategoryName(const FName& PinCate
 
 FText UEdGraphSchema_CustomizableObject::GetPinCategoryFriendlyName(const FName& PinCategory)
 {
-	if (PinCategory == UEdGraphSchema_CustomizableObject::PC_Object)
+	if (PinCategory == PC_Object)
 	{
 		return LOCTEXT("Object_Pin_Category", "Object");
 	}
-	else if (PinCategory == UEdGraphSchema_CustomizableObject::PC_Component)
+	else if (PinCategory == PC_Component)
 	{
 		return LOCTEXT("Component_Pin_Category", "Component");
 	}
-	else if (PinCategory == UEdGraphSchema_CustomizableObject::PC_MeshSection)
+	else if (PinCategory == PC_MeshSection)
 	{
 		return LOCTEXT("MeshSection_Pin_Category", "Mesh Section");
 	}
-	else if (PinCategory == UEdGraphSchema_CustomizableObject::PC_Modifier)
+	else if (PinCategory == PC_Modifier)
 	{
 		return LOCTEXT("Modifier_Pin_Category", "Modifier");
 	}
-	else if (PinCategory == UEdGraphSchema_CustomizableObject::PC_Mesh)
+	else if (PinCategory == PC_Mesh)
 	{
 		return LOCTEXT("Mesh_Pin_Category", "Mesh");
 	}
-	else if (PinCategory == UEdGraphSchema_CustomizableObject::PC_Layout)
+	else if (PinCategory == PC_SkeletalMesh)
+	{
+		return LOCTEXT("Mesh_Pin_Category", "Skeletal Mesh");
+	}
+	else if (PinCategory == PC_PassthroughSkeletalMesh)
 	{
-		return LOCTEXT("Layout_Pin_Category", "Layout");
+		return LOCTEXT("Mesh_Pin_Category", "Passthrough Skeletal Mesh");
 	}
-	else if (PinCategory == UEdGraphSchema_CustomizableObject::PC_Image)
+	else if (PinCategory == PC_Texture)
 	{
 		return LOCTEXT("Image_Pin_Category", "Texture");
 	}
-	else if (PinCategory == UEdGraphSchema_CustomizableObject::PC_PassThroughImage)
+	else if (PinCategory == PC_PassthroughTexture)
 	{
-		return LOCTEXT("PassThrough_Image_Pin_Category", "PassThrough Texture");
+		return LOCTEXT("PassThrough_Image_Pin_Category", "Passthrough Texture");
 	}
-	else if (PinCategory == UEdGraphSchema_CustomizableObject::PC_Projector)
+	else if (PinCategory == PC_Projector)
 	{
 		return LOCTEXT("Projector_Pin_Category", "Projector");
 	}
-	else if (PinCategory == UEdGraphSchema_CustomizableObject::PC_GroupProjector)
+	else if (PinCategory == PC_GroupProjector)
 	{
 		return LOCTEXT("Group_Projector_Pin_Category", "Group Projector");
 	}
-	else if (PinCategory == UEdGraphSchema_CustomizableObject::PC_Color)
+	else if (PinCategory == PC_Color)
 	{
 		return LOCTEXT("Color_Pin_Category", "Color");
 	}
-	else if (PinCategory == UEdGraphSchema_CustomizableObject::PC_Float)
+	else if (PinCategory == PC_Float)
 	{
 		return LOCTEXT("Float_Pin_Category", "Float");
 	}
-	else if (PinCategory == UEdGraphSchema_CustomizableObject::PC_Bool)
+	else if (PinCategory == PC_Bool)
 	{
 		return LOCTEXT("Bool_Pin_Category", "Bool");
 	}
-	else if (PinCategory == UEdGraphSchema_CustomizableObject::PC_Enum)
+	else if (PinCategory == PC_Enum)
 	{
 		return LOCTEXT("Enum_Pin_Category", "Enum");
 	}
-	else if (PinCategory == UEdGraphSchema_CustomizableObject::PC_Stack)
+	else if (PinCategory == PC_Stack)
 	{
 		return LOCTEXT("Stack_Pin_Category", "Stack");
 	}
-	else if (PinCategory == UEdGraphSchema_CustomizableObject::PC_Material)
+	else if (PinCategory == PC_Material)
 	{
 		return LOCTEXT("Material_Asset_Pin_Category", "Material");
 	}
-	else if (PinCategory == UEdGraphSchema_CustomizableObject::PC_Wildcard)
+	else if (PinCategory == PC_Wildcard)
 	{
 		return LOCTEXT("Wildcard_Pin_Category", "Wildcard");
 	}
-	else if (PinCategory == UEdGraphSchema_CustomizableObject::PC_PoseAsset)
+	else if (PinCategory == PC_PoseAsset)
 	{
 		return LOCTEXT("Pose_Pin_Category", "PoseAsset");
 	}
-	else if (PinCategory == UEdGraphSchema_CustomizableObject::PC_Transform)
+	else if (PinCategory == PC_Transform)
 	{
 		return LOCTEXT("Transform_Pin_Category", "Transform");
 	}
-	else if (PinCategory == UEdGraphSchema_CustomizableObject::PC_String)
+	else if (PinCategory == PC_String)
 	{
 		return LOCTEXT("String_Pin_Category", "String");
 	}
@@ -1598,8 +1582,8 @@ FText UEdGraphSchema_CustomizableObject::GetPinCategoryFriendlyName(const FName&
 
 bool UEdGraphSchema_CustomizableObject::IsPassthrough(const FName& PinCategory)
 {
-	return PinCategory == PC_PassThroughImage ||
-		PinCategory == PC_PassThroughMesh;
+	return PinCategory == PC_PassthroughTexture ||
+		PinCategory == PC_PassthroughSkeletalMesh;
 }
 
 
diff --git a/Engine/Plugins/Mutable/Source/CustomizableObjectEditor/Private/MuCOE/GenerateMutableSource/GenerateMutableSource.cpp b/Engine/Plugins/Mutable/Source/CustomizableObjectEditor/Private/MuCOE/GenerateMutableSource/GenerateMutableSource.cpp
--- a/Engine/Plugins/Mutable/Source/CustomizableObjectEditor/Private/MuCOE/GenerateMutableSource/GenerateMutableSource.cpp
+++ b/Engine/Plugins/Mutable/Source/CustomizableObjectEditor/Private/MuCOE/GenerateMutableSource/GenerateMutableSource.cpp
@@ -487,7 +487,7 @@ UTexture2D* FindReferenceImage(const UEdGraphPin* Pin, FMutableGraphGenerationCo
 		{
 			if (const UEdGraphPin* ConnectedPin = FollowInputPin(*TypedNodeLayer->LayerPin(LayerIndex)))
 			{
-				if (ConnectedPin->PinType.PinCategory == Schema->PC_Image)
+				if (ConnectedPin->PinType.PinCategory == Schema->PC_Texture)
 				{
 					Result = FindReferenceImage(ConnectedPin, GenerationContext);
 				}
diff --git a/Engine/Plugins/Mutable/Source/CustomizableObjectEditor/Private/MuCOE/GenerateMutableSource/GenerateMutableSourceImage.cpp b/Engine/Plugins/Mutable/Source/CustomizableObjectEditor/Private/MuCOE/GenerateMutableSource/GenerateMutableSourceImage.cpp
--- a/Engine/Plugins/Mutable/Source/CustomizableObjectEditor/Private/MuCOE/GenerateMutableSource/GenerateMutableSourceImage.cpp
+++ b/Engine/Plugins/Mutable/Source/CustomizableObjectEditor/Private/MuCOE/GenerateMutableSource/GenerateMutableSourceImage.cpp
@@ -320,7 +320,7 @@ mu::Ptr<mu::NodeImage> GenerateMutableSourceImage(const UEdGraphPin* Pin, FMutab
 					GenerationContext.Log(LOCTEXT("ModulateWithoutMask", "Texture layer effect uses Modulate without a mask. It will replace everything below it!"), Node);
 				}
 				
-				if (OtherPin->PinType.PinCategory == Schema->PC_Image)
+				if (OtherPin->PinType.PinCategory == Schema->PC_Texture)
 				{
 					mu::Ptr<mu::NodeImage> BlendNode = GenerateMutableSourceImage(OtherPin, GenerationContext, ReferenceTextureSize);
 
@@ -1092,7 +1092,7 @@ mu::Ptr<mu::NodeImage> GenerateMutableSourceImage(const UEdGraphPin* Pin, FMutab
 							}
 
 							// Encoding the texture mode of the type to allow different texture modes from the same pin
-							if (Pin->PinType.PinCategory == Schema->PC_PassThroughImage)
+							if (Pin->PinType.PinCategory == Schema->PC_PassthroughTexture)
 							{
 								ColumnName += "--PassThrough";
 							}
@@ -1124,7 +1124,7 @@ mu::Ptr<mu::NodeImage> GenerateMutableSourceImage(const UEdGraphPin* Pin, FMutab
 							ImageTableNode->bNoneOption = TypedNodeTable->bAddNoneOption;
 							ImageTableNode->DefaultRowName = TypedNodeTable->DefaultRowName.ToString();
 
-							bool bIsPassthrough = (Pin->PinType.PinCategory == Schema->PC_PassThroughImage);
+							bool bIsPassthrough = (Pin->PinType.PinCategory == Schema->PC_PassthroughTexture);
 
 							UTexture2D* DefaultTexture2D = Cast<UTexture2D>(DefaultTexture);
 							if (!bIsPassthrough && DefaultTexture2D)
diff --git a/Engine/Plugins/Mutable/Source/CustomizableObjectEditor/Private/MuCOE/GraphTraversal.cpp b/Engine/Plugins/Mutable/Source/CustomizableObjectEditor/Private/MuCOE/GraphTraversal.cpp
--- a/Engine/Plugins/Mutable/Source/CustomizableObjectEditor/Private/MuCOE/GraphTraversal.cpp
+++ b/Engine/Plugins/Mutable/Source/CustomizableObjectEditor/Private/MuCOE/GraphTraversal.cpp
@@ -586,8 +586,6 @@ const UEdGraphPin* FindMeshBaseSource(const UEdGraphPin& Pin, const bool bOnlyLo
 {
 	check(Pin.Direction == EGPD_Output);
 	check(Pin.PinType.PinCategory == UEdGraphSchema_CustomizableObject::PC_Mesh
-		||
-		Pin.PinType.PinCategory == UEdGraphSchema_CustomizableObject::PC_PassThroughMesh
 		||
 		Pin.PinType.PinCategory == UEdGraphSchema_CustomizableObject::PC_MeshSection
 		||
diff --git a/Engine/Plugins/Mutable/Source/CustomizableObjectEditor/Private/MuCOE/Nodes/CONodeMaterialConstant.cpp b/Engine/Plugins/Mutable/Source/CustomizableObjectEditor/Private/MuCOE/Nodes/CONodeMaterialConstant.cpp
--- a/Engine/Plugins/Mutable/Source/CustomizableObjectEditor/Private/MuCOE/Nodes/CONodeMaterialConstant.cpp
+++ b/Engine/Plugins/Mutable/Source/CustomizableObjectEditor/Private/MuCOE/Nodes/CONodeMaterialConstant.cpp
@@ -31,7 +31,7 @@ FText UCONodeMaterialConstant::GetNodeTitle(ENodeTitleType::Type TitleType) cons
 
 FLinearColor UCONodeMaterialConstant::GetNodeTitleColor() const
 {
-	return UEdGraphSchema_CustomizableObject::GetPinTypeColor(UEdGraphSchema_CustomizableObject::PC_Material);
+	return GetDefault<UEdGraphSchema_CustomizableObject>()->GetPinTypeColor(UEdGraphSchema_CustomizableObject::PC_Material);
 }
 
 
diff --git a/Engine/Plugins/Mutable/Source/CustomizableObjectEditor/Private/MuCOE/Nodes/CustomizableObjectNodeAnimationPose.cpp b/Engine/Plugins/Mutable/Source/CustomizableObjectEditor/Private/MuCOE/Nodes/CustomizableObjectNodeAnimationPose.cpp
--- a/Engine/Plugins/Mutable/Source/CustomizableObjectEditor/Private/MuCOE/Nodes/CustomizableObjectNodeAnimationPose.cpp
+++ b/Engine/Plugins/Mutable/Source/CustomizableObjectEditor/Private/MuCOE/Nodes/CustomizableObjectNodeAnimationPose.cpp
@@ -96,7 +96,7 @@ FText UCustomizableObjectNodeAnimationPose::GetNodeTitle(ENodeTitleType::Type Ti
 
 FLinearColor UCustomizableObjectNodeAnimationPose::GetNodeTitleColor() const
 {
-	return UEdGraphSchema_CustomizableObject::GetPinTypeColor(UEdGraphSchema_CustomizableObject::PC_Mesh);
+	return GetDefault<UEdGraphSchema_CustomizableObject>()->GetPinTypeColor(UEdGraphSchema_CustomizableObject::PC_Mesh);
 }
 
 
diff --git a/Engine/Plugins/Mutable/Source/CustomizableObjectEditor/Private/MuCOE/Nodes/CustomizableObjectNodeComponent.cpp b/Engine/Plugins/Mutable/Source/CustomizableObjectEditor/Private/MuCOE/Nodes/CustomizableObjectNodeComponent.cpp
--- a/Engine/Plugins/Mutable/Source/CustomizableObjectEditor/Private/MuCOE/Nodes/CustomizableObjectNodeComponent.cpp
+++ b/Engine/Plugins/Mutable/Source/CustomizableObjectEditor/Private/MuCOE/Nodes/CustomizableObjectNodeComponent.cpp
@@ -20,7 +20,7 @@ void UCustomizableObjectNodeComponent::AllocateDefaultPins(UCustomizableObjectNo
 
 FLinearColor UCustomizableObjectNodeComponent::GetNodeTitleColor() const
 {
-	return UEdGraphSchema_CustomizableObject::GetPinTypeColor(UEdGraphSchema_CustomizableObject::PC_Component);
+	return GetDefault<UEdGraphSchema_CustomizableObject>()->GetPinTypeColor(UEdGraphSchema_CustomizableObject::PC_Component);
 }
 
 
diff --git a/Engine/Plugins/Mutable/Source/CustomizableObjectEditor/Private/MuCOE/Nodes/CustomizableObjectNodeComponentMeshAddTo.cpp b/Engine/Plugins/Mutable/Source/CustomizableObjectEditor/Private/MuCOE/Nodes/CustomizableObjectNodeComponentMeshAddTo.cpp
--- a/Engine/Plugins/Mutable/Source/CustomizableObjectEditor/Private/MuCOE/Nodes/CustomizableObjectNodeComponentMeshAddTo.cpp
+++ b/Engine/Plugins/Mutable/Source/CustomizableObjectEditor/Private/MuCOE/Nodes/CustomizableObjectNodeComponentMeshAddTo.cpp
@@ -94,7 +94,7 @@ void UCustomizableObjectNodeComponentMeshAddTo::OnRenameNode(const FString& NewN
 
 FLinearColor UCustomizableObjectNodeComponentMeshAddTo::GetNodeTitleColor() const
 {
-	return UEdGraphSchema_CustomizableObject::GetPinTypeColor(UEdGraphSchema_CustomizableObject::PC_Component);
+	return GetDefault<UEdGraphSchema_CustomizableObject>()->GetPinTypeColor(UEdGraphSchema_CustomizableObject::PC_Component);
 }
 
 
diff --git a/Engine/Plugins/Mutable/Source/CustomizableObjectEditor/Private/MuCOE/Nodes/CustomizableObjectNodeExposePin.cpp b/Engine/Plugins/Mutable/Source/CustomizableObjectEditor/Private/MuCOE/Nodes/CustomizableObjectNodeExposePin.cpp
--- a/Engine/Plugins/Mutable/Source/CustomizableObjectEditor/Private/MuCOE/Nodes/CustomizableObjectNodeExposePin.cpp
+++ b/Engine/Plugins/Mutable/Source/CustomizableObjectEditor/Private/MuCOE/Nodes/CustomizableObjectNodeExposePin.cpp
@@ -111,17 +111,12 @@ void UCustomizableObjectNodeExposePin::BackwardsCompatibleFixup(int32 Customizab
 				Pin->PinName = TEXT("Mesh");
 				Pin->PinFriendlyName = LOCTEXT("Mesh_Pin_Category", "Mesh");
 			}
-			else if (PinType == UEdGraphSchema_CustomizableObject::PC_Layout)
-			{
-				Pin->PinName = TEXT("Layout");
-				Pin->PinFriendlyName = LOCTEXT("Layout_Pin_Category", "Layout");
-			}
-			else if (PinType == UEdGraphSchema_CustomizableObject::PC_Image)
+			else if (PinType == UEdGraphSchema_CustomizableObject::PC_Texture)
 			{
 				Pin->PinName = TEXT("Texture");
 				Pin->PinFriendlyName = LOCTEXT("Image_Pin_Category", "Texture");
 			}
-			else if (PinType == UEdGraphSchema_CustomizableObject::PC_PassThroughImage)
+			else if (PinType == UEdGraphSchema_CustomizableObject::PC_PassthroughTexture)
 			{
 				Pin->PinName = TEXT("PassThrough Texture");
 				Pin->PinFriendlyName = LOCTEXT("PassThrough_Image_Pin_Category", "PassThrough Texture");
diff --git a/Engine/Plugins/Mutable/Source/CustomizableObjectEditor/Private/MuCOE/Nodes/CustomizableObjectNodeExternalPin.cpp b/Engine/Plugins/Mutable/Source/CustomizableObjectEditor/Private/MuCOE/Nodes/CustomizableObjectNodeExternalPin.cpp
--- a/Engine/Plugins/Mutable/Source/CustomizableObjectEditor/Private/MuCOE/Nodes/CustomizableObjectNodeExternalPin.cpp
+++ b/Engine/Plugins/Mutable/Source/CustomizableObjectEditor/Private/MuCOE/Nodes/CustomizableObjectNodeExternalPin.cpp
@@ -69,17 +69,12 @@ void UCustomizableObjectNodeExternalPin::BackwardsCompatibleFixup(int32 Customiz
 				Pin->PinName = TEXT("Mesh");
 				Pin->PinFriendlyName = LOCTEXT("Mesh_Pin_Category", "Mesh");
 			}
-			else if (PinType == UEdGraphSchema_CustomizableObject::PC_Layout)
-			{
-				Pin->PinName = TEXT("Layout");
-				Pin->PinFriendlyName = LOCTEXT("Layout_Pin_Category", "Layout");
-			}
-			else if (PinType == UEdGraphSchema_CustomizableObject::PC_Image)
+			else if (PinType == UEdGraphSchema_CustomizableObject::PC_Texture)
 			{
 				Pin->PinName = TEXT("Texture");
 				Pin->PinFriendlyName = LOCTEXT("Image_Pin_Category", "Texture");
 			}
-			else if (PinType == UEdGraphSchema_CustomizableObject::PC_PassThroughImage)
+			else if (PinType == UEdGraphSchema_CustomizableObject::PC_PassthroughTexture)
 			{
 				Pin->PinName = TEXT("PassThrough Texture");
 				Pin->PinFriendlyName = LOCTEXT("PassThrough_Image_Pin_Category", "PassThrough Texture");
diff --git a/Engine/Plugins/Mutable/Source/CustomizableObjectEditor/Private/MuCOE/Nodes/CustomizableObjectNodeLayoutBlocks.cpp b/Engine/Plugins/Mutable/Source/CustomizableObjectEditor/Private/MuCOE/Nodes/CustomizableObjectNodeLayoutBlocks.cpp
--- a/Engine/Plugins/Mutable/Source/CustomizableObjectEditor/Private/MuCOE/Nodes/CustomizableObjectNodeLayoutBlocks.cpp
+++ b/Engine/Plugins/Mutable/Source/CustomizableObjectEditor/Private/MuCOE/Nodes/CustomizableObjectNodeLayoutBlocks.cpp
@@ -37,14 +37,6 @@ void UCustomizableObjectNodeLayoutBlocks::BackwardsCompatibleFixup(int32 Customi
 }
 
 
-void UCustomizableObjectNodeLayoutBlocks::AllocateDefaultPins(UCustomizableObjectNodeRemapPins* RemapPins)
-{
-	const UEdGraphSchema_CustomizableObject* Schema = GetDefault<UEdGraphSchema_CustomizableObject>();
-
-	UEdGraphPin* OutputPin = CustomCreatePin(EGPD_Output, Schema->PC_Layout, FName("Layout"));
-	OutputPin->bDefaultValueIsIgnored = true;
-}
-
 void UCustomizableObjectNodeLayoutBlocks::Serialize(FArchive& Ar)
 {
 	Super::Serialize(Ar);
diff --git a/Engine/Plugins/Mutable/Source/CustomizableObjectEditor/Private/MuCOE/Nodes/CustomizableObjectNodeLayoutBlocks.h b/Engine/Plugins/Mutable/Source/CustomizableObjectEditor/Private/MuCOE/Nodes/CustomizableObjectNodeLayoutBlocks.h
--- a/Engine/Plugins/Mutable/Source/CustomizableObjectEditor/Private/MuCOE/Nodes/CustomizableObjectNodeLayoutBlocks.h
+++ b/Engine/Plugins/Mutable/Source/CustomizableObjectEditor/Private/MuCOE/Nodes/CustomizableObjectNodeLayoutBlocks.h
@@ -45,16 +45,10 @@ public:
 	TObjectPtr<UCustomizableObjectLayout> Layout = nullptr;
 
 	// EdGraphNode interface
-	UE_API void Serialize(FArchive& Ar) override;
+	UE_API virtual void Serialize(FArchive& Ar) override;
 
 	// UCustomizableObjectNode interface
 	UE_API virtual void BackwardsCompatibleFixup(int32 CustomizableObjectCustomVersion) override;
-	UE_API void AllocateDefaultPins(UCustomizableObjectNodeRemapPins* RemapPins) override;
-
-	UEdGraphPin* OutputPin() const
-	{
-		return FindPin(TEXT("Layout"));
-	}
 };
 
 #undef UE_API
diff --git a/Engine/Plugins/Mutable/Source/CustomizableObjectEditor/Private/MuCOE/Nodes/CustomizableObjectNodeMaterial.cpp b/Engine/Plugins/Mutable/Source/CustomizableObjectEditor/Private/MuCOE/Nodes/CustomizableObjectNodeMaterial.cpp
--- a/Engine/Plugins/Mutable/Source/CustomizableObjectEditor/Private/MuCOE/Nodes/CustomizableObjectNodeMaterial.cpp
+++ b/Engine/Plugins/Mutable/Source/CustomizableObjectEditor/Private/MuCOE/Nodes/CustomizableObjectNodeMaterial.cpp
@@ -176,9 +176,9 @@ FName UCustomizableObjectNodeMaterial::NodePinModeToImagePinMode(const ENodePinM
 	switch (NodePinMode)
 	{
 	case ENodePinMode::Mutable:
-		return UEdGraphSchema_CustomizableObject::PC_Image;
+		return UEdGraphSchema_CustomizableObject::PC_Texture;
 	case ENodePinMode::Passthrough:
-		return UEdGraphSchema_CustomizableObject::PC_PassThroughImage;
+		return UEdGraphSchema_CustomizableObject::PC_PassthroughTexture;
 	default:
 		unimplemented();
 		return {};
@@ -193,9 +193,9 @@ FName UCustomizableObjectNodeMaterial::GetImagePinMode(const EPinMode PinMode) c
 	case EPinMode::Default:
 		return NodePinModeToImagePinMode(TextureParametersMode);
 	case EPinMode::Mutable:
-		return UEdGraphSchema_CustomizableObject::PC_Image;
+		return UEdGraphSchema_CustomizableObject::PC_Texture;
 	case EPinMode::Passthrough:
-		return UEdGraphSchema_CustomizableObject::PC_PassThroughImage;
+		return UEdGraphSchema_CustomizableObject::PC_PassthroughTexture;
 	default:
 		unimplemented();
 		return {};
@@ -503,14 +503,14 @@ void UCustomizableObjectNodeMaterial::BackwardsCompatibleFixup(int32 Customizabl
 		// Check if there are still pins which where not present in the Images, ScalarParams and ScalarParams arrays.
 		for (const UEdGraphPin* Pin : GetAllNonOrphanPins())
 		{
-			if ((Pin->PinType.PinCategory == UEdGraphSchema_CustomizableObject::PC_Image ||
+			if ((Pin->PinType.PinCategory == UEdGraphSchema_CustomizableObject::PC_Texture ||
  				 Pin->PinType.PinCategory == UEdGraphSchema_CustomizableObject::PC_Color ||
 				 Pin->PinType.PinCategory == UEdGraphSchema_CustomizableObject::PC_Float) &&
 				!GetPinData(*Pin))
 			{
 				UCustomizableObjectNodeMaterialPinDataParameter* PinData = [&](UObject* Outer) -> UCustomizableObjectNodeMaterialPinDataParameter*
 				{
-					if (Pin->PinType.PinCategory == UEdGraphSchema_CustomizableObject::PC_Image)
+					if (Pin->PinType.PinCategory == UEdGraphSchema_CustomizableObject::PC_Texture)
 					{
 						return NewObject<UCustomizableObjectNodeMaterialPinDataImage>(Outer);	
 					}
@@ -666,7 +666,7 @@ void UCustomizableObjectNodeMaterial::BackwardsCompatibleFixup(int32 Customizabl
 		
 		for (UEdGraphPin* Pin : GetAllPins())
 		{
-			if (Pin->PinType.PinCategory != UEdGraphSchema_CustomizableObject::PC_Image)
+			if (Pin->PinType.PinCategory != UEdGraphSchema_CustomizableObject::PC_Texture)
 			{
 				continue;
 			}
@@ -677,9 +677,9 @@ void UCustomizableObjectNodeMaterial::BackwardsCompatibleFixup(int32 Customizabl
 
 				UCustomizableObjectNodeMaterialPinDataImage& PinData = GetPinData<UCustomizableObjectNodeMaterialPinDataImage>(*Pin);
 
-				if (LinkedPin->PinType.PinCategory == UEdGraphSchema_CustomizableObject::PC_Image)
+				if (LinkedPin->PinType.PinCategory == UEdGraphSchema_CustomizableObject::PC_Texture)
 				{
-					Pin->PinType.PinCategory = UEdGraphSchema_CustomizableObject::PC_Image;
+					Pin->PinType.PinCategory = UEdGraphSchema_CustomizableObject::PC_Texture;
 
 					if (PinData.PinMode == EPinMode::Default &&
 						TextureParametersMode == ENodePinMode::Passthrough)
@@ -691,9 +691,9 @@ void UCustomizableObjectNodeMaterial::BackwardsCompatibleFixup(int32 Customizabl
 						PinData.PinMode = EPinMode::Mutable;
 					}
 				}
-				else if (LinkedPin->PinType.PinCategory == UEdGraphSchema_CustomizableObject::PC_PassThroughImage)
+				else if (LinkedPin->PinType.PinCategory == UEdGraphSchema_CustomizableObject::PC_PassthroughTexture)
 				{
-					Pin->PinType.PinCategory = UEdGraphSchema_CustomizableObject::PC_PassThroughImage;
+					Pin->PinType.PinCategory = UEdGraphSchema_CustomizableObject::PC_PassthroughTexture;
 
 					if (PinData.PinMode == EPinMode::Default &&
 						TextureParametersMode == ENodePinMode::Mutable)
@@ -710,7 +710,7 @@ void UCustomizableObjectNodeMaterial::BackwardsCompatibleFixup(int32 Customizabl
 			{
 				if (!IsImageMutableMode(*Pin))
 				{
-					Pin->PinType.PinCategory = UEdGraphSchema_CustomizableObject::PC_PassThroughImage;
+					Pin->PinType.PinCategory = UEdGraphSchema_CustomizableObject::PC_PassthroughTexture;
 				}				
 			}
 		}
@@ -950,14 +950,14 @@ bool UCustomizableObjectNodeMaterial::IsImageMutableMode(const int32 ImageIndex)
 	}
 	else
 	{
-		return NodePinModeToImagePinMode(TextureParametersMode) == UEdGraphSchema_CustomizableObject::PC_Image;
+		return NodePinModeToImagePinMode(TextureParametersMode) == UEdGraphSchema_CustomizableObject::PC_Texture;
 	}
 }
 
 
 bool UCustomizableObjectNodeMaterial::IsImageMutableMode(const UEdGraphPin& Pin) const
 {
-	return GetImagePinMode(Pin) == UEdGraphSchema_CustomizableObject::PC_Image; // The ImageMutableMode is stored in the PinData, not in the PinCategory
+	return GetImagePinMode(Pin) == UEdGraphSchema_CustomizableObject::PC_Texture; // The ImageMutableMode is stored in the PinData, not in the PinCategory
 }
 
 
@@ -1434,11 +1434,11 @@ void UCustomizableObjectNodeMaterial::AllocateDefaultParameterPins(const EMateri
 		case EMaterialParameterType::Texture:
 			if (IsImageMutableMode(ParameterIndex)) // If a pin exists, we store the PinMode in its PinData.
 			{
-				PinCategory = UEdGraphSchema_CustomizableObject::PC_Image;				
+				PinCategory = UEdGraphSchema_CustomizableObject::PC_Texture;				
 			}
 			else
 			{
-				PinCategory = UEdGraphSchema_CustomizableObject::PC_PassThroughImage;
+				PinCategory = UEdGraphSchema_CustomizableObject::PC_PassthroughTexture;
 			}
 			break;
 
@@ -1526,11 +1526,11 @@ void UCustomizableObjectNodeMaterial::PinConnectionListChanged(UEdGraphPin* Pin)
 			{
 				EPinMode PinMode = EPinMode::Default;
 		
-				if (LinkedPin->PinType.PinCategory == UEdGraphSchema_CustomizableObject::PC_Image)
+				if (LinkedPin->PinType.PinCategory == UEdGraphSchema_CustomizableObject::PC_Texture)
 				{
 					PinMode = EPinMode::Mutable;
 				}
-				else if (LinkedPin->PinType.PinCategory == UEdGraphSchema_CustomizableObject::PC_PassThroughImage)
+				else if (LinkedPin->PinType.PinCategory == UEdGraphSchema_CustomizableObject::PC_PassthroughTexture)
 				{
 					PinMode = EPinMode::Passthrough;
 				}
@@ -1559,14 +1559,13 @@ bool UCustomizableObjectNodeMaterial::CanConnect(const UEdGraphPin* InOwnedInput
 
 	if (InOwnedInputPin && InOutputPin)
 	{
-		if ((InOwnedInputPin->PinType.PinCategory == Schema->PC_Image && InOutputPin->PinType.PinCategory == Schema->PC_PassThroughImage) ||
-			(InOwnedInputPin->PinType.PinCategory == Schema->PC_PassThroughImage && InOutputPin->PinType.PinCategory == Schema->PC_Image))
+		if ((InOwnedInputPin->PinType.PinCategory == Schema->PC_Texture && InOutputPin->PinType.PinCategory == Schema->PC_PassthroughTexture) ||
+			(InOwnedInputPin->PinType.PinCategory == Schema->PC_PassthroughTexture && InOutputPin->PinType.PinCategory == Schema->PC_Texture))
 		{
 			return true;
 		}
 		
-		if (InOwnedInputPin->PinType.PinCategory == Schema->PC_Mesh &&
-			InOutputPin->PinType.PinCategory == Schema->PC_PassThroughMesh)
+		if (InOwnedInputPin->PinType.PinCategory == Schema->PC_Mesh)
 		{
 			return true;
 		}
diff --git a/Engine/Plugins/Mutable/Source/CustomizableObjectEditor/Private/MuCOE/Nodes/CustomizableObjectNodeMaterialBase.cpp b/Engine/Plugins/Mutable/Source/CustomizableObjectEditor/Private/MuCOE/Nodes/CustomizableObjectNodeMaterialBase.cpp
--- a/Engine/Plugins/Mutable/Source/CustomizableObjectEditor/Private/MuCOE/Nodes/CustomizableObjectNodeMaterialBase.cpp
+++ b/Engine/Plugins/Mutable/Source/CustomizableObjectEditor/Private/MuCOE/Nodes/CustomizableObjectNodeMaterialBase.cpp
@@ -13,5 +13,5 @@ int32 UCustomizableObjectNodeMaterialBase::GetMeshSectionMaxLOD() const
 
 FLinearColor UCustomizableObjectNodeMaterialBase::GetNodeTitleColor() const
 {
-	return UEdGraphSchema_CustomizableObject::GetPinTypeColor(UEdGraphSchema_CustomizableObject::PC_MeshSection);
+	return GetDefault<UEdGraphSchema_CustomizableObject>()->GetPinTypeColor(UEdGraphSchema_CustomizableObject::PC_MeshSection);
 }
diff --git a/Engine/Plugins/Mutable/Source/CustomizableObjectEditor/Private/MuCOE/Nodes/CustomizableObjectNodeMeshMorphStackApplication.cpp b/Engine/Plugins/Mutable/Source/CustomizableObjectEditor/Private/MuCOE/Nodes/CustomizableObjectNodeMeshMorphStackApplication.cpp
--- a/Engine/Plugins/Mutable/Source/CustomizableObjectEditor/Private/MuCOE/Nodes/CustomizableObjectNodeMeshMorphStackApplication.cpp
+++ b/Engine/Plugins/Mutable/Source/CustomizableObjectEditor/Private/MuCOE/Nodes/CustomizableObjectNodeMeshMorphStackApplication.cpp
@@ -64,7 +64,7 @@ FText UCustomizableObjectNodeMeshMorphStackApplication::GetNodeTitle(ENodeTitleT
 
 FLinearColor UCustomizableObjectNodeMeshMorphStackApplication::GetNodeTitleColor() const
 {
-	return UEdGraphSchema_CustomizableObject::GetPinTypeColor(UEdGraphSchema_CustomizableObject::PC_Mesh);
+	return GetDefault<UEdGraphSchema_CustomizableObject>()->GetPinTypeColor(UEdGraphSchema_CustomizableObject::PC_Mesh);
 }
 
 
diff --git a/Engine/Plugins/Mutable/Source/CustomizableObjectEditor/Private/MuCOE/Nodes/CustomizableObjectNodeModifierClipWithUVMask.cpp b/Engine/Plugins/Mutable/Source/CustomizableObjectEditor/Private/MuCOE/Nodes/CustomizableObjectNodeModifierClipWithUVMask.cpp
--- a/Engine/Plugins/Mutable/Source/CustomizableObjectEditor/Private/MuCOE/Nodes/CustomizableObjectNodeModifierClipWithUVMask.cpp
+++ b/Engine/Plugins/Mutable/Source/CustomizableObjectEditor/Private/MuCOE/Nodes/CustomizableObjectNodeModifierClipWithUVMask.cpp
@@ -18,7 +18,7 @@ void UCustomizableObjectNodeModifierClipWithUVMask::AllocateDefaultPins(UCustomi
 {
 	const UEdGraphSchema_CustomizableObject* Schema = GetDefault<UEdGraphSchema_CustomizableObject>();
 
-	UEdGraphPin* ClipMaskPin = CustomCreatePin(EGPD_Input, Schema->PC_Image, FName("Clip Mask"));
+	UEdGraphPin* ClipMaskPin = CustomCreatePin(EGPD_Input, Schema->PC_Texture, FName("Clip Mask"));
 	ClipMaskPin->bDefaultValueIsIgnored = true;
 
 	UEdGraphPin* OutputPin_p = CustomCreatePin(EGPD_Output, Schema->PC_Modifier, FName("Modifier"));
diff --git a/Engine/Plugins/Mutable/Source/CustomizableObjectEditor/Private/MuCOE/Nodes/CustomizableObjectNodeModifierEditMeshSection.cpp b/Engine/Plugins/Mutable/Source/CustomizableObjectEditor/Private/MuCOE/Nodes/CustomizableObjectNodeModifierEditMeshSection.cpp
--- a/Engine/Plugins/Mutable/Source/CustomizableObjectEditor/Private/MuCOE/Nodes/CustomizableObjectNodeModifierEditMeshSection.cpp
+++ b/Engine/Plugins/Mutable/Source/CustomizableObjectEditor/Private/MuCOE/Nodes/CustomizableObjectNodeModifierEditMeshSection.cpp
@@ -34,14 +34,14 @@ void UCustomizableObjectNodeModifierEditMeshSection::AllocateDefaultPins(UCustom
 			PinEditImageData->ImageParamId = GetParameterId(EMaterialParameterType::Texture, ImageIndex);
 			
 			const FName ImageName = GetParameterName(EMaterialParameterType::Texture, ImageIndex);
-			UEdGraphPin* PinImage = CustomCreatePin(EGPD_Input, Schema->PC_Image, ImageName, PinEditImageData);
+			UEdGraphPin* PinImage = CustomCreatePin(EGPD_Input, Schema->PC_Texture, ImageName, PinEditImageData);
 			PinImage->bHidden = true;
 			PinImage->bDefaultValueIsIgnored = true;
 
 			PinsParameterMap.Add(PinEditImageData->ImageParamId, FEdGraphPinReference(PinImage));
 
 			FString PinMaskName = ImageName.ToString() + FString(TEXT(" Mask"));
-			UEdGraphPin* PinMask = CustomCreatePin(EGPD_Input, Schema->PC_Image, *PinMaskName);
+			UEdGraphPin* PinMask = CustomCreatePin(EGPD_Input, Schema->PC_Texture, *PinMaskName);
 			PinMask->bHidden = true;
 			PinMask->bDefaultValueIsIgnored = true;
 			
diff --git a/Engine/Plugins/Mutable/Source/CustomizableObjectEditor/Private/MuCOE/Nodes/CustomizableObjectNodeModifierExtendMeshSection.cpp b/Engine/Plugins/Mutable/Source/CustomizableObjectEditor/Private/MuCOE/Nodes/CustomizableObjectNodeModifierExtendMeshSection.cpp
--- a/Engine/Plugins/Mutable/Source/CustomizableObjectEditor/Private/MuCOE/Nodes/CustomizableObjectNodeModifierExtendMeshSection.cpp
+++ b/Engine/Plugins/Mutable/Source/CustomizableObjectEditor/Private/MuCOE/Nodes/CustomizableObjectNodeModifierExtendMeshSection.cpp
@@ -145,7 +145,7 @@ void UCustomizableObjectNodeModifierExtendMeshSection::AllocateDefaultPins(UCust
 		for (int32 ImageIndex = 0; ImageIndex < NumImages; ++ImageIndex)
 		{
 			const FName ImageName = GetParameterName(EMaterialParameterType::Texture, ImageIndex);
-			UEdGraphPin* PinImage = CustomCreatePin(EGPD_Input, Schema->PC_Image, ImageName);
+			UEdGraphPin* PinImage = CustomCreatePin(EGPD_Input, Schema->PC_Texture, ImageName);
 			PinImage->bDefaultValueIsIgnored = true;
 
 			const FNodeMaterialParameterId ImageId = GetParameterId(EMaterialParameterType::Texture, ImageIndex);
diff --git a/Engine/Plugins/Mutable/Source/CustomizableObjectEditor/Private/MuCOE/Nodes/CustomizableObjectNodeObjectGroup.cpp b/Engine/Plugins/Mutable/Source/CustomizableObjectEditor/Private/MuCOE/Nodes/CustomizableObjectNodeObjectGroup.cpp
--- a/Engine/Plugins/Mutable/Source/CustomizableObjectEditor/Private/MuCOE/Nodes/CustomizableObjectNodeObjectGroup.cpp
+++ b/Engine/Plugins/Mutable/Source/CustomizableObjectEditor/Private/MuCOE/Nodes/CustomizableObjectNodeObjectGroup.cpp
@@ -96,7 +96,7 @@ FText UCustomizableObjectNodeObjectGroup::GetNodeTitle(ENodeTitleType::Type Titl
 
 FLinearColor UCustomizableObjectNodeObjectGroup::GetNodeTitleColor() const
 {
-	return UEdGraphSchema_CustomizableObject::GetPinTypeColor(UEdGraphSchema_CustomizableObject::PC_Object);
+	return GetDefault<UEdGraphSchema_CustomizableObject>()->GetPinTypeColor(UEdGraphSchema_CustomizableObject::PC_Object);
 }
 
 
diff --git a/Engine/Plugins/Mutable/Source/CustomizableObjectEditor/Private/MuCOE/Nodes/CustomizableObjectNodeParameter.cpp b/Engine/Plugins/Mutable/Source/CustomizableObjectEditor/Private/MuCOE/Nodes/CustomizableObjectNodeParameter.cpp
--- a/Engine/Plugins/Mutable/Source/CustomizableObjectEditor/Private/MuCOE/Nodes/CustomizableObjectNodeParameter.cpp
+++ b/Engine/Plugins/Mutable/Source/CustomizableObjectEditor/Private/MuCOE/Nodes/CustomizableObjectNodeParameter.cpp
@@ -55,7 +55,7 @@ FText UCustomizableObjectNodeParameter::GetNodeTitle(ENodeTitleType::Type TitleT
 
 FLinearColor UCustomizableObjectNodeParameter::GetNodeTitleColor() const
 {
-	return UEdGraphSchema_CustomizableObject::GetPinTypeColor(GetCategory());
+	return GetDefault<UEdGraphSchema_CustomizableObject>()->GetPinTypeColor(GetCategory());
 }
 
 
diff --git a/Engine/Plugins/Mutable/Source/CustomizableObjectEditor/Private/MuCOE/Nodes/CustomizableObjectNodePassThroughTexture.cpp b/Engine/Plugins/Mutable/Source/CustomizableObjectEditor/Private/MuCOE/Nodes/CustomizableObjectNodePassThroughTexture.cpp
--- a/Engine/Plugins/Mutable/Source/CustomizableObjectEditor/Private/MuCOE/Nodes/CustomizableObjectNodePassThroughTexture.cpp
+++ b/Engine/Plugins/Mutable/Source/CustomizableObjectEditor/Private/MuCOE/Nodes/CustomizableObjectNodePassThroughTexture.cpp
@@ -39,7 +39,7 @@ void UCustomizableObjectNodePassThroughTexture::BackwardsCompatibleFixup(int32 C
 
 void UCustomizableObjectNodePassThroughTexture::AllocateDefaultPins(UCustomizableObjectNodeRemapPins* RemapPins)
 {
-	UEdGraphPin* PinImagePin = CustomCreatePinSimple(EGPD_Output, UEdGraphSchema_CustomizableObject::PC_PassThroughImage);
+	UEdGraphPin* PinImagePin = CustomCreatePinSimple(EGPD_Output, UEdGraphSchema_CustomizableObject::PC_PassthroughTexture);
 	PinImagePin->bDefaultValueIsIgnored = true;
 }
 
@@ -62,7 +62,7 @@ FText UCustomizableObjectNodePassThroughTexture::GetNodeTitle(ENodeTitleType::Ty
 
 FLinearColor UCustomizableObjectNodePassThroughTexture::GetNodeTitleColor() const
 {
-	return UEdGraphSchema_CustomizableObject::GetPinTypeColor(UEdGraphSchema_CustomizableObject::PC_PassThroughImage);
+	return GetDefault<UEdGraphSchema_CustomizableObject>()->GetPinTypeColor(UEdGraphSchema_CustomizableObject::PC_PassthroughTexture);
 }
 
 
diff --git a/Engine/Plugins/Mutable/Source/CustomizableObjectEditor/Private/MuCOE/Nodes/CustomizableObjectNodePassThroughTextureSwitch.cpp b/Engine/Plugins/Mutable/Source/CustomizableObjectEditor/Private/MuCOE/Nodes/CustomizableObjectNodePassThroughTextureSwitch.cpp
--- a/Engine/Plugins/Mutable/Source/CustomizableObjectEditor/Private/MuCOE/Nodes/CustomizableObjectNodePassThroughTextureSwitch.cpp
+++ b/Engine/Plugins/Mutable/Source/CustomizableObjectEditor/Private/MuCOE/Nodes/CustomizableObjectNodePassThroughTextureSwitch.cpp
@@ -11,7 +11,7 @@
 
 FName UCustomizableObjectNodePassThroughTextureSwitch::GetCategory() const
 {
-	return UEdGraphSchema_CustomizableObject::PC_PassThroughImage;
+	return UEdGraphSchema_CustomizableObject::PC_PassthroughTexture;
 }
 
 
diff --git a/Engine/Plugins/Mutable/Source/CustomizableObjectEditor/Private/MuCOE/Nodes/CustomizableObjectNodePassThroughTextureVariation.cpp b/Engine/Plugins/Mutable/Source/CustomizableObjectEditor/Private/MuCOE/Nodes/CustomizableObjectNodePassThroughTextureVariation.cpp
--- a/Engine/Plugins/Mutable/Source/CustomizableObjectEditor/Private/MuCOE/Nodes/CustomizableObjectNodePassThroughTextureVariation.cpp
+++ b/Engine/Plugins/Mutable/Source/CustomizableObjectEditor/Private/MuCOE/Nodes/CustomizableObjectNodePassThroughTextureVariation.cpp
@@ -13,7 +13,7 @@
 // UCustomizableObjectNodeVariation interface
 FName UCustomizableObjectNodePassThroughTextureVariation::GetCategory() const
 {
-	return UEdGraphSchema_CustomizableObject::PC_PassThroughImage;
+	return UEdGraphSchema_CustomizableObject::PC_PassthroughTexture;
 }
 
 
diff --git a/Engine/Plugins/Mutable/Source/CustomizableObjectEditor/Private/MuCOE/Nodes/CustomizableObjectNodeSkeletalMesh.cpp b/Engine/Plugins/Mutable/Source/CustomizableObjectEditor/Private/MuCOE/Nodes/CustomizableObjectNodeSkeletalMesh.cpp
--- a/Engine/Plugins/Mutable/Source/CustomizableObjectEditor/Private/MuCOE/Nodes/CustomizableObjectNodeSkeletalMesh.cpp
+++ b/Engine/Plugins/Mutable/Source/CustomizableObjectEditor/Private/MuCOE/Nodes/CustomizableObjectNodeSkeletalMesh.cpp
@@ -144,7 +144,7 @@ void UCustomizableObjectNodeSkeletalMesh::AllocateDefaultPins(UCustomizableObjec
 						const FString ImageNameStr = *ImageInfo.Name.ToString();
 						const FString TexturePinName = FString::Printf(TEXT("LOD %i - Section %i - Texture Parameter %s"), LODIndex, SectionIndex, *ImageNameStr);
 						
-						UEdGraphPin* Pin = CustomCreatePin(EGPD_Output, Schema->PC_Image, FName(*TexturePinName), PinData);
+						UEdGraphPin* Pin = CustomCreatePin(EGPD_Output, Schema->PC_Texture, FName(*TexturePinName), PinData);
 						Pin->PinFriendlyName = FText::FromString(FString::Printf(TEXT("LOD %i - %s - %s"), LODIndex, *Section, *ImageNameStr));
 						Pin->bHidden = true;
 					}
@@ -315,12 +315,7 @@ const FSkelMeshSection* UCustomizableObjectNodeSkeletalMesh::GetSkeletalMeshSect
 bool UCustomizableObjectNodeSkeletalMesh::IsPinRelevant(const UEdGraphPin* Pin) const
 {
 	const UEdGraphSchema_CustomizableObject* Schema = GetDefault<UEdGraphSchema_CustomizableObject>();
-
-	if (Pin->Direction == EEdGraphPinDirection::EGPD_Output)
-	{
-		return Pin->PinType.PinCategory == Schema->PC_Layout;
-	}
-
+	
 	if (Pin->Direction == EEdGraphPinDirection::EGPD_Input)
 	{	
 		return Pin->PinType.PinCategory == Schema->PC_Mesh;
diff --git a/Engine/Plugins/Mutable/Source/CustomizableObjectEditor/Private/MuCOE/Nodes/CustomizableObjectNodeStaticMesh.cpp b/Engine/Plugins/Mutable/Source/CustomizableObjectEditor/Private/MuCOE/Nodes/CustomizableObjectNodeStaticMesh.cpp
--- a/Engine/Plugins/Mutable/Source/CustomizableObjectEditor/Private/MuCOE/Nodes/CustomizableObjectNodeStaticMesh.cpp
+++ b/Engine/Plugins/Mutable/Source/CustomizableObjectEditor/Private/MuCOE/Nodes/CustomizableObjectNodeStaticMesh.cpp
@@ -142,7 +142,7 @@ void UCustomizableObjectNodeStaticMesh::AllocateDefaultPins(UCustomizableObjectN
 					for (int32 ImageIndex = 0; ImageIndex < ImageNames.Num(); ++ImageIndex)
 					{
 						FString ImageName = ImageNames[ImageIndex].Name.ToString();
-						Pin = CustomCreatePin(EGPD_Output, Schema->PC_Image, *ImageName);
+						Pin = CustomCreatePin(EGPD_Output, Schema->PC_Texture, *ImageName);
 						Pin->bDefaultValueIsIgnored = true;
 						LODs[i].Materials[MaterialIndex].ImagePinsRef.Add(Pin);
 					}
@@ -221,31 +221,9 @@ UTexture2D* UCustomizableObjectNodeStaticMesh::FindTextureForPin(const UEdGraphP
 }
 
 
-TArray<class UCustomizableObjectLayout*> UCustomizableObjectNodeStaticMesh::GetLayouts(const UEdGraphPin& OutPin) const
+TArray<UCustomizableObjectLayout*> UCustomizableObjectNodeStaticMesh::GetLayouts(const UEdGraphPin& OutPin) const
 {
-	TArray<class UCustomizableObjectLayout*> Result;
-
-	for (int LODIndex = 0; LODIndex<LODs.Num(); ++LODIndex)
-	{
-		for (int MaterialIndex = 0; MaterialIndex<LODs[LODIndex].Materials.Num(); ++MaterialIndex)
-		{
-			if (LODs[LODIndex].Materials[MaterialIndex].MeshPinRef.Get() == &OutPin)
-			{
-				if (UEdGraphPin* LayoutInPin = LODs[LODIndex].Materials[MaterialIndex].LayoutPinRef.Get())
-				{
-					if (const UEdGraphPin* ConnectedPin = FollowInputPin(*LayoutInPin))
-					{
-						const UCustomizableObjectNodeLayoutBlocks* LayoutNode = Cast<UCustomizableObjectNodeLayoutBlocks>(ConnectedPin->GetOwningNode());
-						if (LayoutNode && LayoutNode->Layout)
-						{
-							Result.Add(LayoutNode->Layout);
-						}
-					}
-				}
-			}
-		}
-	}
-
+	TArray<UCustomizableObjectLayout*> Result;
 	return Result;
 }
 
diff --git a/Engine/Plugins/Mutable/Source/CustomizableObjectEditor/Private/MuCOE/Nodes/CustomizableObjectNodeStaticString.cpp b/Engine/Plugins/Mutable/Source/CustomizableObjectEditor/Private/MuCOE/Nodes/CustomizableObjectNodeStaticString.cpp
--- a/Engine/Plugins/Mutable/Source/CustomizableObjectEditor/Private/MuCOE/Nodes/CustomizableObjectNodeStaticString.cpp
+++ b/Engine/Plugins/Mutable/Source/CustomizableObjectEditor/Private/MuCOE/Nodes/CustomizableObjectNodeStaticString.cpp
@@ -24,7 +24,7 @@ FText UCustomizableObjectNodeStaticString::GetTooltipText() const
 
 FLinearColor UCustomizableObjectNodeStaticString::GetNodeTitleColor() const
 {
-	return UEdGraphSchema_CustomizableObject::GetPinTypeColor(UEdGraphSchema_CustomizableObject::PC_String);
+	return GetDefault<UEdGraphSchema_CustomizableObject>()->GetPinTypeColor(UEdGraphSchema_CustomizableObject::PC_String);
 }
 
 
diff --git a/Engine/Plugins/Mutable/Source/CustomizableObjectEditor/Private/MuCOE/Nodes/CustomizableObjectNodeSwitchBase.cpp b/Engine/Plugins/Mutable/Source/CustomizableObjectEditor/Private/MuCOE/Nodes/CustomizableObjectNodeSwitchBase.cpp
--- a/Engine/Plugins/Mutable/Source/CustomizableObjectEditor/Private/MuCOE/Nodes/CustomizableObjectNodeSwitchBase.cpp
+++ b/Engine/Plugins/Mutable/Source/CustomizableObjectEditor/Private/MuCOE/Nodes/CustomizableObjectNodeSwitchBase.cpp
@@ -113,7 +113,7 @@ FText UCustomizableObjectNodeSwitchBase::GetNodeTitle(ENodeTitleType::Type Title
 
 FLinearColor UCustomizableObjectNodeSwitchBase::GetNodeTitleColor() const
 {
-	return UEdGraphSchema_CustomizableObject::GetPinTypeColor(GetCategory());
+	return GetDefault<UEdGraphSchema_CustomizableObject>()->GetPinTypeColor(GetCategory());
 }
 
 
diff --git a/Engine/Plugins/Mutable/Source/CustomizableObjectEditor/Private/MuCOE/Nodes/CustomizableObjectNodeTable.cpp b/Engine/Plugins/Mutable/Source/CustomizableObjectEditor/Private/MuCOE/Nodes/CustomizableObjectNodeTable.cpp
--- a/Engine/Plugins/Mutable/Source/CustomizableObjectEditor/Private/MuCOE/Nodes/CustomizableObjectNodeTable.cpp
+++ b/Engine/Plugins/Mutable/Source/CustomizableObjectEditor/Private/MuCOE/Nodes/CustomizableObjectNodeTable.cpp
@@ -65,8 +65,8 @@ bool UCustomizableObjectNodeTableRemapPins::Equal(const UCustomizableObjectNode&
 	}
 
 	// In this case pin type may have changed but we consider them the same type
-	if (OldPin.PinType.PinCategory == UEdGraphSchema_CustomizableObject::PC_Image ||
-		OldPin.PinType.PinCategory == UEdGraphSchema_CustomizableObject::PC_PassThroughImage)
+	if (OldPin.PinType.PinCategory == UEdGraphSchema_CustomizableObject::PC_Texture ||
+		OldPin.PinType.PinCategory == UEdGraphSchema_CustomizableObject::PC_PassthroughTexture)
 	{
 		return true;
 	}
@@ -136,7 +136,7 @@ void UCustomizableObjectNodeTable::BackwardsCompatibleFixup(int32 CustomizableOb
 	{
 		for (UEdGraphPin* Pin : Pins)
 		{
-			if (Pin->PinType.PinCategory == UEdGraphSchema_CustomizableObject::PC_Image)
+			if (Pin->PinType.PinCategory == UEdGraphSchema_CustomizableObject::PC_Texture)
 			{
 				UCustomizableObjectNodeTableObjectPinData* OldPinData = Cast<UCustomizableObjectNodeTableObjectPinData>(GetPinData(*(Pin)));
 				UCustomizableObjectNodeTableImagePinData* NewPinData = NewObject<UCustomizableObjectNodeTableImagePinData>(this);
@@ -616,7 +616,7 @@ void UCustomizableObjectNodeTable::AllocateDefaultPins(UCustomizableObjectNodeRe
 						PinData->bIsNotTexture2D = false;
 						PinData->NodeTable = this;
 
-						FName PinCategory = PinData->ImageMode == ETableTextureType::PASSTHROUGH_TEXTURE ? Schema->PC_PassThroughImage : Schema->PC_Image;
+						FName PinCategory = PinData->ImageMode == ETableTextureType::PASSTHROUGH_TEXTURE ? Schema->PC_PassthroughTexture : Schema->PC_Texture;
 
 						for (UEdGraphPin* OldPin : OldPins)
 						{
@@ -625,7 +625,7 @@ void UCustomizableObjectNodeTable::AllocateDefaultPins(UCustomizableObjectNodeRe
 							{
 								if (OldPinData->ColumnPropertyName == PropertyName)
 								{
-									PinCategory = OldPinData->ImageMode == ETableTextureType::PASSTHROUGH_TEXTURE ? Schema->PC_PassThroughImage : Schema->PC_Image;
+									PinCategory = OldPinData->ImageMode == ETableTextureType::PASSTHROUGH_TEXTURE ? Schema->PC_PassthroughTexture : Schema->PC_Texture;
 									break;
 								}
 							}
@@ -646,7 +646,7 @@ void UCustomizableObjectNodeTable::AllocateDefaultPins(UCustomizableObjectNodeRe
 						PinData->bIsNotTexture2D = true;
 						PinData->NodeTable = this;
 
-						OutPin = CustomCreatePin(EGPD_Output, Schema->PC_PassThroughImage, FName(*PinName), PinData);
+						OutPin = CustomCreatePin(EGPD_Output, Schema->PC_PassthroughTexture, FName(*PinName), PinData);
 					}
 
 					else if (Object->IsA(UMaterialInterface::StaticClass()))
@@ -844,7 +844,7 @@ bool UCustomizableObjectNodeTable::IsNodeOutDatedAndNeedsRefresh()
 				}
 				else if (Object->IsA(UTexture2D::StaticClass()))
 				{
-					if (CheckPinUpdated(PinName, Schema->PC_Image) && CheckPinUpdated(PinName, Schema->PC_PassThroughImage))
+					if (CheckPinUpdated(PinName, Schema->PC_Texture) && CheckPinUpdated(PinName, Schema->PC_PassthroughTexture))
 					{
 						bNeedsUpdate = true;
 					}
@@ -853,7 +853,7 @@ bool UCustomizableObjectNodeTable::IsNodeOutDatedAndNeedsRefresh()
 				}
 				else if (Object->IsA(UTexture2DArray::StaticClass()))
 				{
-					if (CheckPinUpdated(PinName, Schema->PC_PassThroughImage))
+					if (CheckPinUpdated(PinName, Schema->PC_PassthroughTexture))
 					{
 						bNeedsUpdate = true;
 					}
@@ -967,7 +967,7 @@ void UCustomizableObjectNodeTable::RemapPinsData(const TMap<UEdGraphPin*, UEdGra
 				}
 			}
 		}
-		else if (Pair.Key->PinType.PinCategory == Schema->PC_Image || Pair.Key->PinType.PinCategory == Schema->PC_PassThroughImage)
+		else if (Pair.Key->PinType.PinCategory == Schema->PC_Texture || Pair.Key->PinType.PinCategory == Schema->PC_PassthroughTexture)
 		{
 			UCustomizableObjectNodeTableImagePinData* PinDataOldPin = Cast<UCustomizableObjectNodeTableImagePinData>(GetPinData(*(Pair.Key)));
 			UCustomizableObjectNodeTableImagePinData* PinDataNewPin = Cast<UCustomizableObjectNodeTableImagePinData>(GetPinData(*(Pair.Value)));
@@ -994,8 +994,8 @@ bool UCustomizableObjectNodeTable::IsPinRelevant(const UEdGraphPin* Pin) const
 	
 	return Pin->Direction == EGPD_Input &&
 		(Pin->PinType.PinCategory == Schema->PC_Material ||
-		Pin->PinType.PinCategory == Schema->PC_Image ||
-		Pin->PinType.PinCategory == Schema->PC_PassThroughImage ||
+		Pin->PinType.PinCategory == Schema->PC_Texture ||
+		Pin->PinType.PinCategory == Schema->PC_PassthroughTexture ||
 		Pin->PinType.PinCategory == Schema->PC_Color ||
 		Pin->PinType.PinCategory == Schema->PC_Float ||
 		Pin->PinType.PinCategory == Schema->PC_Mesh);
diff --git a/Engine/Plugins/Mutable/Source/CustomizableObjectEditor/Private/MuCOE/Nodes/CustomizableObjectNodeTexture.cpp b/Engine/Plugins/Mutable/Source/CustomizableObjectEditor/Private/MuCOE/Nodes/CustomizableObjectNodeTexture.cpp
--- a/Engine/Plugins/Mutable/Source/CustomizableObjectEditor/Private/MuCOE/Nodes/CustomizableObjectNodeTexture.cpp
+++ b/Engine/Plugins/Mutable/Source/CustomizableObjectEditor/Private/MuCOE/Nodes/CustomizableObjectNodeTexture.cpp
@@ -32,7 +32,7 @@ void UCustomizableObjectNodeTexture::AllocateDefaultPins(UCustomizableObjectNode
 	const UEdGraphSchema_CustomizableObject* Schema = GetDefault<UEdGraphSchema_CustomizableObject>();
 
 	FString PinName = TEXT("Texture");
-	UEdGraphPin* PinImagePin = CustomCreatePin(EGPD_Output, Schema->PC_Image, FName(*PinName));
+	UEdGraphPin* PinImagePin = CustomCreatePin(EGPD_Output, Schema->PC_Texture, FName(*PinName));
 	PinImagePin->bDefaultValueIsIgnored = true;
 }
 
@@ -70,7 +70,7 @@ FText UCustomizableObjectNodeTexture::GetNodeTitle(ENodeTitleType::Type TitleTyp
 FLinearColor UCustomizableObjectNodeTexture::GetNodeTitleColor() const
 {
 	const UEdGraphSchema_CustomizableObject* Schema = GetDefault<UEdGraphSchema_CustomizableObject>();
-	return Schema->GetPinTypeColor(Schema->PC_Image);
+	return Schema->GetPinTypeColor(Schema->PC_Texture);
 }
 
 
diff --git a/Engine/Plugins/Mutable/Source/CustomizableObjectEditor/Private/MuCOE/Nodes/CustomizableObjectNodeTextureBinarise.cpp b/Engine/Plugins/Mutable/Source/CustomizableObjectEditor/Private/MuCOE/Nodes/CustomizableObjectNodeTextureBinarise.cpp
--- a/Engine/Plugins/Mutable/Source/CustomizableObjectEditor/Private/MuCOE/Nodes/CustomizableObjectNodeTextureBinarise.cpp
+++ b/Engine/Plugins/Mutable/Source/CustomizableObjectEditor/Private/MuCOE/Nodes/CustomizableObjectNodeTextureBinarise.cpp
@@ -14,10 +14,10 @@ class UCustomizableObjectNodeRemapPins;
 
 void UCustomizableObjectNodeTextureBinarise::AllocateDefaultPins(UCustomizableObjectNodeRemapPins* RemapPins)
 {
-	UEdGraphPin* OutputPin = CustomCreatePinSimple(EGPD_Output, UEdGraphSchema_CustomizableObject::PC_Image);
+	UEdGraphPin* OutputPin = CustomCreatePinSimple(EGPD_Output, UEdGraphSchema_CustomizableObject::PC_Texture);
 	OutputPin->bDefaultValueIsIgnored = true;
 	
-	UEdGraphPin* ImagePin = CustomCreatePinSimple(EGPD_Input, UEdGraphSchema_CustomizableObject::PC_Image);
+	UEdGraphPin* ImagePin = CustomCreatePinSimple(EGPD_Input, UEdGraphSchema_CustomizableObject::PC_Texture);
 	BaseImagePinReference = FEdGraphPinReference(ImagePin);
 	
 	const FName PinName = TEXT("Threshold");
@@ -73,7 +73,7 @@ FText UCustomizableObjectNodeTextureBinarise::GetNodeTitle(ENodeTitleType::Type
 
 FLinearColor UCustomizableObjectNodeTextureBinarise::GetNodeTitleColor() const
 {
-	return UEdGraphSchema_CustomizableObject::GetPinTypeColor(UEdGraphSchema_CustomizableObject::PC_Image);
+	return GetDefault<UEdGraphSchema_CustomizableObject>()->GetPinTypeColor(UEdGraphSchema_CustomizableObject::PC_Texture);
 }
 
 
@@ -91,7 +91,7 @@ void UCustomizableObjectNodeTextureBinarise::Serialize(FArchive& Ar)
 
 	if (Ar.CustomVer(FCustomizableObjectCustomVersion::GUID) < FCustomizableObjectCustomVersion::PinsNamesImageToTexture)
 	{
-		const FName TargetPinName = UEdGraphSchema_CustomizableObject::GetPinCategoryName(UEdGraphSchema_CustomizableObject::PC_Image);
+		const FName TargetPinName = UEdGraphSchema_CustomizableObject::GetPinCategoryName(UEdGraphSchema_CustomizableObject::PC_Texture);
 		BaseImagePinReference = FEdGraphPinReference(FindPin(TargetPinName, EGPD_Input));
 	}
 }
diff --git a/Engine/Plugins/Mutable/Source/CustomizableObjectEditor/Private/MuCOE/Nodes/CustomizableObjectNodeTextureColourMap.cpp b/Engine/Plugins/Mutable/Source/CustomizableObjectEditor/Private/MuCOE/Nodes/CustomizableObjectNodeTextureColourMap.cpp
--- a/Engine/Plugins/Mutable/Source/CustomizableObjectEditor/Private/MuCOE/Nodes/CustomizableObjectNodeTextureColourMap.cpp
+++ b/Engine/Plugins/Mutable/Source/CustomizableObjectEditor/Private/MuCOE/Nodes/CustomizableObjectNodeTextureColourMap.cpp
@@ -17,19 +17,19 @@ void UCustomizableObjectNodeTextureColourMap::AllocateDefaultPins(UCustomizableO
 	const UEdGraphSchema_CustomizableObject* Schema = GetDefault<UEdGraphSchema_CustomizableObject>();
 
 	FString PinName = TEXT("Texture");
-	UEdGraphPin* OutputPin = CustomCreatePin(EGPD_Output, Schema->PC_Image, FName(*PinName));
+	UEdGraphPin* OutputPin = CustomCreatePin(EGPD_Output, Schema->PC_Texture, FName(*PinName));
 	OutputPin->bDefaultValueIsIgnored = true;
 
 	PinName = TEXT("Base");
-	UEdGraphPin* SourcePin = CustomCreatePin(EGPD_Input, Schema->PC_Image, FName(*PinName));
+	UEdGraphPin* SourcePin = CustomCreatePin(EGPD_Input, Schema->PC_Texture, FName(*PinName));
 	SourcePin->bDefaultValueIsIgnored = true;
 
 	PinName = TEXT("Mask");
-	UEdGraphPin* MaskPin = CustomCreatePin(EGPD_Input, Schema->PC_Image, FName(*PinName));
+	UEdGraphPin* MaskPin = CustomCreatePin(EGPD_Input, Schema->PC_Texture, FName(*PinName));
 	MaskPin->bDefaultValueIsIgnored = true;
 
 	PinName = TEXT("Map");
-	UEdGraphPin* GradientPin = CustomCreatePin(EGPD_Input, Schema->PC_Image, FName(*PinName));
+	UEdGraphPin* GradientPin = CustomCreatePin(EGPD_Input, Schema->PC_Texture, FName(*PinName));
 	GradientPin->bDefaultValueIsIgnored = true;
 }
 
@@ -58,7 +58,7 @@ FText UCustomizableObjectNodeTextureColourMap::GetNodeTitle(ENodeTitleType::Type
 FLinearColor UCustomizableObjectNodeTextureColourMap::GetNodeTitleColor() const
 {
 	const UEdGraphSchema_CustomizableObject* Schema = GetDefault<UEdGraphSchema_CustomizableObject>();
-	return Schema->GetPinTypeColor(Schema->PC_Image);
+	return Schema->GetPinTypeColor(Schema->PC_Texture);
 }
 
 
diff --git a/Engine/Plugins/Mutable/Source/CustomizableObjectEditor/Private/MuCOE/Nodes/CustomizableObjectNodeTextureFromChannels.cpp b/Engine/Plugins/Mutable/Source/CustomizableObjectEditor/Private/MuCOE/Nodes/CustomizableObjectNodeTextureFromChannels.cpp
--- a/Engine/Plugins/Mutable/Source/CustomizableObjectEditor/Private/MuCOE/Nodes/CustomizableObjectNodeTextureFromChannels.cpp
+++ b/Engine/Plugins/Mutable/Source/CustomizableObjectEditor/Private/MuCOE/Nodes/CustomizableObjectNodeTextureFromChannels.cpp
@@ -17,23 +17,23 @@ void UCustomizableObjectNodeTextureFromChannels::AllocateDefaultPins(UCustomizab
 	const UEdGraphSchema_CustomizableObject* Schema = GetDefault<UEdGraphSchema_CustomizableObject>();
 
 	FString PinName = TEXT("Texture");
-	UEdGraphPin* OutputPin = CustomCreatePin(EGPD_Output, Schema->PC_Image, FName(*PinName));
+	UEdGraphPin* OutputPin = CustomCreatePin(EGPD_Output, Schema->PC_Texture, FName(*PinName));
 	OutputPin->bDefaultValueIsIgnored = true;
 
 	PinName = TEXT("R");
-	UEdGraphPin* RPin = CustomCreatePin(EGPD_Input, Schema->PC_Image, FName(*PinName));
+	UEdGraphPin* RPin = CustomCreatePin(EGPD_Input, Schema->PC_Texture, FName(*PinName));
 	RPin->bDefaultValueIsIgnored = true;
 
 	PinName = TEXT("G");
-	UEdGraphPin* GPin = CustomCreatePin(EGPD_Input, Schema->PC_Image, FName(*PinName));
+	UEdGraphPin* GPin = CustomCreatePin(EGPD_Input, Schema->PC_Texture, FName(*PinName));
 	GPin->bDefaultValueIsIgnored = true;
 
 	PinName = TEXT("B");
-	UEdGraphPin* BPin = CustomCreatePin(EGPD_Input, Schema->PC_Image, FName(*PinName));
+	UEdGraphPin* BPin = CustomCreatePin(EGPD_Input, Schema->PC_Texture, FName(*PinName));
 	BPin->bDefaultValueIsIgnored = true;
 
 	PinName = TEXT("A");
-	UEdGraphPin* APin = CustomCreatePin(EGPD_Input, Schema->PC_Image, FName(*PinName));
+	UEdGraphPin* APin = CustomCreatePin(EGPD_Input, Schema->PC_Texture, FName(*PinName));
 	APin->bDefaultValueIsIgnored = true;
 }
 
@@ -62,7 +62,7 @@ FText UCustomizableObjectNodeTextureFromChannels::GetNodeTitle(ENodeTitleType::T
 FLinearColor UCustomizableObjectNodeTextureFromChannels::GetNodeTitleColor() const
 {
 	const UEdGraphSchema_CustomizableObject* Schema = GetDefault<UEdGraphSchema_CustomizableObject>();
-	return Schema->GetPinTypeColor(Schema->PC_Image);
+	return Schema->GetPinTypeColor(Schema->PC_Texture);
 }
 
 
diff --git a/Engine/Plugins/Mutable/Source/CustomizableObjectEditor/Private/MuCOE/Nodes/CustomizableObjectNodeTextureFromColor.cpp b/Engine/Plugins/Mutable/Source/CustomizableObjectEditor/Private/MuCOE/Nodes/CustomizableObjectNodeTextureFromColor.cpp
--- a/Engine/Plugins/Mutable/Source/CustomizableObjectEditor/Private/MuCOE/Nodes/CustomizableObjectNodeTextureFromColor.cpp
+++ b/Engine/Plugins/Mutable/Source/CustomizableObjectEditor/Private/MuCOE/Nodes/CustomizableObjectNodeTextureFromColor.cpp
@@ -17,7 +17,7 @@ void UCustomizableObjectNodeTextureFromColor::AllocateDefaultPins(UCustomizableO
 	const UEdGraphSchema_CustomizableObject* Schema = GetDefault<UEdGraphSchema_CustomizableObject>();
 
 	FString PinName = TEXT("Texture");
-	UEdGraphPin* OutputPin = CustomCreatePin(EGPD_Output, Schema->PC_Image, FName(*PinName));
+	UEdGraphPin* OutputPin = CustomCreatePin(EGPD_Output, Schema->PC_Texture, FName(*PinName));
 	OutputPin->bDefaultValueIsIgnored = true;
 
 	PinName = TEXT("Color");
@@ -50,7 +50,7 @@ FText UCustomizableObjectNodeTextureFromColor::GetNodeTitle(ENodeTitleType::Type
 FLinearColor UCustomizableObjectNodeTextureFromColor::GetNodeTitleColor() const
 {
 	const UEdGraphSchema_CustomizableObject* Schema = GetDefault<UEdGraphSchema_CustomizableObject>();
-	return Schema->GetPinTypeColor(Schema->PC_Image);
+	return Schema->GetPinTypeColor(Schema->PC_Texture);
 }
 
 
diff --git a/Engine/Plugins/Mutable/Source/CustomizableObjectEditor/Private/MuCOE/Nodes/CustomizableObjectNodeTextureFromFloats.cpp b/Engine/Plugins/Mutable/Source/CustomizableObjectEditor/Private/MuCOE/Nodes/CustomizableObjectNodeTextureFromFloats.cpp
--- a/Engine/Plugins/Mutable/Source/CustomizableObjectEditor/Private/MuCOE/Nodes/CustomizableObjectNodeTextureFromFloats.cpp
+++ b/Engine/Plugins/Mutable/Source/CustomizableObjectEditor/Private/MuCOE/Nodes/CustomizableObjectNodeTextureFromFloats.cpp
@@ -17,7 +17,7 @@ void UCustomizableObjectNodeTextureFromFloats::AllocateDefaultPins(UCustomizable
 	const UEdGraphSchema_CustomizableObject* Schema = GetDefault<UEdGraphSchema_CustomizableObject>();
 
 	FString PinName = TEXT("Texture");
-	UEdGraphPin* OutputPin = CustomCreatePin(EGPD_Output, Schema->PC_Image, FName(*PinName));
+	UEdGraphPin* OutputPin = CustomCreatePin(EGPD_Output, Schema->PC_Texture, FName(*PinName));
 	OutputPin->bDefaultValueIsIgnored = true;
 
 	PinName = TEXT("R");
@@ -47,7 +47,7 @@ FText UCustomizableObjectNodeTextureFromFloats::GetNodeTitle(ENodeTitleType::Typ
 FLinearColor UCustomizableObjectNodeTextureFromFloats::GetNodeTitleColor() const
 {
 	const UEdGraphSchema_CustomizableObject* Schema = GetDefault<UEdGraphSchema_CustomizableObject>();
-	return Schema->GetPinTypeColor(Schema->PC_Image);
+	return Schema->GetPinTypeColor(Schema->PC_Texture);
 }
 
 
diff --git a/Engine/Plugins/Mutable/Source/CustomizableObjectEditor/Private/MuCOE/Nodes/CustomizableObjectNodeTextureInterpolate.cpp b/Engine/Plugins/Mutable/Source/CustomizableObjectEditor/Private/MuCOE/Nodes/CustomizableObjectNodeTextureInterpolate.cpp
--- a/Engine/Plugins/Mutable/Source/CustomizableObjectEditor/Private/MuCOE/Nodes/CustomizableObjectNodeTextureInterpolate.cpp
+++ b/Engine/Plugins/Mutable/Source/CustomizableObjectEditor/Private/MuCOE/Nodes/CustomizableObjectNodeTextureInterpolate.cpp
@@ -35,7 +35,7 @@ void UCustomizableObjectNodeTextureInterpolate::AllocateDefaultPins(UCustomizabl
 	const UEdGraphSchema_CustomizableObject* Schema = GetDefault<UEdGraphSchema_CustomizableObject>();
 
 	FString PinName = TEXT("Texture");
-	UEdGraphPin* OutputPin = CustomCreatePin(EGPD_Output, Schema->PC_Image, FName(*PinName));
+	UEdGraphPin* OutputPin = CustomCreatePin(EGPD_Output, Schema->PC_Texture, FName(*PinName));
 	OutputPin->bDefaultValueIsIgnored = true;
 	
 	PinName = TEXT("Factor");
@@ -45,7 +45,7 @@ void UCustomizableObjectNodeTextureInterpolate::AllocateDefaultPins(UCustomizabl
 	for ( int LayerIndex = 0; LayerIndex < NumTargets; ++LayerIndex )
 	{
 		PinName = FString::Printf(TEXT("Target %d"), LayerIndex);
-		UEdGraphPin* LayerPin = CustomCreatePin(EGPD_Input, Schema->PC_Image, FName(*PinName));
+		UEdGraphPin* LayerPin = CustomCreatePin(EGPD_Input, Schema->PC_Texture, FName(*PinName));
 		LayerPin->bDefaultValueIsIgnored = true;
 	}
 }
@@ -60,7 +60,7 @@ FText UCustomizableObjectNodeTextureInterpolate::GetNodeTitle(ENodeTitleType::Ty
 FLinearColor UCustomizableObjectNodeTextureInterpolate::GetNodeTitleColor() const
 {
 	const UEdGraphSchema_CustomizableObject* Schema = GetDefault<UEdGraphSchema_CustomizableObject>();
-	return Schema->GetPinTypeColor(Schema->PC_Image);
+	return Schema->GetPinTypeColor(Schema->PC_Texture);
 }
 
 
diff --git a/Engine/Plugins/Mutable/Source/CustomizableObjectEditor/Private/MuCOE/Nodes/CustomizableObjectNodeTextureInvert.cpp b/Engine/Plugins/Mutable/Source/CustomizableObjectEditor/Private/MuCOE/Nodes/CustomizableObjectNodeTextureInvert.cpp
--- a/Engine/Plugins/Mutable/Source/CustomizableObjectEditor/Private/MuCOE/Nodes/CustomizableObjectNodeTextureInvert.cpp
+++ b/Engine/Plugins/Mutable/Source/CustomizableObjectEditor/Private/MuCOE/Nodes/CustomizableObjectNodeTextureInvert.cpp
@@ -14,10 +14,10 @@ class UCustomizableObjectNodeRemapPins;
 
 void UCustomizableObjectNodeTextureInvert::AllocateDefaultPins(UCustomizableObjectNodeRemapPins* RemapPins)
 {
-	UEdGraphPin* OutputPin = CustomCreatePinSimple(EGPD_Output, UEdGraphSchema_CustomizableObject::PC_Image);
+	UEdGraphPin* OutputPin = CustomCreatePinSimple(EGPD_Output, UEdGraphSchema_CustomizableObject::PC_Texture);
 	OutputPin->bDefaultValueIsIgnored = true;
 	
-	UEdGraphPin* ImagePin = CustomCreatePinSimple(EGPD_Input, UEdGraphSchema_CustomizableObject::PC_Image);
+	UEdGraphPin* ImagePin = CustomCreatePinSimple(EGPD_Input, UEdGraphSchema_CustomizableObject::PC_Texture);
 	BaseImagePinReference = FEdGraphPinReference(ImagePin);
 }
 
@@ -80,7 +80,7 @@ FText UCustomizableObjectNodeTextureInvert::GetNodeTitle(ENodeTitleType::Type Ti
 
 FLinearColor UCustomizableObjectNodeTextureInvert::GetNodeTitleColor() const
 {
-	return UEdGraphSchema_CustomizableObject::GetPinTypeColor(UEdGraphSchema_CustomizableObject::PC_Image);
+	return GetDefault<UEdGraphSchema_CustomizableObject>()->GetPinTypeColor(UEdGraphSchema_CustomizableObject::PC_Texture);
 }
 
 
diff --git a/Engine/Plugins/Mutable/Source/CustomizableObjectEditor/Private/MuCOE/Nodes/CustomizableObjectNodeTextureLayer.cpp b/Engine/Plugins/Mutable/Source/CustomizableObjectEditor/Private/MuCOE/Nodes/CustomizableObjectNodeTextureLayer.cpp
--- a/Engine/Plugins/Mutable/Source/CustomizableObjectEditor/Private/MuCOE/Nodes/CustomizableObjectNodeTextureLayer.cpp
+++ b/Engine/Plugins/Mutable/Source/CustomizableObjectEditor/Private/MuCOE/Nodes/CustomizableObjectNodeTextureLayer.cpp
@@ -51,22 +51,22 @@ void UCustomizableObjectNodeTextureLayer::AllocateDefaultPins(UCustomizableObjec
 {
 	const UEdGraphSchema_CustomizableObject* Schema = GetDefault<UEdGraphSchema_CustomizableObject>();
 
-	UEdGraphPin* OutputPin = CustomCreatePin(EGPD_Output, Schema->PC_Image, FName(TEXT("Texture")));
+	UEdGraphPin* OutputPin = CustomCreatePin(EGPD_Output, Schema->PC_Texture, FName(TEXT("Texture")));
 	OutputPin->bDefaultValueIsIgnored = true;
 
 	OutputPinReference = FEdGraphPinReference(OutputPin);
 
-	UEdGraphPin* BasePin = CustomCreatePin(EGPD_Input, Schema->PC_Image, FName(TEXT("Base")));
+	UEdGraphPin* BasePin = CustomCreatePin(EGPD_Input, Schema->PC_Texture, FName(TEXT("Base")));
 	BasePin->bDefaultValueIsIgnored = true;
 	
 	for (int32 LayerIndex = 0; LayerIndex < Layers.Num() ; ++LayerIndex )
 	{
 		FString PinName = FString::Printf( TEXT("Layer %d "), LayerIndex );
-		UEdGraphPin* LayerPin = CustomCreatePin(EGPD_Input, Schema->PC_Image, *PinName);
+		UEdGraphPin* LayerPin = CustomCreatePin(EGPD_Input, Schema->PC_Texture, *PinName);
 		LayerPin->bDefaultValueIsIgnored = true;
 
 		PinName = FString::Printf( TEXT("Mask %d "), LayerIndex );
-		UEdGraphPin* MaskPin = CustomCreatePin(EGPD_Input, Schema->PC_Image, *PinName);
+		UEdGraphPin* MaskPin = CustomCreatePin(EGPD_Input, Schema->PC_Texture, *PinName);
 		MaskPin->bDefaultValueIsIgnored = true;
 	}
 }
@@ -96,7 +96,7 @@ FText UCustomizableObjectNodeTextureLayer::GetNodeTitle(ENodeTitleType::Type Tit
 FLinearColor UCustomizableObjectNodeTextureLayer::GetNodeTitleColor() const
 {
 	const UEdGraphSchema_CustomizableObject* Schema = GetDefault<UEdGraphSchema_CustomizableObject>();
-	return Schema->GetPinTypeColor(Schema->PC_Image);
+	return Schema->GetPinTypeColor(Schema->PC_Texture);
 }
 
 
@@ -138,7 +138,7 @@ bool UCustomizableObjectNodeTextureLayer::CanConnect(const UEdGraphPin* InOwnedI
 	{
 		if (LayerPin(LayerIndex) == InOwnedInputPin )
 		{
-			return (InOutputPin->PinType.PinCategory==Schema->PC_Image)
+			return (InOutputPin->PinType.PinCategory==Schema->PC_Texture)
 				|| (InOutputPin->PinType.PinCategory==Schema->PC_Color);
 		}
 	}
diff --git a/Engine/Plugins/Mutable/Source/CustomizableObjectEditor/Private/MuCOE/Nodes/CustomizableObjectNodeTextureParameter.cpp b/Engine/Plugins/Mutable/Source/CustomizableObjectEditor/Private/MuCOE/Nodes/CustomizableObjectNodeTextureParameter.cpp
--- a/Engine/Plugins/Mutable/Source/CustomizableObjectEditor/Private/MuCOE/Nodes/CustomizableObjectNodeTextureParameter.cpp
+++ b/Engine/Plugins/Mutable/Source/CustomizableObjectEditor/Private/MuCOE/Nodes/CustomizableObjectNodeTextureParameter.cpp
@@ -41,7 +41,7 @@ void UCustomizableObjectNodeTextureParameter::BackwardsCompatibleFixup(int32 Cus
 
 FName UCustomizableObjectNodeTextureParameter::GetCategory() const
 {
-	return UEdGraphSchema_CustomizableObject::PC_Image;
+	return UEdGraphSchema_CustomizableObject::PC_Texture;
 }
 
 #undef LOCTEXT_NAMESPACE
diff --git a/Engine/Plugins/Mutable/Source/CustomizableObjectEditor/Private/MuCOE/Nodes/CustomizableObjectNodeTextureProject.cpp b/Engine/Plugins/Mutable/Source/CustomizableObjectEditor/Private/MuCOE/Nodes/CustomizableObjectNodeTextureProject.cpp
--- a/Engine/Plugins/Mutable/Source/CustomizableObjectEditor/Private/MuCOE/Nodes/CustomizableObjectNodeTextureProject.cpp
+++ b/Engine/Plugins/Mutable/Source/CustomizableObjectEditor/Private/MuCOE/Nodes/CustomizableObjectNodeTextureProject.cpp
@@ -83,13 +83,13 @@ void UCustomizableObjectNodeTextureProject::AllocateDefaultPins(UCustomizableObj
 	for ( uint32 LayerIndex = 0; LayerIndex < Textures; ++LayerIndex )
 	{
 		FString PinName = FString::Printf(TEXT("Texture %u"), LayerIndex);
-		UEdGraphPin* TexturePin = CustomCreatePin(EGPD_Input, Schema->PC_Image, FName(*PinName));
+		UEdGraphPin* TexturePin = CustomCreatePin(EGPD_Input, Schema->PC_Texture, FName(*PinName));
 		TexturePin->bDefaultValueIsIgnored = true;
 
 		TexturePinsReferences.Add(FEdGraphPinReference(TexturePin));
 		
 		PinName = FString::Printf(TEXT("Texture %u"), LayerIndex);
-		UEdGraphPin* OutputPin = CustomCreatePin(EGPD_Output, Schema->PC_Image, FName(*PinName));
+		UEdGraphPin* OutputPin = CustomCreatePin(EGPD_Output, Schema->PC_Texture, FName(*PinName));
 		OutputPin->bDefaultValueIsIgnored = true;
 
 		OutputPinsReferences.Add(FEdGraphPinReference(OutputPin));
@@ -100,7 +100,7 @@ void UCustomizableObjectNodeTextureProject::AllocateDefaultPins(UCustomizableObj
 	MeshPin->bDefaultValueIsIgnored = true;
 
 	PinName = TEXT("Mesh Mask");
-	UEdGraphPin* MeshMaskPin = CustomCreatePin(EGPD_Input, Schema->PC_Image, FName(*PinName));
+	UEdGraphPin* MeshMaskPin = CustomCreatePin(EGPD_Input, Schema->PC_Texture, FName(*PinName));
 	MeshMaskPin->bDefaultValueIsIgnored = true;
 
 	PinName = TEXT("Fade Start Angle");
@@ -181,7 +181,7 @@ FText UCustomizableObjectNodeTextureProject::GetNodeTitle(ENodeTitleType::Type T
 FLinearColor UCustomizableObjectNodeTextureProject::GetNodeTitleColor() const
 {
 	const UEdGraphSchema_CustomizableObject* Schema = GetDefault<UEdGraphSchema_CustomizableObject>();
-	return Schema->GetPinTypeColor(Schema->PC_Image);
+	return Schema->GetPinTypeColor(Schema->PC_Texture);
 }
 
 
diff --git a/Engine/Plugins/Mutable/Source/CustomizableObjectEditor/Private/MuCOE/Nodes/CustomizableObjectNodeTextureSample.cpp b/Engine/Plugins/Mutable/Source/CustomizableObjectEditor/Private/MuCOE/Nodes/CustomizableObjectNodeTextureSample.cpp
--- a/Engine/Plugins/Mutable/Source/CustomizableObjectEditor/Private/MuCOE/Nodes/CustomizableObjectNodeTextureSample.cpp
+++ b/Engine/Plugins/Mutable/Source/CustomizableObjectEditor/Private/MuCOE/Nodes/CustomizableObjectNodeTextureSample.cpp
@@ -29,7 +29,7 @@ void UCustomizableObjectNodeTextureSample::AllocateDefaultPins(UCustomizableObje
 	
 	// Inputs
 	
-	UEdGraphPin* TexturePin = CustomCreatePinSimple(EGPD_Input, UEdGraphSchema_CustomizableObject::PC_Image);
+	UEdGraphPin* TexturePin = CustomCreatePinSimple(EGPD_Input, UEdGraphSchema_CustomizableObject::PC_Texture);
 	TexturePin->bDefaultValueIsIgnored = true;
 
 	UEdGraphPin* XPin = CustomCreatePin(EGPD_Input, UEdGraphSchema_CustomizableObject::PC_Float,  FName(TEXT("X")));
@@ -66,7 +66,7 @@ FText UCustomizableObjectNodeTextureSample::GetNodeTitle(ENodeTitleType::Type Ti
 
 FLinearColor UCustomizableObjectNodeTextureSample::GetNodeTitleColor() const
 {
-	return UEdGraphSchema_CustomizableObject::GetPinTypeColor(UEdGraphSchema_CustomizableObject::PC_Color);
+	return GetDefault<UEdGraphSchema_CustomizableObject>()->GetPinTypeColor(UEdGraphSchema_CustomizableObject::PC_Color);
 }
 
 
diff --git a/Engine/Plugins/Mutable/Source/CustomizableObjectEditor/Private/MuCOE/Nodes/CustomizableObjectNodeTextureSample.h b/Engine/Plugins/Mutable/Source/CustomizableObjectEditor/Private/MuCOE/Nodes/CustomizableObjectNodeTextureSample.h
--- a/Engine/Plugins/Mutable/Source/CustomizableObjectEditor/Private/MuCOE/Nodes/CustomizableObjectNodeTextureSample.h
+++ b/Engine/Plugins/Mutable/Source/CustomizableObjectEditor/Private/MuCOE/Nodes/CustomizableObjectNodeTextureSample.h
@@ -36,7 +36,7 @@ public:
 	
 	const UEdGraphPin* TexturePin() const
 	{
-		const FName PinName = UEdGraphSchema_CustomizableObject::GetPinCategoryName(UEdGraphSchema_CustomizableObject::PC_Image);
+		const FName PinName = UEdGraphSchema_CustomizableObject::GetPinCategoryName(UEdGraphSchema_CustomizableObject::PC_Texture);
 		return FindPin(PinName);
 	}
 
diff --git a/Engine/Plugins/Mutable/Source/CustomizableObjectEditor/Private/MuCOE/Nodes/CustomizableObjectNodeTextureSaturate.cpp b/Engine/Plugins/Mutable/Source/CustomizableObjectEditor/Private/MuCOE/Nodes/CustomizableObjectNodeTextureSaturate.cpp
--- a/Engine/Plugins/Mutable/Source/CustomizableObjectEditor/Private/MuCOE/Nodes/CustomizableObjectNodeTextureSaturate.cpp
+++ b/Engine/Plugins/Mutable/Source/CustomizableObjectEditor/Private/MuCOE/Nodes/CustomizableObjectNodeTextureSaturate.cpp
@@ -14,10 +14,10 @@ class UCustomizableObjectNodeRemapPins;
 
 void UCustomizableObjectNodeTextureSaturate::AllocateDefaultPins(UCustomizableObjectNodeRemapPins* RemapPins)
 {
-	UEdGraphPin* OutputPin = CustomCreatePinSimple(EGPD_Output, UEdGraphSchema_CustomizableObject::PC_Image);
+	UEdGraphPin* OutputPin = CustomCreatePinSimple(EGPD_Output, UEdGraphSchema_CustomizableObject::PC_Texture);
 	OutputPin->bDefaultValueIsIgnored = true;
 	
-	UEdGraphPin* ImagePin = CustomCreatePinSimple(EGPD_Input, UEdGraphSchema_CustomizableObject::PC_Image);
+	UEdGraphPin* ImagePin = CustomCreatePinSimple(EGPD_Input, UEdGraphSchema_CustomizableObject::PC_Texture);
 	BaseImagePinReference = FEdGraphPinReference(ImagePin);
 	
 	UEdGraphPin* FactorPin = CustomCreatePin(EGPD_Input, UEdGraphSchema_CustomizableObject::PC_Float, FName(TEXT("Factor")));
@@ -66,7 +66,7 @@ FText UCustomizableObjectNodeTextureSaturate::GetNodeTitle(ENodeTitleType::Type
 
 FLinearColor UCustomizableObjectNodeTextureSaturate::GetNodeTitleColor() const
 {
-	return UEdGraphSchema_CustomizableObject::GetPinTypeColor(UEdGraphSchema_CustomizableObject::PC_Image);
+	return GetDefault<UEdGraphSchema_CustomizableObject>()->GetPinTypeColor(UEdGraphSchema_CustomizableObject::PC_Texture);
 }
 
 
diff --git a/Engine/Plugins/Mutable/Source/CustomizableObjectEditor/Private/MuCOE/Nodes/CustomizableObjectNodeTextureSwitch.cpp b/Engine/Plugins/Mutable/Source/CustomizableObjectEditor/Private/MuCOE/Nodes/CustomizableObjectNodeTextureSwitch.cpp
--- a/Engine/Plugins/Mutable/Source/CustomizableObjectEditor/Private/MuCOE/Nodes/CustomizableObjectNodeTextureSwitch.cpp
+++ b/Engine/Plugins/Mutable/Source/CustomizableObjectEditor/Private/MuCOE/Nodes/CustomizableObjectNodeTextureSwitch.cpp
@@ -32,7 +32,7 @@ void UCustomizableObjectNodeTextureSwitch::BackwardsCompatibleFixup(int32 Custom
 
 FName UCustomizableObjectNodeTextureSwitch::GetCategory() const
 {
-	return UEdGraphSchema_CustomizableObject::PC_Image;
+	return UEdGraphSchema_CustomizableObject::PC_Texture;
 }
 
 
diff --git a/Engine/Plugins/Mutable/Source/CustomizableObjectEditor/Private/MuCOE/Nodes/CustomizableObjectNodeTextureToChannels.cpp b/Engine/Plugins/Mutable/Source/CustomizableObjectEditor/Private/MuCOE/Nodes/CustomizableObjectNodeTextureToChannels.cpp
--- a/Engine/Plugins/Mutable/Source/CustomizableObjectEditor/Private/MuCOE/Nodes/CustomizableObjectNodeTextureToChannels.cpp
+++ b/Engine/Plugins/Mutable/Source/CustomizableObjectEditor/Private/MuCOE/Nodes/CustomizableObjectNodeTextureToChannels.cpp
@@ -38,7 +38,7 @@ void UCustomizableObjectNodeTextureToChannels::AllocateDefaultPins(UCustomizable
 
 	{
 		FString PinName = TEXT("Texture");
-		UEdGraphPin* InputPin = CustomCreatePin(EGPD_Input, Schema->PC_Image, FName(*PinName));
+		UEdGraphPin* InputPin = CustomCreatePin(EGPD_Input, Schema->PC_Texture, FName(*PinName));
 		InputPin->bDefaultValueIsIgnored = true;
 
 		InputPinReference = FEdGraphPinReference(InputPin);
@@ -46,25 +46,25 @@ void UCustomizableObjectNodeTextureToChannels::AllocateDefaultPins(UCustomizable
 
 	{
 		FString PinName = TEXT("R");
-		UEdGraphPin* RPin = CustomCreatePin(EGPD_Output, Schema->PC_Image, FName(*PinName));
+		UEdGraphPin* RPin = CustomCreatePin(EGPD_Output, Schema->PC_Texture, FName(*PinName));
 		RPin->bDefaultValueIsIgnored = true;
 	}
 
 	{
 		FString PinName = TEXT("G");
-		UEdGraphPin* GPin = CustomCreatePin(EGPD_Output, Schema->PC_Image, FName(*PinName));
+		UEdGraphPin* GPin = CustomCreatePin(EGPD_Output, Schema->PC_Texture, FName(*PinName));
 		GPin->bDefaultValueIsIgnored = true;
 	}
 
 	{
 		FString PinName = TEXT("B");
-		UEdGraphPin* BPin = CustomCreatePin(EGPD_Output, Schema->PC_Image, FName(*PinName));
+		UEdGraphPin* BPin = CustomCreatePin(EGPD_Output, Schema->PC_Texture, FName(*PinName));
 		BPin->bDefaultValueIsIgnored = true;
 	}
 
 	{
 		FString PinName = TEXT("A");
-		UEdGraphPin* APin = CustomCreatePin(EGPD_Output, Schema->PC_Image, FName(*PinName));
+		UEdGraphPin* APin = CustomCreatePin(EGPD_Output, Schema->PC_Texture, FName(*PinName));
 		APin->bDefaultValueIsIgnored = true;
 	}
 }
@@ -99,7 +99,7 @@ FText UCustomizableObjectNodeTextureToChannels::GetNodeTitle(ENodeTitleType::Typ
 FLinearColor UCustomizableObjectNodeTextureToChannels::GetNodeTitleColor() const
 {
 	const UEdGraphSchema_CustomizableObject* Schema = GetDefault<UEdGraphSchema_CustomizableObject>();
-	return Schema->GetPinTypeColor(Schema->PC_Image);
+	return Schema->GetPinTypeColor(Schema->PC_Texture);
 }
 
 
diff --git a/Engine/Plugins/Mutable/Source/CustomizableObjectEditor/Private/MuCOE/Nodes/CustomizableObjectNodeTextureTransform.cpp b/Engine/Plugins/Mutable/Source/CustomizableObjectEditor/Private/MuCOE/Nodes/CustomizableObjectNodeTextureTransform.cpp
--- a/Engine/Plugins/Mutable/Source/CustomizableObjectEditor/Private/MuCOE/Nodes/CustomizableObjectNodeTextureTransform.cpp
+++ b/Engine/Plugins/Mutable/Source/CustomizableObjectEditor/Private/MuCOE/Nodes/CustomizableObjectNodeTextureTransform.cpp
@@ -13,11 +13,11 @@ class UCustomizableObjectNodeRemapPins;
 
 void UCustomizableObjectNodeTextureTransform::AllocateDefaultPins(UCustomizableObjectNodeRemapPins* RemapPins)
 {
-	UEdGraphPin* OutputPin = CustomCreatePinSimple(EGPD_Output, UEdGraphSchema_CustomizableObject::PC_Image);
+	UEdGraphPin* OutputPin = CustomCreatePinSimple(EGPD_Output, UEdGraphSchema_CustomizableObject::PC_Texture);
 	OutputPin->bDefaultValueIsIgnored = true;
 	
 	
-	UEdGraphPin* ImagePin = CustomCreatePinSimple(EGPD_Input, UEdGraphSchema_CustomizableObject::PC_Image);
+	UEdGraphPin* ImagePin = CustomCreatePinSimple(EGPD_Input, UEdGraphSchema_CustomizableObject::PC_Texture);
 	BaseImagePinReference = FEdGraphPinReference(ImagePin);
 	
 	UEdGraphPin* OffsetXPin = CustomCreatePin(EGPD_Input, UEdGraphSchema_CustomizableObject::PC_Float, FName(TEXT("Offset X")));
@@ -67,7 +67,7 @@ FText UCustomizableObjectNodeTextureTransform::GetNodeTitle(ENodeTitleType::Type
 
 FLinearColor UCustomizableObjectNodeTextureTransform::GetNodeTitleColor() const
 {
-	return UEdGraphSchema_CustomizableObject::GetPinTypeColor(UEdGraphSchema_CustomizableObject::PC_Image);
+	return GetDefault<UEdGraphSchema_CustomizableObject>()->GetPinTypeColor(UEdGraphSchema_CustomizableObject::PC_Texture);
 }
 
 FText UCustomizableObjectNodeTextureTransform::GetTooltipText() const
diff --git a/Engine/Plugins/Mutable/Source/CustomizableObjectEditor/Private/MuCOE/Nodes/CustomizableObjectNodeTextureVariation.cpp b/Engine/Plugins/Mutable/Source/CustomizableObjectEditor/Private/MuCOE/Nodes/CustomizableObjectNodeTextureVariation.cpp
--- a/Engine/Plugins/Mutable/Source/CustomizableObjectEditor/Private/MuCOE/Nodes/CustomizableObjectNodeTextureVariation.cpp
+++ b/Engine/Plugins/Mutable/Source/CustomizableObjectEditor/Private/MuCOE/Nodes/CustomizableObjectNodeTextureVariation.cpp
@@ -14,7 +14,7 @@ class UCustomizableObjectNodeRemapPins;
 
 FName UCustomizableObjectNodeTextureVariation::GetCategory() const
 {
-	return UEdGraphSchema_CustomizableObject::PC_Image;
+	return UEdGraphSchema_CustomizableObject::PC_Texture;
 }
 
 
@@ -72,7 +72,7 @@ FText UCustomizableObjectNodeTextureVariation::GetNodeTitle(ENodeTitleType::Type
 
 FLinearColor UCustomizableObjectNodeTextureVariation::GetNodeTitleColor() const
 {
-	return UEdGraphSchema_CustomizableObject::GetPinTypeColor(UEdGraphSchema_CustomizableObject::PC_Image);
+	return GetDefault<UEdGraphSchema_CustomizableObject>()->GetPinTypeColor(UEdGraphSchema_CustomizableObject::PC_Texture);
 }
 
 
diff --git a/Engine/Plugins/Mutable/Source/CustomizableObjectEditor/Private/MuCOE/Nodes/SCustomizableObjectNodeMaterial.cpp b/Engine/Plugins/Mutable/Source/CustomizableObjectEditor/Private/MuCOE/Nodes/SCustomizableObjectNodeMaterial.cpp
--- a/Engine/Plugins/Mutable/Source/CustomizableObjectEditor/Private/MuCOE/Nodes/SCustomizableObjectNodeMaterial.cpp
+++ b/Engine/Plugins/Mutable/Source/CustomizableObjectEditor/Private/MuCOE/Nodes/SCustomizableObjectNodeMaterial.cpp
@@ -17,7 +17,7 @@ void SCustomizableObjectNodeMaterial::Construct(const FArguments& InArgs, UCusto
 
 TSharedPtr<SGraphPin> SCustomizableObjectNodeMaterial::CreatePinWidget(UEdGraphPin* Pin) const
 {
-	if ((Pin->PinType.PinCategory == UEdGraphSchema_CustomizableObject::PC_Image || Pin->PinType.PinCategory == UEdGraphSchema_CustomizableObject::PC_PassThroughImage) &&
+	if ((Pin->PinType.PinCategory == UEdGraphSchema_CustomizableObject::PC_Texture || Pin->PinType.PinCategory == UEdGraphSchema_CustomizableObject::PC_PassthroughTexture) &&
 		Pin->Direction == EGPD_Input)
 	{
 		return SNew(SCustomizableObjectNodeMaterialPinImage, Pin);
diff --git a/Engine/Plugins/Mutable/Source/CustomizableObjectEditor/Private/MuCOE/Nodes/SCustomizableObjectNodeMaterialPinImage.cpp b/Engine/Plugins/Mutable/Source/CustomizableObjectEditor/Private/MuCOE/Nodes/SCustomizableObjectNodeMaterialPinImage.cpp
--- a/Engine/Plugins/Mutable/Source/CustomizableObjectEditor/Private/MuCOE/Nodes/SCustomizableObjectNodeMaterialPinImage.cpp
+++ b/Engine/Plugins/Mutable/Source/CustomizableObjectEditor/Private/MuCOE/Nodes/SCustomizableObjectNodeMaterialPinImage.cpp
@@ -49,7 +49,7 @@ FText SCustomizableObjectNodeMaterialPinImage::GetPinTooltipText() const
 		return LOCTEXT("PinModeMutableOrpahan", "Pin not disapearing due to being connected or having a property modified.");
 	}
 	
-	if (GraphPinObj->PinType.PinCategory == UEdGraphSchema_CustomizableObject::PC_Image)
+	if (GraphPinObj->PinType.PinCategory == UEdGraphSchema_CustomizableObject::PC_Texture)
 	{
 		return LOCTEXT("PinModeMutableTooltip", "Texture Parameter goes through Mutable.");
 	}
@@ -62,7 +62,7 @@ FText SCustomizableObjectNodeMaterialPinImage::GetPinTooltipText() const
 
 FText SCustomizableObjectNodeMaterialPinImage::GetDefaultValueText() const
 {
-	if (GraphPinObj->PinType.PinCategory == UEdGraphSchema_CustomizableObject::PC_Image)
+	if (GraphPinObj->PinType.PinCategory == UEdGraphSchema_CustomizableObject::PC_Texture)
 	{
 		return LOCTEXT("PinModeMutable", "mutable");
 	}
diff --git a/Engine/Plugins/Mutable/Source/CustomizableObjectEditor/Private/MuCOE/Nodes/SCustomizableObjectNodeReroutePin.cpp b/Engine/Plugins/Mutable/Source/CustomizableObjectEditor/Private/MuCOE/Nodes/SCustomizableObjectNodeReroutePin.cpp
--- a/Engine/Plugins/Mutable/Source/CustomizableObjectEditor/Private/MuCOE/Nodes/SCustomizableObjectNodeReroutePin.cpp
+++ b/Engine/Plugins/Mutable/Source/CustomizableObjectEditor/Private/MuCOE/Nodes/SCustomizableObjectNodeReroutePin.cpp
@@ -18,7 +18,7 @@ void SCustomizableObjectNodeReroutePin::Construct(const FArguments& InArgs, UEdG
 
 const FSlateBrush* SCustomizableObjectNodeReroutePin::GetPinIcon() const
 {
-	if (GraphPinObj->PinType.PinCategory == UEdGraphSchema_CustomizableObject::PC_PassThroughImage)
+	if (GraphPinObj->PinType.PinCategory == UEdGraphSchema_CustomizableObject::PC_PassthroughTexture)
 	{
 		return GraphPinObj->LinkedTo.Num() ?
 			PassThroughImageConnected :
diff --git a/Engine/Plugins/Mutable/Source/CustomizableObjectEditor/Public/MuCOE/EdGraphSchema_CustomizableObject.h b/Engine/Plugins/Mutable/Source/CustomizableObjectEditor/Public/MuCOE/EdGraphSchema_CustomizableObject.h
--- a/Engine/Plugins/Mutable/Source/CustomizableObjectEditor/Public/MuCOE/EdGraphSchema_CustomizableObject.h
+++ b/Engine/Plugins/Mutable/Source/CustomizableObjectEditor/Public/MuCOE/EdGraphSchema_CustomizableObject.h
@@ -31,10 +31,10 @@ public:
 	static UE_API const FName PC_MeshSection;
 	static UE_API const FName PC_Modifier;
 	static UE_API const FName PC_Mesh;
-	static UE_API const FName PC_PassThroughMesh;
-	static UE_API const FName PC_Layout;
-	static UE_API const FName PC_Image;
-	static UE_API const FName PC_PassThroughImage;
+	static UE_API const FName PC_SkeletalMesh;
+	static UE_API const FName PC_PassthroughSkeletalMesh;
+	static UE_API const FName PC_Texture;
+	static UE_API const FName PC_PassthroughTexture;
 	static UE_API const FName PC_Projector;
 	static UE_API const FName PC_GroupProjector;
 	static UE_API const FName PC_Color;
@@ -54,6 +54,8 @@ public:
 	static UE_API const FText NC_Experimental;
 	static UE_API const FText NC_Material;
 
+	mutable TMap<FName, FLinearColor> PinCategoryColors;
+	
 	// EdGraphSchema interface
 	UE_API virtual void GetAssetsGraphHoverMessage(const TArray<FAssetData>& Assets, const UEdGraph* HoverGraph, FString& OutTooltipText, bool& OutOkIcon) const override;
 	UE_API virtual void GetGraphContextActions(FGraphContextMenuBuilder& ContextMenuBuilder) const override;
@@ -70,8 +72,8 @@ public:
 	UE_API virtual FPinConnectionResponse MovePinLinks(UEdGraphPin& MoveFromPin, UEdGraphPin& MoveToPin, bool bIsIntermediateMove, bool bNotifyLinkedNodes) const override;
 	UE_API virtual FConnectionDrawingPolicy* CreateConnectionDrawingPolicy(int32 InBackLayerID, int32 InFrontLayerID, float InZoomFactor, const FSlateRect& InClippingRect, FSlateWindowElementList& InDrawElements, UEdGraph* InGraphObj) const override;
 
-	//  Own interface
-	static UE_API FLinearColor GetPinTypeColor(const FName& PinType);
+	// Own interface
+	UE_API FLinearColor GetPinTypeColor(const FName& PinType) const;
 	
 	/** Given a pin category name, get its friendly name (user readable). */
 	static UE_API FName GetPinCategoryName(const FName& PinCategory);
