[Mutable] Fixed missing dependency to module "TraceLog"

#rnx

diff --git a/Engine/Plugins/Mutable/Source/CustomizableObject/CustomizableObject.Build.cs b/Engine/Plugins/Mutable/Source/CustomizableObject/CustomizableObject.Build.cs
--- a/Engine/Plugins/Mutable/Source/CustomizableObject/CustomizableObject.Build.cs
+++ b/Engine/Plugins/Mutable/Source/CustomizableObject/CustomizableObject.Build.cs
@@ -24,7 +24,8 @@ public class CustomizableObject : ModuleRules
 			"ClothingSystemRuntimeInterface",
 			//"ClothingSystemEditor",
 			"UMG",
-			"ImageCore"
+			"ImageCore",
+			"TraceLog"
 		]);
 
         PublicDependencyModuleNames.AddRange([
