[Mutable] Allow all pins to be renamed.

Now all pins can optionally be renamed. Added functions to UCONode to be able to rename it.

- Made GetPinTypeColor static.
- Fixed issues with renaming nodes.
- Hidden node names from detials if they have a pin.
- Removed unused Thumbnail classes.
- Fixed Node Expose Pin title.

#rnx

diff --git a/Engine/Plugins/Mutable/Source/CustomizableObjectEditor/Private/MuCOE/COVariable.cpp b/Engine/Plugins/Mutable/Source/CustomizableObjectEditor/Private/MuCOE/COVariable.cpp
--- a/Engine/Plugins/Mutable/Source/CustomizableObjectEditor/Private/MuCOE/COVariable.cpp
+++ b/Engine/Plugins/Mutable/Source/CustomizableObjectEditor/Private/MuCOE/COVariable.cpp
@@ -137,7 +137,7 @@ TSharedRef<SWidget> FCOVariableCustomzation::OnGenerateVariableTypeRow(FName Typ
 		[
 			SNew(SImage)
 				.Image(IconBrush)
-				.ColorAndOpacity(GetDefault<UEdGraphSchema_CustomizableObject>()->GetPinTypeColor(Type))
+				.ColorAndOpacity(UEdGraphSchema_CustomizableObject::GetPinTypeColor(Type))
 		]
 		+ SHorizontalBox::Slot()
 		.AutoWidth()
diff --git a/Engine/Plugins/Mutable/Source/CustomizableObjectEditor/Private/MuCOE/CustomizableObjectConnectionDrawingPolicy.cpp b/Engine/Plugins/Mutable/Source/CustomizableObjectEditor/Private/MuCOE/CustomizableObjectConnectionDrawingPolicy.cpp
--- a/Engine/Plugins/Mutable/Source/CustomizableObjectEditor/Private/MuCOE/CustomizableObjectConnectionDrawingPolicy.cpp
+++ b/Engine/Plugins/Mutable/Source/CustomizableObjectEditor/Private/MuCOE/CustomizableObjectConnectionDrawingPolicy.cpp
@@ -9,7 +9,7 @@ void FCustomizableObjectConnectionDrawingPolicy::DetermineWiringStyle(UEdGraphPi
 {
 	FConnectionDrawingPolicy::DetermineWiringStyle(OutputPin, InputPin, Params);
 
-	Params.WireColor = GetDefault<UEdGraphSchema_CustomizableObject>()->GetPinTypeColor(OutputPin->PinType.PinCategory); // Color.
+	Params.WireColor = UEdGraphSchema_CustomizableObject::GetPinTypeColor(OutputPin->PinType.PinCategory); // Color.
 
 	if (HoveredPins.Num() > 0)
 	{
diff --git a/Engine/Plugins/Mutable/Source/CustomizableObjectEditor/Private/MuCOE/CustomizableObjectEditorModule.cpp b/Engine/Plugins/Mutable/Source/CustomizableObjectEditor/Private/MuCOE/CustomizableObjectEditorModule.cpp
--- a/Engine/Plugins/Mutable/Source/CustomizableObjectEditor/Private/MuCOE/CustomizableObjectEditorModule.cpp
+++ b/Engine/Plugins/Mutable/Source/CustomizableObjectEditor/Private/MuCOE/CustomizableObjectEditorModule.cpp
@@ -247,10 +247,6 @@ void FCustomizableObjectEditorModule::StartupModule()
 {
 	// Delete unused local compiled data
 	DeleteUnusedMutableStreamedDataEditorFiles();
-
-	// Register the thumbnail renderers
-	//UThumbnailManager::Get().RegisterCustomRenderer(UCustomizableObject::StaticClass(), UCustomizableObjectThumbnailRenderer::StaticClass());
-	//UThumbnailManager::Get().RegisterCustomRenderer(UCustomizableObjectInstance::StaticClass(), UCustomizableObjectInstanceThumbnailRenderer::StaticClass());
 	
 	// Property views
 	// Nodes
diff --git a/Engine/Plugins/Mutable/Source/CustomizableObjectEditor/Private/MuCOE/CustomizableObjectMacroLibrary/CustomizableObjectMacroDetails.cpp b/Engine/Plugins/Mutable/Source/CustomizableObjectEditor/Private/MuCOE/CustomizableObjectMacroLibrary/CustomizableObjectMacroDetails.cpp
--- a/Engine/Plugins/Mutable/Source/CustomizableObjectEditor/Private/MuCOE/CustomizableObjectMacroLibrary/CustomizableObjectMacroDetails.cpp
+++ b/Engine/Plugins/Mutable/Source/CustomizableObjectEditor/Private/MuCOE/CustomizableObjectMacroLibrary/CustomizableObjectMacroDetails.cpp
@@ -290,7 +290,7 @@ TSharedRef<SWidget> SCOMacroPinTypeSelector::OnGenerateRow(TSharedPtr<FPinNameRo
 	[
 		SNew(SImage)
 		.Image(IconBrush)
-		.ColorAndOpacity(GetDefault<UEdGraphSchema_CustomizableObject>()->GetPinTypeColor(Option->PinCategory))
+		.ColorAndOpacity(UEdGraphSchema_CustomizableObject::GetPinTypeColor(Option->PinCategory))
 	]
 	+ SHorizontalBox::Slot()
 	.AutoWidth()
@@ -317,7 +317,7 @@ TSharedRef<SWidget> SCOMacroPinTypeSelector::GenerateSelectedContent() const
 		.Image(IconBrush)
 		.ColorAndOpacity_Lambda([this]() 
 		{
-			return GetDefault<UEdGraphSchema_CustomizableObject>()->GetPinTypeColor(Variable->PinCategoryType); 
+			return UEdGraphSchema_CustomizableObject::GetPinTypeColor(Variable->PinCategoryType); 
 		})
 	]
 	+ SHorizontalBox::Slot()
diff --git a/Engine/Plugins/Mutable/Source/CustomizableObjectEditor/Private/MuCOE/EdGraphSchema_CustomizableObject.cpp b/Engine/Plugins/Mutable/Source/CustomizableObjectEditor/Private/MuCOE/EdGraphSchema_CustomizableObject.cpp
--- a/Engine/Plugins/Mutable/Source/CustomizableObjectEditor/Private/MuCOE/EdGraphSchema_CustomizableObject.cpp
+++ b/Engine/Plugins/Mutable/Source/CustomizableObjectEditor/Private/MuCOE/EdGraphSchema_CustomizableObject.cpp
@@ -506,7 +506,6 @@ void UEdGraphSchema_CustomizableObject::GetGraphContextActions(FGraphContextMenu
 			ContextMenuBuilder.CreateTemplateNode<UCustomizableObjectNodeMeshVariation>(),
 			ContextMenuBuilder.CreateTemplateNode<UCustomizableObjectNodeMeshMorphStackDefinition>(),
 			ContextMenuBuilder.CreateTemplateNode<UCustomizableObjectNodeMeshMorphStackApplication>(),
-			// ContextMenuBuilder.CreateTemplateNode<UCustomizableObjectNodeMeshGeometryOperation>(),
 			ContextMenuBuilder.CreateTemplateNode<UCustomizableObjectNodeAnimationPose>(),
 		};
 
@@ -797,8 +796,10 @@ const FPinConnectionResponse UEdGraphSchema_CustomizableObject::CanCreateConnect
 }
 
 
-FLinearColor UEdGraphSchema_CustomizableObject::GetPinTypeColor(const FName& TypeString) const
+FLinearColor UEdGraphSchema_CustomizableObject::GetPinTypeColor(const FName& TypeString)
 {
+	static TMap<FName, FLinearColor> PinCategoryColors;
+
 	// Cache colors
 	if (PinCategoryColors.IsEmpty())
 	{
diff --git a/Engine/Plugins/Mutable/Source/CustomizableObjectEditor/Private/MuCOE/Nodes/CONodeMaterialBreak.cpp b/Engine/Plugins/Mutable/Source/CustomizableObjectEditor/Private/MuCOE/Nodes/CONodeMaterialBreak.cpp
--- a/Engine/Plugins/Mutable/Source/CustomizableObjectEditor/Private/MuCOE/Nodes/CONodeMaterialBreak.cpp
+++ b/Engine/Plugins/Mutable/Source/CustomizableObjectEditor/Private/MuCOE/Nodes/CONodeMaterialBreak.cpp
@@ -67,7 +67,7 @@ FText UCONodeMaterialBreak::GetNodeTitle(ENodeTitleType::Type TitleType) const
 
 FLinearColor UCONodeMaterialBreak::GetNodeTitleColor() const
 {
-	return GetDefault<UEdGraphSchema_CustomizableObject>()->GetPinTypeColor(UEdGraphSchema_CustomizableObject::PC_Material);
+	return UEdGraphSchema_CustomizableObject::GetPinTypeColor(UEdGraphSchema_CustomizableObject::PC_Material);
 }
 
 
diff --git a/Engine/Plugins/Mutable/Source/CustomizableObjectEditor/Private/MuCOE/Nodes/CONodeMaterialConstant.cpp b/Engine/Plugins/Mutable/Source/CustomizableObjectEditor/Private/MuCOE/Nodes/CONodeMaterialConstant.cpp
--- a/Engine/Plugins/Mutable/Source/CustomizableObjectEditor/Private/MuCOE/Nodes/CONodeMaterialConstant.cpp
+++ b/Engine/Plugins/Mutable/Source/CustomizableObjectEditor/Private/MuCOE/Nodes/CONodeMaterialConstant.cpp
@@ -31,7 +31,7 @@ FText UCONodeMaterialConstant::GetNodeTitle(ENodeTitleType::Type TitleType) cons
 
 FLinearColor UCONodeMaterialConstant::GetNodeTitleColor() const
 {
-	return GetDefault<UEdGraphSchema_CustomizableObject>()->GetPinTypeColor(UEdGraphSchema_CustomizableObject::PC_Material);
+	return UEdGraphSchema_CustomizableObject::GetPinTypeColor(UEdGraphSchema_CustomizableObject::PC_Material);
 }
 
 
diff --git a/Engine/Plugins/Mutable/Source/CustomizableObjectEditor/Private/MuCOE/Nodes/CONodeSchema.cpp b/Engine/Plugins/Mutable/Source/CustomizableObjectEditor/Private/MuCOE/Nodes/CONodeSchema.cpp
--- a/Engine/Plugins/Mutable/Source/CustomizableObjectEditor/Private/MuCOE/Nodes/CONodeSchema.cpp
+++ b/Engine/Plugins/Mutable/Source/CustomizableObjectEditor/Private/MuCOE/Nodes/CONodeSchema.cpp
@@ -33,6 +33,8 @@ void UCONodeSchema::AllocateDefaultPins(UCustomizableObjectNodeRemapPins* RemapP
 		UCONodeSchemaPinData* PinData = NewObject<UCONodeSchemaPinData>(this);
 		PinData->Color = PinSchema.Color;
 		PinData->bIsPassthrough = PinSchema.bIsPassthrough;
+		PinData->bIsEditable = PinSchema.bIsEditable;
+		PinData->Name = Pin->PinFriendlyName;
 		
 		AddPinData(*Pin, *PinData);
 	}
@@ -51,6 +53,24 @@ bool UCONodeSchema::IsPassthrough(const UEdGraphPin& Pin) const
 }
 
 
+bool UCONodeSchema::CanRenamePin(const UEdGraphPin& Pin) const
+{
+	return GetPinData<UCONodeSchemaPinData>(Pin).bIsEditable;
+}
+
+
+FText UCONodeSchema::GetPinEditableName(const UEdGraphPin& Pin) const
+{
+	return GetPinData<UCONodeSchemaPinData>(Pin).Name;
+}
+
+
+void UCONodeSchema::SetPinEditableName(const UEdGraphPin& Pin, const FText& Value)
+{
+	GetPinData<UCONodeSchemaPinData>(Pin).Name = Value;
+}
+
+
 UCustomizableObjectNodeRemapPins* UCONodeSchema::CreateRemapPinsDefault() const
 {
 	return NewObject<UCustomizableObjectNodeRemapPinsByPosition>();
diff --git a/Engine/Plugins/Mutable/Source/CustomizableObjectEditor/Private/MuCOE/Nodes/CONodeSchema.h b/Engine/Plugins/Mutable/Source/CustomizableObjectEditor/Private/MuCOE/Nodes/CONodeSchema.h
--- a/Engine/Plugins/Mutable/Source/CustomizableObjectEditor/Private/MuCOE/Nodes/CONodeSchema.h
+++ b/Engine/Plugins/Mutable/Source/CustomizableObjectEditor/Private/MuCOE/Nodes/CONodeSchema.h
@@ -21,6 +21,12 @@ public:
 
 	UPROPERTY()
 	bool bIsPassthrough = false;
+
+	UPROPERTY()
+	bool bIsEditable = false;
+	
+	UPROPERTY()
+	FText Name;
 };
 
 
@@ -40,6 +46,9 @@ struct FPinSchema
 	
 	UPROPERTY(EditAnywhere, Category = Schema)
 	bool bIsPassthrough = false;
+
+	UPROPERTY(EditAnywhere, Category = Schema)
+	bool bIsEditable = false;
 	
 	UPROPERTY(EditAnywhere, Category = Schema)
 	FLinearColor Color = FLinearColor::White;
@@ -63,6 +72,9 @@ public:
 	virtual void AllocateDefaultPins(UCustomizableObjectNodeRemapPins* RemapPins) override;
 	virtual FLinearColor GetPinColor(const UEdGraphPin &Pin) const override;
 	virtual bool IsPassthrough(const UEdGraphPin& Pin) const override;
+	virtual bool CanRenamePin(const UEdGraphPin& Pin) const override;
+	virtual FText GetPinEditableName(const UEdGraphPin& Pin) const override;
+	virtual void SetPinEditableName(const UEdGraphPin& Pin, const FText& Value) override;
 
 protected:
 	virtual UCustomizableObjectNodeRemapPins* CreateRemapPinsDefault() const override;
diff --git a/Engine/Plugins/Mutable/Source/CustomizableObjectEditor/Private/MuCOE/Nodes/CustomizableObjectNode.cpp b/Engine/Plugins/Mutable/Source/CustomizableObjectEditor/Private/MuCOE/Nodes/CustomizableObjectNode.cpp
--- a/Engine/Plugins/Mutable/Source/CustomizableObjectEditor/Private/MuCOE/Nodes/CustomizableObjectNode.cpp
+++ b/Engine/Plugins/Mutable/Source/CustomizableObjectEditor/Private/MuCOE/Nodes/CustomizableObjectNode.cpp
@@ -754,10 +754,26 @@ bool UCustomizableObjectNode::CanPinBeHidden(const UEdGraphPin& Pin) const
 }
 
 
+bool UCustomizableObjectNode::CanRenamePin(const UEdGraphPin& Pin) const
+{
+	return false;
+}
+
+
+FText UCustomizableObjectNode::GetPinEditableName(const UEdGraphPin& Pin) const
+{
+	return {};
+}
+
+
+void UCustomizableObjectNode::SetPinEditableName(const UEdGraphPin& Pin, const FText& Value)
+{
+}
+
+
 FLinearColor UCustomizableObjectNode::GetPinColor(const UEdGraphPin& Pin) const
 {
-	const UEdGraphSchema_CustomizableObject* Schema = GetDefault<UEdGraphSchema_CustomizableObject>();
-	return Schema->GetPinTypeColor(Pin.PinType);
+	return UEdGraphSchema_CustomizableObject::GetPinTypeColor(Pin.PinType.PinCategory);
 }
 
 
@@ -830,12 +846,6 @@ bool UCustomizableObjectNode::IsInMacro() const
 }
 
 
-bool UCustomizableObjectNode::CreateStaticStringPinWidget() const
-{
-	return true;
-}
-
-
 TArray<FName> UCustomizableObjectNode::GetAllowedPinViewerCreationTypes() const
 {
 	return UEdGraphSchema_CustomizableObject::SupportedMacroPinTypes;
diff --git a/Engine/Plugins/Mutable/Source/CustomizableObjectEditor/Private/MuCOE/Nodes/CustomizableObjectNodeAnimationPose.cpp b/Engine/Plugins/Mutable/Source/CustomizableObjectEditor/Private/MuCOE/Nodes/CustomizableObjectNodeAnimationPose.cpp
--- a/Engine/Plugins/Mutable/Source/CustomizableObjectEditor/Private/MuCOE/Nodes/CustomizableObjectNodeAnimationPose.cpp
+++ b/Engine/Plugins/Mutable/Source/CustomizableObjectEditor/Private/MuCOE/Nodes/CustomizableObjectNodeAnimationPose.cpp
@@ -96,7 +96,7 @@ FText UCustomizableObjectNodeAnimationPose::GetNodeTitle(ENodeTitleType::Type Ti
 
 FLinearColor UCustomizableObjectNodeAnimationPose::GetNodeTitleColor() const
 {
-	return GetDefault<UEdGraphSchema_CustomizableObject>()->GetPinTypeColor(UEdGraphSchema_CustomizableObject::PC_Mesh);
+	return UEdGraphSchema_CustomizableObject::GetPinTypeColor(UEdGraphSchema_CustomizableObject::PC_Mesh);
 }
 
 
diff --git a/Engine/Plugins/Mutable/Source/CustomizableObjectEditor/Private/MuCOE/Nodes/CustomizableObjectNodeColorArithmeticOp.cpp b/Engine/Plugins/Mutable/Source/CustomizableObjectEditor/Private/MuCOE/Nodes/CustomizableObjectNodeColorArithmeticOp.cpp
--- a/Engine/Plugins/Mutable/Source/CustomizableObjectEditor/Private/MuCOE/Nodes/CustomizableObjectNodeColorArithmeticOp.cpp
+++ b/Engine/Plugins/Mutable/Source/CustomizableObjectEditor/Private/MuCOE/Nodes/CustomizableObjectNodeColorArithmeticOp.cpp
@@ -71,8 +71,7 @@ FText UCustomizableObjectNodeColorArithmeticOp::GetNodeTitle(ENodeTitleType::Typ
 
 FLinearColor UCustomizableObjectNodeColorArithmeticOp::GetNodeTitleColor() const
 {
-	const UEdGraphSchema_CustomizableObject* Schema = GetDefault<UEdGraphSchema_CustomizableObject>();
-	return Schema->GetPinTypeColor(Schema->PC_Color);
+	return UEdGraphSchema_CustomizableObject::GetPinTypeColor(UEdGraphSchema_CustomizableObject::PC_Color);
 }
 
 
diff --git a/Engine/Plugins/Mutable/Source/CustomizableObjectEditor/Private/MuCOE/Nodes/CustomizableObjectNodeColorConstant.cpp b/Engine/Plugins/Mutable/Source/CustomizableObjectEditor/Private/MuCOE/Nodes/CustomizableObjectNodeColorConstant.cpp
--- a/Engine/Plugins/Mutable/Source/CustomizableObjectEditor/Private/MuCOE/Nodes/CustomizableObjectNodeColorConstant.cpp
+++ b/Engine/Plugins/Mutable/Source/CustomizableObjectEditor/Private/MuCOE/Nodes/CustomizableObjectNodeColorConstant.cpp
@@ -49,8 +49,7 @@ FText UCustomizableObjectNodeColorConstant::GetNodeTitle(ENodeTitleType::Type Ti
 
 FLinearColor UCustomizableObjectNodeColorConstant::GetNodeTitleColor() const
 {
-	const UEdGraphSchema_CustomizableObject* Schema = GetDefault<UEdGraphSchema_CustomizableObject>();
-	return Schema->GetPinTypeColor(Schema->PC_Color);
+	return UEdGraphSchema_CustomizableObject::GetPinTypeColor(UEdGraphSchema_CustomizableObject::PC_Color);
 }
 
 
diff --git a/Engine/Plugins/Mutable/Source/CustomizableObjectEditor/Private/MuCOE/Nodes/CustomizableObjectNodeColorFromFloats.cpp b/Engine/Plugins/Mutable/Source/CustomizableObjectEditor/Private/MuCOE/Nodes/CustomizableObjectNodeColorFromFloats.cpp
--- a/Engine/Plugins/Mutable/Source/CustomizableObjectEditor/Private/MuCOE/Nodes/CustomizableObjectNodeColorFromFloats.cpp
+++ b/Engine/Plugins/Mutable/Source/CustomizableObjectEditor/Private/MuCOE/Nodes/CustomizableObjectNodeColorFromFloats.cpp
@@ -54,8 +54,7 @@ FText UCustomizableObjectNodeColorFromFloats::GetNodeTitle(ENodeTitleType::Type
 
 FLinearColor UCustomizableObjectNodeColorFromFloats::GetNodeTitleColor() const
 {
-	const UEdGraphSchema_CustomizableObject* Schema = GetDefault<UEdGraphSchema_CustomizableObject>();
-	return Schema->GetPinTypeColor(Schema->PC_Color);
+	return UEdGraphSchema_CustomizableObject::GetPinTypeColor(UEdGraphSchema_CustomizableObject::PC_Color);
 }
 
 
diff --git a/Engine/Plugins/Mutable/Source/CustomizableObjectEditor/Private/MuCOE/Nodes/CustomizableObjectNodeColorToSRGB.cpp b/Engine/Plugins/Mutable/Source/CustomizableObjectEditor/Private/MuCOE/Nodes/CustomizableObjectNodeColorToSRGB.cpp
--- a/Engine/Plugins/Mutable/Source/CustomizableObjectEditor/Private/MuCOE/Nodes/CustomizableObjectNodeColorToSRGB.cpp
+++ b/Engine/Plugins/Mutable/Source/CustomizableObjectEditor/Private/MuCOE/Nodes/CustomizableObjectNodeColorToSRGB.cpp
@@ -39,8 +39,7 @@ FText UCustomizableObjectNodeColorToSRGB::GetNodeTitle(ENodeTitleType::Type Titl
 
 FLinearColor UCustomizableObjectNodeColorToSRGB::GetNodeTitleColor() const
 {
-	const UEdGraphSchema_CustomizableObject* Schema = GetDefault<UEdGraphSchema_CustomizableObject>();
-	return Schema->GetPinTypeColor(UEdGraphSchema_CustomizableObject::PC_Color);
+	return UEdGraphSchema_CustomizableObject::GetPinTypeColor(UEdGraphSchema_CustomizableObject::PC_Color);
 }
 
 
diff --git a/Engine/Plugins/Mutable/Source/CustomizableObjectEditor/Private/MuCOE/Nodes/CustomizableObjectNodeComponent.cpp b/Engine/Plugins/Mutable/Source/CustomizableObjectEditor/Private/MuCOE/Nodes/CustomizableObjectNodeComponent.cpp
--- a/Engine/Plugins/Mutable/Source/CustomizableObjectEditor/Private/MuCOE/Nodes/CustomizableObjectNodeComponent.cpp
+++ b/Engine/Plugins/Mutable/Source/CustomizableObjectEditor/Private/MuCOE/Nodes/CustomizableObjectNodeComponent.cpp
@@ -20,7 +20,7 @@ void UCustomizableObjectNodeComponent::AllocateDefaultPins(UCustomizableObjectNo
 
 FLinearColor UCustomizableObjectNodeComponent::GetNodeTitleColor() const
 {
-	return GetDefault<UEdGraphSchema_CustomizableObject>()->GetPinTypeColor(UEdGraphSchema_CustomizableObject::PC_Component);
+	return UEdGraphSchema_CustomizableObject::GetPinTypeColor(UEdGraphSchema_CustomizableObject::PC_Component);
 }
 
 
@@ -44,6 +44,24 @@ void UCustomizableObjectNodeComponent::BackwardsCompatibleFixup(int32 Customizab
 }
 
 
+bool UCustomizableObjectNodeComponent::CanRenamePin(const UEdGraphPin& Pin) const
+{
+	return Pin.PinType.PinCategory == UEdGraphSchema_CustomizableObject::PC_String;
+}
+
+
+FText UCustomizableObjectNodeComponent::GetPinEditableName(const UEdGraphPin& Pin) const
+{
+	return FText::FromName(GetComponentName());
+}
+
+
+void UCustomizableObjectNodeComponent::SetPinEditableName(const UEdGraphPin& Pin, const FText& InValue)
+{
+	SetComponentName(FName(*InValue.ToString()));
+}
+
+
 void UCustomizableObjectNodeComponent::OnRenameNode(const FString& NewName)
 {
 	if (!NewName.IsEmpty())
diff --git a/Engine/Plugins/Mutable/Source/CustomizableObjectEditor/Private/MuCOE/Nodes/CustomizableObjectNodeComponent.h b/Engine/Plugins/Mutable/Source/CustomizableObjectEditor/Private/MuCOE/Nodes/CustomizableObjectNodeComponent.h
--- a/Engine/Plugins/Mutable/Source/CustomizableObjectEditor/Private/MuCOE/Nodes/CustomizableObjectNodeComponent.h
+++ b/Engine/Plugins/Mutable/Source/CustomizableObjectEditor/Private/MuCOE/Nodes/CustomizableObjectNodeComponent.h
@@ -27,6 +27,9 @@ public:
 	// UCustomizableObjectNode interface
 	UE_API virtual bool IsAffectedByLOD() const override;
 	UE_API virtual void BackwardsCompatibleFixup(int32 CustomizableObjectCustomVersion) override;
+	UE_API virtual bool CanRenamePin(const UEdGraphPin& Pin) const override;
+	UE_API virtual FText GetPinEditableName(const UEdGraphPin& Pin) const override;
+	UE_API virtual void SetPinEditableName(const UEdGraphPin& Pin, const FText& Value) override;
 
 	//Own Interface
 	UE_API FName GetComponentName(TArray<const UCustomizableObjectNodeMacroInstance*>* MacroContext = nullptr) const;
@@ -34,20 +37,16 @@ public:
 	UE_API UEdGraphPin* GetComponentNamePin() const;
 
 protected:
-
-	UPROPERTY(EditAnywhere, Category = ComponentMesh)
-	FName ComponentName = "Default name";
+	UPROPERTY()
+	FName ComponentName = "Default name"; // TODO GMT Should be a FString due to FNames not being case sensitive at runtime
 
 public:
-
 	UPROPERTY()
 	FEdGraphPinReference OutputPin;
 
 private:
-	
 	UPROPERTY()
 	FEdGraphPinReference ComponentNamePin;
-
 };
 
 #undef UE_API
diff --git a/Engine/Plugins/Mutable/Source/CustomizableObjectEditor/Private/MuCOE/Nodes/CustomizableObjectNodeComponentMeshAddTo.cpp b/Engine/Plugins/Mutable/Source/CustomizableObjectEditor/Private/MuCOE/Nodes/CustomizableObjectNodeComponentMeshAddTo.cpp
--- a/Engine/Plugins/Mutable/Source/CustomizableObjectEditor/Private/MuCOE/Nodes/CustomizableObjectNodeComponentMeshAddTo.cpp
+++ b/Engine/Plugins/Mutable/Source/CustomizableObjectEditor/Private/MuCOE/Nodes/CustomizableObjectNodeComponentMeshAddTo.cpp
@@ -77,24 +77,33 @@ void UCustomizableObjectNodeComponentMeshAddTo::BackwardsCompatibleFixup(int32 C
 }
 
 
-FText UCustomizableObjectNodeComponentMeshAddTo::GetNodeTitle(ENodeTitleType::Type TitleType) const
+bool UCustomizableObjectNodeComponentMeshAddTo::CanRenamePin(const UEdGraphPin& Pin) const
 {
-	return LOCTEXT("ComponentMeshAdd", "Add To Skeletal Mesh Component");
+	return Pin.PinType.PinCategory == UEdGraphSchema_CustomizableObject::PC_String;
 }
 
 
-void UCustomizableObjectNodeComponentMeshAddTo::OnRenameNode(const FString& NewName)
+FText UCustomizableObjectNodeComponentMeshAddTo::GetPinEditableName(const UEdGraphPin& Pin) const
 {
-	if (!NewName.IsEmpty())
-	{
-		ParentComponentName = FName(*NewName);
-	}
+	return FText::FromName(GetParentComponentName());
+}
+
+
+void UCustomizableObjectNodeComponentMeshAddTo::SetPinEditableName(const UEdGraphPin& Pin, const FText& InValue)
+{
+	SetParentComponentName(FName(*InValue.ToString()));
+}
+
+
+FText UCustomizableObjectNodeComponentMeshAddTo::GetNodeTitle(ENodeTitleType::Type TitleType) const
+{
+	return LOCTEXT("ComponentMeshAdd", "Add To Skeletal Mesh Component");
 }
 
 
 FLinearColor UCustomizableObjectNodeComponentMeshAddTo::GetNodeTitleColor() const
 {
-	return GetDefault<UEdGraphSchema_CustomizableObject>()->GetPinTypeColor(UEdGraphSchema_CustomizableObject::PC_Component);
+	return UEdGraphSchema_CustomizableObject::GetPinTypeColor(UEdGraphSchema_CustomizableObject::PC_Component);
 }
 
 
diff --git a/Engine/Plugins/Mutable/Source/CustomizableObjectEditor/Private/MuCOE/Nodes/CustomizableObjectNodeComponentMeshAddTo.h b/Engine/Plugins/Mutable/Source/CustomizableObjectEditor/Private/MuCOE/Nodes/CustomizableObjectNodeComponentMeshAddTo.h
--- a/Engine/Plugins/Mutable/Source/CustomizableObjectEditor/Private/MuCOE/Nodes/CustomizableObjectNodeComponentMeshAddTo.h
+++ b/Engine/Plugins/Mutable/Source/CustomizableObjectEditor/Private/MuCOE/Nodes/CustomizableObjectNodeComponentMeshAddTo.h
@@ -16,12 +16,9 @@ class UCustomizableObjectNodeComponentMeshAddTo : public UCustomizableObjectNode
 	GENERATED_BODY()
 
 public:
-
 	// UEdGraphNode interface
 	UE_API virtual void AllocateDefaultPins(UCustomizableObjectNodeRemapPins* RemapPins) override;
-	virtual bool GetCanRenameNode() const override { return true; }
 	UE_API virtual FText GetNodeTitle(ENodeTitleType::Type TitleType) const override;
-	UE_API virtual void OnRenameNode(const FString& NewName) override;
 	UE_API virtual FLinearColor GetNodeTitleColor() const override;
 	UE_API virtual void PinConnectionListChanged(UEdGraphPin* Pin) override;
 
@@ -30,16 +27,19 @@ public:
 
 	//UCustomizableObjectNode interface
 	UE_API virtual bool IsAffectedByLOD() const override;
-	UE_API bool IsSingleOutputNode() const;
+	UE_API virtual bool IsSingleOutputNode() const override;
 	UE_API virtual void BackwardsCompatibleFixup(int32 CustomizableObjectCustomVersion) override;
-
+	UE_API virtual bool CanRenamePin(const UEdGraphPin& Pin) const override;
+	UE_API virtual FText GetPinEditableName(const UEdGraphPin& Pin) const override;
+	UE_API virtual void SetPinEditableName(const UEdGraphPin& Pin, const FText& Value) override;
+	
 	// ComponentMesh Interface
 	UE_API virtual int32 GetNumLODs() override;
 	UE_API virtual ECustomizableObjectAutomaticLODStrategy GetAutoLODStrategy() override;
 	UE_API virtual const TArray<FEdGraphPinReference>& GetLODPins() const override;
-	UE_API virtual UEdGraphPin* GetOutputPin() const;
-	UE_API virtual void SetOutputPin(const UEdGraphPin* Pin);
-	UE_API virtual const UCustomizableObjectNode* GetOwningNode() const;
+	UE_API virtual UEdGraphPin* GetOutputPin() const override;
+	UE_API virtual void SetOutputPin(const UEdGraphPin* Pin) override;
+	UE_API virtual const UCustomizableObjectNode* GetOwningNode() const override;
 
 	// Own interface
 	UE_API FName GetParentComponentName(TArray<const UCustomizableObjectNodeMacroInstance*>* MacroContext = nullptr) const;
@@ -47,12 +47,10 @@ public:
 	UE_API UEdGraphPin* GetParentComponentNamePin() const;
 
 private:
-
-	UPROPERTY(EditAnywhere, Category = ComponentMesh)
+	UPROPERTY()
 	FName ParentComponentName;
 
 public:
-
 	UPROPERTY(EditAnywhere, Category = ComponentMesh)
 	int32 NumLODs = 1;
 
@@ -66,10 +64,8 @@ public:
 	FEdGraphPinReference OutputPin;
 
 private:
-	
 	UPROPERTY()
 	FEdGraphPinReference ParentComponentNamePin;
-
 };
 
 #undef UE_API
diff --git a/Engine/Plugins/Mutable/Source/CustomizableObjectEditor/Private/MuCOE/Nodes/CustomizableObjectNodeCurve.cpp b/Engine/Plugins/Mutable/Source/CustomizableObjectEditor/Private/MuCOE/Nodes/CustomizableObjectNodeCurve.cpp
--- a/Engine/Plugins/Mutable/Source/CustomizableObjectEditor/Private/MuCOE/Nodes/CustomizableObjectNodeCurve.cpp
+++ b/Engine/Plugins/Mutable/Source/CustomizableObjectEditor/Private/MuCOE/Nodes/CustomizableObjectNodeCurve.cpp
@@ -138,8 +138,7 @@ FText UCustomizableObjectNodeCurve::GetNodeTitle(ENodeTitleType::Type TitleType)
 
 FLinearColor UCustomizableObjectNodeCurve::GetNodeTitleColor() const
 {
-	const UEdGraphSchema_CustomizableObject* Schema = GetDefault<UEdGraphSchema_CustomizableObject>();
-	return Schema->GetPinTypeColor(Schema->PC_Float);
+	return UEdGraphSchema_CustomizableObject::GetPinTypeColor(UEdGraphSchema_CustomizableObject::PC_Float);
 }
 
 
diff --git a/Engine/Plugins/Mutable/Source/CustomizableObjectEditor/Private/MuCOE/Nodes/CustomizableObjectNodeExposePin.cpp b/Engine/Plugins/Mutable/Source/CustomizableObjectEditor/Private/MuCOE/Nodes/CustomizableObjectNodeExposePin.cpp
--- a/Engine/Plugins/Mutable/Source/CustomizableObjectEditor/Private/MuCOE/Nodes/CustomizableObjectNodeExposePin.cpp
+++ b/Engine/Plugins/Mutable/Source/CustomizableObjectEditor/Private/MuCOE/Nodes/CustomizableObjectNodeExposePin.cpp
@@ -46,17 +46,31 @@ bool UCustomizableObjectNodeExposePin::IsNodeSupportedInMacros() const
 
 FText UCustomizableObjectNodeExposePin::GetNodeTitle(ENodeTitleType::Type TitleType) const
 {
+	FText NodeTitle;
 	FFormatNamedArguments Args;
+	Args.Add(TEXT("NodeName"), FText::FromString(Name));
 	Args.Add(TEXT("PinType"), UEdGraphSchema_CustomizableObject::GetPinCategoryFriendlyName(PinType));
 
-	return FText::Format(LOCTEXT("Expose_Pin_Title_ListView", "Export {PinType} Pin"), Args);
+	if (TitleType == ENodeTitleType::ListView)
+	{
+		NodeTitle = LOCTEXT("Expose_Pin_Title_ListView", "Export {PinType} Pin");
+	}
+	else if (TitleType == ENodeTitleType::EditableTitle)
+	{
+		NodeTitle = LOCTEXT("Expose_Pin_Title_Edit", "{NodeName}");
+	}
+	else
+	{
+		NodeTitle = LOCTEXT("Expose_Pin_Title", "{NodeName}\nExport {PinType} Pin");
+	}
+
+	return FText::Format(NodeTitle, Args);
 }
 
 
 FLinearColor UCustomizableObjectNodeExposePin::GetNodeTitleColor() const
 {
-	const UEdGraphSchema_CustomizableObject* Schema = GetDefault<UEdGraphSchema_CustomizableObject>();
-	return Schema->GetPinTypeColor(PinType);
+	return UEdGraphSchema_CustomizableObject::GetPinTypeColor(PinType);
 }
 
 
diff --git a/Engine/Plugins/Mutable/Source/CustomizableObjectEditor/Private/MuCOE/Nodes/CustomizableObjectNodeExtensionDataVariation.cpp b/Engine/Plugins/Mutable/Source/CustomizableObjectEditor/Private/MuCOE/Nodes/CustomizableObjectNodeExtensionDataVariation.cpp
--- a/Engine/Plugins/Mutable/Source/CustomizableObjectEditor/Private/MuCOE/Nodes/CustomizableObjectNodeExtensionDataVariation.cpp
+++ b/Engine/Plugins/Mutable/Source/CustomizableObjectEditor/Private/MuCOE/Nodes/CustomizableObjectNodeExtensionDataVariation.cpp
@@ -23,8 +23,7 @@ void UCustomizableObjectNodeExtensionDataVariation::PostEditChangeProperty(FProp
 
 FLinearColor UCustomizableObjectNodeExtensionDataVariation::GetNodeTitleColor() const
 {
-	const UEdGraphSchema_CustomizableObject* Schema = GetDefault<UEdGraphSchema_CustomizableObject>();
-	return Schema->GetPinTypeColor(GetCategory());
+	return UEdGraphSchema_CustomizableObject::GetPinTypeColor(GetCategory());
 }
 
 bool UCustomizableObjectNodeExtensionDataVariation::IsAffectedByLOD() const
diff --git a/Engine/Plugins/Mutable/Source/CustomizableObjectEditor/Private/MuCOE/Nodes/CustomizableObjectNodeExternalPin.cpp b/Engine/Plugins/Mutable/Source/CustomizableObjectEditor/Private/MuCOE/Nodes/CustomizableObjectNodeExternalPin.cpp
--- a/Engine/Plugins/Mutable/Source/CustomizableObjectEditor/Private/MuCOE/Nodes/CustomizableObjectNodeExternalPin.cpp
+++ b/Engine/Plugins/Mutable/Source/CustomizableObjectEditor/Private/MuCOE/Nodes/CustomizableObjectNodeExternalPin.cpp
@@ -264,8 +264,7 @@ FText UCustomizableObjectNodeExternalPin::GetNodeTitle(ENodeTitleType::Type Titl
 
 FLinearColor UCustomizableObjectNodeExternalPin::GetNodeTitleColor() const
 {
-	const UEdGraphSchema_CustomizableObject* Schema = GetDefault<UEdGraphSchema_CustomizableObject>();
-	return Schema->GetPinTypeColor(PinType);
+	return UEdGraphSchema_CustomizableObject::GetPinTypeColor(PinType);
 }
 
 
diff --git a/Engine/Plugins/Mutable/Source/CustomizableObjectEditor/Private/MuCOE/Nodes/CustomizableObjectNodeFloatArithmeticOp.cpp b/Engine/Plugins/Mutable/Source/CustomizableObjectEditor/Private/MuCOE/Nodes/CustomizableObjectNodeFloatArithmeticOp.cpp
--- a/Engine/Plugins/Mutable/Source/CustomizableObjectEditor/Private/MuCOE/Nodes/CustomizableObjectNodeFloatArithmeticOp.cpp
+++ b/Engine/Plugins/Mutable/Source/CustomizableObjectEditor/Private/MuCOE/Nodes/CustomizableObjectNodeFloatArithmeticOp.cpp
@@ -60,8 +60,7 @@ FText UCustomizableObjectNodeFloatArithmeticOp::GetNodeTitle(ENodeTitleType::Typ
 
 FLinearColor UCustomizableObjectNodeFloatArithmeticOp::GetNodeTitleColor() const
 {
-	const UEdGraphSchema_CustomizableObject* Schema = GetDefault<UEdGraphSchema_CustomizableObject>();
-	return Schema->GetPinTypeColor(Schema->PC_Float);
+	return UEdGraphSchema_CustomizableObject::GetPinTypeColor(UEdGraphSchema_CustomizableObject::PC_Float);
 }
 
 
diff --git a/Engine/Plugins/Mutable/Source/CustomizableObjectEditor/Private/MuCOE/Nodes/CustomizableObjectNodeFloatConstant.cpp b/Engine/Plugins/Mutable/Source/CustomizableObjectEditor/Private/MuCOE/Nodes/CustomizableObjectNodeFloatConstant.cpp
--- a/Engine/Plugins/Mutable/Source/CustomizableObjectEditor/Private/MuCOE/Nodes/CustomizableObjectNodeFloatConstant.cpp
+++ b/Engine/Plugins/Mutable/Source/CustomizableObjectEditor/Private/MuCOE/Nodes/CustomizableObjectNodeFloatConstant.cpp
@@ -47,8 +47,7 @@ FText UCustomizableObjectNodeFloatConstant::GetNodeTitle(ENodeTitleType::Type Ti
 
 FLinearColor UCustomizableObjectNodeFloatConstant::GetNodeTitleColor() const
 {
-	const UEdGraphSchema_CustomizableObject* Schema = GetDefault<UEdGraphSchema_CustomizableObject>();
-	return Schema->GetPinTypeColor(Schema->PC_Float);
+	return UEdGraphSchema_CustomizableObject::GetPinTypeColor(UEdGraphSchema_CustomizableObject::PC_Float);
 }
 
 
diff --git a/Engine/Plugins/Mutable/Source/CustomizableObjectEditor/Private/MuCOE/Nodes/CustomizableObjectNodeMacroInstance.cpp b/Engine/Plugins/Mutable/Source/CustomizableObjectEditor/Private/MuCOE/Nodes/CustomizableObjectNodeMacroInstance.cpp
--- a/Engine/Plugins/Mutable/Source/CustomizableObjectEditor/Private/MuCOE/Nodes/CustomizableObjectNodeMacroInstance.cpp
+++ b/Engine/Plugins/Mutable/Source/CustomizableObjectEditor/Private/MuCOE/Nodes/CustomizableObjectNodeMacroInstance.cpp
@@ -157,12 +157,6 @@ void UCustomizableObjectNodeMacroInstance::PostEditChangeProperty(FPropertyChang
 }
 
 
-bool UCustomizableObjectNodeMacroInstance::CreateStaticStringPinWidget() const
-{
-	return false;
-}
-
-
 bool UCustomizableObjectNodeMacroInstance::IsExperimental() const
 {
 	return true;
diff --git a/Engine/Plugins/Mutable/Source/CustomizableObjectEditor/Private/MuCOE/Nodes/CustomizableObjectNodeMacroInstance.h b/Engine/Plugins/Mutable/Source/CustomizableObjectEditor/Private/MuCOE/Nodes/CustomizableObjectNodeMacroInstance.h
--- a/Engine/Plugins/Mutable/Source/CustomizableObjectEditor/Private/MuCOE/Nodes/CustomizableObjectNodeMacroInstance.h
+++ b/Engine/Plugins/Mutable/Source/CustomizableObjectEditor/Private/MuCOE/Nodes/CustomizableObjectNodeMacroInstance.h
@@ -62,21 +62,19 @@ public:
 	UE_API virtual FText GetTooltipText() const override;
 	UE_API virtual FLinearColor GetNodeTitleColor() const override;
 	UE_API virtual bool CanUserDeleteNode() const override;
-	UE_API UObject* GetJumpTargetForDoubleClick() const override;
+	UE_API virtual UObject* GetJumpTargetForDoubleClick() const override;
 	UE_API virtual bool CanJumpToDefinition() const override;
 	UE_API virtual void JumpToDefinition() const override;
 
 	// UCustomizableObjectNode interface
-	UE_API void AllocateDefaultPins(UCustomizableObjectNodeRemapPins* RemapPins) override;
+	UE_API virtual void AllocateDefaultPins(UCustomizableObjectNodeRemapPins* RemapPins) override;
 	UE_API virtual UCustomizableObjectNodeMacroInstanceRemapPins* CreateRemapPinsDefault() const override;
 	UE_API virtual void PostEditChangeProperty(FPropertyChangedEvent& PropertyChangedEvent) override;
-	UE_API virtual bool CreateStaticStringPinWidget() const override;
 	UE_API virtual bool IsExperimental() const override;
 
 	// Node MacroInstance Interface
 	/** Returns the pin of the Macro's input/output node with the same name */
 	UE_API UEdGraphPin* GetMacroIOPin(ECOMacroIOType IONodeType, const FName& PinName) const;
-
 };
 
 #undef UE_API
diff --git a/Engine/Plugins/Mutable/Source/CustomizableObjectEditor/Private/MuCOE/Nodes/CustomizableObjectNodeMaterialBase.cpp b/Engine/Plugins/Mutable/Source/CustomizableObjectEditor/Private/MuCOE/Nodes/CustomizableObjectNodeMaterialBase.cpp
--- a/Engine/Plugins/Mutable/Source/CustomizableObjectEditor/Private/MuCOE/Nodes/CustomizableObjectNodeMaterialBase.cpp
+++ b/Engine/Plugins/Mutable/Source/CustomizableObjectEditor/Private/MuCOE/Nodes/CustomizableObjectNodeMaterialBase.cpp
@@ -13,5 +13,5 @@ int32 UCustomizableObjectNodeMaterialBase::GetMeshSectionMaxLOD() const
 
 FLinearColor UCustomizableObjectNodeMaterialBase::GetNodeTitleColor() const
 {
-	return GetDefault<UEdGraphSchema_CustomizableObject>()->GetPinTypeColor(UEdGraphSchema_CustomizableObject::PC_MeshSection);
+	return UEdGraphSchema_CustomizableObject::GetPinTypeColor(UEdGraphSchema_CustomizableObject::PC_MeshSection);
 }
diff --git a/Engine/Plugins/Mutable/Source/CustomizableObjectEditor/Private/MuCOE/Nodes/CustomizableObjectNodeMeshMorph.cpp b/Engine/Plugins/Mutable/Source/CustomizableObjectEditor/Private/MuCOE/Nodes/CustomizableObjectNodeMeshMorph.cpp
--- a/Engine/Plugins/Mutable/Source/CustomizableObjectEditor/Private/MuCOE/Nodes/CustomizableObjectNodeMeshMorph.cpp
+++ b/Engine/Plugins/Mutable/Source/CustomizableObjectEditor/Private/MuCOE/Nodes/CustomizableObjectNodeMeshMorph.cpp
@@ -55,8 +55,7 @@ FText UCustomizableObjectNodeMeshMorph::GetNodeTitle(ENodeTitleType::Type TitleT
 
 FLinearColor UCustomizableObjectNodeMeshMorph::GetNodeTitleColor() const
 {
-	const UEdGraphSchema_CustomizableObject* Schema = GetDefault<UEdGraphSchema_CustomizableObject>();
-	return Schema->GetPinTypeColor(Schema->PC_Mesh);
+	return UEdGraphSchema_CustomizableObject::GetPinTypeColor(UEdGraphSchema_CustomizableObject::PC_Mesh);
 }
 
 
@@ -165,12 +164,6 @@ void UCustomizableObjectNodeMeshMorph::Serialize(FArchive& Ar)
 }
 
 
-bool UCustomizableObjectNodeMeshMorph::CreateStaticStringPinWidget() const
-{
-	return false;
-}
-
-
 void UCustomizableObjectNodeMeshMorph::BackwardsCompatibleFixup(int32 CustomizableObjectCustomVersion)
 {
 	Super::BackwardsCompatibleFixup(CustomizableObjectCustomVersion);
diff --git a/Engine/Plugins/Mutable/Source/CustomizableObjectEditor/Private/MuCOE/Nodes/CustomizableObjectNodeMeshMorph.h b/Engine/Plugins/Mutable/Source/CustomizableObjectEditor/Private/MuCOE/Nodes/CustomizableObjectNodeMeshMorph.h
--- a/Engine/Plugins/Mutable/Source/CustomizableObjectEditor/Private/MuCOE/Nodes/CustomizableObjectNodeMeshMorph.h
+++ b/Engine/Plugins/Mutable/Source/CustomizableObjectEditor/Private/MuCOE/Nodes/CustomizableObjectNodeMeshMorph.h
@@ -33,7 +33,6 @@ public:
 	UE_API virtual bool IsNodeOutDatedAndNeedsRefresh() override;
 	UE_API virtual FString GetRefreshMessage() const override;
 	UE_API virtual void PinConnectionListChanged(UEdGraphPin* Pin) override;
-	UE_API virtual bool CreateStaticStringPinWidget() const override;
 	UE_API virtual void BackwardsCompatibleFixup(int32 CustomizableObjectCustomVersion) override;
 
 
diff --git a/Engine/Plugins/Mutable/Source/CustomizableObjectEditor/Private/MuCOE/Nodes/CustomizableObjectNodeMeshMorphStackApplication.cpp b/Engine/Plugins/Mutable/Source/CustomizableObjectEditor/Private/MuCOE/Nodes/CustomizableObjectNodeMeshMorphStackApplication.cpp
--- a/Engine/Plugins/Mutable/Source/CustomizableObjectEditor/Private/MuCOE/Nodes/CustomizableObjectNodeMeshMorphStackApplication.cpp
+++ b/Engine/Plugins/Mutable/Source/CustomizableObjectEditor/Private/MuCOE/Nodes/CustomizableObjectNodeMeshMorphStackApplication.cpp
@@ -64,7 +64,7 @@ FText UCustomizableObjectNodeMeshMorphStackApplication::GetNodeTitle(ENodeTitleT
 
 FLinearColor UCustomizableObjectNodeMeshMorphStackApplication::GetNodeTitleColor() const
 {
-	return GetDefault<UEdGraphSchema_CustomizableObject>()->GetPinTypeColor(UEdGraphSchema_CustomizableObject::PC_Mesh);
+	return UEdGraphSchema_CustomizableObject::GetPinTypeColor(UEdGraphSchema_CustomizableObject::PC_Mesh);
 }
 
 
diff --git a/Engine/Plugins/Mutable/Source/CustomizableObjectEditor/Private/MuCOE/Nodes/CustomizableObjectNodeMeshMorphStackDefinition.cpp b/Engine/Plugins/Mutable/Source/CustomizableObjectEditor/Private/MuCOE/Nodes/CustomizableObjectNodeMeshMorphStackDefinition.cpp
--- a/Engine/Plugins/Mutable/Source/CustomizableObjectEditor/Private/MuCOE/Nodes/CustomizableObjectNodeMeshMorphStackDefinition.cpp
+++ b/Engine/Plugins/Mutable/Source/CustomizableObjectEditor/Private/MuCOE/Nodes/CustomizableObjectNodeMeshMorphStackDefinition.cpp
@@ -48,8 +48,7 @@ FText UCustomizableObjectNodeMeshMorphStackDefinition::GetNodeTitle(ENodeTitleTy
 
 FLinearColor UCustomizableObjectNodeMeshMorphStackDefinition::GetNodeTitleColor() const
 {
-	const UEdGraphSchema_CustomizableObject* Schema = GetDefault<UEdGraphSchema_CustomizableObject>();
-	return Schema->GetPinTypeColor(Schema->PC_Stack);
+	return UEdGraphSchema_CustomizableObject::GetPinTypeColor(UEdGraphSchema_CustomizableObject::PC_Stack);
 }
 
 
diff --git a/Engine/Plugins/Mutable/Source/CustomizableObjectEditor/Private/MuCOE/Nodes/CustomizableObjectNodeMeshReshape.cpp b/Engine/Plugins/Mutable/Source/CustomizableObjectEditor/Private/MuCOE/Nodes/CustomizableObjectNodeMeshReshape.cpp
--- a/Engine/Plugins/Mutable/Source/CustomizableObjectEditor/Private/MuCOE/Nodes/CustomizableObjectNodeMeshReshape.cpp
+++ b/Engine/Plugins/Mutable/Source/CustomizableObjectEditor/Private/MuCOE/Nodes/CustomizableObjectNodeMeshReshape.cpp
@@ -45,8 +45,7 @@ FText UCustomizableObjectNodeMeshReshape::GetNodeTitle(ENodeTitleType::Type Titl
 
 FLinearColor UCustomizableObjectNodeMeshReshape::GetNodeTitleColor() const
 {
-	const UEdGraphSchema_CustomizableObject* Schema = GetDefault<UEdGraphSchema_CustomizableObject>();
-	return Schema->GetPinTypeColor(Schema->PC_Mesh);
+	return UEdGraphSchema_CustomizableObject::GetPinTypeColor(UEdGraphSchema_CustomizableObject::PC_Mesh);
 }
 
 
diff --git a/Engine/Plugins/Mutable/Source/CustomizableObjectEditor/Private/MuCOE/Nodes/CustomizableObjectNodeModifierBase.cpp b/Engine/Plugins/Mutable/Source/CustomizableObjectEditor/Private/MuCOE/Nodes/CustomizableObjectNodeModifierBase.cpp
--- a/Engine/Plugins/Mutable/Source/CustomizableObjectEditor/Private/MuCOE/Nodes/CustomizableObjectNodeModifierBase.cpp
+++ b/Engine/Plugins/Mutable/Source/CustomizableObjectEditor/Private/MuCOE/Nodes/CustomizableObjectNodeModifierBase.cpp
@@ -16,8 +16,7 @@
 
 FLinearColor UCustomizableObjectNodeModifierBase::GetNodeTitleColor() const
 {
-	const UEdGraphSchema_CustomizableObject* Schema = GetDefault<UEdGraphSchema_CustomizableObject>();
-	return Schema->GetPinTypeColor(Schema->PC_Modifier);
+	return UEdGraphSchema_CustomizableObject::GetPinTypeColor(UEdGraphSchema_CustomizableObject::PC_Modifier);
 }
 
 
diff --git a/Engine/Plugins/Mutable/Source/CustomizableObjectEditor/Private/MuCOE/Nodes/CustomizableObjectNodeModifierMorphMeshSection.cpp b/Engine/Plugins/Mutable/Source/CustomizableObjectEditor/Private/MuCOE/Nodes/CustomizableObjectNodeModifierMorphMeshSection.cpp
--- a/Engine/Plugins/Mutable/Source/CustomizableObjectEditor/Private/MuCOE/Nodes/CustomizableObjectNodeModifierMorphMeshSection.cpp
+++ b/Engine/Plugins/Mutable/Source/CustomizableObjectEditor/Private/MuCOE/Nodes/CustomizableObjectNodeModifierMorphMeshSection.cpp
@@ -68,12 +68,6 @@ UEdGraphPin* UCustomizableObjectNodeModifierMorphMeshSection::MorphTargetNamePin
 }
 
 
-bool UCustomizableObjectNodeModifierMorphMeshSection::CreateStaticStringPinWidget() const
-{
-	return false;
-}
-
-
 void UCustomizableObjectNodeModifierMorphMeshSection::BackwardsCompatibleFixup(int32 CustomizableObjectCustomVersion)
 {
 	Super::BackwardsCompatibleFixup(CustomizableObjectCustomVersion);
diff --git a/Engine/Plugins/Mutable/Source/CustomizableObjectEditor/Private/MuCOE/Nodes/CustomizableObjectNodeModifierMorphMeshSection.h b/Engine/Plugins/Mutable/Source/CustomizableObjectEditor/Private/MuCOE/Nodes/CustomizableObjectNodeModifierMorphMeshSection.h
--- a/Engine/Plugins/Mutable/Source/CustomizableObjectEditor/Private/MuCOE/Nodes/CustomizableObjectNodeModifierMorphMeshSection.h
+++ b/Engine/Plugins/Mutable/Source/CustomizableObjectEditor/Private/MuCOE/Nodes/CustomizableObjectNodeModifierMorphMeshSection.h
@@ -35,7 +35,6 @@ public:
 	UE_API virtual void AllocateDefaultPins(UCustomizableObjectNodeRemapPins* RemapPins) override;
 	UE_API virtual FString GetRefreshMessage() const override;
 	UE_API virtual bool IsSingleOutputNode() const override;
-	UE_API virtual bool CreateStaticStringPinWidget() const override;
 	UE_API virtual void BackwardsCompatibleFixup(int32 CustomizableObjectCustomVersion) override;
 
 	UEdGraphPin* FactorPin() const
diff --git a/Engine/Plugins/Mutable/Source/CustomizableObjectEditor/Private/MuCOE/Nodes/CustomizableObjectNodeObject.cpp b/Engine/Plugins/Mutable/Source/CustomizableObjectEditor/Private/MuCOE/Nodes/CustomizableObjectNodeObject.cpp
--- a/Engine/Plugins/Mutable/Source/CustomizableObjectEditor/Private/MuCOE/Nodes/CustomizableObjectNodeObject.cpp
+++ b/Engine/Plugins/Mutable/Source/CustomizableObjectEditor/Private/MuCOE/Nodes/CustomizableObjectNodeObject.cpp
@@ -619,6 +619,24 @@ bool UCustomizableObjectNodeObject::IsNodeSupportedInMacros() const
 }
 
 
+bool UCustomizableObjectNodeObject::CanRenamePin(const UEdGraphPin& Pin) const
+{
+	return Pin.PinType.PinCategory == UEdGraphSchema_CustomizableObject::PC_String;
+}
+
+
+FText UCustomizableObjectNodeObject::GetPinEditableName(const UEdGraphPin& Pin) const
+{
+	return FText::FromString(GetObjectName());
+}
+
+
+void UCustomizableObjectNodeObject::SetPinEditableName(const UEdGraphPin& Pin, const FText& InValue)
+{
+	SetObjectName(InValue.ToString());
+}
+
+
 FText UCustomizableObjectNodeObject::GetNodeTitle(ENodeTitleType::Type TitleType) const
 {
 	if (bIsBase)
@@ -634,8 +652,7 @@ FText UCustomizableObjectNodeObject::GetNodeTitle(ENodeTitleType::Type TitleType
 
 FLinearColor UCustomizableObjectNodeObject::GetNodeTitleColor() const
 {
-	const UEdGraphSchema_CustomizableObject* Schema = GetDefault<UEdGraphSchema_CustomizableObject>();
-	return Schema->GetPinTypeColor(Schema->PC_Object);
+	return UEdGraphSchema_CustomizableObject::GetPinTypeColor(UEdGraphSchema_CustomizableObject::PC_Object);
 }
 
 
@@ -715,15 +732,6 @@ bool UCustomizableObjectNodeObject::IsSingleOutputNode() const
 }
 
 
-void UCustomizableObjectNodeObject::OnRenameNode(const FString& NewName)
-{
-	if (!NewName.IsEmpty())
-	{
-		SetObjectName(NewName);
-	}
-}
-
-
 void UCustomizableObjectNodeObject::PinConnectionListChanged(UEdGraphPin* Pin)
 {
 	if (Pin == NamePin.Get())
diff --git a/Engine/Plugins/Mutable/Source/CustomizableObjectEditor/Private/MuCOE/Nodes/CustomizableObjectNodeObject.h b/Engine/Plugins/Mutable/Source/CustomizableObjectEditor/Private/MuCOE/Nodes/CustomizableObjectNodeObject.h
--- a/Engine/Plugins/Mutable/Source/CustomizableObjectEditor/Private/MuCOE/Nodes/CustomizableObjectNodeObject.h
+++ b/Engine/Plugins/Mutable/Source/CustomizableObjectEditor/Private/MuCOE/Nodes/CustomizableObjectNodeObject.h
@@ -165,8 +165,6 @@ public:
 	UE_API virtual void PrepareForCopying() override;
 	UE_API virtual void PostPasteNode() override;
 	UE_API virtual void PostDuplicate(bool bDuplicateForPIE) override;
-	virtual bool GetCanRenameNode() const override { return true; }
-	UE_API virtual void OnRenameNode(const FString& NewName) override;
 	UE_API virtual void PinConnectionListChanged(UEdGraphPin* Pin) override;
 
 	// UCustomizableObjectNode interface
@@ -174,6 +172,9 @@ public:
 	UE_API virtual void PostBackwardsCompatibleFixup() override;
 	UE_API virtual void AllocateDefaultPins(UCustomizableObjectNodeRemapPins* RemapPins) override;
 	UE_API virtual bool IsNodeSupportedInMacros() const override;
+	UE_API virtual bool CanRenamePin(const UEdGraphPin& Pin) const override;
+	UE_API virtual FText GetPinEditableName(const UEdGraphPin& Pin) const override;
+	UE_API virtual void SetPinEditableName(const UEdGraphPin& Pin, const FText& Value) override;
 	
 	// Own Interface
 	UE_API void SetParentObject(UCustomizableObject* CustomizableParentObject);
diff --git a/Engine/Plugins/Mutable/Source/CustomizableObjectEditor/Private/MuCOE/Nodes/CustomizableObjectNodeObjectGroup.cpp b/Engine/Plugins/Mutable/Source/CustomizableObjectEditor/Private/MuCOE/Nodes/CustomizableObjectNodeObjectGroup.cpp
--- a/Engine/Plugins/Mutable/Source/CustomizableObjectEditor/Private/MuCOE/Nodes/CustomizableObjectNodeObjectGroup.cpp
+++ b/Engine/Plugins/Mutable/Source/CustomizableObjectEditor/Private/MuCOE/Nodes/CustomizableObjectNodeObjectGroup.cpp
@@ -96,7 +96,7 @@ FText UCustomizableObjectNodeObjectGroup::GetNodeTitle(ENodeTitleType::Type Titl
 
 FLinearColor UCustomizableObjectNodeObjectGroup::GetNodeTitleColor() const
 {
-	return GetDefault<UEdGraphSchema_CustomizableObject>()->GetPinTypeColor(UEdGraphSchema_CustomizableObject::PC_Object);
+	return UEdGraphSchema_CustomizableObject::GetPinTypeColor(UEdGraphSchema_CustomizableObject::PC_Object);
 }
 
 
@@ -159,6 +159,24 @@ void UCustomizableObjectNodeObjectGroup::BackwardsCompatibleFixup(int32 Customiz
 }
 
 
+bool UCustomizableObjectNodeObjectGroup::CanRenamePin(const UEdGraphPin& Pin) const
+{
+	return Pin.PinType.PinCategory == UEdGraphSchema_CustomizableObject::PC_String;
+}
+
+
+FText UCustomizableObjectNodeObjectGroup::GetPinEditableName(const UEdGraphPin& Pin) const
+{
+	return FText::FromString(GetGroupName());
+}
+
+
+void UCustomizableObjectNodeObjectGroup::SetPinEditableName(const UEdGraphPin& Pin, const FText& InValue)
+{
+	SetGroupName(InValue.ToString());
+}
+
+
 UEdGraphPin* UCustomizableObjectNodeObjectGroup::ObjectsPin() const
 {
 	return FindPin(TEXT("Objects"));
diff --git a/Engine/Plugins/Mutable/Source/CustomizableObjectEditor/Private/MuCOE/Nodes/CustomizableObjectNodeObjectGroup.h b/Engine/Plugins/Mutable/Source/CustomizableObjectEditor/Private/MuCOE/Nodes/CustomizableObjectNodeObjectGroup.h
--- a/Engine/Plugins/Mutable/Source/CustomizableObjectEditor/Private/MuCOE/Nodes/CustomizableObjectNodeObjectGroup.h
+++ b/Engine/Plugins/Mutable/Source/CustomizableObjectEditor/Private/MuCOE/Nodes/CustomizableObjectNodeObjectGroup.h
@@ -58,6 +58,9 @@ public:
 	UE_API virtual void AllocateDefaultPins(UCustomizableObjectNodeRemapPins* RemapPins) override;
 	UE_API virtual bool IsSingleOutputNode() const override;
 	UE_API virtual void BackwardsCompatibleFixup(int32 CustomizableObjectCustomVersion) override;
+	UE_API virtual bool CanRenamePin(const UEdGraphPin& Pin) const override;
+	UE_API virtual FText GetPinEditableName(const UEdGraphPin& Pin) const override;
+	UE_API virtual void SetPinEditableName(const UEdGraphPin& Pin, const FText& Value) override;
 
 	// Own interface
 	UE_API UEdGraphPin* ObjectsPin() const;
diff --git a/Engine/Plugins/Mutable/Source/CustomizableObjectEditor/Private/MuCOE/Nodes/CustomizableObjectNodeParameter.cpp b/Engine/Plugins/Mutable/Source/CustomizableObjectEditor/Private/MuCOE/Nodes/CustomizableObjectNodeParameter.cpp
--- a/Engine/Plugins/Mutable/Source/CustomizableObjectEditor/Private/MuCOE/Nodes/CustomizableObjectNodeParameter.cpp
+++ b/Engine/Plugins/Mutable/Source/CustomizableObjectEditor/Private/MuCOE/Nodes/CustomizableObjectNodeParameter.cpp
@@ -30,6 +30,12 @@ void UCustomizableObjectNodeParameter::AllocateDefaultPins(UCustomizableObjectNo
 }
 
 
+bool UCustomizableObjectNodeParameter::IsAffectedByLOD() const
+{
+	return false;
+}
+
+
 void UCustomizableObjectNodeParameter::BackwardsCompatibleFixup(int32 CustomizableObjectCustomVersion)
 {
 	Super::BackwardsCompatibleFixup(CustomizableObjectCustomVersion);
@@ -44,6 +50,24 @@ void UCustomizableObjectNodeParameter::BackwardsCompatibleFixup(int32 Customizab
 }
 
 
+bool UCustomizableObjectNodeParameter::CanRenamePin(const UEdGraphPin& Pin) const
+{
+	return Pin.PinType.PinCategory == UEdGraphSchema_CustomizableObject::PC_String;
+}
+
+
+FText UCustomizableObjectNodeParameter::GetPinEditableName(const UEdGraphPin& Pin) const
+{
+	return FText::FromString(GetParameterName());
+}
+
+
+void UCustomizableObjectNodeParameter::SetPinEditableName(const UEdGraphPin& Pin, const FText& InValue)
+{
+	SetParameterName(InValue.ToString());
+}
+
+
 FText UCustomizableObjectNodeParameter::GetNodeTitle(ENodeTitleType::Type TitleType) const
 {
 	FFormatNamedArguments Args;
@@ -55,7 +79,7 @@ FText UCustomizableObjectNodeParameter::GetNodeTitle(ENodeTitleType::Type TitleT
 
 FLinearColor UCustomizableObjectNodeParameter::GetNodeTitleColor() const
 {
-	return GetDefault<UEdGraphSchema_CustomizableObject>()->GetPinTypeColor(GetCategory());
+	return UEdGraphSchema_CustomizableObject::GetPinTypeColor(GetCategory());
 }
 
 
diff --git a/Engine/Plugins/Mutable/Source/CustomizableObjectEditor/Private/MuCOE/Nodes/CustomizableObjectNodeParameter.h b/Engine/Plugins/Mutable/Source/CustomizableObjectEditor/Private/MuCOE/Nodes/CustomizableObjectNodeParameter.h
--- a/Engine/Plugins/Mutable/Source/CustomizableObjectEditor/Private/MuCOE/Nodes/CustomizableObjectNodeParameter.h
+++ b/Engine/Plugins/Mutable/Source/CustomizableObjectEditor/Private/MuCOE/Nodes/CustomizableObjectNodeParameter.h
@@ -42,8 +42,11 @@ public:
 
 	// UCustomizableObjectNode interface
 	UE_API virtual void AllocateDefaultPins(UCustomizableObjectNodeRemapPins* RemapPins) override;
-	virtual bool IsAffectedByLOD() const override { return false; }
+	UE_API virtual bool IsAffectedByLOD() const override;
 	UE_API virtual void BackwardsCompatibleFixup(int32 CustomizableObjectCustomVersion) override;
+	UE_API virtual bool CanRenamePin(const UEdGraphPin& Pin) const override;
+	UE_API virtual FText GetPinEditableName(const UEdGraphPin& Pin) const override;
+	UE_API virtual void SetPinEditableName(const UEdGraphPin& Pin, const FText& Value) override;
 
 	// Own interface
 	/** Return the pin category of this node. */
diff --git a/Engine/Plugins/Mutable/Source/CustomizableObjectEditor/Private/MuCOE/Nodes/CustomizableObjectNodePassThroughTexture.cpp b/Engine/Plugins/Mutable/Source/CustomizableObjectEditor/Private/MuCOE/Nodes/CustomizableObjectNodePassThroughTexture.cpp
--- a/Engine/Plugins/Mutable/Source/CustomizableObjectEditor/Private/MuCOE/Nodes/CustomizableObjectNodePassThroughTexture.cpp
+++ b/Engine/Plugins/Mutable/Source/CustomizableObjectEditor/Private/MuCOE/Nodes/CustomizableObjectNodePassThroughTexture.cpp
@@ -62,7 +62,7 @@ FText UCustomizableObjectNodePassThroughTexture::GetNodeTitle(ENodeTitleType::Ty
 
 FLinearColor UCustomizableObjectNodePassThroughTexture::GetNodeTitleColor() const
 {
-	return GetDefault<UEdGraphSchema_CustomizableObject>()->GetPinTypeColor(UEdGraphSchema_CustomizableObject::PC_PassthroughTexture);
+	return UEdGraphSchema_CustomizableObject::GetPinTypeColor(UEdGraphSchema_CustomizableObject::PC_PassthroughTexture);
 }
 
 
diff --git a/Engine/Plugins/Mutable/Source/CustomizableObjectEditor/Private/MuCOE/Nodes/CustomizableObjectNodeProjectorConstant.cpp b/Engine/Plugins/Mutable/Source/CustomizableObjectEditor/Private/MuCOE/Nodes/CustomizableObjectNodeProjectorConstant.cpp
--- a/Engine/Plugins/Mutable/Source/CustomizableObjectEditor/Private/MuCOE/Nodes/CustomizableObjectNodeProjectorConstant.cpp
+++ b/Engine/Plugins/Mutable/Source/CustomizableObjectEditor/Private/MuCOE/Nodes/CustomizableObjectNodeProjectorConstant.cpp
@@ -151,8 +151,7 @@ FText UCustomizableObjectNodeProjectorConstant::GetNodeTitle(ENodeTitleType::Typ
 
 FLinearColor UCustomizableObjectNodeProjectorConstant::GetNodeTitleColor() const
 {
-	const UEdGraphSchema_CustomizableObject* Schema = GetDefault<UEdGraphSchema_CustomizableObject>();
-	return Schema->GetPinTypeColor(Schema->PC_Projector);
+	return UEdGraphSchema_CustomizableObject::GetPinTypeColor(UEdGraphSchema_CustomizableObject::PC_Projector);
 }
 
 
diff --git a/Engine/Plugins/Mutable/Source/CustomizableObjectEditor/Private/MuCOE/Nodes/CustomizableObjectNodeSkeletalMesh.cpp b/Engine/Plugins/Mutable/Source/CustomizableObjectEditor/Private/MuCOE/Nodes/CustomizableObjectNodeSkeletalMesh.cpp
--- a/Engine/Plugins/Mutable/Source/CustomizableObjectEditor/Private/MuCOE/Nodes/CustomizableObjectNodeSkeletalMesh.cpp
+++ b/Engine/Plugins/Mutable/Source/CustomizableObjectEditor/Private/MuCOE/Nodes/CustomizableObjectNodeSkeletalMesh.cpp
@@ -185,8 +185,7 @@ FText UCustomizableObjectNodeSkeletalMesh::GetNodeTitle(ENodeTitleType::Type Tit
 
 FLinearColor UCustomizableObjectNodeSkeletalMesh::GetNodeTitleColor() const
 {
-	const UEdGraphSchema_CustomizableObject* Schema = GetDefault<UEdGraphSchema_CustomizableObject>();
-	return Schema->GetPinTypeColor(Schema->PC_Mesh);
+	return UEdGraphSchema_CustomizableObject::GetPinTypeColor(UEdGraphSchema_CustomizableObject::PC_Mesh);
 }
 
 
diff --git a/Engine/Plugins/Mutable/Source/CustomizableObjectEditor/Private/MuCOE/Nodes/CustomizableObjectNodeSkeletalMeshParameter.cpp b/Engine/Plugins/Mutable/Source/CustomizableObjectEditor/Private/MuCOE/Nodes/CustomizableObjectNodeSkeletalMeshParameter.cpp
--- a/Engine/Plugins/Mutable/Source/CustomizableObjectEditor/Private/MuCOE/Nodes/CustomizableObjectNodeSkeletalMeshParameter.cpp
+++ b/Engine/Plugins/Mutable/Source/CustomizableObjectEditor/Private/MuCOE/Nodes/CustomizableObjectNodeSkeletalMeshParameter.cpp
@@ -108,8 +108,7 @@ FText UCustomizableObjectNodeSkeletalMeshParameter::GetNodeTitle(ENodeTitleType:
 
 FLinearColor UCustomizableObjectNodeSkeletalMeshParameter::GetNodeTitleColor() const
 {
-	const UEdGraphSchema_CustomizableObject* Schema = GetDefault<UEdGraphSchema_CustomizableObject>();
-	return Schema->GetPinTypeColor(Schema->PC_Mesh);
+	return UEdGraphSchema_CustomizableObject::GetPinTypeColor(UEdGraphSchema_CustomizableObject::PC_Mesh);
 }
 
 
diff --git a/Engine/Plugins/Mutable/Source/CustomizableObjectEditor/Private/MuCOE/Nodes/CustomizableObjectNodeSkeletalMeshParameter.h b/Engine/Plugins/Mutable/Source/CustomizableObjectEditor/Private/MuCOE/Nodes/CustomizableObjectNodeSkeletalMeshParameter.h
--- a/Engine/Plugins/Mutable/Source/CustomizableObjectEditor/Private/MuCOE/Nodes/CustomizableObjectNodeSkeletalMeshParameter.h
+++ b/Engine/Plugins/Mutable/Source/CustomizableObjectEditor/Private/MuCOE/Nodes/CustomizableObjectNodeSkeletalMeshParameter.h
@@ -49,17 +49,10 @@ private:
 UCLASS(MinimalAPI)
 class UCustomizableObjectNodeSkeletalMeshParameter : public UCustomizableObjectNodeParameter, public ICustomizableObjectNodeMeshInterface
 {
-public:
 	GENERATED_BODY()
-
-	/** Default value of the parameter. */
-	UPROPERTY(EditAnywhere, Category=CustomizableObject)
-	TSoftObjectPtr<USkeletalMesh> DefaultValue;
-
-	UPROPERTY(EditAnywhere, Category=CustomizableObject)
-	TSoftObjectPtr<USkeletalMesh> ReferenceValue;
 	
-	// UObject interface.
+public:
+	// UObject interface
 	UE_API virtual void PostEditChangeProperty(FPropertyChangedEvent& PropertyChangedEvent) override;
 
 	// EdGraphNode interface
@@ -70,11 +63,10 @@ public:
 
 	// UCustomizableObjectNode interface
 	UE_API virtual void AllocateDefaultPins(UCustomizableObjectNodeRemapPins* RemapPins) override;
-	virtual bool IsAffectedByLOD() const override { return false; }
 	UE_API virtual bool IsExperimental() const override;
 	UE_API virtual void BackwardsCompatibleFixup(int32 CustomizableObjectCustomVersion) override;
-
-	// UCustomizableObjectNodeParameter iterface
+	
+	// UCustomizableObjectNodeParameter interface
 	virtual FName GetCategory() const override;
 	
 	// ICustomizableObjectNodeMeshInterface interface
@@ -90,6 +82,13 @@ public:
 	UE_API FSkeletalMaterial* GetSkeletalMaterialFor(const int32 SectionIndex) const;
 	UE_API int32 GetSkeletalMaterialIndexFor(const UEdGraphPin& Pin) const;
 	UE_API int32 GetSkeletalMaterialIndexFor(const int32 SectionIndex) const;
+
+	/** Default value of the parameter. */
+	UPROPERTY(EditAnywhere, Category=CustomizableObject)
+	TSoftObjectPtr<USkeletalMesh> DefaultValue;
+
+	UPROPERTY(EditAnywhere, Category=CustomizableObject)
+	TSoftObjectPtr<USkeletalMesh> ReferenceValue;
 };
 
 #undef UE_API
diff --git a/Engine/Plugins/Mutable/Source/CustomizableObjectEditor/Private/MuCOE/Nodes/CustomizableObjectNodeStaticString.cpp b/Engine/Plugins/Mutable/Source/CustomizableObjectEditor/Private/MuCOE/Nodes/CustomizableObjectNodeStaticString.cpp
--- a/Engine/Plugins/Mutable/Source/CustomizableObjectEditor/Private/MuCOE/Nodes/CustomizableObjectNodeStaticString.cpp
+++ b/Engine/Plugins/Mutable/Source/CustomizableObjectEditor/Private/MuCOE/Nodes/CustomizableObjectNodeStaticString.cpp
@@ -24,7 +24,7 @@ FText UCustomizableObjectNodeStaticString::GetTooltipText() const
 
 FLinearColor UCustomizableObjectNodeStaticString::GetNodeTitleColor() const
 {
-	return GetDefault<UEdGraphSchema_CustomizableObject>()->GetPinTypeColor(UEdGraphSchema_CustomizableObject::PC_String);
+	return UEdGraphSchema_CustomizableObject::GetPinTypeColor(UEdGraphSchema_CustomizableObject::PC_String);
 }
 
 
@@ -32,10 +32,29 @@ void UCustomizableObjectNodeStaticString::AllocateDefaultPins(UCustomizableObjec
 {
 	const UEdGraphSchema_CustomizableObject* Schema = GetDefault<UEdGraphSchema_CustomizableObject>();
 
-	UEdGraphPin* ValuePin = CustomCreatePin(EGPD_Output, Schema->PC_String, FName("Value"));
+	CustomCreatePin(EGPD_Output, Schema->PC_String, FName("Value"));
+
 	UEdGraphPin* StringPin = CustomCreatePin(EGPD_Input, Schema->PC_String, FName("String"));
 	StringPin->bNotConnectable = true;
 }
 
 
+bool UCustomizableObjectNodeStaticString::CanRenamePin(const UEdGraphPin& Pin) const
+{
+	return Pin.PinType.PinCategory == UEdGraphSchema_CustomizableObject::PC_String;
+}
+
+
+FText UCustomizableObjectNodeStaticString::GetPinEditableName(const UEdGraphPin& Pin) const
+{
+	return FText::FromString(Value);
+}
+
+
+void UCustomizableObjectNodeStaticString::SetPinEditableName(const UEdGraphPin& Pin, const FText& InValue)
+{
+	Value = InValue.ToString();
+}
+
+
 #undef LOCTEXT_NAMESPACE
diff --git a/Engine/Plugins/Mutable/Source/CustomizableObjectEditor/Private/MuCOE/Nodes/CustomizableObjectNodeStaticString.h b/Engine/Plugins/Mutable/Source/CustomizableObjectEditor/Private/MuCOE/Nodes/CustomizableObjectNodeStaticString.h
--- a/Engine/Plugins/Mutable/Source/CustomizableObjectEditor/Private/MuCOE/Nodes/CustomizableObjectNodeStaticString.h
+++ b/Engine/Plugins/Mutable/Source/CustomizableObjectEditor/Private/MuCOE/Nodes/CustomizableObjectNodeStaticString.h
@@ -8,27 +8,29 @@
 
 #define UE_API CUSTOMIZABLEOBJECTEDITOR_API
 
-
 UCLASS(MinimalAPI)
 class UCustomizableObjectNodeStaticString : public UCustomizableObjectNode
 {
-public:
 	GENERATED_BODY()
 
-	/** Macro Library where the macro to instantiate belongs. */
-	UPROPERTY(EditAnywhere, Category = CustomizableObject)
-	FString Value;
-
 public:
-
 	// EdGraphNode interface
 	UE_API virtual FText GetNodeTitle(ENodeTitleType::Type TitleType) const override;
 	UE_API virtual FText GetTooltipText() const override;
 	UE_API virtual FLinearColor GetNodeTitleColor() const override;
 	
 	// UCustomizableObjectNode interface
-	UE_API void AllocateDefaultPins(UCustomizableObjectNodeRemapPins* RemapPins) override;
-
+	UE_API virtual void AllocateDefaultPins(UCustomizableObjectNodeRemapPins* RemapPins) override;
+	UE_API virtual bool CanRenamePin(const UEdGraphPin& Pin) const override;
+	UE_API virtual FText GetPinEditableName(const UEdGraphPin& Pin) const override;
+	UE_API virtual void SetPinEditableName(const UEdGraphPin& Pin, const FText& Value) override;
+	
+	/** Macro Library where the macro to instantiate belongs. */
+	UPROPERTY()
+	FString Value;
+	
+	UPROPERTY()
+	FEdGraphPinReference ValuePin;
 };
 
 #undef UE_API
diff --git a/Engine/Plugins/Mutable/Source/CustomizableObjectEditor/Private/MuCOE/Nodes/CustomizableObjectNodeSwitchBase.cpp b/Engine/Plugins/Mutable/Source/CustomizableObjectEditor/Private/MuCOE/Nodes/CustomizableObjectNodeSwitchBase.cpp
--- a/Engine/Plugins/Mutable/Source/CustomizableObjectEditor/Private/MuCOE/Nodes/CustomizableObjectNodeSwitchBase.cpp
+++ b/Engine/Plugins/Mutable/Source/CustomizableObjectEditor/Private/MuCOE/Nodes/CustomizableObjectNodeSwitchBase.cpp
@@ -113,7 +113,7 @@ FText UCustomizableObjectNodeSwitchBase::GetNodeTitle(ENodeTitleType::Type Title
 
 FLinearColor UCustomizableObjectNodeSwitchBase::GetNodeTitleColor() const
 {
-	return GetDefault<UEdGraphSchema_CustomizableObject>()->GetPinTypeColor(GetCategory());
+	return UEdGraphSchema_CustomizableObject::GetPinTypeColor(GetCategory());
 }
 
 
diff --git a/Engine/Plugins/Mutable/Source/CustomizableObjectEditor/Private/MuCOE/Nodes/CustomizableObjectNodeTable.cpp b/Engine/Plugins/Mutable/Source/CustomizableObjectEditor/Private/MuCOE/Nodes/CustomizableObjectNodeTable.cpp
--- a/Engine/Plugins/Mutable/Source/CustomizableObjectEditor/Private/MuCOE/Nodes/CustomizableObjectNodeTable.cpp
+++ b/Engine/Plugins/Mutable/Source/CustomizableObjectEditor/Private/MuCOE/Nodes/CustomizableObjectNodeTable.cpp
@@ -381,8 +381,7 @@ FText UCustomizableObjectNodeTable::GetNodeTitle(ENodeTitleType::Type TitleType)
 
 FLinearColor UCustomizableObjectNodeTable::GetNodeTitleColor() const
 {
-	const UEdGraphSchema_CustomizableObject* Schema = GetDefault<UEdGraphSchema_CustomizableObject>();
-	return Schema->GetPinTypeColor(Schema->PC_Object);
+	return UEdGraphSchema_CustomizableObject::GetPinTypeColor(UEdGraphSchema_CustomizableObject::PC_Object);
 }
 
 
diff --git a/Engine/Plugins/Mutable/Source/CustomizableObjectEditor/Private/MuCOE/Nodes/CustomizableObjectNodeTexture.cpp b/Engine/Plugins/Mutable/Source/CustomizableObjectEditor/Private/MuCOE/Nodes/CustomizableObjectNodeTexture.cpp
--- a/Engine/Plugins/Mutable/Source/CustomizableObjectEditor/Private/MuCOE/Nodes/CustomizableObjectNodeTexture.cpp
+++ b/Engine/Plugins/Mutable/Source/CustomizableObjectEditor/Private/MuCOE/Nodes/CustomizableObjectNodeTexture.cpp
@@ -69,8 +69,7 @@ FText UCustomizableObjectNodeTexture::GetNodeTitle(ENodeTitleType::Type TitleTyp
 
 FLinearColor UCustomizableObjectNodeTexture::GetNodeTitleColor() const
 {
-	const UEdGraphSchema_CustomizableObject* Schema = GetDefault<UEdGraphSchema_CustomizableObject>();
-	return Schema->GetPinTypeColor(Schema->PC_Texture);
+	return UEdGraphSchema_CustomizableObject::GetPinTypeColor(UEdGraphSchema_CustomizableObject::PC_Texture);
 }
 
 
diff --git a/Engine/Plugins/Mutable/Source/CustomizableObjectEditor/Private/MuCOE/Nodes/CustomizableObjectNodeTextureBinarise.cpp b/Engine/Plugins/Mutable/Source/CustomizableObjectEditor/Private/MuCOE/Nodes/CustomizableObjectNodeTextureBinarise.cpp
--- a/Engine/Plugins/Mutable/Source/CustomizableObjectEditor/Private/MuCOE/Nodes/CustomizableObjectNodeTextureBinarise.cpp
+++ b/Engine/Plugins/Mutable/Source/CustomizableObjectEditor/Private/MuCOE/Nodes/CustomizableObjectNodeTextureBinarise.cpp
@@ -73,7 +73,7 @@ FText UCustomizableObjectNodeTextureBinarise::GetNodeTitle(ENodeTitleType::Type
 
 FLinearColor UCustomizableObjectNodeTextureBinarise::GetNodeTitleColor() const
 {
-	return GetDefault<UEdGraphSchema_CustomizableObject>()->GetPinTypeColor(UEdGraphSchema_CustomizableObject::PC_Texture);
+	return UEdGraphSchema_CustomizableObject::GetPinTypeColor(UEdGraphSchema_CustomizableObject::PC_Texture);
 }
 
 
diff --git a/Engine/Plugins/Mutable/Source/CustomizableObjectEditor/Private/MuCOE/Nodes/CustomizableObjectNodeTextureColourMap.cpp b/Engine/Plugins/Mutable/Source/CustomizableObjectEditor/Private/MuCOE/Nodes/CustomizableObjectNodeTextureColourMap.cpp
--- a/Engine/Plugins/Mutable/Source/CustomizableObjectEditor/Private/MuCOE/Nodes/CustomizableObjectNodeTextureColourMap.cpp
+++ b/Engine/Plugins/Mutable/Source/CustomizableObjectEditor/Private/MuCOE/Nodes/CustomizableObjectNodeTextureColourMap.cpp
@@ -57,8 +57,7 @@ FText UCustomizableObjectNodeTextureColourMap::GetNodeTitle(ENodeTitleType::Type
 
 FLinearColor UCustomizableObjectNodeTextureColourMap::GetNodeTitleColor() const
 {
-	const UEdGraphSchema_CustomizableObject* Schema = GetDefault<UEdGraphSchema_CustomizableObject>();
-	return Schema->GetPinTypeColor(Schema->PC_Texture);
+	return UEdGraphSchema_CustomizableObject::GetPinTypeColor(UEdGraphSchema_CustomizableObject::PC_Texture);
 }
 
 
diff --git a/Engine/Plugins/Mutable/Source/CustomizableObjectEditor/Private/MuCOE/Nodes/CustomizableObjectNodeTextureFromChannels.cpp b/Engine/Plugins/Mutable/Source/CustomizableObjectEditor/Private/MuCOE/Nodes/CustomizableObjectNodeTextureFromChannels.cpp
--- a/Engine/Plugins/Mutable/Source/CustomizableObjectEditor/Private/MuCOE/Nodes/CustomizableObjectNodeTextureFromChannels.cpp
+++ b/Engine/Plugins/Mutable/Source/CustomizableObjectEditor/Private/MuCOE/Nodes/CustomizableObjectNodeTextureFromChannels.cpp
@@ -61,8 +61,7 @@ FText UCustomizableObjectNodeTextureFromChannels::GetNodeTitle(ENodeTitleType::T
 
 FLinearColor UCustomizableObjectNodeTextureFromChannels::GetNodeTitleColor() const
 {
-	const UEdGraphSchema_CustomizableObject* Schema = GetDefault<UEdGraphSchema_CustomizableObject>();
-	return Schema->GetPinTypeColor(Schema->PC_Texture);
+	return UEdGraphSchema_CustomizableObject::GetPinTypeColor(UEdGraphSchema_CustomizableObject::PC_Texture);
 }
 
 
diff --git a/Engine/Plugins/Mutable/Source/CustomizableObjectEditor/Private/MuCOE/Nodes/CustomizableObjectNodeTextureFromColor.cpp b/Engine/Plugins/Mutable/Source/CustomizableObjectEditor/Private/MuCOE/Nodes/CustomizableObjectNodeTextureFromColor.cpp
--- a/Engine/Plugins/Mutable/Source/CustomizableObjectEditor/Private/MuCOE/Nodes/CustomizableObjectNodeTextureFromColor.cpp
+++ b/Engine/Plugins/Mutable/Source/CustomizableObjectEditor/Private/MuCOE/Nodes/CustomizableObjectNodeTextureFromColor.cpp
@@ -49,8 +49,7 @@ FText UCustomizableObjectNodeTextureFromColor::GetNodeTitle(ENodeTitleType::Type
 
 FLinearColor UCustomizableObjectNodeTextureFromColor::GetNodeTitleColor() const
 {
-	const UEdGraphSchema_CustomizableObject* Schema = GetDefault<UEdGraphSchema_CustomizableObject>();
-	return Schema->GetPinTypeColor(Schema->PC_Texture);
+	return UEdGraphSchema_CustomizableObject::GetPinTypeColor(UEdGraphSchema_CustomizableObject::PC_Texture);
 }
 
 
diff --git a/Engine/Plugins/Mutable/Source/CustomizableObjectEditor/Private/MuCOE/Nodes/CustomizableObjectNodeTextureFromFloats.cpp b/Engine/Plugins/Mutable/Source/CustomizableObjectEditor/Private/MuCOE/Nodes/CustomizableObjectNodeTextureFromFloats.cpp
--- a/Engine/Plugins/Mutable/Source/CustomizableObjectEditor/Private/MuCOE/Nodes/CustomizableObjectNodeTextureFromFloats.cpp
+++ b/Engine/Plugins/Mutable/Source/CustomizableObjectEditor/Private/MuCOE/Nodes/CustomizableObjectNodeTextureFromFloats.cpp
@@ -46,8 +46,7 @@ FText UCustomizableObjectNodeTextureFromFloats::GetNodeTitle(ENodeTitleType::Typ
 
 FLinearColor UCustomizableObjectNodeTextureFromFloats::GetNodeTitleColor() const
 {
-	const UEdGraphSchema_CustomizableObject* Schema = GetDefault<UEdGraphSchema_CustomizableObject>();
-	return Schema->GetPinTypeColor(Schema->PC_Texture);
+	return UEdGraphSchema_CustomizableObject::GetPinTypeColor(UEdGraphSchema_CustomizableObject::PC_Texture);
 }
 
 
diff --git a/Engine/Plugins/Mutable/Source/CustomizableObjectEditor/Private/MuCOE/Nodes/CustomizableObjectNodeTextureInterpolate.cpp b/Engine/Plugins/Mutable/Source/CustomizableObjectEditor/Private/MuCOE/Nodes/CustomizableObjectNodeTextureInterpolate.cpp
--- a/Engine/Plugins/Mutable/Source/CustomizableObjectEditor/Private/MuCOE/Nodes/CustomizableObjectNodeTextureInterpolate.cpp
+++ b/Engine/Plugins/Mutable/Source/CustomizableObjectEditor/Private/MuCOE/Nodes/CustomizableObjectNodeTextureInterpolate.cpp
@@ -59,8 +59,7 @@ FText UCustomizableObjectNodeTextureInterpolate::GetNodeTitle(ENodeTitleType::Ty
 
 FLinearColor UCustomizableObjectNodeTextureInterpolate::GetNodeTitleColor() const
 {
-	const UEdGraphSchema_CustomizableObject* Schema = GetDefault<UEdGraphSchema_CustomizableObject>();
-	return Schema->GetPinTypeColor(Schema->PC_Texture);
+	return UEdGraphSchema_CustomizableObject::GetPinTypeColor(UEdGraphSchema_CustomizableObject::PC_Texture);
 }
 
 
diff --git a/Engine/Plugins/Mutable/Source/CustomizableObjectEditor/Private/MuCOE/Nodes/CustomizableObjectNodeTextureInvert.cpp b/Engine/Plugins/Mutable/Source/CustomizableObjectEditor/Private/MuCOE/Nodes/CustomizableObjectNodeTextureInvert.cpp
--- a/Engine/Plugins/Mutable/Source/CustomizableObjectEditor/Private/MuCOE/Nodes/CustomizableObjectNodeTextureInvert.cpp
+++ b/Engine/Plugins/Mutable/Source/CustomizableObjectEditor/Private/MuCOE/Nodes/CustomizableObjectNodeTextureInvert.cpp
@@ -80,7 +80,7 @@ FText UCustomizableObjectNodeTextureInvert::GetNodeTitle(ENodeTitleType::Type Ti
 
 FLinearColor UCustomizableObjectNodeTextureInvert::GetNodeTitleColor() const
 {
-	return GetDefault<UEdGraphSchema_CustomizableObject>()->GetPinTypeColor(UEdGraphSchema_CustomizableObject::PC_Texture);
+	return UEdGraphSchema_CustomizableObject::GetPinTypeColor(UEdGraphSchema_CustomizableObject::PC_Texture);
 }
 
 
diff --git a/Engine/Plugins/Mutable/Source/CustomizableObjectEditor/Private/MuCOE/Nodes/CustomizableObjectNodeTextureLayer.cpp b/Engine/Plugins/Mutable/Source/CustomizableObjectEditor/Private/MuCOE/Nodes/CustomizableObjectNodeTextureLayer.cpp
--- a/Engine/Plugins/Mutable/Source/CustomizableObjectEditor/Private/MuCOE/Nodes/CustomizableObjectNodeTextureLayer.cpp
+++ b/Engine/Plugins/Mutable/Source/CustomizableObjectEditor/Private/MuCOE/Nodes/CustomizableObjectNodeTextureLayer.cpp
@@ -95,8 +95,7 @@ FText UCustomizableObjectNodeTextureLayer::GetNodeTitle(ENodeTitleType::Type Tit
 
 FLinearColor UCustomizableObjectNodeTextureLayer::GetNodeTitleColor() const
 {
-	const UEdGraphSchema_CustomizableObject* Schema = GetDefault<UEdGraphSchema_CustomizableObject>();
-	return Schema->GetPinTypeColor(Schema->PC_Texture);
+	return UEdGraphSchema_CustomizableObject::GetPinTypeColor(UEdGraphSchema_CustomizableObject::PC_Texture);
 }
 
 
diff --git a/Engine/Plugins/Mutable/Source/CustomizableObjectEditor/Private/MuCOE/Nodes/CustomizableObjectNodeTextureProject.cpp b/Engine/Plugins/Mutable/Source/CustomizableObjectEditor/Private/MuCOE/Nodes/CustomizableObjectNodeTextureProject.cpp
--- a/Engine/Plugins/Mutable/Source/CustomizableObjectEditor/Private/MuCOE/Nodes/CustomizableObjectNodeTextureProject.cpp
+++ b/Engine/Plugins/Mutable/Source/CustomizableObjectEditor/Private/MuCOE/Nodes/CustomizableObjectNodeTextureProject.cpp
@@ -180,8 +180,7 @@ FText UCustomizableObjectNodeTextureProject::GetNodeTitle(ENodeTitleType::Type T
 
 FLinearColor UCustomizableObjectNodeTextureProject::GetNodeTitleColor() const
 {
-	const UEdGraphSchema_CustomizableObject* Schema = GetDefault<UEdGraphSchema_CustomizableObject>();
-	return Schema->GetPinTypeColor(Schema->PC_Texture);
+	return UEdGraphSchema_CustomizableObject::GetPinTypeColor(UEdGraphSchema_CustomizableObject::PC_Texture);
 }
 
 
diff --git a/Engine/Plugins/Mutable/Source/CustomizableObjectEditor/Private/MuCOE/Nodes/CustomizableObjectNodeTextureSample.cpp b/Engine/Plugins/Mutable/Source/CustomizableObjectEditor/Private/MuCOE/Nodes/CustomizableObjectNodeTextureSample.cpp
--- a/Engine/Plugins/Mutable/Source/CustomizableObjectEditor/Private/MuCOE/Nodes/CustomizableObjectNodeTextureSample.cpp
+++ b/Engine/Plugins/Mutable/Source/CustomizableObjectEditor/Private/MuCOE/Nodes/CustomizableObjectNodeTextureSample.cpp
@@ -66,7 +66,7 @@ FText UCustomizableObjectNodeTextureSample::GetNodeTitle(ENodeTitleType::Type Ti
 
 FLinearColor UCustomizableObjectNodeTextureSample::GetNodeTitleColor() const
 {
-	return GetDefault<UEdGraphSchema_CustomizableObject>()->GetPinTypeColor(UEdGraphSchema_CustomizableObject::PC_Color);
+	return UEdGraphSchema_CustomizableObject::GetPinTypeColor(UEdGraphSchema_CustomizableObject::PC_Color);
 }
 
 
diff --git a/Engine/Plugins/Mutable/Source/CustomizableObjectEditor/Private/MuCOE/Nodes/CustomizableObjectNodeTextureSaturate.cpp b/Engine/Plugins/Mutable/Source/CustomizableObjectEditor/Private/MuCOE/Nodes/CustomizableObjectNodeTextureSaturate.cpp
--- a/Engine/Plugins/Mutable/Source/CustomizableObjectEditor/Private/MuCOE/Nodes/CustomizableObjectNodeTextureSaturate.cpp
+++ b/Engine/Plugins/Mutable/Source/CustomizableObjectEditor/Private/MuCOE/Nodes/CustomizableObjectNodeTextureSaturate.cpp
@@ -66,7 +66,7 @@ FText UCustomizableObjectNodeTextureSaturate::GetNodeTitle(ENodeTitleType::Type
 
 FLinearColor UCustomizableObjectNodeTextureSaturate::GetNodeTitleColor() const
 {
-	return GetDefault<UEdGraphSchema_CustomizableObject>()->GetPinTypeColor(UEdGraphSchema_CustomizableObject::PC_Texture);
+	return UEdGraphSchema_CustomizableObject::GetPinTypeColor(UEdGraphSchema_CustomizableObject::PC_Texture);
 }
 
 
diff --git a/Engine/Plugins/Mutable/Source/CustomizableObjectEditor/Private/MuCOE/Nodes/CustomizableObjectNodeTextureToChannels.cpp b/Engine/Plugins/Mutable/Source/CustomizableObjectEditor/Private/MuCOE/Nodes/CustomizableObjectNodeTextureToChannels.cpp
--- a/Engine/Plugins/Mutable/Source/CustomizableObjectEditor/Private/MuCOE/Nodes/CustomizableObjectNodeTextureToChannels.cpp
+++ b/Engine/Plugins/Mutable/Source/CustomizableObjectEditor/Private/MuCOE/Nodes/CustomizableObjectNodeTextureToChannels.cpp
@@ -98,8 +98,7 @@ FText UCustomizableObjectNodeTextureToChannels::GetNodeTitle(ENodeTitleType::Typ
 
 FLinearColor UCustomizableObjectNodeTextureToChannels::GetNodeTitleColor() const
 {
-	const UEdGraphSchema_CustomizableObject* Schema = GetDefault<UEdGraphSchema_CustomizableObject>();
-	return Schema->GetPinTypeColor(Schema->PC_Texture);
+	return UEdGraphSchema_CustomizableObject::GetPinTypeColor(UEdGraphSchema_CustomizableObject::PC_Texture);
 }
 
 
diff --git a/Engine/Plugins/Mutable/Source/CustomizableObjectEditor/Private/MuCOE/Nodes/CustomizableObjectNodeTextureTransform.cpp b/Engine/Plugins/Mutable/Source/CustomizableObjectEditor/Private/MuCOE/Nodes/CustomizableObjectNodeTextureTransform.cpp
--- a/Engine/Plugins/Mutable/Source/CustomizableObjectEditor/Private/MuCOE/Nodes/CustomizableObjectNodeTextureTransform.cpp
+++ b/Engine/Plugins/Mutable/Source/CustomizableObjectEditor/Private/MuCOE/Nodes/CustomizableObjectNodeTextureTransform.cpp
@@ -67,7 +67,7 @@ FText UCustomizableObjectNodeTextureTransform::GetNodeTitle(ENodeTitleType::Type
 
 FLinearColor UCustomizableObjectNodeTextureTransform::GetNodeTitleColor() const
 {
-	return GetDefault<UEdGraphSchema_CustomizableObject>()->GetPinTypeColor(UEdGraphSchema_CustomizableObject::PC_Texture);
+	return UEdGraphSchema_CustomizableObject::GetPinTypeColor(UEdGraphSchema_CustomizableObject::PC_Texture);
 }
 
 FText UCustomizableObjectNodeTextureTransform::GetTooltipText() const
diff --git a/Engine/Plugins/Mutable/Source/CustomizableObjectEditor/Private/MuCOE/Nodes/CustomizableObjectNodeTextureVariation.cpp b/Engine/Plugins/Mutable/Source/CustomizableObjectEditor/Private/MuCOE/Nodes/CustomizableObjectNodeTextureVariation.cpp
--- a/Engine/Plugins/Mutable/Source/CustomizableObjectEditor/Private/MuCOE/Nodes/CustomizableObjectNodeTextureVariation.cpp
+++ b/Engine/Plugins/Mutable/Source/CustomizableObjectEditor/Private/MuCOE/Nodes/CustomizableObjectNodeTextureVariation.cpp
@@ -72,7 +72,7 @@ FText UCustomizableObjectNodeTextureVariation::GetNodeTitle(ENodeTitleType::Type
 
 FLinearColor UCustomizableObjectNodeTextureVariation::GetNodeTitleColor() const
 {
-	return GetDefault<UEdGraphSchema_CustomizableObject>()->GetPinTypeColor(UEdGraphSchema_CustomizableObject::PC_Texture);
+	return UEdGraphSchema_CustomizableObject::GetPinTypeColor(UEdGraphSchema_CustomizableObject::PC_Texture);
 }
 
 
diff --git a/Engine/Plugins/Mutable/Source/CustomizableObjectEditor/Private/MuCOE/Nodes/CustomizableObjectNodeTransformConstant.cpp b/Engine/Plugins/Mutable/Source/CustomizableObjectEditor/Private/MuCOE/Nodes/CustomizableObjectNodeTransformConstant.cpp
--- a/Engine/Plugins/Mutable/Source/CustomizableObjectEditor/Private/MuCOE/Nodes/CustomizableObjectNodeTransformConstant.cpp
+++ b/Engine/Plugins/Mutable/Source/CustomizableObjectEditor/Private/MuCOE/Nodes/CustomizableObjectNodeTransformConstant.cpp
@@ -218,8 +218,7 @@ FText UCustomizableObjectNodeTransformConstant::GetNodeTitle(ENodeTitleType::Typ
 
 FLinearColor UCustomizableObjectNodeTransformConstant::GetNodeTitleColor() const
 {
-	const UEdGraphSchema_CustomizableObject* Schema = GetDefault<UEdGraphSchema_CustomizableObject>();
-	return Schema->GetPinTypeColor(Schema->PC_Transform);
+	return UEdGraphSchema_CustomizableObject::GetPinTypeColor(UEdGraphSchema_CustomizableObject::PC_Transform);
 }
 
 
diff --git a/Engine/Plugins/Mutable/Source/CustomizableObjectEditor/Private/MuCOE/Nodes/CustomizableObjectNodeTunnel.cpp b/Engine/Plugins/Mutable/Source/CustomizableObjectEditor/Private/MuCOE/Nodes/CustomizableObjectNodeTunnel.cpp
--- a/Engine/Plugins/Mutable/Source/CustomizableObjectEditor/Private/MuCOE/Nodes/CustomizableObjectNodeTunnel.cpp
+++ b/Engine/Plugins/Mutable/Source/CustomizableObjectEditor/Private/MuCOE/Nodes/CustomizableObjectNodeTunnel.cpp
@@ -105,9 +105,4 @@ void UCustomizableObjectNodeTunnel::AllocateDefaultPins(UCustomizableObjectNodeR
 }
 
 
-bool UCustomizableObjectNodeTunnel::CreateStaticStringPinWidget() const
-{
-	return false;
-}
-
 #undef LOCTEXT_NAMESPACE
diff --git a/Engine/Plugins/Mutable/Source/CustomizableObjectEditor/Private/MuCOE/Nodes/CustomizableObjectNodeTunnel.h b/Engine/Plugins/Mutable/Source/CustomizableObjectEditor/Private/MuCOE/Nodes/CustomizableObjectNodeTunnel.h
--- a/Engine/Plugins/Mutable/Source/CustomizableObjectEditor/Private/MuCOE/Nodes/CustomizableObjectNodeTunnel.h
+++ b/Engine/Plugins/Mutable/Source/CustomizableObjectEditor/Private/MuCOE/Nodes/CustomizableObjectNodeTunnel.h
@@ -69,8 +69,6 @@ public:
 	// UCustomizableObjectNode interface
 	UE_API void AllocateDefaultPins(UCustomizableObjectNodeRemapPins* RemapPins) override;
 	UE_API virtual UCustomizableObjectNodeTunnelRemapPins* CreateRemapPinsDefault() const override;
-	UE_API virtual bool CreateStaticStringPinWidget() const override;
-
 };
 
 #undef UE_API
diff --git a/Engine/Plugins/Mutable/Source/CustomizableObjectEditor/Private/MuCOE/Nodes/CustomizableObjectNodeVariation.cpp b/Engine/Plugins/Mutable/Source/CustomizableObjectEditor/Private/MuCOE/Nodes/CustomizableObjectNodeVariation.cpp
--- a/Engine/Plugins/Mutable/Source/CustomizableObjectEditor/Private/MuCOE/Nodes/CustomizableObjectNodeVariation.cpp
+++ b/Engine/Plugins/Mutable/Source/CustomizableObjectEditor/Private/MuCOE/Nodes/CustomizableObjectNodeVariation.cpp
@@ -62,6 +62,40 @@ void UCustomizableObjectNodeVariation::AllocateDefaultPins(UCustomizableObjectNo
 }
 
 
+bool UCustomizableObjectNodeVariation::CanRenamePin(const UEdGraphPin& Pin) const
+{
+	return Pin.PinType.PinCategory == UEdGraphSchema_CustomizableObject::PC_String;
+}
+
+
+FText UCustomizableObjectNodeVariation::GetPinEditableName(const UEdGraphPin& Pin) const
+{
+	for (int32 VariationIndex = 0; VariationIndex < GetNumVariations(); ++VariationIndex)
+	{
+		if (VariationTagPin(VariationIndex) == &Pin)
+		{
+			return FText::FromString(GetVariationTag(VariationIndex));
+		}
+	}
+
+	return {};
+}
+
+
+void UCustomizableObjectNodeVariation::SetPinEditableName(const UEdGraphPin& Pin, const FText& InValue)
+{
+	for (int32 VariationIndex = 0; VariationIndex < GetNumVariations(); ++VariationIndex)
+	{
+		if (VariationTagPin(VariationIndex) == &Pin && VariationsData.IsValidIndex(VariationIndex))
+		{
+			VariationsData[VariationIndex].Tag = InValue.ToString();
+
+			break;
+		}
+	}
+}
+
+
 bool UCustomizableObjectNodeVariation::IsInputPinArray() const
 {
 	return false;
@@ -138,8 +172,7 @@ FText UCustomizableObjectNodeVariation::GetNodeTitle(ENodeTitleType::Type TitleT
 
 FLinearColor UCustomizableObjectNodeVariation::GetNodeTitleColor() const
 {
-	const UEdGraphSchema_CustomizableObject* Schema = GetDefault<UEdGraphSchema_CustomizableObject>();
-	return Schema->GetPinTypeColor(GetCategory());
+	return UEdGraphSchema_CustomizableObject::GetPinTypeColor(GetCategory());
 }
 
 
diff --git a/Engine/Plugins/Mutable/Source/CustomizableObjectEditor/Private/MuCOE/Nodes/CustomizableObjectNodeVariation.h b/Engine/Plugins/Mutable/Source/CustomizableObjectEditor/Private/MuCOE/Nodes/CustomizableObjectNodeVariation.h
--- a/Engine/Plugins/Mutable/Source/CustomizableObjectEditor/Private/MuCOE/Nodes/CustomizableObjectNodeVariation.h
+++ b/Engine/Plugins/Mutable/Source/CustomizableObjectEditor/Private/MuCOE/Nodes/CustomizableObjectNodeVariation.h
@@ -35,20 +35,6 @@ class UCustomizableObjectNodeVariation : public UCustomizableObjectNode
 	GENERATED_BODY()
 
 public:
-
-	UPROPERTY(EditAnywhere, Category = CustomizableObject)
-	TArray<FCustomizableObjectVariation> VariationsData; // The variable name can not be Variations due issues with the on to UObject Serialization system
-
-protected:
-
-	UPROPERTY()
-	TArray<FEdGraphPinReference> VariationsPins;
-
-	UPROPERTY()
-	TArray<FEdGraphPinReference> VariationTagPins;
-	
-public:
-
 	// UObject interface.
 	UE_API virtual void PostEditChangeProperty(FPropertyChangedEvent& PropertyChangedEvent) override;
 
@@ -60,7 +46,10 @@ public:
 	// UCustomizableObjectNode interface
 	UE_API virtual void BackwardsCompatibleFixup(int32 CustomizableObjectCustomVersion) override;
 	UE_API virtual void AllocateDefaultPins(UCustomizableObjectNodeRemapPins* RemapPins) override;
-
+	UE_API virtual bool CanRenamePin(const UEdGraphPin& Pin) const override;
+	UE_API virtual FText GetPinEditableName(const UEdGraphPin& Pin) const override;
+	UE_API virtual void SetPinEditableName(const UEdGraphPin& Pin, const FText& Value) override;
+	
 	// Own interface
 	/** Return the pin category of this node. */
 	virtual FName GetCategory() const PURE_VIRTUAL(UCustomizableObjectNodeVariation::GetCategory, return {}; );
@@ -85,7 +74,17 @@ public:
 
 	/** Get the Variation Tag Input Pin */
 	UE_API UEdGraphPin* VariationTagPin(int32 Index) const;
+	
+	UPROPERTY(EditAnywhere, Category = CustomizableObject)
+	TArray<FCustomizableObjectVariation> VariationsData; // The variable name can not be Variations due issues with the UObject Serialization system
+
+protected:
+
+	UPROPERTY()
+	TArray<FEdGraphPinReference> VariationsPins;
 
+	UPROPERTY()
+	TArray<FEdGraphPinReference> VariationTagPins;
 };
 
 #undef UE_API
diff --git a/Engine/Plugins/Mutable/Source/CustomizableObjectEditor/Private/MuCOE/Nodes/SCustomizableObjectNode.cpp b/Engine/Plugins/Mutable/Source/CustomizableObjectEditor/Private/MuCOE/Nodes/SCustomizableObjectNode.cpp
--- a/Engine/Plugins/Mutable/Source/CustomizableObjectEditor/Private/MuCOE/Nodes/SCustomizableObjectNode.cpp
+++ b/Engine/Plugins/Mutable/Source/CustomizableObjectEditor/Private/MuCOE/Nodes/SCustomizableObjectNode.cpp
@@ -4,12 +4,7 @@
 
 #include "SCustomizableObjectNodePin.h"
 #include "MuCOE/EdGraphSchema_CustomizableObject.h"
-
 #include "MuCOE/Nodes/CustomizableObjectNode.h"
-#include "MuCOE/Nodes/CustomizableObjectNodeStaticString.h"
-
-#include "MuCOE/Nodes/SCustomizableObjectNodeStaticStringPin.h"
-
 #include "Widgets/SBoxPanel.h"
 
 #define LOCTEXT_NAMESPACE "CustomizableObjectEditor"
@@ -27,13 +22,6 @@ TSharedPtr<SGraphPin> SCustomizableObjectNode::CreatePinWidget(UEdGraphPin* Pin)
 {
 	check(Pin->GetOwningNode());
 
-	const UCustomizableObjectNode* Node = Cast<UCustomizableObjectNode>(Pin->GetOwningNode());
-
-	if (Pin->PinType.PinCategory == UEdGraphSchema_CustomizableObject::PC_String && Pin->PinType.ContainerType != EPinContainerType::Array && Node->CreateStaticStringPinWidget())
-	{
-		return SNew(SCustomizableObjectNodeStaticStringPin, Pin);
-	}
-
 	return SNew(SCustomizableObjectNodePin, Pin);
 }
 
diff --git a/Engine/Plugins/Mutable/Source/CustomizableObjectEditor/Private/MuCOE/Nodes/SCustomizableObjectNodePin.cpp b/Engine/Plugins/Mutable/Source/CustomizableObjectEditor/Private/MuCOE/Nodes/SCustomizableObjectNodePin.cpp
--- a/Engine/Plugins/Mutable/Source/CustomizableObjectEditor/Private/MuCOE/Nodes/SCustomizableObjectNodePin.cpp
+++ b/Engine/Plugins/Mutable/Source/CustomizableObjectEditor/Private/MuCOE/Nodes/SCustomizableObjectNodePin.cpp
@@ -5,6 +5,7 @@
 #include "MuCOE/EdGraphSchema_CustomizableObject.h"
 #include "MuCOE/UnrealEditorPortabilityHelpers.h"
 #include "MuCOE/Nodes/CustomizableObjectNode.h"
+#include "Widgets/Input/SEditableTextBox.h"
 
 
 void SCustomizableObjectNodePin::Construct(const FArguments& InArgs, UEdGraphPin* InGraphPinObj)
@@ -14,6 +15,32 @@ void SCustomizableObjectNodePin::Construct(const FArguments& InArgs, UEdGraphPin
 	// Cache pin icons.
 	PassThroughImageConnected = UE_MUTABLE_GET_BRUSH(TEXT("Graph.ExecPin.Connected"));
 	PassThroughImageDisconnected = UE_MUTABLE_GET_BRUSH(TEXT("Graph.ExecPin.Disconnected"));
+
+	// Hide pin Icon and Label
+	if (InGraphPinObj->bNotConnectable)
+	{
+		SetShowLabel(false);
+		PinImage->SetVisibility(EVisibility::Collapsed);
+	}
+}
+
+
+TSharedRef<SWidget> SCustomizableObjectNodePin::GetDefaultValueWidget()
+{
+	UEdGraphPin* GraphPin = GetPinObj();
+	if (GraphPin && !GraphPin->IsPendingKill())
+	{
+		UCustomizableObjectNode* Node = CastChecked<UCustomizableObjectNode>(GetPinObj()->GetOwningNode());
+		if (Node && Node->CanRenamePin(*GraphPin))
+		{
+			return SNew(SEditableTextBox)
+				.Text(this, &SCustomizableObjectNodePin::GetNodeStringValue)
+				.OnTextCommitted(this, &SCustomizableObjectNodePin::OnTextCommited)
+				.Visibility(this, &SCustomizableObjectNodePin::GetWidgetVisibility);
+		}	
+	}
+
+	return SGraphPin::GetDefaultValueWidget();
 }
 
 
@@ -73,3 +100,39 @@ FSlateColor SCustomizableObjectNodePin::GetPinColor() const
 
 	return FLinearColor::White;
 }
+
+
+FText SCustomizableObjectNodePin::GetNodeStringValue() const
+{
+	UEdGraphPin* GraphPin = GetPinObj();
+	if (GraphPin && !GraphPin->IsPendingKill())
+	{
+		UCustomizableObjectNode* Node = CastChecked<UCustomizableObjectNode>(GetPinObj()->GetOwningNode());
+		if (Node)
+		{
+			return Node->GetPinEditableName(*GraphPin);
+		}
+	}
+
+	return {};
+}
+
+
+void SCustomizableObjectNodePin::OnTextCommited(const FText& InValue, ETextCommit::Type InCommitInfo)
+{
+	UEdGraphPin* GraphPin = GetPinObj();
+	if (GraphPin && !GraphPin->IsPendingKill())
+	{
+		UCustomizableObjectNode* Node = CastChecked<UCustomizableObjectNode>(GetPinObj()->GetOwningNode());
+		if (Node)
+		{
+			Node->SetPinEditableName(*GraphPin, InValue);
+		}
+	}
+}
+
+
+EVisibility SCustomizableObjectNodePin::GetWidgetVisibility() const
+{
+	return GraphPinObj->LinkedTo.Num() ? EVisibility::Collapsed : EVisibility::Visible;
+}
diff --git a/Engine/Plugins/Mutable/Source/CustomizableObjectEditor/Private/MuCOE/Nodes/SCustomizableObjectNodePin.h b/Engine/Plugins/Mutable/Source/CustomizableObjectEditor/Private/MuCOE/Nodes/SCustomizableObjectNodePin.h
--- a/Engine/Plugins/Mutable/Source/CustomizableObjectEditor/Private/MuCOE/Nodes/SCustomizableObjectNodePin.h
+++ b/Engine/Plugins/Mutable/Source/CustomizableObjectEditor/Private/MuCOE/Nodes/SCustomizableObjectNodePin.h
@@ -19,11 +19,18 @@ public:
 	void Construct(const FArguments& InArgs, UEdGraphPin* InGraphPinObj);
 	
 	// SGraphPin interface
-protected:
+	virtual TSharedRef<SWidget>	GetDefaultValueWidget() override;
 	virtual const FSlateBrush* GetPinIcon() const override;
 	virtual FSlateColor GetPinColor() const override;
 
 private:
+	FText GetNodeStringValue() const;
+	
+	void OnTextCommited(const FText& InValue, ETextCommit::Type InCommitInfo);
+	
+	EVisibility GetWidgetVisibility() const;
+
+
 	const FSlateBrush* PassThroughImageConnected = nullptr;
-	const FSlateBrush* PassThroughImageDisconnected = nullptr;
-};
+		const FSlateBrush* PassThroughImageDisconnected = nullptr;
+	};
diff --git a/Engine/Plugins/Mutable/Source/CustomizableObjectEditor/Private/MuCOE/Nodes/SCustomizableObjectNodeStaticStringPin.cpp b/Engine/Plugins/Mutable/Source/CustomizableObjectEditor/Private/MuCOE/Nodes/SCustomizableObjectNodeStaticStringPin.cpp
deleted file mode 100644
diff --git a/Engine/Plugins/Mutable/Source/CustomizableObjectEditor/Private/MuCOE/Nodes/SCustomizableObjectNodeStaticStringPin.h b/Engine/Plugins/Mutable/Source/CustomizableObjectEditor/Private/MuCOE/Nodes/SCustomizableObjectNodeStaticStringPin.h
deleted file mode 100644
diff --git a/Engine/Plugins/Mutable/Source/CustomizableObjectEditor/Private/MuCOE/Thumbnails/CustomizableObjectInstanceThumbnailRenderer.cpp b/Engine/Plugins/Mutable/Source/CustomizableObjectEditor/Private/MuCOE/Thumbnails/CustomizableObjectInstanceThumbnailRenderer.cpp
deleted file mode 100644
diff --git a/Engine/Plugins/Mutable/Source/CustomizableObjectEditor/Private/MuCOE/Thumbnails/CustomizableObjectThumbnailRenderer.cpp b/Engine/Plugins/Mutable/Source/CustomizableObjectEditor/Private/MuCOE/Thumbnails/CustomizableObjectThumbnailRenderer.cpp
deleted file mode 100644
diff --git a/Engine/Plugins/Mutable/Source/CustomizableObjectEditor/Public/MuCOE/EdGraphSchema_CustomizableObject.h b/Engine/Plugins/Mutable/Source/CustomizableObjectEditor/Public/MuCOE/EdGraphSchema_CustomizableObject.h
--- a/Engine/Plugins/Mutable/Source/CustomizableObjectEditor/Public/MuCOE/EdGraphSchema_CustomizableObject.h
+++ b/Engine/Plugins/Mutable/Source/CustomizableObjectEditor/Public/MuCOE/EdGraphSchema_CustomizableObject.h
@@ -22,9 +22,9 @@ class ICustomizableObjectEditor;
 UCLASS(MinimalAPI)
 class UEdGraphSchema_CustomizableObject : public UEdGraphSchema
 {
-public:
 	GENERATED_BODY()
 
+public:
 	// Allowable PinType.PinCategory values
 	static UE_API const FName PC_Object;
 	static UE_API const FName PC_Component;
@@ -54,8 +54,6 @@ public:
 	static UE_API const FText NC_Experimental;
 	static UE_API const FText NC_Material;
 
-	mutable TMap<FName, FLinearColor> PinCategoryColors;
-	
 	// EdGraphSchema interface
 	UE_API virtual void GetAssetsGraphHoverMessage(const TArray<FAssetData>& Assets, const UEdGraph* HoverGraph, FString& OutTooltipText, bool& OutOkIcon) const override;
 	UE_API virtual void GetGraphContextActions(FGraphContextMenuBuilder& ContextMenuBuilder) const override;
@@ -73,9 +71,7 @@ public:
 	UE_API virtual FConnectionDrawingPolicy* CreateConnectionDrawingPolicy(int32 InBackLayerID, int32 InFrontLayerID, float InZoomFactor, const FSlateRect& InClippingRect, FSlateWindowElementList& InDrawElements, UEdGraph* InGraphObj) const override;
 
 	// Own interface
-	UE_API FLinearColor GetPinTypeColor(const FName& PinType) const;
-	
-	/** Given a pin category name, get its friendly name (user readable). */
+	static UE_API FLinearColor GetPinTypeColor(const FName& PinType);
 	static UE_API FName GetPinCategoryName(const FName& PinCategory);
 	static UE_API FText GetPinCategoryFriendlyName(const FName& PinCategory);
 	static UE_API bool IsPassthrough(const FName& PinCategory);
diff --git a/Engine/Plugins/Mutable/Source/CustomizableObjectEditor/Public/MuCOE/Nodes/CustomizableObjectNode.h b/Engine/Plugins/Mutable/Source/CustomizableObjectEditor/Public/MuCOE/Nodes/CustomizableObjectNode.h
--- a/Engine/Plugins/Mutable/Source/CustomizableObjectEditor/Public/MuCOE/Nodes/CustomizableObjectNode.h
+++ b/Engine/Plugins/Mutable/Source/CustomizableObjectEditor/Public/MuCOE/Nodes/CustomizableObjectNode.h
@@ -269,6 +269,12 @@ public:
 	/** Return true if a pin can be hidden. Override it to have a custom behaviour. */
 	UE_API virtual bool CanPinBeHidden(const UEdGraphPin& Pin) const;
 
+	UE_API virtual bool CanRenamePin(const UEdGraphPin& Pin) const;
+
+	UE_API virtual FText GetPinEditableName(const UEdGraphPin& Pin) const;
+
+	UE_API virtual void SetPinEditableName(const UEdGraphPin& Pin, const FText& Value);
+	
 	UE_API virtual FLinearColor GetPinColor(const UEdGraphPin &Pin) const;
 	
 	UE_API virtual bool IsPassthrough(const UEdGraphPin &Pin) const;
@@ -305,9 +311,6 @@ public:
 	/** Returns true if this node is inside a Macro Graph */
 	UE_API bool IsInMacro() const;
 
-	/** If a node has a static string pin, this function determines if it should construct an input text widget next to the pin or not. */
-	UE_API virtual bool CreateStaticStringPinWidget() const;
-
 	// Returns an array with all the PC_Types that can be created from the pin viewer
 	UE_API virtual TArray<FName> GetAllowedPinViewerCreationTypes() const;
 	
diff --git a/Engine/Plugins/Mutable/Source/CustomizableObjectEditor/Public/MuCOE/Thumbnails/CustomizableObjectInstanceThumbnailRenderer.h b/Engine/Plugins/Mutable/Source/CustomizableObjectEditor/Public/MuCOE/Thumbnails/CustomizableObjectInstanceThumbnailRenderer.h
deleted file mode 100644
diff --git a/Engine/Plugins/Mutable/Source/CustomizableObjectEditor/Public/MuCOE/Thumbnails/CustomizableObjectThumbnailRenderer.h b/Engine/Plugins/Mutable/Source/CustomizableObjectEditor/Public/MuCOE/Thumbnails/CustomizableObjectThumbnailRenderer.h
deleted file mode 100644
