[mutable] Fixed Build Engine Localization Text conflict.
#tests Editor
#rnx

diff --git a/Engine/Plugins/Mutable/Source/CustomizableObjectEditor/Private/MuCOE/Nodes/CONodeMaterialBreak.cpp b/Engine/Plugins/Mutable/Source/CustomizableObjectEditor/Private/MuCOE/Nodes/CONodeMaterialBreak.cpp
--- a/Engine/Plugins/Mutable/Source/CustomizableObjectEditor/Private/MuCOE/Nodes/CONodeMaterialBreak.cpp
+++ b/Engine/Plugins/Mutable/Source/CustomizableObjectEditor/Private/MuCOE/Nodes/CONodeMaterialBreak.cpp
@@ -73,7 +73,7 @@ FLinearColor UCONodeMaterialBreak::GetNodeTitleColor() const
 
 FText UCONodeMaterialBreak::GetTooltipText() const
 {
-	return LOCTEXT("Material_Constant_Tooltip", "Allows to get the parameter value from the input material node.");
+	return LOCTEXT("Material_Break_Tooltip", "Allows to get the parameter value from the input material node.");
 }
 
 
