[Mutable] Remove unused member in Node Static String.

- Removed incorrect comment.

#rnx

diff --git a/Engine/Plugins/Mutable/Source/CustomizableObjectEditor/Private/MuCOE/Nodes/CustomizableObjectNodeStaticString.h b/Engine/Plugins/Mutable/Source/CustomizableObjectEditor/Private/MuCOE/Nodes/CustomizableObjectNodeStaticString.h
--- a/Engine/Plugins/Mutable/Source/CustomizableObjectEditor/Private/MuCOE/Nodes/CustomizableObjectNodeStaticString.h
+++ b/Engine/Plugins/Mutable/Source/CustomizableObjectEditor/Private/MuCOE/Nodes/CustomizableObjectNodeStaticString.h
@@ -25,12 +25,8 @@ public:
 	UE_API virtual FText GetPinEditableName(const UEdGraphPin& Pin) const override;
 	UE_API virtual void SetPinEditableName(const UEdGraphPin& Pin, const FText& Value) override;
 	
-	/** Macro Library where the macro to instantiate belongs. */
 	UPROPERTY()
 	FString Value;
-	
-	UPROPERTY()
-	FEdGraphPinReference ValuePin;
 };
 
 #undef UE_API
