[Mutable] Fix crash when updating without streaming

#rnx

diff --git a/Engine/Plugins/Mutable/Source/CustomizableObject/Private/MuCO/CustomizableObjectInstance.cpp b/Engine/Plugins/Mutable/Source/CustomizableObject/Private/MuCO/CustomizableObjectInstance.cpp
--- a/Engine/Plugins/Mutable/Source/CustomizableObject/Private/MuCO/CustomizableObjectInstance.cpp
+++ b/Engine/Plugins/Mutable/Source/CustomizableObject/Private/MuCO/CustomizableObjectInstance.cpp
@@ -2275,8 +2275,12 @@ bool UCustomizableInstancePrivate::UpdateSkeletalMesh_PostBeginUpdate0(UCustomiz
 			continue;
 		}
 
+		const int32 FirstLOD = Context->bStreamMeshLODs ?
+			Context->FirstLODAvailable[ComponentName] :
+			Context->GetFirstRequestedLOD()[ComponentName];
+		
 		uint32 MeshHash = INDEX_NONE;
-		for (int32 LODIndex = Context->FirstLODAvailable[ComponentName]; LODIndex < Component->LODCount && LODIndex < MAX_MESH_LOD_COUNT; ++LODIndex)
+		for (int32 LODIndex = FirstLOD; LODIndex < Component->LODCount && LODIndex < MAX_MESH_LOD_COUNT; ++LODIndex)
 		{
 			const mu::FGeneratedMeshKey* MeshKey = (*MeshDescriptors)[LODIndex].GetKey();
 			check(MeshKey);
