[mutable] Fixed invalid interactions with the StaticMesh TSoftObjtPtr
- Added backwards compatibility code to handle the regeneration of the output pins alongside the overriding of the "LODs" property.
- A TSoftObjPtr does require loading before being able to access the data it points.
- Improved property value changes management.
- Removed unnecessary includes from source files.
- Renamed incorrectly named variables and removed unused ones.
- Changed the type of the int variables (int to int32)

#rnx

diff --git a/Engine/Plugins/Mutable/Source/CustomizableObject/Internal/MuCO/CustomizableObjectCustomVersion.h b/Engine/Plugins/Mutable/Source/CustomizableObject/Internal/MuCO/CustomizableObjectCustomVersion.h
--- a/Engine/Plugins/Mutable/Source/CustomizableObject/Internal/MuCO/CustomizableObjectCustomVersion.h
+++ b/Engine/Plugins/Mutable/Source/CustomizableObject/Internal/MuCO/CustomizableObjectCustomVersion.h
@@ -197,6 +197,8 @@ struct FCustomizableObjectCustomVersion
 		BaseObjectNodeNamePin,
 
 		FixMaterialNodeMaterialPinIncorrectLocalization,
+
+		FixStaticMeshNodeLODMaterialDataInvalidation,
 		
 		// -----<new versions can be added above this line>-------------------------------------------------
 		VersionPlusOne,
diff --git a/Engine/Plugins/Mutable/Source/CustomizableObjectEditor/Private/MuCOE/Nodes/CustomizableObjectNodeStaticMesh.cpp b/Engine/Plugins/Mutable/Source/CustomizableObjectEditor/Private/MuCOE/Nodes/CustomizableObjectNodeStaticMesh.cpp
--- a/Engine/Plugins/Mutable/Source/CustomizableObjectEditor/Private/MuCOE/Nodes/CustomizableObjectNodeStaticMesh.cpp
+++ b/Engine/Plugins/Mutable/Source/CustomizableObjectEditor/Private/MuCOE/Nodes/CustomizableObjectNodeStaticMesh.cpp
@@ -6,13 +6,12 @@
 #include "ISinglePropertyView.h"
 #include "Modules/ModuleManager.h"
 #include "MuCO/CustomizableObjectCustomVersion.h"
+#include "MuCO/LoadUtils.h"
 #include "MuCOE/CustomizableObjectEditor.h"
 #include "MuCOE/CustomizableObjectEditorStyle.h"
-#include "MuCOE/CustomizableObjectEditorUtilities.h"
 #include "MuCOE/CustomizableObjectLayout.h"
 #include "MuCOE/EdGraphSchema_CustomizableObject.h"
 #include "MuCOE/GraphTraversal.h"
-#include "MuCOE/MutableUtils.h"
 #include "MuCOE/Nodes/CustomizableObjectNodeLayoutBlocks.h"
 #include "MuCOE/RemapPins/CustomizableObjectNodeRemapPinsByNameDefaultPin.h"
 #include "MuCOE/UnrealEditorPortabilityHelpers.h"
@@ -34,7 +33,7 @@ static const TCHAR* STATIC_MESH_PIN_NAME = TEXT("Static Mesh");
 void UCustomizableObjectNodeStaticMesh::PostEditChangeProperty(FPropertyChangedEvent& PropertyChangedEvent)
 {
 	FProperty* PropertyThatChanged = PropertyChangedEvent.Property;
-	if (PropertyThatChanged && PropertyThatChanged->GetName() == TEXT("StaticMesh"))
+	if (PropertyThatChanged && PropertyThatChanged->GetName() == GET_MEMBER_NAME_STRING_CHECKED(UCustomizableObjectNodeStaticMesh, StaticMesh))
 	{
 		ReconstructNode();
 	}
@@ -92,27 +91,22 @@ void UCustomizableObjectNodeStaticMesh::AllocateDefaultPins(UCustomizableObjectN
 	
 	const UEdGraphSchema_CustomizableObject* Schema = GetDefault<UEdGraphSchema_CustomizableObject>();
 
-	if (StaticMesh)
+	if (MutablePrivate::LoadObject(StaticMesh))
 	{
-		const int NumLODs = StaticMesh->GetRenderData()->LODResources.Num();
+		const int32 NumLODs = StaticMesh->GetRenderData()->LODResources.Num();
 
 		LODs.Reset(NumLODs);
 		LODs.SetNum(NumLODs);
-		for (int32 i = 0; i < NumLODs; ++i)
+		for (int32 LODIndex = 0; LODIndex < NumLODs; ++LODIndex)
 		{
-			FString LODName = FString::Printf(TEXT("LOD %d, "), i);
+			const int32 NumMaterials = StaticMesh->GetRenderData()->LODResources[LODIndex].Sections.Num();
 
-			const int NumMaterials = StaticMesh->GetRenderData()->LODResources[i].Sections.Num();
-
-			LODs[i].Materials.Reset(NumMaterials);
-			LODs[i].Materials.SetNum(NumMaterials);
+			LODs[LODIndex].Materials.Reset(NumMaterials);
+			LODs[LODIndex].Materials.SetNum(NumMaterials);
 			for (int32 MaterialIndex = 0; MaterialIndex < NumMaterials; ++MaterialIndex)
 			{
-				LODs[i].Materials.Reset(NumMaterials);
-				LODs[i].Materials.SetNum(NumMaterials);
-				
 				UMaterialInterface* MaterialInterface = 0;
-				int StaticMeshMaterialIndex = StaticMesh->GetRenderData()->LODResources[i].Sections[MaterialIndex].MaterialIndex;
+				int32 StaticMeshMaterialIndex = StaticMesh->GetRenderData()->LODResources[LODIndex].Sections[MaterialIndex].MaterialIndex;
 				if (StaticMeshMaterialIndex >= 0 && StaticMeshMaterialIndex < StaticMesh->GetStaticMaterials().Num())
 				{
 					MaterialInterface = StaticMesh->GetStaticMaterials()[StaticMeshMaterialIndex].MaterialInterface;
@@ -123,13 +117,13 @@ void UCustomizableObjectNodeStaticMesh::AllocateDefaultPins(UCustomizableObjectN
 				{
 					MaterialName = MaterialInterface->GetName();
 				}
-				LODs[i].Materials[MaterialIndex].Name = MaterialName;
-
+				LODs[LODIndex].Materials[MaterialIndex].Name = MaterialName;
+				
 				// Mesh
 				const FString MeshName = MaterialName + TEXT(" Mesh");
 				UEdGraphPin* Pin = CustomCreatePin(EGPD_Output, Schema->PC_Mesh, FName(*MeshName));
 				Pin->bDefaultValueIsIgnored = true;
-				LODs[i].Materials[MaterialIndex].MeshPinRef.SetPin(Pin);
+				LODs[LODIndex].Materials[MaterialIndex].MeshPinRef.SetPin(Pin);
 
 				// Images
 				if (MaterialInterface)
@@ -144,7 +138,7 @@ void UCustomizableObjectNodeStaticMesh::AllocateDefaultPins(UCustomizableObjectN
 						FString ImageName = ImageNames[ImageIndex].Name.ToString();
 						Pin = CustomCreatePin(EGPD_Output, Schema->PC_Texture, *ImageName);
 						Pin->bDefaultValueIsIgnored = true;
-						LODs[i].Materials[MaterialIndex].ImagePinsRef.Add(Pin);
+						LODs[LODIndex].Materials[MaterialIndex].ImagePinsRef.Add(Pin);
 					}
 				}
 			}
@@ -165,7 +159,7 @@ void UCustomizableObjectNodeStaticMesh::AllocateDefaultPins(UCustomizableObjectN
 
 FText UCustomizableObjectNodeStaticMesh::GetNodeTitle(ENodeTitleType::Type TitleType) const
 {
-	if (StaticMesh)
+	if (MutablePrivate::LoadObject(StaticMesh))
 	{
 		FFormatNamedArguments Args;
 		Args.Add(TEXT("MeshName"), FText::FromString(StaticMesh->GetName()));
@@ -181,17 +175,17 @@ FText UCustomizableObjectNodeStaticMesh::GetNodeTitle(ENodeTitleType::Type Title
 
 UTexture2D* UCustomizableObjectNodeStaticMesh::FindTextureForPin(const UEdGraphPin* Pin) const
 {
-	if (!StaticMesh)
+	if (!MutablePrivate::LoadObject(StaticMesh))
 	{
 		return 0;
 	}
 
-	for (int LODIndex = 0; LODIndex < LODs.Num(); ++LODIndex)
+	for (int32 LODIndex = 0; LODIndex < LODs.Num(); ++LODIndex)
 	{
-		for (int MaterialIndex = 0; MaterialIndex<LODs[LODIndex].Materials.Num(); ++MaterialIndex)
+		for (int32 MaterialIndex = 0; MaterialIndex<LODs[LODIndex].Materials.Num(); ++MaterialIndex)
 		{
 			UMaterialInterface* MaterialInterface = 0;
-			int StaticMeshMaterialIndex = -1;
+			int32 StaticMeshMaterialIndex = -1;
 			if (StaticMesh->GetRenderData()->LODResources.Num() > LODIndex && StaticMesh->GetRenderData()->LODResources[LODIndex].Sections.Num()>MaterialIndex)
 			{
 				StaticMeshMaterialIndex = StaticMesh->GetRenderData()->LODResources[LODIndex].Sections[MaterialIndex].MaterialIndex;
@@ -204,7 +198,7 @@ UTexture2D* UCustomizableObjectNodeStaticMesh::FindTextureForPin(const UEdGraphP
 
 			if (MaterialInterface)
 			{
-				for (int ImageIndex = 0; ImageIndex < LODs[LODIndex].Materials[MaterialIndex].ImagePinsRef.Num(); ++ImageIndex)
+				for (int32 ImageIndex = 0; ImageIndex < LODs[LODIndex].Materials[MaterialIndex].ImagePinsRef.Num(); ++ImageIndex)
 				{
 					if (LODs[LODIndex].Materials[MaterialIndex].ImagePinsRef[ImageIndex].Get() == Pin)
 					{
@@ -278,26 +272,26 @@ bool UCustomizableObjectNodeStaticMesh::IsNodeOutDatedAndNeedsRefresh()
 {
 	bool Result = false;
 
-	if (StaticMesh)
+	if (MutablePrivate::LoadObject(StaticMesh))
 	{
-		int32 numLods = StaticMesh->GetRenderData()->LODResources.Num();
-		if (LODs.Num() != numLods)
+		const int32 LODCount = StaticMesh->GetRenderData()->LODResources.Num();
+		if (LODs.Num() != LODCount)
 		{
 			Result = true;
 		}
 
-		for (int LODIndex = 0; !Result && LODIndex < LODs.Num(); ++LODIndex)
+		for (int32 LODIndex = 0; !Result && LODIndex < LODs.Num(); ++LODIndex)
 		{
-			int32 NumMaterials = StaticMesh->GetRenderData()->LODResources[LODIndex].Sections.Num();
+			const int32 NumMaterials = StaticMesh->GetRenderData()->LODResources[LODIndex].Sections.Num();
 			if (LODs[LODIndex].Materials.Num() != NumMaterials)
 			{
 				Result = true;
 			}
 
-			for (int MaterialIndex = 0; !Result && MaterialIndex < NumMaterials; ++MaterialIndex)
+			for (int32 MaterialIndex = 0; !Result && MaterialIndex < NumMaterials; ++MaterialIndex)
 			{
 				UMaterialInterface* MaterialInterface = nullptr;
-				int SkeletalMeshMaterialIndex = StaticMesh->GetRenderData()->LODResources[LODIndex].Sections[MaterialIndex].MaterialIndex;
+				int32 SkeletalMeshMaterialIndex = StaticMesh->GetRenderData()->LODResources[LODIndex].Sections[MaterialIndex].MaterialIndex;
 				if (SkeletalMeshMaterialIndex >= 0 && SkeletalMeshMaterialIndex < StaticMesh->GetStaticMaterials().Num())
 				{
 					MaterialInterface = StaticMesh->GetStaticMaterials()[SkeletalMeshMaterialIndex].MaterialInterface;
@@ -322,7 +316,7 @@ bool UCustomizableObjectNodeStaticMesh::IsNodeOutDatedAndNeedsRefresh()
 						break;
 					}
 
-					for (int ImageIndex = 0; ImageIndex < ImageInfos.Num(); ++ImageIndex)
+					for (int32 ImageIndex = 0; ImageIndex < ImageInfos.Num(); ++ImageIndex)
 					{
 						const FString ImagePinName = LODs[LODIndex].Materials[MaterialIndex].ImagePinsRef[ImageIndex].Get()->GetName();
 						if (!ImagePinName.Contains(ImageInfos[ImageIndex].Name.ToString()))
@@ -346,12 +340,95 @@ bool UCustomizableObjectNodeStaticMesh::IsNodeOutDatedAndNeedsRefresh()
 	return Result;
 }
 
+
 FString UCustomizableObjectNodeStaticMesh::GetRefreshMessage() const
 {
 	return "Node data outdated. Please refresh node.";
 }
 
 
+void UCustomizableObjectNodeStaticMesh::BackwardsCompatibleFixup(int32 CustomizableObjectCustomVersion)
+{
+	Super::BackwardsCompatibleFixup(CustomizableObjectCustomVersion);
+
+	if (CustomizableObjectCustomVersion == FCustomizableObjectCustomVersion::FixStaticMeshNodeLODMaterialDataInvalidation)
+	{
+		FixupReconstructPins(CreateRemapPinsByName(), [&](UCustomizableObjectNode* Node, UCustomizableObjectNodeRemapPins* Action)
+		{
+			// Extracted from AllocateDefaultPins
+			// Re-generate the node LOD data property alongside the pins that represent its contents as the node was saved with invalid data in said property
+				
+			const UEdGraphSchema_CustomizableObject* Schema = GetDefault<UEdGraphSchema_CustomizableObject>();
+
+			if (MutablePrivate::LoadObject(StaticMesh))
+			{
+				const int32 NumLODs = StaticMesh->GetRenderData()->LODResources.Num();
+
+				LODs.Reset(NumLODs);
+				LODs.SetNum(NumLODs);
+				for (int32 LODIndex = 0; LODIndex < NumLODs; ++LODIndex)
+				{
+					const int32 NumMaterials = StaticMesh->GetRenderData()->LODResources[LODIndex].Sections.Num();
+
+					LODs[LODIndex].Materials.Reset(NumMaterials);
+					LODs[LODIndex].Materials.SetNum(NumMaterials);
+					for (int32 MaterialIndex = 0; MaterialIndex < NumMaterials; ++MaterialIndex)
+					{
+						UMaterialInterface* MaterialInterface = nullptr;
+						int32 StaticMeshMaterialIndex = StaticMesh->GetRenderData()->LODResources[LODIndex].Sections[MaterialIndex].MaterialIndex;
+						if (StaticMeshMaterialIndex >= 0 && StaticMeshMaterialIndex < StaticMesh->GetStaticMaterials().Num())
+						{
+							MaterialInterface = StaticMesh->GetStaticMaterials()[StaticMeshMaterialIndex].MaterialInterface;
+						}
+
+						FString MaterialName = TEXT("Unnamed Material");
+						if (MaterialInterface)
+						{
+							MaterialName = MaterialInterface->GetName();
+						}
+						LODs[LODIndex].Materials[MaterialIndex].Name = MaterialName;
+						
+						// Mesh
+						const FString MeshName = MaterialName + TEXT(" Mesh");
+						UEdGraphPin* Pin = CustomCreatePin(EGPD_Output, Schema->PC_Mesh, FName(*MeshName));
+						Pin->bDefaultValueIsIgnored = true;
+						LODs[LODIndex].Materials[MaterialIndex].MeshPinRef.SetPin(Pin);
+
+						// Images
+						if (MaterialInterface)
+						{
+							TArray<FMaterialParameterInfo> ImageNames;
+							TArray<FGuid> ImageIds;
+
+							UMaterial* Material = MaterialInterface->GetMaterial();
+							Material->GetAllTextureParameterInfo(ImageNames, ImageIds);
+							for (int32 ImageIndex = 0; ImageIndex < ImageNames.Num(); ++ImageIndex)
+							{
+								FString ImageName = ImageNames[ImageIndex].Name.ToString();
+								Pin = CustomCreatePin(EGPD_Output, Schema->PC_Texture, *ImageName);
+								Pin->bDefaultValueIsIgnored = true;
+								LODs[LODIndex].Materials[MaterialIndex].ImagePinsRef.Add(Pin);
+							}
+						}
+					}
+				}
+
+				DefaultPin = FEdGraphPinReference();
+			}
+			else
+			{
+				UEdGraphPin* Pin = CustomCreatePin(EGPD_Output, Schema->PC_Mesh, FName(STATIC_MESH_PIN_NAME));
+				Pin->bDefaultValueIsIgnored = true;
+				Pin->PinFriendlyName = FText::FromString(STATIC_MESH_PIN_NAME);
+
+				DefaultPin = FEdGraphPinReference(Pin);
+			}
+		});
+	}
+}
+
+
+
 FText UCustomizableObjectNodeStaticMesh::GetTooltipText() const
 {
 	return LOCTEXT("Static_Mesh_Tooltip", "Static Mesh");
@@ -381,7 +458,7 @@ void UCustomizableObjectNodeStaticMesh::GetPinSection(const UEdGraphPin& Pin, in
 			}
 			
 			// Is an Image pin?
-			for (int ImageIndex = 0; ImageIndex < Material.ImagePinsRef.Num(); ++ImageIndex)
+			for (int32 ImageIndex = 0; ImageIndex < Material.ImagePinsRef.Num(); ++ImageIndex)
 			{
 				if (Material.ImagePinsRef[ImageIndex].Get() == &Pin)
 				{
@@ -397,7 +474,7 @@ void UCustomizableObjectNodeStaticMesh::GetPinSection(const UEdGraphPin& Pin, in
 
 UMaterialInterface* UCustomizableObjectNodeStaticMesh::GetMaterialFor(const UEdGraphPin* Pin) const
 {
-	if (StaticMesh)
+	if (MutablePrivate::LoadObject(StaticMesh))
 	{
 		check(Pin);
 		int32 LODIndex;
@@ -431,7 +508,8 @@ TSharedPtr<SGraphNode> UCustomizableObjectNodeStaticMesh::CreateVisualWidget()
 void SGraphNodeStaticMesh::Construct(const FArguments& InArgs, UEdGraphNode* InGraphNode)
 {
 	NodeStaticMesh = Cast< UCustomizableObjectNodeStaticMesh >(InGraphNode);
-
+	MutablePrivate::LoadObject(NodeStaticMesh->StaticMesh);
+	
 	WidgetSize = 128.0f;
 	ThumbnailSize = 128;
 
@@ -447,7 +525,7 @@ void SGraphNodeStaticMesh::Construct(const FArguments& InArgs, UEdGraphNode* InG
 	SingleDetails.NotifyHook = Editor.Get();
 	SingleDetails.bHideAssetThumbnail = true;
 
-	StaticMeshSelector = PropPlugin.CreateSingleProperty(NodeStaticMesh, "StaticMesh", SingleDetails);	
+	StaticMeshSelector = PropPlugin.CreateSingleProperty(NodeStaticMesh, GET_MEMBER_NAME_STRING_CHECKED(UCustomizableObjectNodeStaticMesh, StaticMesh), SingleDetails);	
 
 	SCustomizableObjectNode::Construct({}, InGraphNode);
 }
diff --git a/Engine/Plugins/Mutable/Source/CustomizableObjectEditor/Private/MuCOE/Nodes/CustomizableObjectNodeStaticMesh.h b/Engine/Plugins/Mutable/Source/CustomizableObjectEditor/Private/MuCOE/Nodes/CustomizableObjectNodeStaticMesh.h
--- a/Engine/Plugins/Mutable/Source/CustomizableObjectEditor/Private/MuCOE/Nodes/CustomizableObjectNodeStaticMesh.h
+++ b/Engine/Plugins/Mutable/Source/CustomizableObjectEditor/Private/MuCOE/Nodes/CustomizableObjectNodeStaticMesh.h
@@ -148,6 +148,7 @@ public:
 	UE_API virtual bool HasPinViewer() const override;
 	UE_API virtual bool IsNodeOutDatedAndNeedsRefresh() override;
 	UE_API virtual FString GetRefreshMessage() const override;
+	UE_API virtual void BackwardsCompatibleFixup(int32 CustomizableObjectCustomVersion) override;
 
 	// UCustomizableObjectNodeMesh interface
 	UE_API virtual UTexture2D* FindTextureForPin(const UEdGraphPin* Pin) const override;
