[mutable] Fixed repeated LOCTEXT keys in the same mutable namespace.

#rnx

diff --git a/Engine/Plugins/Mutable/Source/CustomizableObjectEditor/Private/MuCOE/EdGraphSchema_CustomizableObject.cpp b/Engine/Plugins/Mutable/Source/CustomizableObjectEditor/Private/MuCOE/EdGraphSchema_CustomizableObject.cpp
--- a/Engine/Plugins/Mutable/Source/CustomizableObjectEditor/Private/MuCOE/EdGraphSchema_CustomizableObject.cpp
+++ b/Engine/Plugins/Mutable/Source/CustomizableObjectEditor/Private/MuCOE/EdGraphSchema_CustomizableObject.cpp
@@ -1504,11 +1504,11 @@ FText UEdGraphSchema_CustomizableObject::GetPinCategoryFriendlyName(const FName&
 	}
 	else if (PinCategory == PC_SkeletalMesh)
 	{
-		return LOCTEXT("Mesh_Pin_Category", "Skeletal Mesh");
+		return LOCTEXT("Skeletal_Mesh_Pin_Category", "Skeletal Mesh");
 	}
 	else if (PinCategory == PC_PassthroughSkeletalMesh)
 	{
-		return LOCTEXT("Mesh_Pin_Category", "Passthrough Skeletal Mesh");
+		return LOCTEXT("Passthrough_Skeletal_Mesh_Pin_Category", "Passthrough Skeletal Mesh");
 	}
 	else if (PinCategory == PC_Texture)
 	{
@@ -1516,7 +1516,7 @@ FText UEdGraphSchema_CustomizableObject::GetPinCategoryFriendlyName(const FName&
 	}
 	else if (PinCategory == PC_PassthroughTexture)
 	{
-		return LOCTEXT("PassThrough_Image_Pin_Category", "Passthrough Texture");
+		return LOCTEXT("Passthrough_Texture_Pin_Category", "Passthrough Texture");
 	}
 	else if (PinCategory == PC_Projector)
 	{
