[Mutable] Fix Materials comparison for merging

#13666

#DisabledValidations Preflight

diff --git a/Engine/Plugins/Mutable/Source/CustomizableObject/Private/MuCO/CustomizableObjectInstance.cpp b/Engine/Plugins/Mutable/Source/CustomizableObject/Private/MuCO/CustomizableObjectInstance.cpp
--- a/Engine/Plugins/Mutable/Source/CustomizableObject/Private/MuCO/CustomizableObjectInstance.cpp
+++ b/Engine/Plugins/Mutable/Source/CustomizableObject/Private/MuCO/CustomizableObjectInstance.cpp
@@ -120,7 +120,14 @@ struct FMutableMaterialPlaceholder
 
 		bool operator<(const FMutableMaterialPlaceHolderParam& Other) const
 		{
-			return Type < Other.Type || ParamName.CompareIndexes(Other.ParamName);
+			if (Type != Other.Type)
+			{
+				return Type < Other.Type;
+			}
+			else
+			{
+				return ParamName.LexicalLess(Other.ParamName);
+			}
 		}
 
 		bool operator==(const FMutableMaterialPlaceHolderParam& Other) const = default;
