[Mutable] Compile API typo

diff --git a/Engine/Plugins/Mutable/Config/BaseMutable.ini b/Engine/Plugins/Mutable/Config/BaseMutable.ini
--- a/Engine/Plugins/Mutable/Config/BaseMutable.ini
+++ b/Engine/Plugins/Mutable/Config/BaseMutable.ini
@@ -25,4 +25,5 @@
 +FunctionRedirects=(OldName="/Script/CustomizableObject.CustomizableObject.GetIntParameterOptionDataTable",NewName="/Script/CustomizableObject.CustomizableObject.GetEnumParameterValueDataTable")
 +FunctionRedirects=(OldName="/Script/CustomizableObject.CustomizableObject.GetIntParameterDefaultValue",NewName="/Script/CustomizableObject.CustomizableObject.GetEnumParameterDefaultValue")
 +PropertyRedirects=(OldName="/Script/CustomizableObjectEditor.CustomizableObjectNodeComponentPassthroughMesh.Mesh",NewName="/Script/CustomizableObjectEditor.CustomizableObjectNodeComponentPassthroughMesh.SkeletalMesh")
-+PropertyRedirects=(OldName="/Script/CustomizableObject.CustomizableObject.bEnableAnimBpPhysicsAssetsManipualtion",NewName="/Script/CustomizableObject.CustomizableObject.bEnableAnimBpPhysicsAssetsManipulation")
\ No newline at end of file
++PropertyRedirects=(OldName="/Script/CustomizableObject.CustomizableObject.bEnableAnimBpPhysicsAssetsManipualtion",NewName="/Script/CustomizableObject.CustomizableObject.bEnableAnimBpPhysicsAssetsManipulation")
++PropertyRedirects=(OldName="/Script/CustomizableObject.CompileParams.bGatherReferneces",NewName="/Script/CustomizableObject.CompileParams.bGatherReferences")
\ No newline at end of file
diff --git a/Engine/Plugins/Mutable/Source/CustomizableObject/Public/MuCO/CustomizableObject.h b/Engine/Plugins/Mutable/Source/CustomizableObject/Public/MuCO/CustomizableObject.h
--- a/Engine/Plugins/Mutable/Source/CustomizableObject/Public/MuCO/CustomizableObject.h
+++ b/Engine/Plugins/Mutable/Source/CustomizableObject/Public/MuCO/CustomizableObject.h
@@ -188,7 +188,7 @@ struct FCompileParams
 	
 	/** Gathers all asset references used in this Customizable Object. Marks the object as modified.*/
 	UPROPERTY(BlueprintReadWrite, Category = Compile)
-	bool bGatherReferneces = false;
+	bool bGatherReferences = false;
 
 	/** If true, compile the Customizable Object asynchronously. If synchronously and the Customizable Object is not ready the compilation will fail. */
 	UPROPERTY(BlueprintReadWrite, Category = Compile)
diff --git a/Engine/Plugins/Mutable/Source/CustomizableObjectEditor/Private/MuCOE/CustomizableObjectEditorFunctionLibrary.cpp b/Engine/Plugins/Mutable/Source/CustomizableObjectEditor/Private/MuCOE/CustomizableObjectEditorFunctionLibrary.cpp
--- a/Engine/Plugins/Mutable/Source/CustomizableObjectEditor/Private/MuCOE/CustomizableObjectEditorFunctionLibrary.cpp
+++ b/Engine/Plugins/Mutable/Source/CustomizableObjectEditor/Private/MuCOE/CustomizableObjectEditorFunctionLibrary.cpp
@@ -35,7 +35,7 @@ ECustomizableObjectCompilationState UCustomizableObjectEditorFunctionLibrary::Co
 	Params.bAsync = false;
 	Params.OptimizationLevel = InOptimizationLevel;
 	Params.TextureCompression = InTextureCompression;
-	Params.bGatherReferneces = bGatherReferences;
+	Params.bGatherReferences = bGatherReferences;
 
 	bool bCompilationSuccess = false;
 	Params.CallbackNative.BindLambda([&bCompilationSuccess](const FCompileCallbackParams& Params)
diff --git a/Engine/Plugins/Mutable/Source/CustomizableObjectEditor/Private/MuCOE/CustomizableObjectEditorModule.cpp b/Engine/Plugins/Mutable/Source/CustomizableObjectEditor/Private/MuCOE/CustomizableObjectEditorModule.cpp
--- a/Engine/Plugins/Mutable/Source/CustomizableObjectEditor/Private/MuCOE/CustomizableObjectEditorModule.cpp
+++ b/Engine/Plugins/Mutable/Source/CustomizableObjectEditor/Private/MuCOE/CustomizableObjectEditorModule.cpp
@@ -1153,7 +1153,7 @@ void FCustomizableObjectEditorModule::CompileCustomizableObject(UCustomizableObj
 		CompileRequest->bSkipIfCompiled = Params->bSkipIfCompiled;
 		CompileRequest->bSkipIfOutOfDate = Params->bSkipIfOutOfDate;
 		CompileRequest->Options.TextureCompression = Params->TextureCompression;
-		CompileRequest->Options.bGatherReferences = Params->bGatherReferneces;
+		CompileRequest->Options.bGatherReferences = Params->bGatherReferences;
 		CompileRequest->Callback = Params->Callback;
 		CompileRequest->CallbackNative = Params->CallbackNative;
 
