[Mutable] Fix Textures and Meshes not being properly cached

- Their hash was not deterministic.

diff --git a/Engine/Plugins/Mutable/Source/MutableRuntime/Private/MuR/System.cpp b/Engine/Plugins/Mutable/Source/MutableRuntime/Private/MuR/System.cpp
--- a/Engine/Plugins/Mutable/Source/MutableRuntime/Private/MuR/System.cpp
+++ b/Engine/Plugins/Mutable/Source/MutableRuntime/Private/MuR/System.cpp
@@ -2093,7 +2093,7 @@ namespace mu
 			case EParameterType::Image:
 			{
 				FString Path = Params->GetPrivate()->Values[ParamIndex].Get<FParamTextureType>()->GetPathName();
-				Blob.Serialize(GetData(Path), Path.GetAllocatedSize());
+				Blob.Serialize(GetData(Path), Path.Len() * sizeof(FString::ElementType));
 
 				// Multi-values
 				if (Params->GetPrivate()->HasMultipleValues(ParamIndex))
@@ -2103,7 +2103,7 @@ namespace mu
 					for (const TPair< TArray<int32>, FParameterValue >& v : Multi)
 					{
 						Path = v.Value.Get<FParamTextureType>()->GetPathName();
-						Blob.Serialize(GetData(Path), Path.GetAllocatedSize());
+						Blob.Serialize(GetData(Path), Path.Len() * sizeof(FString::ElementType));
 					}
 				}
 				break;
@@ -2112,7 +2112,7 @@ namespace mu
 			case EParameterType::Mesh:
 			{
 				FString Path = Params->GetPrivate()->Values[ParamIndex].Get<FParamSkeletalMeshType>()->GetPathName();
-				Blob.Serialize(GetData(Path), Path.GetAllocatedSize());
+				Blob.Serialize(GetData(Path), Path.Len() * sizeof(FString::ElementType));
 
 				// Multi-values
 				if (Params->GetPrivate()->HasMultipleValues(ParamIndex))
@@ -2122,7 +2122,7 @@ namespace mu
 					for (const TPair< TArray<int32>, FParameterValue >& v : Multi)
 					{
 						Path = v.Value.Get<FParamSkeletalMeshType>()->GetPathName();
-						Blob.Serialize(GetData(Path), Path.GetAllocatedSize());
+						Blob.Serialize(GetData(Path), Path.Len() * sizeof(FString::ElementType));
 					}
 				}
 				break;
