[Mutable] Remove the legacy code that used to gather additional participating objects when converting skeletal meshes to Mutable.

#rnx

diff --git a/Engine/Plugins/Mutable/Source/CustomizableObjectEditor/Private/MuCOE/GenerateMutableSource/GenerateMutableSourceMesh.cpp b/Engine/Plugins/Mutable/Source/CustomizableObjectEditor/Private/MuCOE/GenerateMutableSource/GenerateMutableSourceMesh.cpp
--- a/Engine/Plugins/Mutable/Source/CustomizableObjectEditor/Private/MuCOE/GenerateMutableSource/GenerateMutableSourceMesh.cpp
+++ b/Engine/Plugins/Mutable/Source/CustomizableObjectEditor/Private/MuCOE/GenerateMutableSource/GenerateMutableSourceMesh.cpp
@@ -2559,77 +2559,6 @@ TSharedPtr<mu::FMesh> ConvertSkeletalMeshToMutable(
 		}
 	}
 
-	// Add additional participating objects
-	{
-		const FName LongPackageFName = Source.Mesh.GetLongPackageFName();
-
-		FAssetRegistryModule& AssetRegistryModule = FModuleManager::LoadModuleChecked<FAssetRegistryModule>("AssetRegistry");
-		TArray<FAssetData> Assets;
-		{
-			TArray<FName> PackageDependencies;
-			AssetRegistryModule.Get().GetDependencies(LongPackageFName, PackageDependencies);
-
-			FARFilter Filter;
-			for (const FName& PackageName : PackageDependencies)
-			{
-				if (!PackageName.ToString().StartsWith(TEXT("/TempAutosave")))
-				{
-					Filter.PackageNames.Add(PackageName);
-				}
-			}
-			AssetRegistryModule.Get().GetAssets(Filter, Assets);
-		}
-
-		// Add the skeleton dependency
-		const bool bIgnoreSkeleton = EnumHasAnyFlags(Source.Flags, EMutableMeshConversionFlags::IgnoreSkinning);
-		if (!bIgnoreSkeleton)
-		{
-			bool bSkeletonFound = false;
-			for (const FAssetData& Dep : Assets)
-			{
-				const UClass* AssetClass = Dep.GetClass(EResolveClass::Yes);
-				if (AssetClass == USkeleton::StaticClass())
-				{
-					bSkeletonFound = true;
-				}
-			}
-
-			if (!bSkeletonFound)
-			{
-				FString Msg = FString::Printf(TEXT("No skeleton provided when converting SkeletalMesh [%s]."), *LongPackageFName.ToString());
-				GenerationContext.Log(FText::FromString(Msg), Source.MessageContext);
-				return nullptr;
-			}
-		}
-
-		// Add clothing asset participating objects
-		if (GenerationContext.CompilationContext->Options.bClothingEnabled)
-		{
-			for (const FAssetData& Dep : Assets)
-			{
-				const UClass* AssetClass = Dep.GetClass(EResolveClass::Yes);
-				if (AssetClass->ImplementsInterface(UClothingAssetCommon::StaticClass()))
-				{
-					TArray<FAssetData> ClothAssets;
-					{
-						TArray<FName> ClothPackageDependencies;
-						AssetRegistryModule.Get().GetDependencies(LongPackageFName, ClothPackageDependencies);
-
-						FARFilter ClothFilter;
-						for (const FName& PackageName : ClothPackageDependencies)
-						{
-							if (!PackageName.ToString().StartsWith(TEXT("/TempAutosave")))
-							{
-								ClothFilter.PackageNames.Add(PackageName);
-							}
-						}
-						AssetRegistryModule.Get().GetAssets(ClothFilter, ClothAssets);
-					}
-				}
-			}
-		}
-	}
-
 	TSharedPtr<mu::FMesh> MutableMesh;
 	if (!bForceImmediateConversion)
 	{
