[mutable] CIS: Updated CIS Baking commandlet so it uses the compilation method of the Baking system instead of a custom one.
- Made the compilation for baking use the higher quality settings for the texture compression (compilation)
- Removed CIS centric compilation approach in favor of the baking system defined one. This improves testing overlap between the CIS test and the real thing.
- Removed overrides that made the baking system prepared for reporting benchmarking data.
- Made the for baking compilation always run (for CIS tests) and to request high-quality textures.
- Renamed compilation variable named bSkipIfOutOfDate with bSkipIfNotOutOfDate as the name was representing the inverse of the behavior that the variable did expose during the compilation of the CO.
- Disabled for now the execution of limited CO compilations to prevent the apparition of warning messages that should not appear in the context of limited compilations.

#rnx

diff --git a/Engine/Plugins/Mutable/Config/BaseMutable.ini b/Engine/Plugins/Mutable/Config/BaseMutable.ini
--- a/Engine/Plugins/Mutable/Config/BaseMutable.ini
+++ b/Engine/Plugins/Mutable/Config/BaseMutable.ini
@@ -26,4 +26,5 @@
 +FunctionRedirects=(OldName="/Script/CustomizableObject.CustomizableObject.GetIntParameterDefaultValue",NewName="/Script/CustomizableObject.CustomizableObject.GetEnumParameterDefaultValue")
 +PropertyRedirects=(OldName="/Script/CustomizableObjectEditor.CustomizableObjectNodeComponentPassthroughMesh.Mesh",NewName="/Script/CustomizableObjectEditor.CustomizableObjectNodeComponentPassthroughMesh.SkeletalMesh")
 +PropertyRedirects=(OldName="/Script/CustomizableObject.CustomizableObject.bEnableAnimBpPhysicsAssetsManipualtion",NewName="/Script/CustomizableObject.CustomizableObject.bEnableAnimBpPhysicsAssetsManipulation")
-+PropertyRedirects=(OldName="/Script/CustomizableObject.CompileParams.bGatherReferneces",NewName="/Script/CustomizableObject.CompileParams.bGatherReferences")
\ No newline at end of file
++PropertyRedirects=(OldName="/Script/CustomizableObject.CompileParams.bGatherReferneces",NewName="/Script/CustomizableObject.CompileParams.bGatherReferences")
++PropertyRedirects=(OldName="/Script/CustomizableObject.CompileParams.bSkipIfOutOfDate",NewName="/Script/CustomizableObject.CompileParams.bSkipIfNotOutOfDate")
diff --git a/Engine/Plugins/Mutable/Source/CustomizableObject/Public/MuCO/CustomizableObject.h b/Engine/Plugins/Mutable/Source/CustomizableObject/Public/MuCO/CustomizableObject.h
--- a/Engine/Plugins/Mutable/Source/CustomizableObject/Public/MuCO/CustomizableObject.h
+++ b/Engine/Plugins/Mutable/Source/CustomizableObject/Public/MuCO/CustomizableObject.h
@@ -182,9 +182,9 @@ struct FCompileParams
 	UPROPERTY(BlueprintReadWrite, Category = Compile)
 	bool bSkipIfCompiled = false;
 
-	/** If true, skip the compilation if the Customizable Object or its contents are out of date (child Customizable Objects, Data Tables, referenced Materials...). */
+	/** If true, skip the compilation if the Customizable Object or its contents are not out of date (child Customizable Objects, Data Tables, referenced Materials...). */
 	UPROPERTY(BlueprintReadWrite, Category = Compile)
-	bool bSkipIfOutOfDate = true;
+	bool bSkipIfNotOutOfDate = true;
 	
 	/** Gathers all asset references used in this Customizable Object. Marks the object as modified.*/
 	UPROPERTY(BlueprintReadWrite, Category = Compile)
diff --git a/Engine/Plugins/Mutable/Source/CustomizableObjectEditor/Internal/MuCOE/CompileRequest.h b/Engine/Plugins/Mutable/Source/CustomizableObjectEditor/Internal/MuCOE/CompileRequest.h
--- a/Engine/Plugins/Mutable/Source/CustomizableObjectEditor/Internal/MuCOE/CompileRequest.h
+++ b/Engine/Plugins/Mutable/Source/CustomizableObjectEditor/Internal/MuCOE/CompileRequest.h
@@ -48,7 +48,7 @@ public:
 
 	bool bAsync = true;
 	bool bSkipIfCompiled = false;
-	bool bSkipIfOutOfDate = false;
+	bool bSkipIfNotOutOfDate = false;
 	bool bSilentCompilation = true;
 };
 
diff --git a/Engine/Plugins/Mutable/Source/CustomizableObjectEditor/Internal/MuCOE/CustomizableObjectInstanceBakingUtils.h b/Engine/Plugins/Mutable/Source/CustomizableObjectEditor/Internal/MuCOE/CustomizableObjectInstanceBakingUtils.h
--- a/Engine/Plugins/Mutable/Source/CustomizableObjectEditor/Internal/MuCOE/CustomizableObjectInstanceBakingUtils.h
+++ b/Engine/Plugins/Mutable/Source/CustomizableObjectEditor/Internal/MuCOE/CustomizableObjectInstanceBakingUtils.h
@@ -9,11 +9,12 @@
 /**
  * Schedules the async compilation of a given Customizable object and, if a Customizable Object Instance is provided, tries to compile only the data required
  * for that given COI. Use the delegate to know when the operation has completed.
- * @param InCustomizableObject The Customizable Object we want to compile.
  * @param InTargetInstance The Customizable Object Instance generated from the provided COI that we want to use to limit the scope of the CO compilation.
- * @param InCustomizableObjectCompilationDelegate Delegate invoked once the compilation has completed. 
+ * @param InCustomizableObjectCompilationDelegate Delegate invoked once the compilation has completed.
+ * @param bPerformPartialCompilation Instead of compiling all the data from the InTargetInstance Co just compile the data required for the provided instance.
+ * todo: UE-315780 Make this option true by default when the compilation warnings caused by it are resolved
  */
-CUSTOMIZABLEOBJECTEDITOR_API void ScheduleCOCompilationForBaking(UCustomizableObjectInstance& InTargetInstance, const FCompileNativeDelegate& InCustomizableObjectCompilationDelegate);
+CUSTOMIZABLEOBJECTEDITOR_API void ScheduleCOCompilationForBaking(UCustomizableObjectInstance& InTargetInstance, const FCompileNativeDelegate& InCustomizableObjectCompilationDelegate, const bool bPerformPartialCompilation = false);
 
 
 /**
diff --git a/Engine/Plugins/Mutable/Source/CustomizableObjectEditor/Private/MuCOE/CustomizableObjectCompiler.cpp b/Engine/Plugins/Mutable/Source/CustomizableObjectEditor/Private/MuCOE/CustomizableObjectCompiler.cpp
--- a/Engine/Plugins/Mutable/Source/CustomizableObjectEditor/Private/MuCOE/CustomizableObjectCompiler.cpp
+++ b/Engine/Plugins/Mutable/Source/CustomizableObjectEditor/Private/MuCOE/CustomizableObjectCompiler.cpp
@@ -254,7 +254,7 @@ void FCustomizableObjectCompiler::Compile(const TSharedRef<FCompilationRequest>&
 		return;
 	}
 
-	if (InCompileRequest->bSkipIfOutOfDate)
+	if (InCompileRequest->bSkipIfNotOutOfDate)
 	{
 		TArray<FName> OutOfDatePackages;
 		TArray<FName> AddedPackages;
diff --git a/Engine/Plugins/Mutable/Source/CustomizableObjectEditor/Private/MuCOE/CustomizableObjectEditor.cpp b/Engine/Plugins/Mutable/Source/CustomizableObjectEditor/Private/MuCOE/CustomizableObjectEditor.cpp
--- a/Engine/Plugins/Mutable/Source/CustomizableObjectEditor/Private/MuCOE/CustomizableObjectEditor.cpp
+++ b/Engine/Plugins/Mutable/Source/CustomizableObjectEditor/Private/MuCOE/CustomizableObjectEditor.cpp
@@ -1937,7 +1937,7 @@ void FCustomizableObjectEditor::OnCustomizableObjectStatusChanged(FCustomizableO
 		else if (CurrentState == FCustomizableObjectStatusTypes::EState::NoModel)
 		{
 			TSharedRef<FCompilationRequest> CompileRequest = MakeShared<FCompilationRequest>(*CustomizableObject);
-			CompileRequest->bSkipIfOutOfDate = true;
+			CompileRequest->bSkipIfNotOutOfDate = true;
 			CompileRequest->bSilentCompilation = false;
 			CompileRequest->SetDerivedDataCachePolicy(GetDerivedDataCachePolicyForEditor());
 
diff --git a/Engine/Plugins/Mutable/Source/CustomizableObjectEditor/Private/MuCOE/CustomizableObjectEditorModule.cpp b/Engine/Plugins/Mutable/Source/CustomizableObjectEditor/Private/MuCOE/CustomizableObjectEditorModule.cpp
--- a/Engine/Plugins/Mutable/Source/CustomizableObjectEditor/Private/MuCOE/CustomizableObjectEditorModule.cpp
+++ b/Engine/Plugins/Mutable/Source/CustomizableObjectEditor/Private/MuCOE/CustomizableObjectEditorModule.cpp
@@ -1151,7 +1151,7 @@ void FCustomizableObjectEditorModule::CompileCustomizableObject(UCustomizableObj
 	{
 		CompileRequest->bAsync = Params->bAsync;
 		CompileRequest->bSkipIfCompiled = Params->bSkipIfCompiled;
-		CompileRequest->bSkipIfOutOfDate = Params->bSkipIfOutOfDate;
+		CompileRequest->bSkipIfNotOutOfDate = Params->bSkipIfNotOutOfDate;
 		CompileRequest->Options.TextureCompression = Params->TextureCompression;
 		CompileRequest->Options.bGatherReferences = Params->bGatherReferences;
 		CompileRequest->Callback = Params->Callback;
diff --git a/Engine/Plugins/Mutable/Source/CustomizableObjectEditor/Private/MuCOE/CustomizableObjectInstanceBakingUtils.cpp b/Engine/Plugins/Mutable/Source/CustomizableObjectEditor/Private/MuCOE/CustomizableObjectInstanceBakingUtils.cpp
--- a/Engine/Plugins/Mutable/Source/CustomizableObjectEditor/Private/MuCOE/CustomizableObjectInstanceBakingUtils.cpp
+++ b/Engine/Plugins/Mutable/Source/CustomizableObjectEditor/Private/MuCOE/CustomizableObjectInstanceBakingUtils.cpp
@@ -296,7 +296,7 @@ void OnInstanceUpdateFinish(const FUpdateContext& Result)
 }
 
 
-void ScheduleCOCompilationForBaking(UCustomizableObjectInstance& InTargetInstance, const FCompileNativeDelegate& InCustomizableObjectCompilationDelegate)
+void ScheduleCOCompilationForBaking(UCustomizableObjectInstance& InTargetInstance, const FCompileNativeDelegate& InCustomizableObjectCompilationDelegate, const bool bPerformPartialCompilation /*= false*/)
 {
 	// This prevents the queue of updates to have more than one instance for baking.
 	if (bIsBakeOperationAlreadyScheduled)
@@ -320,8 +320,26 @@ void ScheduleCOCompilationForBaking(UCustomizableObjectInstance& InTargetInstanc
 	FCompileParams CompileParams;
 	CompileParams.bAsync = true;
 	CompileParams.CallbackNative = InCustomizableObjectCompilationDelegate;
-	CompileParams.CompileOnlySelectedInstance = &InTargetInstance;
+	CompileParams.CompileOnlySelectedInstance = bPerformPartialCompilation ? &InTargetInstance : nullptr;
+	CompileParams.TextureCompression =  ECustomizableObjectTextureCompression::HighQuality;
+
+	// Decide if the recompilation should be forced or not based on the Texture quality settings used in the last compilation
+	CompileParams.bSkipIfCompiled = false;
+	CompileParams.bSkipIfNotOutOfDate = false;
 	
+	const UCustomizableObjectPrivate* CustomizableObjectPrivate = CustomizableObject->GetPrivate();
+	if (ensure(CustomizableObjectPrivate))
+	{
+		if (const UModelResources* ModelResources = CustomizableObjectPrivate->GetModelResources())
+		{
+			// Force the recompilation if the last compilation was not using high quality texture compression settings
+			const bool bShouldRecompileTheCO = !ModelResources->bCompiledWithHDTextureCompression;
+			
+			CompileParams.bSkipIfCompiled = !bShouldRecompileTheCO;
+			CompileParams.bSkipIfNotOutOfDate = !bShouldRecompileTheCO;
+		}
+	}
+
 	CustomizableObject->Compile(CompileParams);
 }
 
@@ -329,7 +347,7 @@ void ScheduleCOCompilationForBaking(UCustomizableObjectInstance& InTargetInstanc
 void ScheduleInstanceUpdateForBaking(UCustomizableObjectInstance& InInstance, FInstanceUpdateNativeDelegate& InInstanceUpdateDelegate)
 {
 	check(InInstance.GetCustomizableObject() && InInstance.GetCustomizableObject()->IsCompiled());
-	ensure(bIsBakeOperationAlreadyScheduled);	
+	ensureMsgf(bIsBakeOperationAlreadyScheduled, TEXT("In order to perform an update for baking you first need to schedule a compilation for bakng. That will ensure the compilation settings are the appropaite bor a bake."));	
 
 	InInstanceUpdateDelegate.AddStatic(&OnInstanceUpdateFinish);
 	
@@ -346,7 +364,6 @@ void ScheduleInstanceUpdateForBaking(UCustomizableObjectInstance& InInstance, FI
 
 	UCustomizableObjectSystemPrivate* SystemPrivate = UCustomizableObjectSystem::GetInstance()->GetPrivate();
 	SystemPrivate->EnqueueUpdateSkeletalMesh(Context);
-	
 
 	FCustomizableObjectEditorLogger::CreateLog(
 	LOCTEXT("CustomizableObjectInstanceBakingUtils_UpdateScheduled", "The COInstance Update operation for baking has been scheduled. Please hold."))
diff --git a/Engine/Plugins/Mutable/Source/CustomizableObjectEditor/Private/MuCOE/UnrealBakeHelpers.h b/Engine/Plugins/Mutable/Source/CustomizableObjectEditor/Private/MuCOE/UnrealBakeHelpers.h
--- a/Engine/Plugins/Mutable/Source/CustomizableObjectEditor/Private/MuCOE/UnrealBakeHelpers.h
+++ b/Engine/Plugins/Mutable/Source/CustomizableObjectEditor/Private/MuCOE/UnrealBakeHelpers.h
@@ -76,20 +76,20 @@ public:
 			TArray<FMaterialParameterInfo> StaticSwitchParameterInfoArray;
 			TArray<FGuid> GuidArray;
 			OriginMaterial.GetAllStaticSwitchParameterInfo(StaticSwitchParameterInfoArray, GuidArray);
-			for (int i = 0; i < StaticSwitchParameterInfoArray.Num(); ++i)
+			for (int32 ParametersArrayIndex = 0; ParametersArrayIndex < StaticSwitchParameterInfoArray.Num(); ++ParametersArrayIndex)
 			{
 				bool Value = false;
 				FGuid ExpressionsGuid;
-				if (OriginMaterial.GetStaticSwitchParameterValue(StaticSwitchParameterInfoArray[i].Name, Value, ExpressionsGuid, true))
+				if (OriginMaterial.GetStaticSwitchParameterValue(StaticSwitchParameterInfoArray[ParametersArrayIndex].Name, Value, ExpressionsGuid, true))
 				{
 					// For some reason UMaterialInstance::SetStaticSwitchParameterValueEditorOnly signature is different than UMaterial::SetStaticSwitchParameterValueEditorOnly
 					if constexpr (std::is_same_v<MaterialType, UMaterial>)
 					{
-						DestMaterial.SetStaticSwitchParameterValueEditorOnly(StaticSwitchParameterInfoArray[i].Name, Value, ExpressionsGuid);
+						DestMaterial.SetStaticSwitchParameterValueEditorOnly(StaticSwitchParameterInfoArray[ParametersArrayIndex].Name, Value, ExpressionsGuid);
 					}
 					else if (std::is_same_v<MaterialType, UMaterialInstanceConstant>)
 					{
-						DestMaterial.SetStaticSwitchParameterValueEditorOnly(StaticSwitchParameterInfoArray[i].Name, Value);
+						DestMaterial.SetStaticSwitchParameterValueEditorOnly(StaticSwitchParameterInfoArray[ParametersArrayIndex].Name, Value);
 					}
 					else
 					{
diff --git a/Engine/Plugins/Mutable/Source/MutableValidation/Private/MuV/COIBakingTestCommandlet.cpp b/Engine/Plugins/Mutable/Source/MutableValidation/Private/MuV/COIBakingTestCommandlet.cpp
--- a/Engine/Plugins/Mutable/Source/MutableValidation/Private/MuV/COIBakingTestCommandlet.cpp
+++ b/Engine/Plugins/Mutable/Source/MutableValidation/Private/MuV/COIBakingTestCommandlet.cpp
@@ -1,24 +1,25 @@
 ﻿// Copyright Epic Games, Inc. All Rights Reserved.
 
 #include "MuV/COIBakingTestCommandlet.h"
-#include "CustomizableObjectCompilationUtility.h"
 #include "ScopedLogSection.h"
 #include "ShaderCompiler.h"
 #include "ValidationUtils.h"
-#include "AssetRegistry/AssetData.h"
 #include "HAL/FileManager.h"
 #include "MuCO/CustomizableObject.h"
 #include "MuCO/CustomizableObjectInstance.h"
 #include "MuCO/CustomizableObjectSystem.h"
 #include "MuCOE/CustomizableObjectInstanceBakingUtils.h"
-#include "HAL/FileManager.h"
-#include "Interfaces/ITargetPlatformManagerModule.h"
 #include "MuCO/CustomizableObjectPrivate.h"
 #include "MuCO/CustomizableObjectSystemPrivate.h"
 #include "MuCO/LoadUtils.h"
 
 #include UE_INLINE_GENERATED_CPP_BY_NAME(COIBakingTestCommandlet)
 
+/** Flag that indicates if the CO is being compiled or not. */
+bool bIsCustomizableObjectBeingCompiled = false;
+
+/** Did the CO Compilation finish with a successful status? */
+bool bWasCustomizableObjectCompilationSuccessful = false;
 
 /** Flag useful to know if we are currently updating an instance or not */
 bool bIsInstanceBeingUpdated = false;
@@ -27,14 +28,34 @@ bool bIsInstanceBeingUpdated = false;
 bool bWasInstanceUpdateSuccessful = false;
 
 
-void OnInstanceUpdate(const FUpdateContext& Result)
+void OnCustomizableObjectCompilationEnd(const FCompileCallbackParams& CompilationParameters)
 {
-	const EUpdateResult InstanceUpdateResult = Result.UpdateResult;
+	UE_LOG(LogMutable,Display,
+		TEXT("The Customizable Object finished the compilation with state : "
+		"RequestFailed : %hhd , bWarnings : %hhd , bErrors : %hhd , bSkipped : %hhd , bCompiled : %hhd."),
+		CompilationParameters.bRequestFailed,
+		CompilationParameters.bWarnings,
+		CompilationParameters.bErrors,
+		CompilationParameters.bSkipped,
+		CompilationParameters.bCompiled);
+
+	check(!CompilationParameters.bSkipped);
+	
+	bWasCustomizableObjectCompilationSuccessful = !CompilationParameters.bRequestFailed;
+	
+	ensure(bIsCustomizableObjectBeingCompiled);
+	bIsCustomizableObjectBeingCompiled = false;
+}
 
+
+void OnInstanceUpdateEnd(const FUpdateContext& Result)
+{
+	const EUpdateResult InstanceUpdateResult = Result.UpdateResult;
 	UE_LOG(LogMutable,Display,TEXT("Instance finished update with state : %s."), *UEnum::GetValueAsString(InstanceUpdateResult));
 	bWasInstanceUpdateSuccessful = UCustomizableObjectSystem::IsUpdateResultValid(InstanceUpdateResult);
 	
 	// Clear update flag so we can exit the update while loop
+	ensure(bIsInstanceBeingUpdated);
 	bIsInstanceBeingUpdated = false;
 }
 
@@ -42,8 +63,8 @@ void OnInstanceUpdate(const FUpdateContext& Result)
 int32 UCOIBakingTestCommandlet::Main(const FString& Params)
 {
 	// Ensure we have set the mutable system to the benchmarking mode and that we are reporting benchmarking data
-	FLogBenchmarkUtil::SetBenchmarkReportingStateOverride(true);
-	UCustomizableObjectSystemPrivate::SetUsageOfBenchmarkingSettings(true);
+	// FLogBenchmarkUtil::SetBenchmarkReportingStateOverride(true);					// We are not caching benchmarking data from the baking tests
+	// UCustomizableObjectSystemPrivate::SetUsageOfBenchmarkingSettings(true);		// try to be as close as the end user in terms of settings
 	
 	// Ensure we do not show any OK dialog since we are not an user that can interact with them
 	GIsRunningUnattendedScript = true;
@@ -99,19 +120,30 @@ int32 UCOIBakingTestCommandlet::Main(const FString& Params)
 		return 1;
 	}
 	
-	// Set the target platform to be using for the compilation. Must not be a nullptr
-	FCompilationOptions CompilationOptions = InstanceCustomizableObject->GetPrivate()->GetCompileOptions();
-	ITargetPlatformManagerModule& TPM = GetTargetPlatformManagerRef();
-	CompilationOptions.TargetPlatform = TPM.GetRunningTargetPlatform();
-	CompilationOptions.bUseDiskCompilation = false;
-
-	TSharedRef<FCustomizableObjectCompilationUtility> CompilationUtility = MakeShared<FCustomizableObjectCompilationUtility>();
-	if(!CompilationUtility->CompileCustomizableObject(*InstanceCustomizableObject, true, &CompilationOptions ))
+	// Compile the Customizable object of the Instance
 	{
-		UE_LOG(LogMutable,Error,TEXT("Failed to compile the target CO. Exiting commandlet."));
-		return 1;
+		const FScopedLogSection CompilationSection (EMutableLogSection::Compilation);
+		
+		FCompileNativeDelegate CompilationEndDelegate;
+		CompilationEndDelegate.BindStatic(&OnCustomizableObjectCompilationEnd);
+
+		bIsCustomizableObjectBeingCompiled = true;
+
+		// To prevent the warnings related to partial compilation from appearing in CIS
+		constexpr bool bPerformPartialCompilation = false;
+
+		UE_LOG(LogMutable, Display, TEXT("Scheduling Customizable Object compilation."));
+		ScheduleCOCompilationForBaking(*TargetInstance, CompilationEndDelegate, bPerformPartialCompilation);
+
+		WaitUntilConditionChanges(bIsCustomizableObjectBeingCompiled);
+
+		if (!bWasCustomizableObjectCompilationSuccessful)
+		{
+			UE_LOG(LogMutable, Error, TEXT("Failed to successfully compile the target COIs CO. Exiting commandlet."));
+			return 1;
+		}
 	}
-	
+
 	
 	// Update the instance
 	{
@@ -122,24 +154,14 @@ int32 UCOIBakingTestCommandlet::Main(const FString& Params)
 
 		// Instance update delegate
 		FInstanceUpdateNativeDelegate InstanceUpdateDelegate;
-		InstanceUpdateDelegate.AddStatic(&OnInstanceUpdate);
+		InstanceUpdateDelegate.AddStatic(&OnInstanceUpdateEnd);
 		
 		bIsInstanceBeingUpdated = true;
 		UE_LOG(LogMutable, Display, TEXT("Scheduling instance update."));
 		ScheduleInstanceUpdateForBaking(*TargetInstance, InstanceUpdateDelegate);
-		
+
 		// Now tick the engine so the instance gets updated while running in the commandlet context
-        while (bIsInstanceBeingUpdated)
-        {
-            // Tick the engine
-            CommandletHelpers::TickEngine();
-
-            // Stop if exit was requested
-            if (IsEngineExitRequested())
-            {
-            	break;
-            }
-        }
+		WaitUntilConditionChanges(bIsInstanceBeingUpdated);
 		
 		// Check the end status of the instance update
 		if (!bWasInstanceUpdateSuccessful)
@@ -167,7 +189,7 @@ int32 UCOIBakingTestCommandlet::Main(const FString& Params)
 		// Delete it then but notify the user since this may mean that we have duplicated COIs
 		if (FPaths::DirectoryExists(GlobalBakingDirectory))
 		{
-			UE_LOG(LogMutable,Warning,TEXT("The directory with path  \" %s \" does already exist. This may be produced by an incompelte execution of a previous test. Clearing it out before continuing..."), *GlobalBakingDirectory);
+			UE_LOG(LogMutable,Warning,TEXT("The directory with path  \" %s \" does already exist. This may be produced by an incomplete execution of a previous test. Clearing it out before continuing..."), *GlobalBakingDirectory);
 			if (!IFileManager::Get().DeleteDirectory(*GlobalBakingDirectory,false,true))
 			{
 				UE_LOG(LogMutable,Error,TEXT("Failed to delete the baking directory at path \" %s \"."), *GlobalBakingDirectory);
@@ -195,12 +217,27 @@ int32 UCOIBakingTestCommandlet::Main(const FString& Params)
 					Configuration,
 					true,
 					SavedPackages);
+
+				// Report the baked resources
+				
+				const UEnum* Enum = StaticEnum<EPackageSaveResolutionType>();
+				check(Enum);
+				
+				for (const TPair<UPackage*, EPackageSaveResolutionType>& SavedPackagePair : SavedPackages)
+				{
+					const UPackage* SavedPackage = SavedPackagePair.Key;
+					if (ensureMsgf(SavedPackage, TEXT("One entry of the \"SavedPackages\" map is a null pointer. All entries are expected to point to non null objects.")))
+					{
+						const FString SaveTypeString = Enum->GetNameByIndex(static_cast<uint8>(SavedPackagePair.Value)).ToString();
+						UE_LOG(LogMutable, Display, TEXT("Stored package \"%s\" with \"%s\" as save resolution type."), *SavedPackage->GetName(), *SaveTypeString)
+					}
+				}
 			}
 			
 			// Delete the target directory where we did save the baked instance
 			if (!IFileManager::Get().DeleteDirectory(*GlobalBakingDirectory,true,true))
 			{
-				UE_LOG(LogMutable,Error,TEXT("Failed to deleta baking directory at path \" %s \"."), *GlobalBakingDirectory);
+				UE_LOG(LogMutable,Error,TEXT("Failed to delete the baking directory at path \" %s \"."), *GlobalBakingDirectory);
 				return 1;
 			}
 		}
@@ -213,7 +250,7 @@ int32 UCOIBakingTestCommandlet::Main(const FString& Params)
 	
 	if (bWasBakingSuccessful)
 	{
-		UE_LOG(LogMutable, Display, TEXT("Instance Baking operation has been completed succesfully."));
+		UE_LOG(LogMutable, Display, TEXT("Instance Baking operation has been completed successfully."));
 		return 0;
 	}
 	else
diff --git a/Engine/Plugins/Mutable/Source/MutableValidation/Private/MuV/ValidationUtils.cpp b/Engine/Plugins/Mutable/Source/MutableValidation/Private/MuV/ValidationUtils.cpp
--- a/Engine/Plugins/Mutable/Source/MutableValidation/Private/MuV/ValidationUtils.cpp
+++ b/Engine/Plugins/Mutable/Source/MutableValidation/Private/MuV/ValidationUtils.cpp
@@ -51,8 +51,8 @@ void Wait(const double ToWaitSeconds)
 {
 	check (ToWaitSeconds > 0);
 	
-	const double EndSeconds = FPlatformTime::Seconds() + ToWaitSeconds;
 	UE_LOG(LogMutable,Display,TEXT("Holding test execution for %f seconds."),ToWaitSeconds);
+	const double EndSeconds = FPlatformTime::Seconds() + ToWaitSeconds;
 	while (FPlatformTime::Seconds() < EndSeconds)
 	{
 		// Tick the engine
@@ -69,6 +69,24 @@ void Wait(const double ToWaitSeconds)
 }
 
 
+void WaitUntilConditionChanges(const bool& bInWaitCondition)
+{
+	const bool bStartingCondition = bInWaitCondition;
+	// Exit the wait process once the provided conditional value changes
+	while (bInWaitCondition == bStartingCondition) 
+	{ 
+		// Tick the engine
+		CommandletHelpers::TickEngine();
+
+		// Stop if exit was requested
+		if (IsEngineExitRequested())
+		{
+			break;
+		}
+	}
+}
+
+
 FCompilationOptions GetCompilationOptionsForBenchmarking (const UCustomizableObject& ReferenceCustomizableObject)
 {
 	// Override some configurations that may have been changed by the user
diff --git a/Engine/Plugins/Mutable/Source/MutableValidation/Private/MuV/ValidationUtils.h b/Engine/Plugins/Mutable/Source/MutableValidation/Private/MuV/ValidationUtils.h
--- a/Engine/Plugins/Mutable/Source/MutableValidation/Private/MuV/ValidationUtils.h
+++ b/Engine/Plugins/Mutable/Source/MutableValidation/Private/MuV/ValidationUtils.h
@@ -49,6 +49,12 @@ void PrepareAssetRegistry();
 void Wait(const double ToWaitSeconds);
 
 
+/**
+ * old the thread until the provided variable evaluates to the inverse of the original value while ticking the engine.
+ * @param bInWaitCondition Externally updatable variable that will trigger the end of the wait once the value it has differs from the original one.
+ */
+void WaitUntilConditionChanges(const bool& bInWaitCondition);
+
 /**
  * Logs some configuration data related to how mutable will compile and then generate instances. We do this so we can later
  * Isolate tests using different configurations.
@@ -87,4 +93,4 @@ TArray<FAssetData> FindAllAssetsAtPath(FName SearchPath, const UClass* TargetObj
  */
 bool TestCustomizableObject(UCustomizableObject& InTargetCustomizableObject, const ITargetPlatform& TargetCompilationPlatform,
 	const uint32 InstancesToGenerate,const bool bUseDiskCompilation);
-	
\ No newline at end of file
+	
