[mutable] Static analisis bug fix (infinite loop)

#rnx

diff --git a/Engine/Plugins/Mutable/Source/MutableValidation/Private/MuV/ValidationUtils.cpp b/Engine/Plugins/Mutable/Source/MutableValidation/Private/MuV/ValidationUtils.cpp
--- a/Engine/Plugins/Mutable/Source/MutableValidation/Private/MuV/ValidationUtils.cpp
+++ b/Engine/Plugins/Mutable/Source/MutableValidation/Private/MuV/ValidationUtils.cpp
@@ -69,7 +69,7 @@ void Wait(const double ToWaitSeconds)
 }
 
 
-void WaitUntilConditionChanges(const bool& bInWaitCondition)
+void WaitUntilConditionChanges(bool& bInWaitCondition)
 {
 	const bool bStartingCondition = bInWaitCondition;
 	// Exit the wait process once the provided conditional value changes
diff --git a/Engine/Plugins/Mutable/Source/MutableValidation/Private/MuV/ValidationUtils.h b/Engine/Plugins/Mutable/Source/MutableValidation/Private/MuV/ValidationUtils.h
--- a/Engine/Plugins/Mutable/Source/MutableValidation/Private/MuV/ValidationUtils.h
+++ b/Engine/Plugins/Mutable/Source/MutableValidation/Private/MuV/ValidationUtils.h
@@ -53,7 +53,7 @@ void Wait(const double ToWaitSeconds);
  * old the thread until the provided variable evaluates to the inverse of the original value while ticking the engine.
  * @param bInWaitCondition Externally updatable variable that will trigger the end of the wait once the value it has differs from the original one.
  */
-void WaitUntilConditionChanges(const bool& bInWaitCondition);
+void WaitUntilConditionChanges(bool& bInWaitCondition);
 
 /**
  * Logs some configuration data related to how mutable will compile and then generate instances. We do this so we can later
