[mutable] Compilations with errors will now not allow for the baking of a COI
- This change affects all paths available for the requesting of COI bakes.

#rnx

diff --git a/Engine/Plugins/Mutable/Source/CustomizableObjectEditor/Private/MuCOE/CustomizableObjectEditorViewportClient.cpp b/Engine/Plugins/Mutable/Source/CustomizableObjectEditor/Private/MuCOE/CustomizableObjectEditorViewportClient.cpp
--- a/Engine/Plugins/Mutable/Source/CustomizableObjectEditor/Private/MuCOE/CustomizableObjectEditorViewportClient.cpp
+++ b/Engine/Plugins/Mutable/Source/CustomizableObjectEditor/Private/MuCOE/CustomizableObjectEditorViewportClient.cpp
@@ -1637,9 +1637,9 @@ void FCustomizableObjectEditorViewportClient::BakeInstance()
 }
 
 
-void FCustomizableObjectEditorViewportClient::OnCOForBakingCompilationEnd(const FCompileCallbackParams& CompileCallbackParams)
+void FCustomizableObjectEditorViewportClient::OnCOForBakingCompilationEnd(const FCompileCallbackParams& InCompileCallbackParams)
 {
-	if (CompileCallbackParams.bRequestFailed)
+	if (InCompileCallbackParams.bRequestFailed || InCompileCallbackParams.bErrors)
 	{
 		UE_LOG(LogMutable, Warning ,TEXT("The Customizable Object failed the compilation. Skipping instance baking"));
 		return;
diff --git a/Engine/Plugins/Mutable/Source/CustomizableObjectEditor/Private/MuCOE/CustomizableObjectEditorViewportClient.h b/Engine/Plugins/Mutable/Source/CustomizableObjectEditor/Private/MuCOE/CustomizableObjectEditorViewportClient.h
--- a/Engine/Plugins/Mutable/Source/CustomizableObjectEditor/Private/MuCOE/CustomizableObjectEditorViewportClient.h
+++ b/Engine/Plugins/Mutable/Source/CustomizableObjectEditor/Private/MuCOE/CustomizableObjectEditorViewportClient.h
@@ -138,9 +138,9 @@ public:
 private:
 	/**
 	 * Callback executed after the Customizable Object instance for baking had it's customizable object compiled
-	 * @param CompileCallbackParams Data returned on compilation that includes the end state of the compilation
+	 * @param InCompileCallbackParams Data returned on compilation that includes the end state of the compilation
 	 */
-	void OnCOForBakingCompilationEnd(const FCompileCallbackParams& CompileCallbackParams);
+	void OnCOForBakingCompilationEnd(const FCompileCallbackParams& InCompileCallbackParams);
 	
 	/**
 	 * Callback executed after the instance in the editor gets updated for baking its contents.
diff --git a/Engine/Plugins/Mutable/Source/CustomizableObjectEditor/Private/MuCOE/CustomizableObjectInstanceBaker.cpp b/Engine/Plugins/Mutable/Source/CustomizableObjectEditor/Private/MuCOE/CustomizableObjectInstanceBaker.cpp
--- a/Engine/Plugins/Mutable/Source/CustomizableObjectEditor/Private/MuCOE/CustomizableObjectInstanceBaker.cpp
+++ b/Engine/Plugins/Mutable/Source/CustomizableObjectEditor/Private/MuCOE/CustomizableObjectInstanceBaker.cpp
@@ -61,7 +61,7 @@ void UCustomizableObjectInstanceBaker::BakeInstance(UCustomizableObjectInstance*
 
 void UCustomizableObjectInstanceBaker::OnCompilationEnd(const FCompileCallbackParams& InCompileCallbackParams)
 {
-	if (InCompileCallbackParams.bRequestFailed)
+	if (InCompileCallbackParams.bRequestFailed || InCompileCallbackParams.bErrors)
 	{
 		// Exit if the compilation failed
 		UE_LOG(LogMutable, Display, TEXT("The Customizable Object compilation for baking failed. Aborting instance baking operation."));
diff --git a/Engine/Plugins/Mutable/Source/MutableValidation/Private/MuV/COIBakingTestCommandlet.cpp b/Engine/Plugins/Mutable/Source/MutableValidation/Private/MuV/COIBakingTestCommandlet.cpp
--- a/Engine/Plugins/Mutable/Source/MutableValidation/Private/MuV/COIBakingTestCommandlet.cpp
+++ b/Engine/Plugins/Mutable/Source/MutableValidation/Private/MuV/COIBakingTestCommandlet.cpp
@@ -39,10 +39,12 @@ void OnCustomizableObjectCompilationEnd(const FCompileCallbackParams& Compilatio
 		CompilationParameters.bSkipped,
 		CompilationParameters.bCompiled);
 
+	// You can ignore this check if running this locally while trying to debug stuff.
+	// It is possible that you already have the CO compiled with the image compression settings required
 	check(!CompilationParameters.bSkipped);
 	
-	bWasCustomizableObjectCompilationSuccessful = !CompilationParameters.bRequestFailed;
-	
+	bWasCustomizableObjectCompilationSuccessful = !(CompilationParameters.bRequestFailed || CompilationParameters.bErrors);
+
 	ensure(bIsCustomizableObjectBeingCompiled);
 	bIsCustomizableObjectBeingCompiled = false;
 }
