[Mutable] Removed LazyAddress form inside FImage.

This was causing a large patch size.

diff --git a/Engine/Plugins/Mutable/Source/CustomizableObject/Internal/MuCO/CustomizableObjectPrivate.h b/Engine/Plugins/Mutable/Source/CustomizableObject/Internal/MuCO/CustomizableObjectPrivate.h
--- a/Engine/Plugins/Mutable/Source/CustomizableObject/Internal/MuCO/CustomizableObjectPrivate.h
+++ b/Engine/Plugins/Mutable/Source/CustomizableObject/Internal/MuCO/CustomizableObjectPrivate.h
@@ -1697,6 +1697,8 @@ enum class ECustomizableObjectVersions : int32
 
 	MaterialBreakCheckFix,
 
+	LazyAddressOutsideFImage,
+
 	// -----<new versions can be added above this line>--------
 	LastCustomizableObjectVersion
 };
diff --git a/Engine/Plugins/Mutable/Source/MutableRuntime/Internal/MuR/Image.h b/Engine/Plugins/Mutable/Source/MutableRuntime/Internal/MuR/Image.h
--- a/Engine/Plugins/Mutable/Source/MutableRuntime/Internal/MuR/Image.h
+++ b/Engine/Plugins/Mutable/Source/MutableRuntime/Internal/MuR/Image.h
@@ -132,9 +132,7 @@ namespace mu
 
 		/** Only valid if the right flags are set, this identifies a referenced image. */
 		uint32 ReferenceID = 0;
-
-		OP::ADDRESS LazyAddress = 0;
-
+    	
 		/** Pixel data for all lods. */
 		FImageDataStorage DataStorage;
 
@@ -154,7 +152,6 @@ namespace mu
 			Result->RelevancyMaxY = RelevancyMaxY;
 			Result->DataStorage = DataStorage;
 			Result->ReferenceID = ReferenceID;
-			Result->LazyAddress = LazyAddress;
 
 			return Result;
 		}
@@ -175,7 +172,6 @@ namespace mu
 			RelevancyMaxY = Other->RelevancyMaxY;
 			DataStorage = Other->DataStorage;
 			ReferenceID = Other->ReferenceID;
-			LazyAddress = Other->LazyAddress;
 		}
 
 		void CopyMove(FImage* Other)
@@ -192,7 +188,6 @@ namespace mu
 			RelevancyMinY = Other->RelevancyMinY;
 			RelevancyMaxY = Other->RelevancyMaxY;
 			ReferenceID = Other->ReferenceID;
-			LazyAddress = Other->LazyAddress;
 
 			DataStorage = MoveTemp(Other->DataStorage);
 		}
@@ -209,8 +204,7 @@ namespace mu
 		{
 			return 
 				(DataStorage == Other.DataStorage) &&
-				(ReferenceID == Other.ReferenceID) &&
-				(LazyAddress == Other.LazyAddress);
+				(ReferenceID == Other.ReferenceID);
 		}
 
 		//-----------------------------------------------------------------------------------------
@@ -253,8 +247,6 @@ namespace mu
 		//! Get the rect inside the image bounding the non-black content of the image.
 		UE_API void GetNonBlackRect(FImageRect& OutRect) const;
 		UE_API void GetNonBlackRect_Reference(FImageRect& OutRect) const;
-
-		UE_API bool IsLazy() const;
 	};
 
 	/** This struct contains image operations that may need to allocate and free temporary images to work.
diff --git a/Engine/Plugins/Mutable/Source/MutableRuntime/Internal/MuR/ParametersPrivate.h b/Engine/Plugins/Mutable/Source/MutableRuntime/Internal/MuR/ParametersPrivate.h
--- a/Engine/Plugins/Mutable/Source/MutableRuntime/Internal/MuR/ParametersPrivate.h
+++ b/Engine/Plugins/Mutable/Source/MutableRuntime/Internal/MuR/ParametersPrivate.h
@@ -3,6 +3,7 @@
 #pragma once
 
 #include "Misc/TVariant.h"
+#include "MuR/TVariant.h"
 #include "Math/MathFwd.h"
 
 #include "MuR/Parameters.h"
@@ -267,25 +268,6 @@ namespace Private
     // static_assert to track ParameterValue size changes. It is ok to change if needed.
     static_assert(sizeof(FParameterValue) == 8*4, "ParameterValue size has changed.");
 
-	// TVariant currently does not support this operator. Once supported remove it.
-	inline bool operator==(const FParameterValue& ValueA, const FParameterValue& ValueB)
-	{
-		const int32 IndexA = ValueA.GetIndex();
-		const int32 IndexB = ValueB.GetIndex();
-
-		if (IndexA != IndexB)
-		{
-			return false;
-		}
-
-		return Visit([&ValueB](const auto& StoredValueA)
-		{
-			using Type = typename TDecay<decltype(StoredValueA)>::Type;
-			return StoredValueA == ValueB.Get<Type>();
-		}, ValueA);
-	}
-	
-
     struct FParameterDesc
     {
         FString Name;
diff --git a/Engine/Plugins/Mutable/Source/MutableRuntime/Internal/MuR/TVariant.h b/Engine/Plugins/Mutable/Source/MutableRuntime/Internal/MuR/TVariant.h
new file mode 100644
--- /dev/null
+++ b/Engine/Plugins/Mutable/Source/MutableRuntime/Internal/MuR/TVariant.h
@@ -0,0 +1,29 @@
+// Copyright Epic Games, Inc. All Rights Reserved.
+
+#pragma once
+
+#include "Misc/TVariant.h"
+#include "Templates/Decay.h"
+
+
+namespace mu
+{
+	// TVariant currently does not support this operator. Once supported remove it.
+	template<typename ...T>
+	bool operator==(const TVariant<T...>& ValueA, const TVariant<T...>& ValueB)
+	{
+		const int32 IndexA = ValueA.GetIndex();
+		const int32 IndexB = ValueB.GetIndex();
+
+		if (IndexA != IndexB)
+		{
+			return false;
+		}
+
+		return Visit([&ValueB](const auto& StoredValueA)
+		{
+			using Type = typename TDecay<decltype(StoredValueA)>::Type;
+			return StoredValueA == ValueB.template Get<Type>();
+		}, ValueA);
+	}
+}
\ No newline at end of file
diff --git a/Engine/Plugins/Mutable/Source/MutableRuntime/Private/MuR/CodeRunner.cpp b/Engine/Plugins/Mutable/Source/MutableRuntime/Private/MuR/CodeRunner.cpp
--- a/Engine/Plugins/Mutable/Source/MutableRuntime/Private/MuR/CodeRunner.cpp
+++ b/Engine/Plugins/Mutable/Source/MutableRuntime/Private/MuR/CodeRunner.cpp
@@ -5512,7 +5512,7 @@ namespace mu
 
 				check(Args.ParameterName < (uint32)Program.ConstantStrings.Num());
 				const FString& ParameterName = Program.ConstantStrings[Args.ParameterName];
-				const TSharedPtr<const FImage>* Image = Material->ImageParameters.Find(FName(ParameterName));
+				const TVariant<OP::ADDRESS, TSharedPtr<const FImage>>* Image = Material->ImageParameters.Find(FName(ParameterName));
 				
 				if (!Image)
 				{
@@ -5520,17 +5520,17 @@ namespace mu
 				}
 				else
 				{
-					if (Image->Get()->IsLazy())
+					if (Image->IsType<OP::ADDRESS>())
 					{
 						FScheduledOp ItemNextStage(Item.At, Item, 2);
-						ItemNextStage.CustomState = Image->Get()->LazyAddress;
+						ItemNextStage.CustomState = Image->Get<OP::ADDRESS>();
 
 						AddOp(ItemNextStage,
-							FScheduledOp(Image->Get()->LazyAddress, Item));
+							FScheduledOp(Image->Get<OP::ADDRESS>(), Item));
 					}
 					else
 					{
-						StoreImage(Item, *Image);
+						StoreImage(Item, Image->Get<TSharedPtr<const FImage>>());
 					}
 				}
 				break;
@@ -5541,9 +5541,9 @@ namespace mu
 
 				check(Args.ParameterName < (uint32)Program.ConstantStrings.Num());
 				const FString& ParameterName = Program.ConstantStrings[Args.ParameterName];
-				const TSharedPtr<const FImage>* Image = Material->ImageParameters.Find(FName(ParameterName));
+				const TVariant<OP::ADDRESS, TSharedPtr<const FImage>>* Image = Material->ImageParameters.Find(FName(ParameterName));
 
-				TSharedPtr<const FImage> pImage = LoadImage(FCacheAddress(Image->Get()->LazyAddress, Item));
+				TSharedPtr<const FImage> pImage = LoadImage(FCacheAddress(Image->Get<OP::ADDRESS>(), Item));
 				StoreImage(Item, pImage);
 				break;
 			}
@@ -6446,10 +6446,9 @@ namespace mu
 
 					if (bParameterFound)
 					{
-						TSharedPtr<FImage> NewParameterImage = MakeShared<FImage>();
-						NewParameterImage->Texture = TStrongObjectPtr(Texture);
-						NewParameterImage->LazyAddress = ImageParameterOperations[ParameterIndex];
-
+						TVariant<OP::ADDRESS, TSharedPtr<const FImage>> NewParameterImage;
+						NewParameterImage.Set<OP::ADDRESS>(ImageParameterOperations[ParameterIndex]);
+						
 						Material->ImageParameters.Add(FName(ParameterName), NewParameterImage);
 					}
 				}
@@ -7879,7 +7878,7 @@ namespace mu
 				{
 					check(Args.ParameterName < (uint32)Program.ConstantStrings.Num());
 					const FString& ParameterName = Program.ConstantStrings[Args.ParameterName];
-					const TSharedPtr<const FImage>* Image = Material->ImageParameters.Find(FName(ParameterName));
+					const TVariant<OP::ADDRESS, TSharedPtr<const FImage>>* Image = Material->ImageParameters.Find(FName(ParameterName));
 
 					if (!Image)
 					{
@@ -7888,13 +7887,13 @@ namespace mu
 					}
 					else
 					{
-						if (Image->Get()->IsLazy())
+						if (Image->IsType<OP::ADDRESS>())
 						{
 							FScheduledOp ItemNextStage(Item.At, Item, 2);
-							ItemNextStage.CustomState = Image->Get()->LazyAddress;
+							ItemNextStage.CustomState = Image->Get<OP::ADDRESS>();
 
 							AddOp(ItemNextStage,
-								FScheduledOp(Image->Get()->LazyAddress, Item));
+								FScheduledOp(Image->Get<OP::ADDRESS>(), Item));
 						}
 						else
 						{
@@ -7921,10 +7920,10 @@ namespace mu
 				{
 					check(Args.ParameterName < (uint32)Program.ConstantStrings.Num());
 					const FString& ParameterName = Program.ConstantStrings[Args.ParameterName];
-					const TSharedPtr<const FImage>* Image = Material->ImageParameters.Find(FName(ParameterName));
+					const TVariant<OP::ADDRESS, TSharedPtr<const FImage>>* Image = Material->ImageParameters.Find(FName(ParameterName));
 
 					FExtendedImageDesc& Result = ImageDescResults.FindOrAdd(Item.At);
-					Result = ImageDescResults.FindChecked(Image->Get()->LazyAddress);
+					Result = ImageDescResults.FindChecked(Image->Get<OP::ADDRESS>());
 				}
 
 				StoreValidDesc(Item);
diff --git a/Engine/Plugins/Mutable/Source/MutableRuntime/Private/MuR/Image.cpp b/Engine/Plugins/Mutable/Source/MutableRuntime/Private/MuR/Image.cpp
--- a/Engine/Plugins/Mutable/Source/MutableRuntime/Private/MuR/Image.cpp
+++ b/Engine/Plugins/Mutable/Source/MutableRuntime/Private/MuR/Image.cpp
@@ -72,14 +72,12 @@ namespace mu
 		// Remove non-persistent flags.
 		uint8 flags = Flags & ~IF_HAS_RELEVANCY_MAP;
 		Arch << flags;
-		Arch << LazyAddress;
 	}
 
 	void FImage::Unserialise(FInputArchive& Arch)
 	{
 		Arch >> DataStorage;
 		Arch >> Flags;
-		Arch >> LazyAddress;
 	}
 
 
@@ -617,13 +615,7 @@ namespace mu
 		 Rect.size[0] = Right - Left + 1;
 		 Rect.size[1] = Bottom - Top + 1;
 	 }
-
-
-	bool FImage::IsLazy() const
-	{
-		return LazyAddress != 0;
-	}
-
+	
 
     //---------------------------------------------------------------------------------------------
     void FImage::GetNonBlackRect(FImageRect& rect) const
diff --git a/Engine/Plugins/Mutable/Source/MutableRuntime/Public/MuR/Material.h b/Engine/Plugins/Mutable/Source/MutableRuntime/Public/MuR/Material.h
--- a/Engine/Plugins/Mutable/Source/MutableRuntime/Public/MuR/Material.h
+++ b/Engine/Plugins/Mutable/Source/MutableRuntime/Public/MuR/Material.h
@@ -5,6 +5,8 @@
 #include "MuR/Ptr.h"
 #include "MuR/RefCounted.h"
 #include "MuR/Types.h"
+#include "MuR/Operations.h"
+#include "MuR/TVariant.h"
 
 #include "UObject/StrongObjectPtr.h"
 #include "Materials/MaterialInterface.h"
@@ -50,7 +52,7 @@ namespace mu
 
 		int32 ReferenceID = INDEX_NONE;
 
-		TMap<FName, TSharedPtr<const FImage>> ImageParameters;
+		TMap<FName, TVariant<OP::ADDRESS, TSharedPtr<const FImage>>> ImageParameters;
 		TMap<FName, FVector4f> ColorParameters;
 		TMap<FName, float> ScalarParameters;
 	};
diff --git a/Engine/Plugins/Mutable/Source/MutableTools/Private/MuT/ASTOpConstantResource.cpp b/Engine/Plugins/Mutable/Source/MutableTools/Private/MuT/ASTOpConstantResource.cpp
--- a/Engine/Plugins/Mutable/Source/MutableTools/Private/MuT/ASTOpConstantResource.cpp
+++ b/Engine/Plugins/Mutable/Source/MutableTools/Private/MuT/ASTOpConstantResource.cpp
@@ -682,8 +682,8 @@ namespace mu
 					// Store linked child images into the material constant
 					for (const TPair<FName, ASTChild>& ImageOperation : ImageOperations)
 					{
-						TSharedPtr<FImage> NewImage = MakeShared<FImage>();
-						NewImage->LazyAddress = ImageOperation.Value->linkedAddress;
+						TVariant<OP::ADDRESS, TSharedPtr<const FImage>> NewImage;
+						NewImage.Set<OP::ADDRESS>(ImageOperation.Value->linkedAddress);
 
 						MaterialValue->ImageParameters.Add(FName(ImageOperation.Key), NewImage);
 					}
