[Mutable] Fixed ambiguity with Resource ID hash

diff --git a/Engine/Plugins/Mutable/Source/CustomizableObject/Private/MuCO/CustomizableObjectInstance.cpp b/Engine/Plugins/Mutable/Source/CustomizableObject/Private/MuCO/CustomizableObjectInstance.cpp
--- a/Engine/Plugins/Mutable/Source/CustomizableObject/Private/MuCO/CustomizableObjectInstance.cpp
+++ b/Engine/Plugins/Mutable/Source/CustomizableObject/Private/MuCO/CustomizableObjectInstance.cpp
@@ -148,7 +148,7 @@ public:
 
 	friend uint32 GetTypeHash(const FMutableMaterialPlaceholder& PlaceHolder);
 	
-	friend uint32 GetHash(const FMutableMaterialPlaceholder& PlaceHolder, const bool bGenerateDeterministicHash);
+	friend uint32 GetHash(const FMutableMaterialPlaceholder& PlaceHolder, const bool bPersistent);
 };
 
 
@@ -186,19 +186,19 @@ uint32 GetTypeHash(const FMutableMaterialPlaceholder& PlaceHolder)
 }
 
 
-uint32 GetHash(const FMutableMaterialPlaceholder& PlaceHolder, bool bGenerateDeterministicHash)
+uint32 GetHash(const FMutableMaterialPlaceholder& PlaceHolder, bool bPersistent)
 {
-	static auto FastApproach = [](const uint32 HashA, const uint32 HashB) -> uint32
+	auto FastApproach = [](const uint32 HashA, const uint32 HashB) -> uint32
 	{
 		return HashCombineFast(HashA, HashB);
 	};
 
-	static auto BackwardsCompatibleApproach = [](const uint32 HashA, const uint32 HashB) -> uint32
+	auto BackwardsCompatibleApproach = [](const uint32 HashA, const uint32 HashB) -> uint32
 	{
 		return HashCombine(HashA,HashB);
 	};
 
-	auto HashingMethod = bGenerateDeterministicHash ? BackwardsCompatibleApproach : FastApproach;
+	auto HashingMethod = bPersistent ? BackwardsCompatibleApproach : FastApproach;
 	
 	uint32 Hash = GetTypeHash(PlaceHolder.ParentMaterialID);
 
@@ -2291,7 +2291,7 @@ bool UCustomizableInstancePrivate::UpdateSkeletalMesh_PostBeginUpdate0(UCustomiz
 		{
 			const mu::FGeneratedMeshKey* MeshKey = (*MeshDescriptors)[LODIndex].GetKey();
 			check(MeshKey);
-			const uint32 MeshIDHash = GetTypeHash(*MeshKey, Context->bBake);
+			const uint32 MeshIDHash = GetTypeHashPersistent(*MeshKey);
 
 			if (!Context->bBake)
 			{
@@ -6217,7 +6217,7 @@ void UCustomizableInstancePrivate::BuildMaterials(const TSharedRef<FUpdateContex
 
 								FString MutableTextureName;
 								{
-									const uint32 MutableImageHash = GetTypeHash(*Image.ImageID.GetKey(), Context->bBake);
+									const uint32 MutableImageHash = mu::GetTypeHashPersistent(*Image.ImageID.GetKey());
 									MutableTextureName = FString::Printf(TEXT("T_%s_%lu"), *Props.TextureParameterName, MutableImageHash);
 									MutableTextureName.ReplaceInline(TEXT(" "), TEXT("_"));
 
diff --git a/Engine/Plugins/Mutable/Source/MutableRuntime/Internal/MuR/ResourceID.h b/Engine/Plugins/Mutable/Source/MutableRuntime/Internal/MuR/ResourceID.h
--- a/Engine/Plugins/Mutable/Source/MutableRuntime/Internal/MuR/ResourceID.h
+++ b/Engine/Plugins/Mutable/Source/MutableRuntime/Internal/MuR/ResourceID.h
@@ -22,15 +22,11 @@ namespace mu
 	struct FGeneratedMeshData
 	{
 	};
+	
+	MUTABLERUNTIME_API uint32 GetTypeHash(const FGeneratedMeshKey& Key);
 
-	/**
-	 * Generates a hash based on the provided mesh key
-	 * @param Key The key you want to hash
-	 * @param bGenerateDeterministicHash Determines what hashing approach to follow.
-	 * @return The mesh key hash
-	 */
-	MUTABLERUNTIME_API uint32 GetTypeHash(const FGeneratedMeshKey& Key, const bool bGenerateDeterministicHash = false);
-
+	MUTABLERUNTIME_API uint32 GetTypeHashPersistent(const FGeneratedMeshKey& Key);
+	
 	typedef TRegistry<FGeneratedMeshKey, FGeneratedMeshData>::FHandle FMeshId;
 
 	typedef TRegistry<FGeneratedMeshKey, FGeneratedMeshData> FMeshIdRegistry;
@@ -50,13 +46,9 @@ namespace mu
 	{
 	};
 	
-	/**
-     * Generates a hash based on the provided image key
-     * @param Key The key you want to hash
-     * @param bGenerateDeterministicHash Determines what hashing approach to follow.
-     * @return The image key hash
-     */
-	MUTABLERUNTIME_API uint32 GetTypeHash(const FGeneratedImageKey& Key, const bool bGenerateDeterministicHash = false);
+	MUTABLERUNTIME_API uint32 GetTypeHash(const FGeneratedImageKey& Key);
+
+	MUTABLERUNTIME_API uint32 GetTypeHashPersistent(const FGeneratedImageKey& Key);
 
 	typedef TRegistry<FGeneratedImageKey, FGeneratedImageData>::FHandle FImageId;
 
@@ -79,6 +71,8 @@ namespace mu
 	
 	MUTABLERUNTIME_API uint32 GetTypeHash(const FGeneratedMaterialKey& Key);
 
+	MUTABLERUNTIME_API uint32 GetTypeHashPersistent(const FGeneratedMaterialKey& Key);
+	
 	typedef TRegistry<FGeneratedMaterialKey, FGeneratedMaterialData>::FHandle FMaterialId;
 
 	typedef TRegistry<FGeneratedMaterialKey, FGeneratedMaterialData> FMaterialIdRegistry;
diff --git a/Engine/Plugins/Mutable/Source/MutableRuntime/Private/MuR/ResourceID.cpp b/Engine/Plugins/Mutable/Source/MutableRuntime/Private/MuR/ResourceID.cpp
--- a/Engine/Plugins/Mutable/Source/MutableRuntime/Private/MuR/ResourceID.cpp
+++ b/Engine/Plugins/Mutable/Source/MutableRuntime/Private/MuR/ResourceID.cpp
@@ -5,45 +5,37 @@
 #include "Hash/xxhash.h"
 
 
-uint32 mu::GetTypeHash(const FGeneratedMeshKey& Key, const bool bGenerateDeterministicHash)
+uint32 mu::GetTypeHash(const FGeneratedMeshKey& Key)
 {
-	uint32 Hash = Key.Address;
-
-	if (!bGenerateDeterministicHash)
-	{
-		Hash = HashCombineFast(Hash, static_cast<uint32>(FXxHash64::HashBuffer(Key.ParameterValuesBlob.GetData(), Key.ParameterValuesBlob.Num()).Hash));
-	}
-	else
-	{
-		Hash = HashCombine(Hash, static_cast<uint32>(FXxHash64::HashBuffer(Key.ParameterValuesBlob.GetData(), Key.ParameterValuesBlob.Num()).Hash));
-	}
-
-	return Hash;
+	return HashCombineFast(Key.Address, static_cast<uint32>(FXxHash64::HashBuffer(Key.ParameterValuesBlob.GetData(), Key.ParameterValuesBlob.Num()).Hash));
 }
 
 
-uint32 mu::GetTypeHash(const FGeneratedImageKey& Key, const bool bGenerateDeterministicHash)
+uint32 mu::GetTypeHashPersistent(const FGeneratedMeshKey& Key)
 {
-	uint32 Hash = Key.Address;
-
-	if (!bGenerateDeterministicHash)
-	{
-		Hash = HashCombineFast(Hash, static_cast<uint32>(FXxHash64::HashBuffer(Key.ParameterValuesBlob.GetData(), Key.ParameterValuesBlob.Num()).Hash));
-	}
-	else
-	{
-		Hash = HashCombine(Hash, static_cast<uint32>(FXxHash64::HashBuffer(Key.ParameterValuesBlob.GetData(), Key.ParameterValuesBlob.Num()).Hash));
-	}
-
-	return Hash;
+	return HashCombine(Key.Address, static_cast<uint32>(FXxHash64::HashBuffer(Key.ParameterValuesBlob.GetData(), Key.ParameterValuesBlob.Num()).Hash));
 }
 
 
-uint32 mu::GetTypeHash(const FGeneratedMaterialKey& Key)
+uint32 mu::GetTypeHash(const FGeneratedImageKey& Key)
 {
-	uint32 Hash = Key.Address;
+	return HashCombineFast(Key.Address, static_cast<uint32>(FXxHash64::HashBuffer(Key.ParameterValuesBlob.GetData(), Key.ParameterValuesBlob.Num()).Hash));
+}
 
-	Hash = HashCombineFast(Hash, static_cast<uint32>(FXxHash64::HashBuffer(Key.ParameterValuesBlob.GetData(), Key.ParameterValuesBlob.Num()).Hash));
 
-	return Hash;
+uint32 mu::GetTypeHashPersistent(const FGeneratedImageKey& Key)
+{
+	return HashCombine(Key.Address, static_cast<uint32>(FXxHash64::HashBuffer(Key.ParameterValuesBlob.GetData(), Key.ParameterValuesBlob.Num()).Hash));
 }
+
+
+uint32 mu::GetTypeHash(const FGeneratedMaterialKey& Key)
+{
+	return HashCombineFast(Key.Address, static_cast<uint32>(FXxHash64::HashBuffer(Key.ParameterValuesBlob.GetData(), Key.ParameterValuesBlob.Num()).Hash));
+}
+
+
+uint32 mu::GetTypeHashPersistent(const FGeneratedMaterialKey& Key)
+{
+	return HashCombine(Key.Address, static_cast<uint32>(FXxHash64::HashBuffer(Key.ParameterValuesBlob.GetData(), Key.ParameterValuesBlob.Num()).Hash));
+}
\ No newline at end of file
