[Mutable] Change namespace name

Change namespace to follow UE standard. Also this change will allow to clearly differenciate between public and private namespaces.

mu -> UE::Mutable::Private
MutablePrivate -> UE::Mutable::Private

diff --git a/Engine/Plugins/Experimental/HairStrandsMutable/Source/HairStrandsMutableEditor/Private/CustomizableObjectNodeGroomConstant.cpp b/Engine/Plugins/Experimental/HairStrandsMutable/Source/HairStrandsMutableEditor/Private/CustomizableObjectNodeGroomConstant.cpp
--- a/Engine/Plugins/Experimental/HairStrandsMutable/Source/HairStrandsMutableEditor/Private/CustomizableObjectNodeGroomConstant.cpp
+++ b/Engine/Plugins/Experimental/HairStrandsMutable/Source/HairStrandsMutableEditor/Private/CustomizableObjectNodeGroomConstant.cpp
@@ -48,12 +48,12 @@ bool UCustomizableObjectNodeGroomConstant::IsExperimental() const
 }
 
 
-mu::Ptr<mu::NodeExtensionData> UCustomizableObjectNodeGroomConstant::GenerateMutableNode(FExtensionDataCompilerInterface& CompilerInterface) const
+UE::Mutable::Private::Ptr<UE::Mutable::Private::NodeExtensionData> UCustomizableObjectNodeGroomConstant::GenerateMutableNode(FExtensionDataCompilerInterface& CompilerInterface) const
 {
 	check(IsInGameThread());
 
 	// Create node and extension data container
-	mu::Ptr<mu::NodeExtensionDataConstant> Result = new mu::NodeExtensionDataConstant();
+	UE::Mutable::Private::Ptr<UE::Mutable::Private::NodeExtensionDataConstant> Result = new UE::Mutable::Private::NodeExtensionDataConstant();
 
 	FInstancedStruct Struct;
 	Struct.InitializeAs<FGroomPinData>(GroomData);
diff --git a/Engine/Plugins/Experimental/HairStrandsMutable/Source/HairStrandsMutableEditor/Public/CustomizableObjectNodeGroomConstant.h b/Engine/Plugins/Experimental/HairStrandsMutable/Source/HairStrandsMutableEditor/Public/CustomizableObjectNodeGroomConstant.h
--- a/Engine/Plugins/Experimental/HairStrandsMutable/Source/HairStrandsMutableEditor/Public/CustomizableObjectNodeGroomConstant.h
+++ b/Engine/Plugins/Experimental/HairStrandsMutable/Source/HairStrandsMutableEditor/Public/CustomizableObjectNodeGroomConstant.h
@@ -26,7 +26,7 @@ public:
 	virtual bool IsExperimental() const override;
 	
 	/** ICustomizableObjectExtensionNode interface */
-	virtual mu::Ptr<mu::NodeExtensionData> GenerateMutableNode(FExtensionDataCompilerInterface& CompilerInterface) const override;
+	virtual UE::Mutable::Private::Ptr<UE::Mutable::Private::NodeExtensionData> GenerateMutableNode(FExtensionDataCompilerInterface& CompilerInterface) const override;
 
 	UPROPERTY(EditAnywhere, Category = CustomizableObject)
 	FGroomPinData GroomData;
diff --git a/Engine/Plugins/Experimental/RigLogicMutable/Source/RigLogicMutableEditor/Private/CustomizableObjectNodeDNAConstant.cpp b/Engine/Plugins/Experimental/RigLogicMutable/Source/RigLogicMutableEditor/Private/CustomizableObjectNodeDNAConstant.cpp
--- a/Engine/Plugins/Experimental/RigLogicMutable/Source/RigLogicMutableEditor/Private/CustomizableObjectNodeDNAConstant.cpp
+++ b/Engine/Plugins/Experimental/RigLogicMutable/Source/RigLogicMutableEditor/Private/CustomizableObjectNodeDNAConstant.cpp
@@ -50,12 +50,12 @@ bool UCustomizableObjectNodeDNAConstant::IsExperimental() const
 }
 
 
-mu::Ptr<mu::NodeExtensionData> UCustomizableObjectNodeDNAConstant::GenerateMutableNode(FExtensionDataCompilerInterface& CompilerInterface) const
+UE::Mutable::Private::Ptr<UE::Mutable::Private::NodeExtensionData> UCustomizableObjectNodeDNAConstant::GenerateMutableNode(FExtensionDataCompilerInterface& CompilerInterface) const
 {
 	check(IsInGameThread());
 	
 	// Create node and extension data container
-	mu::Ptr<mu::NodeExtensionDataConstant> Result = new mu::NodeExtensionDataConstant();
+	UE::Mutable::Private::Ptr<UE::Mutable::Private::NodeExtensionDataConstant> Result = new UE::Mutable::Private::NodeExtensionDataConstant();
 
 	FDNAPinData PinData;
 	PinData.ComponentName = ComponentName;
diff --git a/Engine/Plugins/Experimental/RigLogicMutable/Source/RigLogicMutableEditor/Public/CustomizableObjectNodeDNAConstant.h b/Engine/Plugins/Experimental/RigLogicMutable/Source/RigLogicMutableEditor/Public/CustomizableObjectNodeDNAConstant.h
--- a/Engine/Plugins/Experimental/RigLogicMutable/Source/RigLogicMutableEditor/Public/CustomizableObjectNodeDNAConstant.h
+++ b/Engine/Plugins/Experimental/RigLogicMutable/Source/RigLogicMutableEditor/Public/CustomizableObjectNodeDNAConstant.h
@@ -25,7 +25,7 @@ public:
 	virtual bool IsExperimental() const override;
 
 	/** ICustomizableObjectExtensionNode interface */
-	virtual mu::Ptr<mu::NodeExtensionData> GenerateMutableNode(FExtensionDataCompilerInterface& CompilerInterface) const override;
+	virtual UE::Mutable::Private::Ptr<UE::Mutable::Private::NodeExtensionData> GenerateMutableNode(FExtensionDataCompilerInterface& CompilerInterface) const override;
 
 	/** The Skeletal Mesh to copy DNA from */
 	UPROPERTY(EditAnywhere, Category = CustomizableObject)
diff --git a/Engine/Plugins/Mutable/Source/CustomizableObject/Internal/MuCO/BoneNames.h b/Engine/Plugins/Mutable/Source/CustomizableObject/Internal/MuCO/BoneNames.h
--- a/Engine/Plugins/Mutable/Source/CustomizableObject/Internal/MuCO/BoneNames.h
+++ b/Engine/Plugins/Mutable/Source/CustomizableObject/Internal/MuCO/BoneNames.h
@@ -12,12 +12,12 @@ class FBoneNames
 public:
 	CUSTOMIZABLEOBJECT_API FBoneNames(const UModelResources& ModelResources);
 
-	mu::FBoneName* Find(const FName& BoneName);
+	UE::Mutable::Private::FBoneName* Find(const FName& BoneName);
 	
-	mu::FBoneName FindOrAdd(const FName& BoneName);
+	UE::Mutable::Private::FBoneName FindOrAdd(const FName& BoneName);
 	
 private:
-	TMap<FString, mu::FBoneName> BoneNamesMap;
+	TMap<FString, UE::Mutable::Private::FBoneName> BoneNamesMap;
 
 	FCriticalSection CriticalSection;
 };
diff --git a/Engine/Plugins/Mutable/Source/CustomizableObject/Internal/MuCO/CustomizableObjectInstancePrivate.h b/Engine/Plugins/Mutable/Source/CustomizableObject/Internal/MuCO/CustomizableObjectInstancePrivate.h
--- a/Engine/Plugins/Mutable/Source/CustomizableObject/Internal/MuCO/CustomizableObjectInstancePrivate.h
+++ b/Engine/Plugins/Mutable/Source/CustomizableObject/Internal/MuCO/CustomizableObjectInstancePrivate.h
@@ -17,7 +17,7 @@
 
 #define UE_API CUSTOMIZABLEOBJECT_API
 
-namespace mu 
+namespace UE::Mutable::Private 
 {
 	class FPhysicsBody;
 	class FMesh;
@@ -49,7 +49,7 @@ CUSTOMIZABLEOBJECT_API extern const FString POSE_PARAMETER_POSTFIX;
 
 /** \param OnlyLOD: If not 0, extract and convert only one single LOD from the source image.
   * \param ExtractChannel: If different than -1, extract a single-channel image with the specified source channel data. */
-CUSTOMIZABLEOBJECT_API void ConvertImage(UTexture2D* Texture, TSharedPtr<const mu::FImage> MutableImage, const FMutableModelImageProperties& Props, int32 OnlyLOD = -1, int32 ExtractChannel = -1);
+CUSTOMIZABLEOBJECT_API void ConvertImage(UTexture2D* Texture, TSharedPtr<const UE::Mutable::Private::FImage> MutableImage, const FMutableModelImageProperties& Props, int32 OnlyLOD = -1, int32 ExtractChannel = -1);
 
 
 /** CustomizableObject Instance flags for internal use  */
@@ -147,7 +147,7 @@ struct FCustomizableInstanceComponentData
 	/** Array of generated MeshIds per each LOD, used to decide if the mesh should be updated or not.
 	 *  Size == NumLODsAvailable
 	 *  LODs without mesh will be set to the maximum value of FResourceID (Max_uint64). */
-	TArray<mu::FMeshId> LastMeshIdPerLOD;
+	TArray<UE::Mutable::Private::FMeshId> LastMeshIdPerLOD;
 
 	UPROPERTY(Transient)
 	TArray<TObjectPtr<UMaterialInterface>> OverrideMaterials;
@@ -262,7 +262,7 @@ public:
 	TMap<FTextureReuseCacheKey, TWeakObjectPtr<UTexture2D>> TextureReuseCache;
 	
 	// Only used in LiveUpdateMode to reuse core instances between updates and their temp data to speed up updates, but spend way more memory
-	mu::FInstance::FID LiveUpdateModeInstanceID = 0;
+	UE::Mutable::Private::FInstance::FID LiveUpdateModeInstanceID = 0;
 
 #if WITH_EDITOR
 	UE_API virtual void PostDuplicate(bool bDuplicateForPIE) override;
@@ -443,7 +443,7 @@ public:
 };
 
 
-UE_API UPhysicsAsset* GetOrBuildMainPhysicsAsset(const TSharedRef<FUpdateContextPrivate>& Context, TObjectPtr<UPhysicsAsset> TamplateAsset, const mu::FPhysicsBody* PhysicsBody, bool bDisableCollisionBetweenAssets, FCustomizableObjectInstanceComponentIndex InstanceComponentIndex);
+UE_API UPhysicsAsset* GetOrBuildMainPhysicsAsset(const TSharedRef<FUpdateContextPrivate>& Context, TObjectPtr<UPhysicsAsset> TamplateAsset, const UE::Mutable::Private::FPhysicsBody* PhysicsBody, bool bDisableCollisionBetweenAssets, FCustomizableObjectInstanceComponentIndex InstanceComponentIndex);
 
 UE_API USkeleton* MergeSkeletons(UCustomizableObject& CustomizableObject, FCustomizableInstanceComponentData& ObjectComponentIndex, bool& bOutCreatedNewSkeleton);
 
@@ -453,7 +453,7 @@ UE_API void BuildOrCopyClothingData(const TSharedRef<FUpdateContextPrivate>& Con
 
 UE_API bool BuildSkeletonData(const TSharedRef<FUpdateContextPrivate>& Context, USkeletalMesh& SkeletalMesh, FCustomizableObjectInstanceComponentIndex InstanceComponentIndex);
 
-UE_API void BuildMeshSockets(USkeletalMesh* SkeletalMesh, const UModelResources& ModelResources, const FMutableRefSkeletalMeshData& RefSkeletalMeshData, TSharedPtr<const mu::FMesh> MutableMesh);
+UE_API void BuildMeshSockets(USkeletalMesh* SkeletalMesh, const UModelResources& ModelResources, const FMutableRefSkeletalMeshData& RefSkeletalMeshData, TSharedPtr<const UE::Mutable::Private::FMesh> MutableMesh);
 
 UE_API bool BuildOrCopyRenderData(const TSharedRef<FUpdateContextPrivate>& Context, USkeletalMesh* SkeletalMesh, UCustomizableObjectInstance* CustomizableObjectInstance, FCustomizableObjectInstanceComponentIndex InstanceComponentIndex);
 
diff --git a/Engine/Plugins/Mutable/Source/CustomizableObject/Internal/MuCO/CustomizableObjectMipDataProvider.h b/Engine/Plugins/Mutable/Source/CustomizableObject/Internal/MuCO/CustomizableObjectMipDataProvider.h
--- a/Engine/Plugins/Mutable/Source/CustomizableObject/Internal/MuCO/CustomizableObjectMipDataProvider.h
+++ b/Engine/Plugins/Mutable/Source/CustomizableObject/Internal/MuCO/CustomizableObjectMipDataProvider.h
@@ -15,8 +15,8 @@
 #include "CustomizableObjectMipDataProvider.generated.h"
 
 enum EPixelFormat : uint8;
-namespace mu { class FParameters; }
-namespace mu { class System; }
+namespace UE::Mutable::Private { class FParameters; }
+namespace UE::Mutable::Private { class System; }
 
 class FThreadSafeCounter;
 class UCustomizableObjectInstance;
@@ -41,21 +41,21 @@ public:
 
 	FName CustomizableObjectName;
 
-	TSharedPtr<mu::FSystem> System;
-	TSharedPtr<mu::FModel> Model;
+	TSharedPtr<UE::Mutable::Private::FSystem> System;
+	TSharedPtr<UE::Mutable::Private::FModel> Model;
 
-	TSharedPtr<mu::FMeshIdRegistry> MeshIdRegistry;
-	TSharedPtr<mu::FImageIdRegistry> ImageIdRegistry;
-	TSharedPtr<mu::FMaterialIdRegistry> MaterialIdRegistry;
+	TSharedPtr<UE::Mutable::Private::FMeshIdRegistry> MeshIdRegistry;
+	TSharedPtr<UE::Mutable::Private::FImageIdRegistry> ImageIdRegistry;
+	TSharedPtr<UE::Mutable::Private::FMaterialIdRegistry> MaterialIdRegistry;
 
 	TSharedPtr<FUnrealMutableResourceProvider> ExternalResourceProvider;
 
 	TSharedPtr<FModelStreamableBulkData> ModelStreamableBulkData;
 	
-	TSharedPtr<const mu::FParameters> Parameters;
+	TSharedPtr<const UE::Mutable::Private::FParameters> Parameters;
 	int32 State = -1;
 
-	TArray<TSharedPtr<const mu::FImage>> ImageParameterValues;
+	TArray<TSharedPtr<const UE::Mutable::Private::FImage>> ImageParameterValues;
 };
 
 
@@ -72,7 +72,7 @@ struct FMutableMipUpdateLevel
 	EPixelFormat Format;
 };
 
-namespace mu::MemoryCounters
+namespace UE::Mutable::Private::MemoryCounters
 {
 	struct FPrefetchMemoryCounter
 	{
@@ -93,13 +93,13 @@ struct FMutableImageOperationData
 
 	TSharedPtr<FMutableUpdateImageContext> UpdateContext;
 
-	TSharedPtr<const mu::FImage> Result;
+	TSharedPtr<const UE::Mutable::Private::FImage> Result;
 
 	TArray<FMutableMipUpdateLevel> Levels;
 
 	uint32 MutableTaskId = FMutableTaskGraph::INVALID_ID;
 
-	using FPrefetchArray = TArray<uint8, mu::FDefaultMemoryTrackingAllocator<mu::MemoryCounters::FPrefetchMemoryCounter>>;
+	using FPrefetchArray = TArray<uint8, UE::Mutable::Private::FDefaultMemoryTrackingAllocator<UE::Mutable::Private::MemoryCounters::FPrefetchMemoryCounter>>;
 	FPrefetchArray AllocatedMemory;
 
 	// Used to sync with the FMutableTextureMipDataProvider and FRenderAssetUpdate::Tick
diff --git a/Engine/Plugins/Mutable/Source/CustomizableObject/Internal/MuCO/CustomizableObjectPrivate.h b/Engine/Plugins/Mutable/Source/CustomizableObject/Internal/MuCO/CustomizableObjectPrivate.h
--- a/Engine/Plugins/Mutable/Source/CustomizableObject/Internal/MuCO/CustomizableObjectPrivate.h
+++ b/Engine/Plugins/Mutable/Source/CustomizableObject/Internal/MuCO/CustomizableObjectPrivate.h
@@ -25,14 +25,14 @@
 
 class UCustomizableObjectSkeletalMesh;
 
-namespace mu 
+namespace UE::Mutable::Private 
 { 
 	class FModel;
 	struct FBoneName;
 }
 
 #if WITH_EDITOR
-namespace MutablePrivate
+namespace UE::Mutable::Private
 {
 	struct FClassifyNode;
 }
@@ -257,12 +257,12 @@ struct FMutableModelParameterProperties
 class FMeshCache
 {
 public:
-	UCustomizableObjectSkeletalMesh* Get(const TArray<mu::FMeshId>& Key);
+	UCustomizableObjectSkeletalMesh* Get(const TArray<UE::Mutable::Private::FMeshId>& Key);
 
-	void Add(const TArray<mu::FMeshId>& Key, UCustomizableObjectSkeletalMesh* Value);
+	void Add(const TArray<UE::Mutable::Private::FMeshId>& Key, UCustomizableObjectSkeletalMesh* Value);
 
 private:
-	using FMeshCacheRegistry = TRegistry<TArray<mu::FMeshId>, TWeakObjectPtr<UCustomizableObjectSkeletalMesh>>;
+	using FMeshCacheRegistry = TRegistry<TArray<UE::Mutable::Private::FMeshId>, TWeakObjectPtr<UCustomizableObjectSkeletalMesh>>;
 	
 	TSharedPtr<FMeshCacheRegistry> GeneratedMeshes = MakeShared<FMeshCacheRegistry>();
 
@@ -277,7 +277,7 @@ class FTextureCache
 public:
 	struct FId
 	{
-		mu::FImageId Resource;
+		UE::Mutable::Private::FImageId Resource;
 
 		int32 SkippedMips = 0;
 
@@ -910,7 +910,7 @@ enum class EMutableFileFlags : uint8
 };
 
 
-namespace MutablePrivate
+namespace UE::Mutable::Private
 {
 	enum class EStreamableDataType : uint32 // uint32 for padding and DDC purposes
 	{
@@ -1062,8 +1062,8 @@ namespace MutablePrivate
 
 	struct FMutableCachedPlatformData
 	{
-		/** mu::Model */
-		TSharedPtr<mu::FModel, ESPMode::ThreadSafe> Model;
+		/** UE::Mutable::Private::Model */
+		TSharedPtr<UE::Mutable::Private::FModel, ESPMode::ThreadSafe> Model;
 		
 		/** UModelResources */
 		TStrongObjectPtr<UModelResources> ModelResources;
@@ -1089,7 +1089,7 @@ namespace MutablePrivate
 	 *  Resources will be split into two buckets for non-optional and optional BulkData.
 	 */
 	void CUSTOMIZABLEOBJECT_API GenerateBulkDataFilesListWithFileLimit(
-		TSharedPtr<const mu::FModel, ESPMode::ThreadSafe> Model,
+		TSharedPtr<const UE::Mutable::Private::FModel, ESPMode::ThreadSafe> Model,
 		FModelStreamableBulkData& ModelStreamableBulkData,
 		uint32 NumFilesPerBucket,
 		bool bAllowSplit,
@@ -1098,7 +1098,7 @@ namespace MutablePrivate
 	/** Generate the list of BulkData files with a soft restriction to the size of the files.
 	 */
 	void CUSTOMIZABLEOBJECT_API GenerateBulkDataFilesListWithSizeLimit(
-		TSharedPtr<const mu::FModel, ESPMode::ThreadSafe> Model,
+		TSharedPtr<const UE::Mutable::Private::FModel, ESPMode::ThreadSafe> Model,
 		FModelStreamableBulkData& ModelStreamableBulkData,
 		const ITargetPlatform* TargetPlatform,
 		uint64 TargetBulkDataFileBytes,
@@ -1108,7 +1108,7 @@ namespace MutablePrivate
 	 * the streamable's FileIds and Offsets.
 	 */
 	void GenerateBulkDataFilesList(
-		TSharedPtr<const mu::FModel, ESPMode::ThreadSafe> Model,
+		TSharedPtr<const UE::Mutable::Private::FModel, ESPMode::ThreadSafe> Model,
 		FModelStreamableBulkData& StreamableBulkData,
 		bool bUseRomTypeAndFlagsToFilter,
 		TFunctionRef<void(const FFileCategoryID&, const FClassifyNode&, TArray<FFile>&)> CreateFileList,
@@ -1212,28 +1212,28 @@ public:
 	UPROPERTY()
 	TArray<FMutableRefSkeletalMeshData> ReferenceSkeletalMeshesData;	
 
-	/** Skeletons used by the compiled mu::FModel. */
+	/** Skeletons used by the compiled UE::Mutable::Private::FModel. */
 	UPROPERTY()
 	TArray<TSoftObjectPtr<USkeleton>> Skeletons;
 
-	/** Materials used by the compiled mu::FModel */
+	/** Materials used by the compiled UE::Mutable::Private::FModel */
 	UPROPERTY()
 	TArray<TSoftObjectPtr<UMaterialInterface>> Materials;
 
-	/** PassThrough textures used by the mu::FModel. */
+	/** PassThrough textures used by the UE::Mutable::Private::FModel. */
 	UPROPERTY()
 	TArray<TSoftObjectPtr<UTexture>> PassThroughTextures;
 
-	/** PassThrough meshes used by the mu::FModel. */
+	/** PassThrough meshes used by the UE::Mutable::Private::FModel. */
 	UPROPERTY()
 	TArray<TSoftObjectPtr<UStreamableRenderAsset>> PassThroughMeshes;
 
 #if WITH_EDITORONLY_DATA
-	/** Runtime referenced textures used by the mu::FModel. */
+	/** Runtime referenced textures used by the UE::Mutable::Private::FModel. */
 	UPROPERTY()
 	TArray<TSoftObjectPtr<UTexture2D>> RuntimeReferencedTextures;
 	
-	/** Runtime referenced meshes used by the mu::FModel.
+	/** Runtime referenced meshes used by the UE::Mutable::Private::FModel.
 	* TODO: Move to FMutableSourceMeshData when runtime is really implemented.
 	*/
 	UPROPERTY()
@@ -1340,11 +1340,11 @@ public:
 	UPROPERTY()
 	TArray<FCustomizableObjectStreamedResourceData> StreamedResourceData;
 
-	// mu::FExtensionData::Index is an index into this array when mu::ExtensionData::Origin is ConstantAlwaysLoaded
+	// UE::Mutable::Private::FExtensionData::Index is an index into this array when UE::Mutable::Private::ExtensionData::Origin is ConstantAlwaysLoaded
 	UPROPERTY()
 	TArray<FCustomizableObjectResourceData> AlwaysLoadedExtensionData;
 
-	// mu::FExtensionData::Index is an index into this array when mu::ExtensionData::Origin is ConstantStreamed
+	// UE::Mutable::Private::FExtensionData::Index is an index into this array when UE::Mutable::Private::ExtensionData::Origin is ConstantStreamed
 	UPROPERTY()
 	TArray<FCustomizableObjectStreamedResourceData> StreamedExtensionData;
 	
@@ -1710,7 +1710,7 @@ class UCustomizableObjectPrivate : public UObject
 {
 	GENERATED_BODY()
 
-	TSharedPtr<mu::FModel, ESPMode::ThreadSafe> MutableModel;
+	TSharedPtr<UE::Mutable::Private::FModel, ESPMode::ThreadSafe> MutableModel;
 
 	/** Stores streamable data info to be used by MutableModel In-Game. Cooked resources. */
 	UPROPERTY()
@@ -1736,9 +1736,9 @@ class UCustomizableObjectPrivate : public UObject
 
 public:
 	/** Must be called after unlocking the CustomizableObject. */
-	CUSTOMIZABLEOBJECT_API void SetModel(const TSharedPtr<mu::FModel, ESPMode::ThreadSafe>& Model, const FGuid Identifier);
-	CUSTOMIZABLEOBJECT_API const TSharedPtr<mu::FModel, ESPMode::ThreadSafe>& GetModel();
-	CUSTOMIZABLEOBJECT_API const TSharedPtr<const mu::FModel, ESPMode::ThreadSafe> GetModel() const;
+	CUSTOMIZABLEOBJECT_API void SetModel(const TSharedPtr<UE::Mutable::Private::FModel, ESPMode::ThreadSafe>& Model, const FGuid Identifier);
+	CUSTOMIZABLEOBJECT_API const TSharedPtr<UE::Mutable::Private::FModel, ESPMode::ThreadSafe>& GetModel();
+	CUSTOMIZABLEOBJECT_API const TSharedPtr<const UE::Mutable::Private::FModel, ESPMode::ThreadSafe> GetModel() const;
 
 	CUSTOMIZABLEOBJECT_API UModelResources* GetModelResources();
 	CUSTOMIZABLEOBJECT_API const UModelResources* GetModelResources() const;
@@ -1793,7 +1793,7 @@ public:
 #endif
 	
 	/** Rebuild ParameterProperties from the current compiled model. */
-	CUSTOMIZABLEOBJECT_API void UpdateParameterPropertiesFromModel(const TSharedPtr<mu::FModel>& Model);
+	CUSTOMIZABLEOBJECT_API void UpdateParameterPropertiesFromModel(const TSharedPtr<UE::Mutable::Private::FModel>& Model);
 
 	CUSTOMIZABLEOBJECT_API void AddUncompiledCOWarning(const FString& AdditionalLoggingInfo);
 
@@ -1832,7 +1832,7 @@ public:
 
 	CUSTOMIZABLEOBJECT_API void SetIsChildObject(bool bIsChildObject);
 
-	/** Return the names used by mutable to identify which mu::FImage should be considered of LowPriority. */
+	/** Return the names used by mutable to identify which UE::Mutable::Private::FImage should be considered of LowPriority. */
 	CUSTOMIZABLEOBJECT_API void GetLowPriorityTextureNames(TArray<FString>& OutTextureNames);
 
 	/** Return the MinLOD index to generate based on the active LODSettings (PerPlatformMinLOD or PerQualityLevelMinLOD) */
@@ -1906,11 +1906,11 @@ public:
 	/** Cache of merged Skeletons */
 	FSkeletonCache SkeletonCache;
 
-	TSharedRef<mu::FMeshIdRegistry> MeshIdRegistry = MakeShared<mu::FMeshIdRegistry>();
+	TSharedRef<UE::Mutable::Private::FMeshIdRegistry> MeshIdRegistry = MakeShared<UE::Mutable::Private::FMeshIdRegistry>();
 	
-	TSharedRef<mu::FImageIdRegistry> ImageIdRegistry = MakeShared<mu::FImageIdRegistry>();
+	TSharedRef<UE::Mutable::Private::FImageIdRegistry> ImageIdRegistry = MakeShared<UE::Mutable::Private::FImageIdRegistry>();
 
-	TSharedRef<mu::FMaterialIdRegistry> MaterialIdRegistry = MakeShared<mu::FMaterialIdRegistry>();
+	TSharedRef<UE::Mutable::Private::FMaterialIdRegistry> MaterialIdRegistry = MakeShared<UE::Mutable::Private::FMaterialIdRegistry>();
 	
 	// See UCustomizableObjectSystem::LockObject. Must only be modified from the game thread
 	bool bLocked = false;
@@ -1928,7 +1928,7 @@ public:
 	FPostCompileDelegate PostCompileDelegate;
 
 	/** Map of PlatformName to CachedPlatformData. Only valid while cooking. */
-	TMap<FString, MutablePrivate::FMutableCachedPlatformData> CachedPlatformsData;
+	TMap<FString, UE::Mutable::Private::FMutableCachedPlatformData> CachedPlatformsData;
 
 #endif
 
@@ -1983,18 +1983,18 @@ public:
 
 /** Returns the DDC ValueId of the file owning a streamable resource.
 	* @ param StreamableDataType - UE level data type
-	* @ param ResourceType - mu::EDataType
+	* @ param ResourceType - UE::Mutable::Private::EDataType
 	*/
-CUSTOMIZABLEOBJECT_API UE::DerivedData::FValueId GetDerivedDataValueIdForResource(MutablePrivate::EStreamableDataType StreamableDataType, uint32 FileId, uint16 ResourceType, uint16 Flags);
+CUSTOMIZABLEOBJECT_API UE::DerivedData::FValueId GetDerivedDataValueIdForResource(UE::Mutable::Private::EStreamableDataType StreamableDataType, uint32 FileId, uint16 ResourceType, uint16 Flags);
 
 // Compose folder name where the data is stored
 CUSTOMIZABLEOBJECT_API FString GetCompiledDataFolderPath();
-CUSTOMIZABLEOBJECT_API FString GetDataTypeExtension(MutablePrivate::EStreamableDataType DataType);
+CUSTOMIZABLEOBJECT_API FString GetDataTypeExtension(UE::Mutable::Private::EStreamableDataType DataType);
 
 CUSTOMIZABLEOBJECT_API uint32 GetECustomizableObjectVersionEnumHash();
 
 CUSTOMIZABLEOBJECT_API TObjectPtr<UModelResources> LoadModelResources_Internal(FArchive& Ar, const UCustomizableObject* Outer, const ITargetPlatform* InTargetPlatform, bool bSkipEditorOnlyData = false);
 CUSTOMIZABLEOBJECT_API const TSharedPtr<FModelStreamableBulkData> LoadModelStreamableBulk_Internal(FArchive& Ar);
-CUSTOMIZABLEOBJECT_API const TSharedPtr<mu::FModel, ESPMode::ThreadSafe> LoadModel_Internal(FArchive& Ar);
+CUSTOMIZABLEOBJECT_API const TSharedPtr<UE::Mutable::Private::FModel, ESPMode::ThreadSafe> LoadModel_Internal(FArchive& Ar);
 
 #endif
diff --git a/Engine/Plugins/Mutable/Source/CustomizableObject/Internal/MuCO/CustomizableObjectSystemPrivate.h b/Engine/Plugins/Mutable/Source/CustomizableObject/Internal/MuCO/CustomizableObjectSystemPrivate.h
--- a/Engine/Plugins/Mutable/Source/CustomizableObject/Internal/MuCO/CustomizableObjectSystemPrivate.h
+++ b/Engine/Plugins/Mutable/Source/CustomizableObject/Internal/MuCO/CustomizableObjectSystemPrivate.h
@@ -54,7 +54,7 @@ extern TAutoConsoleVariable<bool> CVarMutableHighPriorityLoading;
 /** Key to identify an image inside a generated mutable runtime instance. */
 struct FMutableImageReference
 {
-	mu::FImageId ImageID;
+	UE::Mutable::Private::FImageId ImageID;
 
 	uint32 SurfaceId = 0;
 
@@ -108,23 +108,23 @@ struct FMappedMorphTargetMeshData
  * @param Mesh Used to know which MappedMorphTargets vertices must be removed.
  * @param GlobalNames
  * @param OutMorphTargets resulting Morph Targets with the indices pointing to GlobalNames. */
-void ReconstructMorphTargets(const mu::FMesh& Mesh, const TArray<FName>& GlobalNames, const TMap<uint32, FMappedMorphTargetMeshData>& MappedMorphTargets, TArray<FMorphTargetLODModel>& OutMorphTargets);
+void ReconstructMorphTargets(const UE::Mutable::Private::FMesh& Mesh, const TArray<FName>& GlobalNames, const TMap<uint32, FMappedMorphTargetMeshData>& MappedMorphTargets, TArray<FMorphTargetLODModel>& OutMorphTargets);
 
 
 /** Request the Mutable Data Streamer to load Morph Target blocks. Does not stream them.
  *  @param MutableDataStreamer Streamer to request the blocks.
  *  @param Mesh Mutable mesh to obtain which blocks to load.
  *  @param StreamingResult Key is the block. Value is the container where the data will be deserialized once streamed. */
-void LoadMorphTargetsData(FMutableStreamRequest& MutableDataStreamer, const TSharedRef<const mu::FMesh>& Mesh, TMap<uint32, FMorphTargetMeshData>& StreamingResult);
+void LoadMorphTargetsData(FMutableStreamRequest& MutableDataStreamer, const TSharedRef<const UE::Mutable::Private::FMesh>& Mesh, TMap<uint32, FMorphTargetMeshData>& StreamingResult);
 
-void LoadMorphTargetsMetadata(const FMutableStreamRequest& MutableDataStreamer, const TSharedRef<const mu::FMesh>& Mesh, TMap<uint32, FMorphTargetMeshData>& StreamingResult);
+void LoadMorphTargetsMetadata(const FMutableStreamRequest& MutableDataStreamer, const TSharedRef<const UE::Mutable::Private::FMesh>& Mesh, TMap<uint32, FMorphTargetMeshData>& StreamingResult);
 
 
 /** Request the Mutable Data Streamer to load Cloth blocks. Does not stream them.
  *  @param MutableDataStreamer Streamer to request the blocks.
  *  @param Mesh Mutable mesh to obtain which blocks to load.
  *  @param StreamingResult Key is the block. Value is the container where the data will be deserialized once streamed. */
-void LoadClothing(FMutableStreamRequest& MutableDataStreamer, const TSharedRef<const mu::FMesh>& Mesh, TMap<uint32, FClothingMeshData>& StreamingResult);
+void LoadClothing(FMutableStreamRequest& MutableDataStreamer, const TSharedRef<const UE::Mutable::Private::FMesh>& Mesh, TMap<uint32, FClothingMeshData>& StreamingResult);
 
 
 /** End a Customizable Object Instance Update. All code paths of an update have to end here. */
@@ -273,7 +273,7 @@ class FMutablePendingInstanceWork
 
 	TSet<FMutablePendingInstanceDiscard, FPendingInstanceDiscardKeyFuncs> PendingInstanceDiscards;
 
-	TSet<mu::FInstance::FID> PendingIDsToRelease;
+	TSet<UE::Mutable::Private::FInstance::FID> PendingIDsToRelease;
 
 public:
 	// Returns the number of pending instance updates, LOD Updates and discards.
@@ -301,13 +301,13 @@ public:
 		return PendingInstanceDiscards.CreateIterator();
 	}
 
-	TSet<mu::FInstance::FID>::TIterator GetIDsToReleaseIterator()
+	TSet<UE::Mutable::Private::FInstance::FID>::TIterator GetIDsToReleaseIterator()
 	{
 		return PendingIDsToRelease.CreateIterator();
 	}
 
 	void AddDiscard(const FMutablePendingInstanceDiscard& TaskToEnqueue);
-	void AddIDRelease(mu::FInstance::FID IDToRelease);
+	void AddIDRelease(UE::Mutable::Private::FInstance::FID IDToRelease);
 
 	void RemoveAllUpdatesAndDiscardsAndReleases();
 };
@@ -344,7 +344,7 @@ struct FGeneratedMaterial
 	uint32 SurfaceId = 0;
 
 	// Core Material. Strong object pointer to a material or a material index
-	TSharedPtr<const mu::FMaterial> Material;
+	TSharedPtr<const UE::Mutable::Private::FMaterial> Material;
 
 #if WITH_EDITORONLY_DATA
 	// Name of the component that contains this material
@@ -374,14 +374,14 @@ struct FInstanceUpdateData
 	struct FImage
 	{
 		FName Name;
-		mu::FImageId ImageID;
+		UE::Mutable::Private::FImageId ImageID;
 		
 		// LOD of the ImageId. If the texture is shared between LOD, first LOD where this image can be found. 
 		int32 BaseLOD;
 		int32 BaseMip;
 		
 		uint16 FullImageSizeX, FullImageSizeY;
-		TSharedPtr<const mu::FImage> Image;
+		TSharedPtr<const UE::Mutable::Private::FImage> Image;
 		TWeakObjectPtr<UTexture2D> Cached;
 
 		TArray<int32> ConstantImagesNeededToGenerate;
@@ -423,14 +423,14 @@ struct FInstanceUpdateData
 		/** Id of the metadata associated with this surface. */ 
 		uint32 SurfaceMetadataId = 0;
 
-		mu::FMaterialId MaterialId;
-		TSharedPtr<const mu::FMaterial> Material;
+		UE::Mutable::Private::FMaterialId MaterialId;
+		TSharedPtr<const UE::Mutable::Private::FMaterial> Material;
 	};
 
 	struct FLOD
 	{		
-		mu::FMeshId MeshId;
-		TSharedPtr<const mu::FMesh> Mesh;
+		UE::Mutable::Private::FMeshId MeshId;
+		TSharedPtr<const UE::Mutable::Private::FMesh> Mesh;
 
 		/** Range in the Surfaces array */
 		uint16 FirstSurface = 0;
@@ -454,8 +454,8 @@ struct FInstanceUpdateData
 		uint16 LODCount = 0;
 
 		/** Overlay material. Only one per mesh, if any. */
-		mu::FMaterialId OverlayMaterialId;
-		TSharedPtr<const mu::FMaterial> OverlayMaterial;
+		UE::Mutable::Private::FMaterialId OverlayMaterialId;
+		TSharedPtr<const UE::Mutable::Private::FMaterial> OverlayMaterial;
 	};
 
 	TArray<FComponent> Components;
@@ -465,15 +465,15 @@ struct FInstanceUpdateData
 	TArray<FVector> Vectors;
 	TArray<FScalar> Scalars;
 
-	TArray<mu::FBoneName> ActiveBones;
-	TArray<mu::FBoneName> BoneMaps;
+	TArray<UE::Mutable::Private::FBoneName> ActiveBones;
+	TArray<UE::Mutable::Private::FBoneName> BoneMaps;
 	
 	struct FBone
 	{
-		mu::FBoneName Name;
+		UE::Mutable::Private::FBoneName Name;
 		FMatrix44f MatrixWithScale;
 
-		bool operator==(const mu::FBoneName& OtherName) const { return Name == OtherName; };
+		bool operator==(const UE::Mutable::Private::FBoneName& OtherName) const { return Name == OtherName; };
 	};
 	
 	/** Key is the Block Id. Value is the LoadAdditionalAssetsAndData read destination. */
@@ -493,7 +493,7 @@ struct FInstanceUpdateData
 
 		TArray<FBone> BonePose;
 
-		TMap<mu::FBoneName, TPair<FName, uint16>> BoneInfoMap;
+		TMap<UE::Mutable::Private::FBoneName, TPair<FName, uint16>> BoneInfoMap;
 	};
 
 	// Access by instance component index
@@ -501,7 +501,7 @@ struct FInstanceUpdateData
 
 	struct FNamedExtensionData
 	{
-		TSharedPtr<const mu::FExtensionData> Data;
+		TSharedPtr<const UE::Mutable::Private::FExtensionData> Data;
 		FName Name;
 	};
 	TArray<FNamedExtensionData> ExtendedInputPins;
@@ -573,8 +573,8 @@ private:
 
 public:
 	/** Instance parameters at the time of the operation request. */
-	TSharedPtr<mu::FParameters> Parameters; 
-	TSharedPtr<mu::FSystem> MutableSystem;
+	TSharedPtr<UE::Mutable::Private::FParameters> Parameters; 
+	TSharedPtr<UE::Mutable::Private::FSystem> MutableSystem;
 
 	bool bOnlyUpdateIfNotGenerated = false;
 	bool bIgnoreCloseDist = false;
@@ -619,13 +619,13 @@ public:
 	/** This option comes from the operation request. It is used to reduce the number of mipmaps that mutable must generate for images.  */
 	int32 MipsToSkip = 0;
 
-	mu::FInstance::FID InstanceID = 0; // Redundant
-	TSharedPtr<const mu::FInstance> MutableInstance;
+	UE::Mutable::Private::FInstance::FID InstanceID = 0; // Redundant
+	TSharedPtr<const UE::Mutable::Private::FInstance> MutableInstance;
 
-	TSharedPtr<mu::FModel> Model;
-	TSharedPtr<mu::FMeshIdRegistry> MeshIdRegistry;
-	TSharedPtr<mu::FImageIdRegistry> ImageIdRegistry;
-	TSharedPtr<mu::FMaterialIdRegistry> MaterialIdRegistry;
+	TSharedPtr<UE::Mutable::Private::FModel> Model;
+	TSharedPtr<UE::Mutable::Private::FMeshIdRegistry> MeshIdRegistry;
+	TSharedPtr<UE::Mutable::Private::FImageIdRegistry> ImageIdRegistry;
+	TSharedPtr<UE::Mutable::Private::FMaterialIdRegistry> MaterialIdRegistry;
 
 	TSharedPtr<FUnrealMutableResourceProvider> ExternalResourceProvider;
 
@@ -655,22 +655,22 @@ public:
 	
 	TMap<FName, uint8> FirstResidentLOD;
 
-	TMap<mu::FImageId, FTexturePlatformData*> ImageToPlatformDataMap;
+	TMap<UE::Mutable::Private::FImageId, FTexturePlatformData*> ImageToPlatformDataMap;
 
 	EUpdateResult UpdateResult = EUpdateResult::Success;
 
-	mu::FImageOperator::FImagePixelFormatFunc PixelFormatOverride;
+	UE::Mutable::Private::FImageOperator::FImagePixelFormatFunc PixelFormatOverride;
 
 private:
 
 	/** Mutable Meshes required for each component. Outermost index is the object component index, inner index is the LOD. */
-	TArray<TArray<mu::FMeshId>> MeshDescriptors;
+	TArray<TArray<UE::Mutable::Private::FMeshId>> MeshDescriptors;
 
 public:
 
 	void InitMeshDescriptors(int32 Size);
-	const TArray<TArray<mu::FMeshId>>& GetMeshDescriptors() const;
-	TArray<mu::FMeshId>* GetMeshDescriptors(FCustomizableObjectComponentIndex Index);
+	const TArray<TArray<UE::Mutable::Private::FMeshId>>& GetMeshDescriptors() const;
+	TArray<UE::Mutable::Private::FMeshId>* GetMeshDescriptors(FCustomizableObjectComponentIndex Index);
 
 	/** Used to know if the updated instances' meshes are different from the previous ones. 
 	  * The index of the array is the instance component's index.
@@ -715,14 +715,14 @@ public:
 	// between AdditionalAssetsAsyncLoaded() and their setting into the generated materials in BuildMaterials()
 	TArray<TStrongObjectPtr<const UStreamableRenderAsset>> LoadedPassThroughMeshesPendingSetMaterial;
 	
-	TArray<mu::FImageId> CachedTextures;
+	TArray<UE::Mutable::Private::FImageId> CachedTextures;
 };
 
 
 /** Runtime data used during a mutable instance update */
 struct FMutableReleasePlatformOperationData
 {
-	TMap<mu::FImageId, FTexturePlatformData*> ImageToPlatformDataMap;
+	TMap<UE::Mutable::Private::FImageId, FTexturePlatformData*> ImageToPlatformDataMap;
 };
 
 
@@ -822,7 +822,7 @@ public:
 	static UCustomizableObjectSystem* SSystem;
 
 	// Pointer to the lower level mutable system that actually does the work.
-	TSharedPtr<mu::FSystem> MutableSystem;
+	TSharedPtr<UE::Mutable::Private::FSystem> MutableSystem;
 
 	/** Store the last streaming memory size in bytes, to change it when it is safe. */
 	uint64 LastWorkingMemoryBytes = 0;
@@ -843,7 +843,7 @@ public:
 	static int32 SkeletalMeshMinLodQualityLevel;
 
 #if WITH_EDITOR
-	mu::FImageOperator::FImagePixelFormatFunc ImageFormatOverrideFunc;
+	UE::Mutable::Private::FImageOperator::FImagePixelFormatFunc ImageFormatOverrideFunc;
 #endif
 
 	// IStreamingManager interface
@@ -868,7 +868,7 @@ public:
 	void InitDiscardResourcesSkeletalMesh(UCustomizableObjectInstance* InCustomizableObjectInstance);
 
 	// Init the async release of a Mutable Core Instance ID and all the temp resources associated with it
-	void InitInstanceIDRelease(mu::FInstance::FID IDToRelease);
+	void InitInstanceIDRelease(UE::Mutable::Private::FInstance::FID IDToRelease);
 
 	void GetMipStreamingConfig(const UCustomizableObjectInstance& Instance, bool bTextureStreamingEnabled, bool& bOutNeverStream, int32& OutMipsToSkip) const;
 	
@@ -1004,7 +1004,7 @@ namespace impl
 
 
 /** Set OnlyLOD to -1 to generate all mips */
-CUSTOMIZABLEOBJECT_API FTexturePlatformData* MutableCreateImagePlatformData(TSharedPtr<const mu::FImage> MutableImage, int32 OnlyLOD, uint16 FullSizeX, uint16 FullSizeY);
+CUSTOMIZABLEOBJECT_API FTexturePlatformData* MutableCreateImagePlatformData(TSharedPtr<const UE::Mutable::Private::FImage> MutableImage, int32 OnlyLOD, uint16 FullSizeX, uint16 FullSizeY);
 
 
 /** Return true if Streaming is enabled for the given Object. */
diff --git a/Engine/Plugins/Mutable/Source/CustomizableObject/Internal/MuCO/LoadUtils.h b/Engine/Plugins/Mutable/Source/CustomizableObject/Internal/MuCO/LoadUtils.h
--- a/Engine/Plugins/Mutable/Source/CustomizableObject/Internal/MuCO/LoadUtils.h
+++ b/Engine/Plugins/Mutable/Source/CustomizableObject/Internal/MuCO/LoadUtils.h
@@ -11,7 +11,7 @@ struct FAssetData;
 struct FSoftObjectPtr;
 
 
-namespace MutablePrivate
+namespace UE::Mutable::Private
 {
 	CUSTOMIZABLEOBJECT_API UObject* LoadObject(const FAssetData& DataAsset);
 
diff --git a/Engine/Plugins/Mutable/Source/CustomizableObject/Internal/MuCO/MutableMeshBufferUtils.h b/Engine/Plugins/Mutable/Source/CustomizableObject/Internal/MuCO/MutableMeshBufferUtils.h
--- a/Engine/Plugins/Mutable/Source/CustomizableObject/Internal/MuCO/MutableMeshBufferUtils.h
+++ b/Engine/Plugins/Mutable/Source/CustomizableObject/Internal/MuCO/MutableMeshBufferUtils.h
@@ -10,7 +10,7 @@
 #include "BoneIndices.h"
 
 class USkeletalMesh;
-namespace mu { class FMeshBufferSet; }
+namespace UE::Mutable::Private { class FMeshBufferSet; }
 
 /** Publicly accessible methods designed to prepare the Mutable Mesh buffers with the internal structure  required by Unreal and
  * our systems to be able to work with the mesh owning those buffers.
@@ -27,14 +27,14 @@ namespace  MutableMeshBufferUtils
 	 * @param InCurrentVertexBuffer Index of the Vertex Buffer being set up inside the OutTargetVertexBuffers.
 	 * @param OutTargetVertexBuffers Mesh buffer to be set up. 
 	 */
-	CUSTOMIZABLEOBJECT_API void SetupVertexPositionsBuffer(const int32& InCurrentVertexBuffer, mu::FMeshBufferSet& OutTargetVertexBuffers);
+	CUSTOMIZABLEOBJECT_API void SetupVertexPositionsBuffer(const int32& InCurrentVertexBuffer, UE::Mutable::Private::FMeshBufferSet& OutTargetVertexBuffers);
 
 	/**
 	 * Sets up all the channels that ought to be part of the Vertex Tangent Buffer of the Mutable mesh owning the provided MeshBufferSet
 	 * @param InCurrentVertexBuffer Index of the Vertex Buffer being set up inside the OutTargetVertexBuffers.
 	 * @param OutTargetVertexBuffers Mesh buffer to be set up. 
 	 */
-	CUSTOMIZABLEOBJECT_API void SetupTangentBuffer(const int32& InCurrentVertexBuffer, mu::FMeshBufferSet& OutTargetVertexBuffers);
+	CUSTOMIZABLEOBJECT_API void SetupTangentBuffer(const int32& InCurrentVertexBuffer, UE::Mutable::Private::FMeshBufferSet& OutTargetVertexBuffers);
 
 	/**
 	 * Sets up all the channels that ought to be part of the Vertex Texture Coordinates Buffer of the Mutable mesh owning the provided MeshBufferSet
@@ -43,7 +43,7 @@ namespace  MutableMeshBufferUtils
 	 * @param OutTargetVertexBuffers Mesh buffer to be set up. 
 	 * @param InTextureSemanticIndicesOverride (Optional) If you are certain of the texture semantic indices that should be used populate this argument, if not, leave it as null
 	 */
-	CUSTOMIZABLEOBJECT_API void SetupTexCoordinatesBuffer(const int32& InCurrentVertexBuffer,const int32& InChannelCount, bool bHighPrecision, mu::FMeshBufferSet& OutTargetVertexBuffers, const int32* InTextureSemanticIndicesOverride = nullptr);
+	CUSTOMIZABLEOBJECT_API void SetupTexCoordinatesBuffer(const int32& InCurrentVertexBuffer,const int32& InChannelCount, bool bHighPrecision, UE::Mutable::Private::FMeshBufferSet& OutTargetVertexBuffers, const int32* InTextureSemanticIndicesOverride = nullptr);
 
 	/**
 	 * Sets up all the channels that ought to be part of the Vertex Skin Buffer of the Mutable mesh owning the provided MeshBufferSet
@@ -53,14 +53,14 @@ namespace  MutableMeshBufferUtils
 	 * @param OutTargetVertexBuffers Mesh buffer to be set up. 
 	 */
 	CUSTOMIZABLEOBJECT_API void SetupSkinBuffer(const int32& InCurrentVertexBuffer,const int32& MaxBoneIndexTypeSizeBytes, const int32& MaxBoneWeightTypeSizeBytes,
-		const int32& MaxNumBonesPerVertex, mu::FMeshBufferSet& OutTargetVertexBuffers);
+		const int32& MaxNumBonesPerVertex, UE::Mutable::Private::FMeshBufferSet& OutTargetVertexBuffers);
 	
 	/**
 	 * Sets up all the channels that ought to be part of the Vertex Color Buffer of the Mutable mesh owning the provided MeshBufferSet
 	 * @param InCurrentVertexBuffer Index of the Vertex Buffer being set up inside the OutTargetVertexBuffers.
 	 * @param OutTargetVertexBuffers Mesh buffer to be set up. 
 	 */
-	CUSTOMIZABLEOBJECT_API void SetupVertexColorBuffer(const int32& InCurrentVertexBuffer, mu::FMeshBufferSet& OutTargetVertexBuffers);
+	CUSTOMIZABLEOBJECT_API void SetupVertexColorBuffer(const int32& InCurrentVertexBuffer, UE::Mutable::Private::FMeshBufferSet& OutTargetVertexBuffers);
 
 	// INDEX BUFFERS
 
@@ -68,7 +68,7 @@ namespace  MutableMeshBufferUtils
 	 * Sets up all the channels that ought to be part of the Index Buffer of the Mutable mesh owning the provided MeshBufferSet
 	 * @param OutTargetIndexBuffers Index buffer to be set up. 
 	 */
-	CUSTOMIZABLEOBJECT_API void SetupIndexBuffer(mu::FMeshBufferSet& OutTargetIndexBuffers);
+	CUSTOMIZABLEOBJECT_API void SetupIndexBuffer(UE::Mutable::Private::FMeshBufferSet& OutTargetIndexBuffers);
 
 	/**
 	 * Sets up all the channels that ought to be part of the Vertex SkinWeightProfile Buffer of the Mutable mesh owning the provided MeshBufferSet
@@ -80,6 +80,6 @@ namespace  MutableMeshBufferUtils
 	 * @param OutTargetVertexBuffers Mesh buffer to be set up.
 	 */
 	CUSTOMIZABLEOBJECT_API void SetupSkinWeightProfileBuffer(const int32& InCurrentVertexBuffer, const int32& MaxBoneIndexTypeSizeBytes,const int32& MaxBoneWeightTypeSizeBytes, const int32& MaxNumBonesPerVertex,
-		const int32 SemanticsIndex, mu::FMeshBufferSet& OutTargetVertexBuffers);
+		const int32 SemanticsIndex, UE::Mutable::Private::FMeshBufferSet& OutTargetVertexBuffers);
 
 }
diff --git a/Engine/Plugins/Mutable/Source/CustomizableObject/Internal/MuCO/MutableProjectorTypeUtils.h b/Engine/Plugins/Mutable/Source/CustomizableObject/Internal/MuCO/MutableProjectorTypeUtils.h
--- a/Engine/Plugins/Mutable/Source/CustomizableObject/Internal/MuCO/MutableProjectorTypeUtils.h
+++ b/Engine/Plugins/Mutable/Source/CustomizableObject/Internal/MuCO/MutableProjectorTypeUtils.h
@@ -12,13 +12,13 @@ namespace ProjectorUtils
 	 * @return The equivalent projector type to the one provided to be used with the CO and COI systems.
 	 */
 	CUSTOMIZABLEOBJECT_API
-	ECustomizableObjectProjectorType GetEquivalentProjectorType (mu::EProjectorType ProjectorType);
+	ECustomizableObjectProjectorType GetEquivalentProjectorType (UE::Mutable::Private::EProjectorType ProjectorType);
 
 	/** Get the equivalent mutable projector type to the CO mutable projector type.
 	 * @param ProjectorType The mutable CO projector type to get the mutable's internal equivalent projector type.
 	 * @return the projector type provided but as a mutable core projector type.
 	 */
 	CUSTOMIZABLEOBJECT_API
-	mu::EProjectorType GetEquivalentProjectorType (ECustomizableObjectProjectorType ProjectorType);
+	UE::Mutable::Private::EProjectorType GetEquivalentProjectorType (ECustomizableObjectProjectorType ProjectorType);
 }
 
diff --git a/Engine/Plugins/Mutable/Source/CustomizableObject/Internal/MuCO/UnrealConversionUtils.h b/Engine/Plugins/Mutable/Source/CustomizableObject/Internal/MuCO/UnrealConversionUtils.h
--- a/Engine/Plugins/Mutable/Source/CustomizableObject/Internal/MuCO/UnrealConversionUtils.h
+++ b/Engine/Plugins/Mutable/Source/CustomizableObject/Internal/MuCO/UnrealConversionUtils.h
@@ -20,7 +20,7 @@ struct FCustomizableObjectMeshToMeshVertData;
 struct FClothingMeshData;
 struct FMorphTargetMeshData;
 
-namespace mu
+namespace UE::Mutable::Private
 {
 	class FMesh;
 	class FMeshBufferSet;
@@ -63,9 +63,9 @@ namespace UnrealConversionUtils
 	 */
 	CUSTOMIZABLEOBJECT_API void SetupRenderSections(
 		FSkeletalMeshLODRenderData& LODResource,
-		const mu::FMesh* InMutableMesh,
-		const TArray<mu::FBoneName>& InBoneMap,
-		const TMap<mu::FBoneName, TPair<FName, uint16>>& BoneInfoMap,
+		const UE::Mutable::Private::FMesh* InMutableMesh,
+		const TArray<UE::Mutable::Private::FBoneName>& InBoneMap,
+		const TMap<UE::Mutable::Private::FBoneName, TPair<FName, uint16>>& BoneInfoMap,
 		const int32 InFirstBoneMapIndex,
 		const TArray<const FMutableSurfaceMetadata*>& SectionMetadata);
 
@@ -77,7 +77,7 @@ namespace UnrealConversionUtils
 	 */
 	CUSTOMIZABLEOBJECT_API void InitVertexBuffersWithDummyData(
 		FSkeletalMeshLODRenderData& LODResource,
-		const mu::FMesh* InMutableMesh,
+		const UE::Mutable::Private::FMesh* InMutableMesh,
 		const bool bAllowCPUAccess);
 
 	/** Performs a copy of the data found on the vertex buffers on the mutable mesh to the buffers of the skeletal mesh
@@ -87,7 +87,7 @@ namespace UnrealConversionUtils
 	 */
 	CUSTOMIZABLEOBJECT_API void CopyMutableVertexBuffers(
 		FSkeletalMeshLODRenderData& LODResource,
-		const mu::FMesh* InMutableMesh,
+		const UE::Mutable::Private::FMesh* InMutableMesh,
 		const bool bAllowCPUAccess);
 
 	
@@ -98,7 +98,7 @@ namespace UnrealConversionUtils
 	 */
 	CUSTOMIZABLEOBJECT_API void InitIndexBuffersWithDummyData(
 		FSkeletalMeshLODRenderData& LODResource,
-		const mu::FMesh* InMutableMesh);
+		const UE::Mutable::Private::FMesh* InMutableMesh);
 
 	/**
 	 *Performs a copy of the data found on the index buffers on the mutable mesh to the buffers of the skeletal mesh
@@ -108,7 +108,7 @@ namespace UnrealConversionUtils
 	 */
 	CUSTOMIZABLEOBJECT_API bool CopyMutableIndexBuffers(
 		FSkeletalMeshLODRenderData& LODResource,
-		const mu::FMesh* InMutableMesh,
+		const UE::Mutable::Private::FMesh* InMutableMesh,
 		const TArray<uint32>& SurfaceIDs,
 		bool& bOutMarkRenderStateDirty);
 	
@@ -125,7 +125,7 @@ namespace UnrealConversionUtils
 		FSkeletalMeshLODRenderData& LODResource,
 		USkeletalMesh& Owner,
 		int32 LODIndex,
-		const mu::FMesh* InMutableMesh,
+		const UE::Mutable::Private::FMesh* InMutableMesh,
 		const TArray<TPair<uint32, FName>>& InActiveProfiles
 		);
 
@@ -154,16 +154,16 @@ namespace UnrealConversionUtils
 		FSkeletalMeshLODRenderData& LODResource
 	);
 
-	void MorphTargetVertexInfoBuffers(FSkeletalMeshLODRenderData& LODResource, const USkeletalMesh& Owner, const mu::FMesh& MutableMesh, const TMap<uint32, FMorphTargetMeshData>& MorphTargetMeshData, int32 LODIndex);
+	void MorphTargetVertexInfoBuffers(FSkeletalMeshLODRenderData& LODResource, const USkeletalMesh& Owner, const UE::Mutable::Private::FMesh& MutableMesh, const TMap<uint32, FMorphTargetMeshData>& MorphTargetMeshData, int32 LODIndex);
 
-	CUSTOMIZABLEOBJECT_API UE::Tasks::FTask ConvertSkeletalMeshFromRuntimeData(USkeletalMesh* SkeletalMesh, int32 LODIndex, int32 SectionIndex, TSharedPtr<mu::FMesh> Result, const TSharedRef<FBoneNames>& BoneNames);
+	CUSTOMIZABLEOBJECT_API UE::Tasks::FTask ConvertSkeletalMeshFromRuntimeData(USkeletalMesh* SkeletalMesh, int32 LODIndex, int32 SectionIndex, TSharedPtr<UE::Mutable::Private::FMesh> Result, const TSharedRef<FBoneNames>& BoneNames);
 
-	void GetSectionClothData(const mu::FMesh& MutableMesh, int32 LODIndex, const TMap<uint32, FClothingMeshData>& ClothingMeshData, TArray<FSectionClothData>& SectionsClothData, int32& NumClothingDataNotFound);
+	void GetSectionClothData(const UE::Mutable::Private::FMesh& MutableMesh, int32 LODIndex, const TMap<uint32, FClothingMeshData>& ClothingMeshData, TArray<FSectionClothData>& SectionsClothData, int32& NumClothingDataNotFound);
 	
 	void CopyMeshToMeshClothData(TArray<FSectionClothData>& SectionsClothData);
 
 	/** Based on FSkeletalMeshLODModel::GetClothMappingData() */
 	void CreateClothMapping(const FSectionClothData& SectionClothData, TArray<FMeshToMeshVertData>& MappingData, TArray<FClothBufferIndexMapping>& ClothIndexMapping);
 	
-	void ClothVertexBuffers(FSkeletalMeshLODRenderData& LODResource, const mu::FMesh& MutableMesh, const TMap<uint32, FClothingMeshData>& ClothingMeshData, int32 LODIndex);
+	void ClothVertexBuffers(FSkeletalMeshLODRenderData& LODResource, const UE::Mutable::Private::FMesh& MutableMesh, const TMap<uint32, FClothingMeshData>& ClothingMeshData, int32 LODIndex);
 }
diff --git a/Engine/Plugins/Mutable/Source/CustomizableObject/Internal/MuCO/UnrealMutableImageProvider.h b/Engine/Plugins/Mutable/Source/CustomizableObject/Internal/MuCO/UnrealMutableImageProvider.h
--- a/Engine/Plugins/Mutable/Source/CustomizableObject/Internal/MuCO/UnrealMutableImageProvider.h
+++ b/Engine/Plugins/Mutable/Source/CustomizableObject/Internal/MuCO/UnrealMutableImageProvider.h
@@ -12,29 +12,29 @@ class UTexture;
 class USkeletalMesh;
 
 /** Implementation of a mutable core provider for image parameters that are application-specific. */
-class FUnrealMutableResourceProvider : public mu::FExternalResourceProvider
+class FUnrealMutableResourceProvider : public UE::Mutable::Private::FExternalResourceProvider
 {
 public:
 	CUSTOMIZABLEOBJECT_API FUnrealMutableResourceProvider(const TSharedRef<FBoneNames>& InBoneNames);
 	
-	// mu::ExternalResourceProvider interface
+	// UE::Mutable::Private::ExternalResourceProvider interface
 	// Thread: worker
-	CUSTOMIZABLEOBJECT_API virtual TTuple<UE::Tasks::FTask, TFunction<void()>> GetImageAsync(UTexture* Texture, uint8 MipmapsToSkip, TFunction<void(TSharedPtr<mu::FImage>)>& ResultCallback) override;
-	CUSTOMIZABLEOBJECT_API virtual TTuple<UE::Tasks::FTask, TFunction<void()>> GetReferencedImageAsync( int32 Id, uint8 MipmapsToSkip, TFunction<void(TSharedPtr<mu::FImage>)>& ResultCallback) override;
-	CUSTOMIZABLEOBJECT_API virtual mu::FExtendedImageDesc GetImageDesc(UTexture* Texture) override;
-	CUSTOMIZABLEOBJECT_API virtual TTuple<UE::Tasks::FTask, TFunction<void()>> GetMeshAsync(USkeletalMesh* SkeletalMesh, int32 LODIndex, int32 SectionIndex, TFunction<void(TSharedPtr<mu::FMesh>)>& ResultCallback) override;
+	CUSTOMIZABLEOBJECT_API virtual TTuple<UE::Tasks::FTask, TFunction<void()>> GetImageAsync(UTexture* Texture, uint8 MipmapsToSkip, TFunction<void(TSharedPtr<UE::Mutable::Private::FImage>)>& ResultCallback) override;
+	CUSTOMIZABLEOBJECT_API virtual TTuple<UE::Tasks::FTask, TFunction<void()>> GetReferencedImageAsync( int32 Id, uint8 MipmapsToSkip, TFunction<void(TSharedPtr<UE::Mutable::Private::FImage>)>& ResultCallback) override;
+	CUSTOMIZABLEOBJECT_API virtual UE::Mutable::Private::FExtendedImageDesc GetImageDesc(UTexture* Texture) override;
+	CUSTOMIZABLEOBJECT_API virtual TTuple<UE::Tasks::FTask, TFunction<void()>> GetMeshAsync(USkeletalMesh* SkeletalMesh, int32 LODIndex, int32 SectionIndex, TFunction<void(TSharedPtr<UE::Mutable::Private::FMesh>)>& ResultCallback) override;
 
 #if WITH_EDITOR
 	void CacheRuntimeReferencedImages(const TArray<TSoftObjectPtr<UTexture2D>>& RuntimeReferencedTextures);
 #endif
 	
 private:
-	static inline const mu::FImageDesc DUMMY_IMAGE_DESC = 
-			mu::FImageDesc {mu::FImageSize(32, 32), mu::EImageFormat::RGBA_UByte, 1};
+	static inline const UE::Mutable::Private::FImageDesc DUMMY_IMAGE_DESC = 
+			UE::Mutable::Private::FImageDesc {UE::Mutable::Private::FImageSize(32, 32), UE::Mutable::Private::EImageFormat::RGBA_UByte, 1};
 
 	/** This will be called if an image Id has been requested by Mutable core but it has not been provided by any provider. */
-	static TSharedPtr<mu::FImage> CreateDummy();
-	static mu::FExtendedImageDesc CreateDummyDesc();
+	static TSharedPtr<UE::Mutable::Private::FImage> CreateDummy();
+	static UE::Mutable::Private::FExtendedImageDesc CreateDummyDesc();
 
 #if WITH_EDITOR
 	TArray<TStrongObjectPtr<UTexture>> Images;
diff --git a/Engine/Plugins/Mutable/Source/CustomizableObject/Internal/MuCO/UnrealMutableModelDiskStreamer.h b/Engine/Plugins/Mutable/Source/CustomizableObject/Internal/MuCO/UnrealMutableModelDiskStreamer.h
--- a/Engine/Plugins/Mutable/Source/CustomizableObject/Internal/MuCO/UnrealMutableModelDiskStreamer.h
+++ b/Engine/Plugins/Mutable/Source/CustomizableObject/Internal/MuCO/UnrealMutableModelDiskStreamer.h
@@ -14,19 +14,19 @@ class IAsyncReadFileHandle;
 class IAsyncReadRequest;
 class IBulkDataIORequest;
 class UCustomizableObject;
-namespace mu {
+namespace UE::Mutable::Private {
 	class FModel;
 	enum class EDataType : uint8;
 }
 struct FMutableStreamableBlock;
 
-class FUnrealMutableInputStream : public mu::FInputStream
+class FUnrealMutableInputStream : public UE::Mutable::Private::FInputStream
 {
 public:
 
 	FUnrealMutableInputStream(FArchive& ar);
 
-	// mu::FInputStream interface
+	// UE::Mutable::Private::FInputStream interface
 	void Read(void* pData, uint64 size) override;
 
 private:
@@ -44,7 +44,7 @@ struct FDDCReadRequest
 #endif
 
 // Implementation of a mutable streamer using bulk storage.
-class FUnrealMutableModelBulkReader : public mu::FModelReader
+class FUnrealMutableModelBulkReader : public UE::Mutable::Private::FModelReader
 {
 public:
 	// 
@@ -71,9 +71,9 @@ public:
 	/** Release all the pending resources. This disables treamings for all objects. */
 	CUSTOMIZABLEOBJECT_API void EndStreaming();
 
-	// mu::ModelReader interface
-	CUSTOMIZABLEOBJECT_API bool DoesBlockExist(const mu::FModel*, uint32 BlockKey) override;
-	CUSTOMIZABLEOBJECT_API FOperationID BeginReadBlock(const mu::FModel*, uint32 BlockKey, void* pBuffer, uint64 size, mu::EDataType ResourceType, TFunction<void(bool bSuccess)>* CompletionCallback) override;
+	// UE::Mutable::Private::ModelReader interface
+	CUSTOMIZABLEOBJECT_API bool DoesBlockExist(const UE::Mutable::Private::FModel*, uint32 BlockKey) override;
+	CUSTOMIZABLEOBJECT_API FOperationID BeginReadBlock(const UE::Mutable::Private::FModel*, uint32 BlockKey, void* pBuffer, uint64 size, UE::Mutable::Private::EDataType ResourceType, TFunction<void(bool bSuccess)>* CompletionCallback) override;
 	CUSTOMIZABLEOBJECT_API bool IsReadCompleted(FOperationID) override;
 	CUSTOMIZABLEOBJECT_API bool EndRead(FOperationID) override;
 
@@ -93,7 +93,7 @@ protected:
 	/** Streaming data for one object. */
 	struct FObjectData
 	{
-		TWeakPtr<const mu::FModel> Model;
+		TWeakPtr<const UE::Mutable::Private::FModel> Model;
 
 		FString BulkFilePrefix;
 		TMap<FOperationID, FReadRequest> CurrentReadRequests;
@@ -113,13 +113,13 @@ protected:
 
 #if WITH_EDITOR
 
-class UnrealMutableOutputStream : public mu::FOutputStream
+class UnrealMutableOutputStream : public UE::Mutable::Private::FOutputStream
 {
 public:
 
 	UnrealMutableOutputStream(FArchive& ar);
 
-	// mu::FOutputStream interface
+	// UE::Mutable::Private::FOutputStream interface
 	void Write(const void* pData, uint64 size) override;
 
 private:
@@ -128,13 +128,13 @@ private:
 
 
 // Implementation of a mutable streamer using bulk storage.
-class FUnrealMutableModelBulkWriterEditor : public mu::FModelWriter
+class FUnrealMutableModelBulkWriterEditor : public UE::Mutable::Private::FModelWriter
 {
 public:
 	// 
 	CUSTOMIZABLEOBJECT_API FUnrealMutableModelBulkWriterEditor(FArchive* InMainDataArchive = nullptr, FArchive* InStreamedDataArchive = nullptr);
 
-	// mu::FModelWriter interface
+	// UE::Mutable::Private::FModelWriter interface
 	CUSTOMIZABLEOBJECT_API void OpenWriteFile(uint32 BlockKey, bool bIsStreamable) override;
 	CUSTOMIZABLEOBJECT_API void Write(const void* pBuffer, uint64 size) override;
 	CUSTOMIZABLEOBJECT_API void CloseWriteFile() override;
@@ -153,13 +153,13 @@ protected:
 
 
 // Implementation of a mutable streamer using bulk storage.
-class FUnrealMutableModelBulkWriterCook : public mu::FModelWriter
+class FUnrealMutableModelBulkWriterCook : public UE::Mutable::Private::FModelWriter
 {
 public:
 	// 
-	CUSTOMIZABLEOBJECT_API FUnrealMutableModelBulkWriterCook(FArchive* InMainDataArchive = nullptr, MutablePrivate::FModelStreamableData* InStreamedData = nullptr);
+	CUSTOMIZABLEOBJECT_API FUnrealMutableModelBulkWriterCook(FArchive* InMainDataArchive = nullptr, UE::Mutable::Private::FModelStreamableData* InStreamedData = nullptr);
 
-	// mu::FModelWriter interface
+	// UE::Mutable::Private::FModelWriter interface
 	CUSTOMIZABLEOBJECT_API void OpenWriteFile(uint32 BlockKey, bool bIsStreamable) override;
 	CUSTOMIZABLEOBJECT_API void Write(const void* pBuffer, uint64 size) override;
 	CUSTOMIZABLEOBJECT_API void CloseWriteFile() override;
@@ -170,7 +170,7 @@ protected:
 	FArchive* MainDataArchive = nullptr;
 
 	// Non-owned pointer to an archive where we'll store the resouces (streamable)
-	MutablePrivate::FModelStreamableData* StreamedData = nullptr;
+	UE::Mutable::Private::FModelStreamableData* StreamedData = nullptr;
 
 	uint32 CurrentKey = 0;
 	bool CurrentIsStreamable = false;
diff --git a/Engine/Plugins/Mutable/Source/CustomizableObject/Internal/MuCO/UnrealToMutableTextureConversionUtils.h b/Engine/Plugins/Mutable/Source/CustomizableObject/Internal/MuCO/UnrealToMutableTextureConversionUtils.h
--- a/Engine/Plugins/Mutable/Source/CustomizableObject/Internal/MuCO/UnrealToMutableTextureConversionUtils.h
+++ b/Engine/Plugins/Mutable/Source/CustomizableObject/Internal/MuCO/UnrealToMutableTextureConversionUtils.h
@@ -15,7 +15,7 @@ class UAnimInstance;
 class USkeletalMesh;
 struct FMutableCompilationContext;
 
-namespace mu
+namespace UE::Mutable::Private
 {
 	class FImage;
 	class FMesh;
@@ -137,7 +137,7 @@ enum class EUnrealToMutableConversionError
     Unknown
 };
 
-UE_API EUnrealToMutableConversionError ConvertTextureUnrealSourceToMutable(mu::FImage* OutResult, FMutableSourceTextureData&, uint8 MipmapsToSkip);
+UE_API EUnrealToMutableConversionError ConvertTextureUnrealSourceToMutable(UE::Mutable::Private::FImage* OutResult, FMutableSourceTextureData&, uint8 MipmapsToSkip);
 
 /** Mesh conversion context. */
 class FMeshConversionContext
@@ -156,7 +156,7 @@ public:
 	// Keep strong references to meshes and other UObjects involved in the conversion
 	TArray<TStrongObjectPtr<UObject>> ReferencedObjects;
 
-	TSharedPtr<mu::FMesh> Result;
+	TSharedPtr<UE::Mutable::Private::FMesh> Result;
 };
 
 #endif
diff --git a/Engine/Plugins/Mutable/Source/CustomizableObject/Private/BoneNames.cpp b/Engine/Plugins/Mutable/Source/CustomizableObject/Private/BoneNames.cpp
--- a/Engine/Plugins/Mutable/Source/CustomizableObject/Private/BoneNames.cpp
+++ b/Engine/Plugins/Mutable/Source/CustomizableObject/Private/BoneNames.cpp
@@ -16,7 +16,7 @@ FBoneNames::FBoneNames(const UModelResources& ModelResources)
 }
 
 
-mu::FBoneName* FBoneNames::Find(const FName& BoneName)
+UE::Mutable::Private::FBoneName* FBoneNames::Find(const FName& BoneName)
 {
 	FScopeLock Lock(&CriticalSection);
 
@@ -24,13 +24,13 @@ mu::FBoneName* FBoneNames::Find(const FName& BoneName)
 }
 
 
-mu::FBoneName FBoneNames::FindOrAdd(const FName& BoneName)
+UE::Mutable::Private::FBoneName FBoneNames::FindOrAdd(const FName& BoneName)
 {
 	FScopeLock Lock(&CriticalSection);
 
 	const FString BoneNameString = BoneName.ToString().ToLower();
 
-	if (mu::FBoneName* Result = BoneNamesMap.Find(BoneNameString))
+	if (UE::Mutable::Private::FBoneName* Result = BoneNamesMap.Find(BoneNameString))
 	{
 		return *Result;
 	}
@@ -43,7 +43,7 @@ mu::FBoneName FBoneNames::FindOrAdd(const FName& BoneName)
 	while (!bUnique)
 	{
 		bUnique = true;
-		TMap<FString, mu::FBoneName>::TConstIterator MapIterator = BoneNamesMap.CreateConstIterator();
+		TMap<FString, UE::Mutable::Private::FBoneName>::TConstIterator MapIterator = BoneNamesMap.CreateConstIterator();
 		while (MapIterator)
 		{
 			if (MapIterator.Value() == NewBoneId)
diff --git a/Engine/Plugins/Mutable/Source/CustomizableObject/Private/MuCO/CustomizableObject.cpp b/Engine/Plugins/Mutable/Source/CustomizableObject/Private/MuCO/CustomizableObject.cpp
--- a/Engine/Plugins/Mutable/Source/CustomizableObject/Private/MuCO/CustomizableObject.cpp
+++ b/Engine/Plugins/Mutable/Source/CustomizableObject/Private/MuCO/CustomizableObject.cpp
@@ -192,14 +192,14 @@ void UCustomizableObject::PreSave(FObjectPreSaveContext ObjectSaveContext)
 			const bool bUseBulkData = CVarMutableUseBulkData.GetValueOnAnyThread();
 			if (bUseBulkData)
 			{
-				MutablePrivate::FMutableCachedPlatformData& CachedPlatformData = *GetPrivate()->CachedPlatformsData.Find(TargetPlatform->PlatformName());
+				UE::Mutable::Private::FMutableCachedPlatformData& CachedPlatformData = *GetPrivate()->CachedPlatformsData.Find(TargetPlatform->PlatformName());
 				TSharedPtr<FModelStreamableBulkData> ModelStreamableBulkData = GetPrivate()->GetModelStreamableBulkData(true);
 
 				const int32 NumBulkDataFiles = CachedPlatformData.BulkDataFiles.Num();
 
 				ModelStreamableBulkData->StreamableBulkData.SetNum(NumBulkDataFiles);
 
-				const auto WriteBulkData = [ModelStreamableBulkData](MutablePrivate::FFile& File, TArray64<uint8>& FileBulkData, uint32 FileIndex)
+				const auto WriteBulkData = [ModelStreamableBulkData](UE::Mutable::Private::FFile& File, TArray64<uint8>& FileBulkData, uint32 FileIndex)
 					{
 						MUTABLE_CPUPROFILER_SCOPE(WriteBulkData);
 
@@ -222,7 +222,7 @@ void UCustomizableObject::PreSave(FObjectPreSaveContext ObjectSaveContext)
 					};
 
 				bool bDropData = true;
-				MutablePrivate::SerializeBulkDataFiles(CachedPlatformData, CachedPlatformData.BulkDataFiles, WriteBulkData, bDropData);
+				UE::Mutable::Private::SerializeBulkDataFiles(CachedPlatformData, CachedPlatformData.BulkDataFiles, WriteBulkData, bDropData);
 			}
 			else 
 			{
@@ -256,7 +256,7 @@ void UCustomizableObject::PostSaveRoot(FObjectPostSaveRootContext ObjectSaveCont
 	{
 		const ITargetPlatform* TargetPlatform = ObjectSaveContext.GetTargetPlatform();
 
-		const MutablePrivate::FMutableCachedPlatformData* PlatformData = GetPrivate()->CachedPlatformsData.Find(TargetPlatform->PlatformName());
+		const UE::Mutable::Private::FMutableCachedPlatformData* PlatformData = GetPrivate()->CachedPlatformsData.Find(TargetPlatform->PlatformName());
 		if (PlatformData && PlatformData->ModelResources)
 		{
 			// Set the outer to TransientPkg again. This is done to avoid serializing multiple ModelResources when cooking more than one platform at once.
@@ -294,7 +294,7 @@ void UCustomizableObjectPrivate::SetIsChildObject(bool bIsChildObject)
 
 bool UCustomizableObjectPrivate::TryLoadCompiledCookDataForPlatform(const ITargetPlatform* TargetPlatform)
 {
-	const MutablePrivate::FMutableCachedPlatformData* PlatformData = CachedPlatformsData.Find(TargetPlatform->PlatformName());
+	const UE::Mutable::Private::FMutableCachedPlatformData* PlatformData = CachedPlatformsData.Find(TargetPlatform->PlatformName());
 	if (!PlatformData)
 	{
 		return false;
@@ -543,11 +543,11 @@ const TSharedPtr<FModelStreamableBulkData> LoadModelStreamableBulk_Internal(FArc
 }
 
 
-const TSharedPtr<mu::FModel, ESPMode::ThreadSafe> LoadModel_Internal(FArchive& MemoryReader)
+const TSharedPtr<UE::Mutable::Private::FModel, ESPMode::ThreadSafe> LoadModel_Internal(FArchive& MemoryReader)
 {
 	FUnrealMutableInputStream Stream(MemoryReader);
-	mu::FInputArchive Arch(&Stream);
-	return mu::FModel::StaticUnserialise(Arch);
+	UE::Mutable::Private::FInputArchive Arch(&Stream);
+	return UE::Mutable::Private::FModel::StaticUnserialise(Arch);
 }
 
 
@@ -575,7 +575,7 @@ void UModelResources::InitCookData(UObject& InCustomizableObject)
 			if (FCustomizableObjectAssetUserData* AUDResource = Container->Data.Data.GetMutablePtr<FCustomizableObjectAssetUserData>())
 			{
 				// Find or duplicate the AUD replacing the outer
-				UAssetUserData* SourceAssetUserData = MutablePrivate::LoadObject(AUDResource->AssetUserDataEditor);
+				UAssetUserData* SourceAssetUserData = UE::Mutable::Private::LoadObject(AUDResource->AssetUserDataEditor);
 				const FString AssetName = GetNameSafe(SourceAssetUserData);
 				check(SourceAssetUserData != nullptr);
 
@@ -646,9 +646,9 @@ void UCustomizableObjectPrivate::LoadCompiledDataFromDisk()
 	FGuid CompiledVersionId;
 
 	bool bHasCompiledData = true;
-	for (int32 DataType = 0; DataType < static_cast<int32>(MutablePrivate::EStreamableDataType::DataTypeCount); ++DataType)
+	for (int32 DataType = 0; DataType < static_cast<int32>(UE::Mutable::Private::EStreamableDataType::DataTypeCount); ++DataType)
 	{
-		const FString FilePath = FullFileName + GetDataTypeExtension(static_cast<MutablePrivate::EStreamableDataType>(DataType));
+		const FString FilePath = FullFileName + GetDataTypeExtension(static_cast<UE::Mutable::Private::EStreamableDataType>(DataType));
 		if (IFileHandle* FileHandle = FPlatformFileManager::Get().GetPlatformFile().OpenRead(*FilePath))
 		{
 			FileHandles.Emplace(FileHandle);
@@ -853,14 +853,14 @@ FString UCustomizableObjectPrivate::GetCompiledDataFileName(const ITargetPlatfor
 }
 
 
-FString GetDataTypeExtension(MutablePrivate::EStreamableDataType DataType)
+FString GetDataTypeExtension(UE::Mutable::Private::EStreamableDataType DataType)
 {
 	switch (DataType)
 	{
-	case MutablePrivate::EStreamableDataType::None:	return TEXT(".mut");
-	case MutablePrivate::EStreamableDataType::Model: return TEXT("_M.mut");
-	case MutablePrivate::EStreamableDataType::RealTimeMorph: return TEXT("_RTM.mut");
-	case MutablePrivate::EStreamableDataType::Clothing:	return TEXT("_C.mut");
+	case UE::Mutable::Private::EStreamableDataType::None:	return TEXT(".mut");
+	case UE::Mutable::Private::EStreamableDataType::Model: return TEXT("_M.mut");
+	case UE::Mutable::Private::EStreamableDataType::RealTimeMorph: return TEXT("_RTM.mut");
+	case UE::Mutable::Private::EStreamableDataType::Clothing:	return TEXT("_C.mut");
 	default:
 	{
 		unimplemented();
@@ -882,7 +882,7 @@ void UCustomizableObjectPrivate::SaveEmbeddedData(FArchive& Ar)
 {
 	UE_LOG(LogMutable, Verbose, TEXT("Saving embedded data for Customizable Object [%s] now at position %d."), *GetName(), int(Ar.Tell()));
 
-	TSharedPtr<mu::FModel> Model = GetModel();
+	TSharedPtr<UE::Mutable::Private::FModel> Model = GetModel();
 
 	int64 InternalVersion = Model ? (int64)GetECustomizableObjectVersionEnumHash() : INDEX_NONE;
 	Ar << InternalVersion;
@@ -894,8 +894,8 @@ void UCustomizableObjectPrivate::SaveEmbeddedData(FArchive& Ar)
 			Model->GetPrivate()->Program.RomsCompileData.Empty();
 		}
 		UnrealMutableOutputStream Stream(Ar);
-		mu::FOutputArchive Arch(&Stream);
-		mu::FModel::Serialise(Model.Get(), Arch);
+		UE::Mutable::Private::FOutputArchive Arch(&Stream);
+		UE::Mutable::Private::FModel::Serialise(Model.Get(), Arch);
 
 		UE_LOG(LogMutable, Verbose, TEXT("Saved embedded data for Customizable Object [%s] now at position %d."), *GetName(), int(Ar.Tell()));
 	}
@@ -916,8 +916,8 @@ void UCustomizableObjectPrivate::LoadEmbeddedData(FArchive& Ar)
 	{		
 		// Load model
 		FUnrealMutableInputStream Stream(Ar);
-		mu::FInputArchive Arch(&Stream);
-		TSharedPtr<mu::FModel, ESPMode::ThreadSafe> Model = mu::FModel::StaticUnserialise(Arch);
+		UE::Mutable::Private::FInputArchive Arch(&Stream);
+		TSharedPtr<UE::Mutable::Private::FModel, ESPMode::ThreadSafe> Model = UE::Mutable::Private::FModel::StaticUnserialise(Arch);
 
 		SetModel(Model, FGuid());
 	}
@@ -1060,8 +1060,8 @@ FString UCustomizableObject::GetStateParameterName(const FString& StateName, int
 #if WITH_EDITORONLY_DATA
 void UCustomizableObjectPrivate::PostCompile()
 {
-	MeshIdRegistry = MakeShared<mu::FMeshIdRegistry>();
-	ImageIdRegistry = MakeShared<mu::FImageIdRegistry>();
+	MeshIdRegistry = MakeShared<UE::Mutable::Private::FMeshIdRegistry>();
+	ImageIdRegistry = MakeShared<UE::Mutable::Private::FImageIdRegistry>();
 	
 	for (TObjectIterator<UCustomizableObjectInstance> It; It; ++It)
 	{
@@ -1373,15 +1373,15 @@ const FString& UCustomizableObject::GetParameterName(int32 ParamIndex) const
 }
 
 
-void UCustomizableObjectPrivate::UpdateParameterPropertiesFromModel(const TSharedPtr<mu::FModel>& Model)
+void UCustomizableObjectPrivate::UpdateParameterPropertiesFromModel(const TSharedPtr<UE::Mutable::Private::FModel>& Model)
 {
 	if (Model)
 	{
-		TSharedPtr<mu::FParameters> MutableParameters = mu::FModel::NewParameters(Model);
+		TSharedPtr<UE::Mutable::Private::FParameters> MutableParameters = UE::Mutable::Private::FModel::NewParameters(Model);
 		const int32 NumParameters = MutableParameters->GetCount();
 
 		TArray<int32> TypedParametersCount;
-		TypedParametersCount.SetNum(static_cast<int32>(mu::EParameterType::Count));
+		TypedParametersCount.SetNum(static_cast<int32>(UE::Mutable::Private::EParameterType::Count));
 		
 		ParameterProperties.Reset(NumParameters);
 		ParameterPropertiesLookupTable.Empty(NumParameters);
@@ -1392,16 +1392,16 @@ void UCustomizableObjectPrivate::UpdateParameterPropertiesFromModel(const TShare
 			Data.Name = MutableParameters->GetName(Index);
 			Data.Type = EMutableParameterType::None;
 
-			mu::EParameterType ParameterType = MutableParameters->GetType(Index);
+			UE::Mutable::Private::EParameterType ParameterType = MutableParameters->GetType(Index);
 			switch (ParameterType)
 			{
-			case mu::EParameterType::Bool:
+			case UE::Mutable::Private::EParameterType::Bool:
 			{
 				Data.Type = EMutableParameterType::Bool;
 				break;
 			}
 
-			case mu::EParameterType::Int:
+			case UE::Mutable::Private::EParameterType::Int:
 			{
 				Data.Type = EMutableParameterType::Int;
 
@@ -1416,43 +1416,43 @@ void UCustomizableObjectPrivate::UpdateParameterPropertiesFromModel(const TShare
 				break;
 			}
 
-			case mu::EParameterType::Float:
+			case UE::Mutable::Private::EParameterType::Float:
 			{
 				Data.Type = EMutableParameterType::Float;
 				break;
 			}
 
-			case mu::EParameterType::Color:
+			case UE::Mutable::Private::EParameterType::Color:
 			{
 				Data.Type = EMutableParameterType::Color;
 				break;
 			}
 
-			case mu::EParameterType::Projector:
+			case UE::Mutable::Private::EParameterType::Projector:
 			{
 				Data.Type = EMutableParameterType::Projector;
 				break;
 			}
 
-			case mu::EParameterType::Matrix:
+			case UE::Mutable::Private::EParameterType::Matrix:
 			{
 				Data.Type = EMutableParameterType::Transform;
 				break;
 			}
 				
-			case mu::EParameterType::Image:
+			case UE::Mutable::Private::EParameterType::Image:
 			{
 				Data.Type = EMutableParameterType::Texture;
 				break;
 			}
 
-			case mu::EParameterType::Mesh:
+			case UE::Mutable::Private::EParameterType::Mesh:
 			{
 				Data.Type = EMutableParameterType::SkeletalMesh;
 				break;
 			}
 			
-			case mu::EParameterType::Material:
+			case UE::Mutable::Private::EParameterType::Material:
 			{
 				Data.Type = EMutableParameterType::Material;
 				break;
@@ -1684,7 +1684,7 @@ float UCustomizableObject::GetFloatParameterDefaultValue(const FString& InParame
 		return FCustomizableObjectFloatParameterValue::DEFAULT_PARAMETER_VALUE;
 	}
 
-	const TSharedPtr<const mu::FModel>& Model = GetPrivate()->GetModel();
+	const TSharedPtr<const UE::Mutable::Private::FModel>& Model = GetPrivate()->GetModel();
 	if (!Model)
 	{
 		checkNoEntry();
@@ -1704,7 +1704,7 @@ int32 UCustomizableObject::GetEnumParameterDefaultValue(const FString& InParamet
 		return FCustomizableObjectIntParameterValue::DEFAULT_PARAMETER_VALUE;
 	}
 
-	const TSharedPtr<const mu::FModel>& Model = GetPrivate()->GetModel();
+	const TSharedPtr<const UE::Mutable::Private::FModel>& Model = GetPrivate()->GetModel();
 	if (!Model)
 	{
 		checkNoEntry();
@@ -1724,7 +1724,7 @@ bool UCustomizableObject::GetBoolParameterDefaultValue(const FString& InParamete
 		return FCustomizableObjectBoolParameterValue::DEFAULT_PARAMETER_VALUE;
 	}
 
-	const TSharedPtr<const mu::FModel>& Model = GetPrivate()->GetModel();
+	const TSharedPtr<const UE::Mutable::Private::FModel>& Model = GetPrivate()->GetModel();
 	if (!Model)
 	{
 		checkNoEntry();
@@ -1744,7 +1744,7 @@ FLinearColor UCustomizableObject::GetColorParameterDefaultValue(const FString& I
 		return FCustomizableObjectVectorParameterValue::DEFAULT_PARAMETER_VALUE;
 	}
 
-	const TSharedPtr<const mu::FModel>& Model = GetPrivate()->GetModel();
+	const TSharedPtr<const UE::Mutable::Private::FModel>& Model = GetPrivate()->GetModel();
 	if (!Model)
 	{
 		checkNoEntry();
@@ -1767,7 +1767,7 @@ FTransform UCustomizableObject::GetTransformParameterDefaultValue(const FString&
 		return FCustomizableObjectTransformParameterValue::DEFAULT_PARAMETER_VALUE;
 	}
 
-	const TSharedPtr<const mu::FModel>& Model = GetPrivate()->GetModel();
+	const TSharedPtr<const UE::Mutable::Private::FModel>& Model = GetPrivate()->GetModel();
 	if (!Model)
 	{
 		checkNoEntry();
@@ -1789,7 +1789,7 @@ FCustomizableObjectProjector UCustomizableObject::GetProjectorParameterDefaultVa
 		return FCustomizableObjectProjectorParameterValue::DEFAULT_PARAMETER_VALUE;
 	}
 
-	const TSharedPtr<const mu::FModel>& Model = GetPrivate()->GetModel();
+	const TSharedPtr<const UE::Mutable::Private::FModel>& Model = GetPrivate()->GetModel();
 	if (!Model)
 	{
 		checkNoEntry();
@@ -1797,7 +1797,7 @@ FCustomizableObjectProjector UCustomizableObject::GetProjectorParameterDefaultVa
 	}
 
 	FCustomizableObjectProjector Value;
-	mu::EProjectorType Type;
+	UE::Mutable::Private::EProjectorType Type;
 	Model->GetProjectorDefaultValue(ParameterIndex, &Type, &Value.Position, &Value.Direction, &Value.Up, &Value.Scale, &Value.Angle);
 	Value.ProjectionType = ProjectorUtils::GetEquivalentProjectorType(Type);
 	
@@ -1861,7 +1861,7 @@ bool UCustomizableObjectPrivate::IsParameterMultidimensional(const int32& InPara
 {
 	check(InParamIndex != INDEX_NONE);
 	
-	if (const TSharedPtr<const mu::FModel>& Model = GetModel())
+	if (const TSharedPtr<const UE::Mutable::Private::FModel>& Model = GetModel())
 	{
 		return Model->IsParameterMultidimensional(InParamIndex);
 	}
@@ -1979,14 +1979,14 @@ uint8 UCustomizableObjectPrivate::GetMinLODIndex(const FName& ComponentName) con
 //-------------------------------------------------------------------------------------------------
 //-------------------------------------------------------------------------------------------------
 
-UCustomizableObjectSkeletalMesh* FMeshCache::Get(const TArray<mu::FMeshId>& Key)
+UCustomizableObjectSkeletalMesh* FMeshCache::Get(const TArray<UE::Mutable::Private::FMeshId>& Key)
 {
 	const TWeakObjectPtr<UCustomizableObjectSkeletalMesh>* Result = GeneratedMeshes->Find(Key);
 	return Result ? Result->Get() : nullptr;
 }
 
 
-void FMeshCache::Add(const TArray<mu::FMeshId>& Key, UCustomizableObjectSkeletalMesh* Value)
+void FMeshCache::Add(const TArray<UE::Mutable::Private::FMeshId>& Key, UCustomizableObjectSkeletalMesh* Value)
 {
 	if (!Value)
 	{
@@ -2152,7 +2152,7 @@ void UModelStreamableData::PostLoad()
 }
 
 
-void UCustomizableObjectPrivate::SetModel(const TSharedPtr<mu::FModel, ESPMode::ThreadSafe>& Model, const FGuid Id)
+void UCustomizableObjectPrivate::SetModel(const TSharedPtr<UE::Mutable::Private::FModel, ESPMode::ThreadSafe>& Model, const FGuid Id)
 {
 	if (MutableModel == Model
 #if WITH_EDITOR
@@ -2182,13 +2182,13 @@ void UCustomizableObjectPrivate::SetModel(const TSharedPtr<mu::FModel, ESPMode::
 }
 
 
-const TSharedPtr<mu::FModel, ESPMode::ThreadSafe>& UCustomizableObjectPrivate::GetModel()
+const TSharedPtr<UE::Mutable::Private::FModel, ESPMode::ThreadSafe>& UCustomizableObjectPrivate::GetModel()
 {
 	return MutableModel;
 }
 
 
-const TSharedPtr<const mu::FModel, ESPMode::ThreadSafe> UCustomizableObjectPrivate::GetModel() const
+const TSharedPtr<const UE::Mutable::Private::FModel, ESPMode::ThreadSafe> UCustomizableObjectPrivate::GetModel() const
 {
 	return MutableModel;
 }
@@ -2361,7 +2361,7 @@ FCompilationOptions UCustomizableObjectPrivate::GetCompileOptions() const
 #endif
 
 #if WITH_EDITOR
-namespace MutablePrivate
+namespace UE::Mutable::Private
 {
 	uint64 FFile::GetSize() const
 	{
@@ -2472,7 +2472,7 @@ namespace MutablePrivate
 	}
 	
 	void GenerateBulkDataFilesListWithFileLimit(
-		TSharedPtr<const mu::FModel, ESPMode::ThreadSafe> Model,
+		TSharedPtr<const UE::Mutable::Private::FModel, ESPMode::ThreadSafe> Model,
 		FModelStreamableBulkData& ModelStreamableBulkData,
 		uint32 NumFilesPerBucket,
 		bool bAllowSplit,
@@ -2502,17 +2502,17 @@ namespace MutablePrivate
 		// Model Roms. Iterate all Model roms to distribute them in categories.
 		{
 			// Add meshes and low-res textures to the Default bucket 
-			DefaultBucket.Categories.Add({ FFileCategoryID(EStreamableDataType::Model, uint16(mu::EDataType::Mesh), 0), 0, 0, 0 });
-			DefaultBucket.Categories.Add({ FFileCategoryID(EStreamableDataType::Model, uint16(mu::EDataType::Image), 0), 0, 0, 0 });
+			DefaultBucket.Categories.Add({ FFileCategoryID(EStreamableDataType::Model, uint16(UE::Mutable::Private::EDataType::Mesh), 0), 0, 0, 0 });
+			DefaultBucket.Categories.Add({ FFileCategoryID(EStreamableDataType::Model, uint16(UE::Mutable::Private::EDataType::Image), 0), 0, 0, 0 });
 
 			// Add High-res textures to the Optional bucket
-			OptionalBucket.Categories.Add({ FFileCategoryID(EStreamableDataType::Model, uint16(mu::EDataType::Image), uint16(EMutableFileFlags::HighRes)), 0, 0, 0});
+			OptionalBucket.Categories.Add({ FFileCategoryID(EStreamableDataType::Model, uint16(UE::Mutable::Private::EDataType::Image), uint16(EMutableFileFlags::HighRes)), 0, 0, 0});
 
 			const int32 NumRoms = Model->GetRomCount();
 			for (int32 RomIndex = 0; RomIndex < NumRoms; ++RomIndex)
 			{
 				const uint32 BlockSize = Model->GetRomSize(RomIndex);
-				const uint16 BlockResourceType = Model->IsMeshData(RomIndex) ? uint16(mu::EDataType::Mesh) : uint16(mu::EDataType::Image);
+				const uint16 BlockResourceType = Model->IsMeshData(RomIndex) ? uint16(UE::Mutable::Private::EDataType::Mesh) : uint16(UE::Mutable::Private::EDataType::Image);
 				const EMutableFileFlags BlockFlags = Model->IsRomHighRes(RomIndex) ? EMutableFileFlags::HighRes : EMutableFileFlags::None;
 
 				FFileCategoryID CategoryID = { EStreamableDataType::Model, BlockResourceType, uint16(BlockFlags) };
@@ -2672,7 +2672,7 @@ namespace MutablePrivate
 	
 
 	void GenerateBulkDataFilesListWithSizeLimit(
-		TSharedPtr<const mu::FModel, ESPMode::ThreadSafe> Model,
+		TSharedPtr<const UE::Mutable::Private::FModel, ESPMode::ThreadSafe> Model,
 		FModelStreamableBulkData& ModelStreamableBulkData,
 		const ITargetPlatform* TargetPlatform,
 		uint64 TargetBulkDataFileBytes,
@@ -2759,7 +2759,7 @@ namespace MutablePrivate
 
 
 	void GenerateBulkDataFilesList(
-		TSharedPtr<const mu::FModel, ESPMode::ThreadSafe> Model,
+		TSharedPtr<const UE::Mutable::Private::FModel, ESPMode::ThreadSafe> Model,
 		FModelStreamableBulkData& ModelStreamableBulkData,
 		bool bUseRomTypeAndFlagsToFilter,
 		TFunctionRef<void(const FFileCategoryID&, const FClassifyNode&, TArray<FFile>&)> CreateFileList,
@@ -2790,7 +2790,7 @@ namespace MutablePrivate
 				// TODO: Temp. Remove after unifying generated output files code between editor an package. UE-222777 
 				if (bUseRomTypeAndFlagsToFilter)
 				{
-					BlockResourceType = Model->IsMeshData(RomIndex) ? uint16(mu::EDataType::Mesh) : uint16(mu::EDataType::Image);
+					BlockResourceType = Model->IsMeshData(RomIndex) ? uint16(UE::Mutable::Private::EDataType::Mesh) : uint16(UE::Mutable::Private::EDataType::Image);
 					BlockFlags = Model->IsRomHighRes(RomIndex) ? EMutableFileFlags::HighRes : EMutableFileFlags::None;
 				}
 
@@ -2810,7 +2810,7 @@ namespace MutablePrivate
 		{
 			const TMap<uint32, FRealTimeMorphStreamable>& RealTimeMorphStreamables = ModelStreamableBulkData.RealTimeMorphStreamables;
 
-			const FFileCategoryID RealTimeMorphCategory = { EStreamableDataType::RealTimeMorph, uint16(mu::EDataType::None), uint16(EMutableFileFlags::None) };
+			const FFileCategoryID RealTimeMorphCategory = { EStreamableDataType::RealTimeMorph, uint16(UE::Mutable::Private::EDataType::None), uint16(EMutableFileFlags::None) };
 
 			for (const TPair<uint32, FRealTimeMorphStreamable>& MorphStreamable : RealTimeMorphStreamables)
 			{
@@ -2831,7 +2831,7 @@ namespace MutablePrivate
 		{
 			const TMap<uint32, FClothingStreamable>& ClothingStreamables = ModelStreamableBulkData.ClothingStreamables;
 
-			const FFileCategoryID ClothingCategory = { EStreamableDataType::Clothing, (uint16)mu::EDataType::None, (uint16)EMutableFileFlags::None };
+			const FFileCategoryID ClothingCategory = { EStreamableDataType::Clothing, (uint16)UE::Mutable::Private::EDataType::None, (uint16)EMutableFileFlags::None };
 
 			for (const TPair<uint32, FClothingStreamable>& ClothStreamable : ClothingStreamables)
 			{
@@ -2926,7 +2926,7 @@ namespace MutablePrivate
 		const uint32 NumBulkDataFiles = BulkDataFiles.Num();
 		for (uint32 FileIndex = 0; FileIndex < NumBulkDataFiles; ++FileIndex)
 		{
-			MutablePrivate::FFile& CurrentFile = BulkDataFiles[FileIndex];
+			UE::Mutable::Private::FFile& CurrentFile = BulkDataFiles[FileIndex];
 
 			const int64 FileSize = CurrentFile.GetSize();
 			FileBulkData.SetNumUninitialized(FileSize, EAllowShrinking::No);
@@ -3111,11 +3111,11 @@ UE::DerivedData::FCacheKey UCustomizableObjectPrivate::GetDerivedDataCacheKeyFor
 }
 
 
-UE::DerivedData::FValueId GetDerivedDataValueIdForResource(MutablePrivate::EStreamableDataType StreamableDataType, uint32 FileId, uint16 ResourceType, uint16 Flags)
+UE::DerivedData::FValueId GetDerivedDataValueIdForResource(UE::Mutable::Private::EStreamableDataType StreamableDataType, uint32 FileId, uint16 ResourceType, uint16 Flags)
 {
 	UE::DerivedData::FValueId::ByteArray ValueIdBytes = {};
 	
-	static_assert(sizeof(UE::DerivedData::FValueId::ByteArray) >= sizeof(MutablePrivate::EStreamableDataType) + sizeof(uint32) + sizeof(uint16) + sizeof(uint16));
+	static_assert(sizeof(UE::DerivedData::FValueId::ByteArray) >= sizeof(UE::Mutable::Private::EStreamableDataType) + sizeof(uint32) + sizeof(uint16) + sizeof(uint16));
 
 	int8 ValueIdOffset = 0;
 	FMemory::Memcpy(&ValueIdBytes[ValueIdOffset], &StreamableDataType, sizeof(StreamableDataType));
@@ -3153,10 +3153,10 @@ void UCustomizableObjectPrivate::LogMemory()
 {
 	// Log in-memory data for the compiled CustomizableObject
 	UE_LOG(LogMutable, Log, TEXT("CustomizableObject [%s] memory report:"), *GetPublic()->GetName());
-	const mu::FModel* MuModel = GetModel().Get();
+	const UE::Mutable::Private::FModel* MuModel = GetModel().Get();
 	if (MuModel)
 	{
-		const mu::FProgram& Program = MuModel->GetPrivate()->Program;
+		const UE::Mutable::Private::FProgram& Program = MuModel->GetPrivate()->Program;
 		int32 ByteCode = Program.ByteCode.GetAllocatedSize();
 		int32 OpAddress = Program.OpAddress.GetAllocatedSize();
 		int32 Roms = Program.Roms.GetAllocatedSize();
@@ -3166,7 +3166,7 @@ void UCustomizableObjectPrivate::LogMemory()
 		int32 ConstantImageLODStreamed = Program.ConstantImageLODsStreamed.GetAllocatedSize();
 
 		int32 PermanentImages = 0;
-		for (const TSharedPtr<const mu::FImage>& Entry : Program.ConstantImageLODsPermanent)
+		for (const TSharedPtr<const UE::Mutable::Private::FImage>& Entry : Program.ConstantImageLODsPermanent)
 		{
 			if (Entry)
 			{
@@ -3181,7 +3181,7 @@ void UCustomizableObjectPrivate::LogMemory()
 		int32 ConstantMeshesStreamed = Program.ConstantMeshesStreamed.GetAllocatedSize();
 
 		int32 PermanentMeshes = 0;
-		for (const TSharedPtr<const mu::FMesh>& Entry : Program.ConstantMeshesPermanent)
+		for (const TSharedPtr<const UE::Mutable::Private::FMesh>& Entry : Program.ConstantMeshesPermanent)
 		{
 			if (Entry)
 			{
@@ -3197,26 +3197,26 @@ void UCustomizableObjectPrivate::LogMemory()
 		}
 
 		int32 SkeletonsTotal = Program.ConstantSkeletons.GetAllocatedSize();
-		for (const TSharedPtr<const mu::FSkeleton>& Entry : Program.ConstantSkeletons)
+		for (const TSharedPtr<const UE::Mutable::Private::FSkeleton>& Entry : Program.ConstantSkeletons)
 		{
-			SkeletonsTotal += sizeof(mu::FSkeleton) + Entry->DebugBoneNames.GetAllocatedSize() + Entry->BoneIds.GetAllocatedSize() + Entry->BoneParents.GetAllocatedSize();
+			SkeletonsTotal += sizeof(UE::Mutable::Private::FSkeleton) + Entry->DebugBoneNames.GetAllocatedSize() + Entry->BoneIds.GetAllocatedSize() + Entry->BoneParents.GetAllocatedSize();
 		}
 
 		int32 PhysicsTotal = Program.ConstantPhysicsBodies.GetAllocatedSize();
-		for (const TSharedPtr<const mu::FPhysicsBody>& Entry : Program.ConstantPhysicsBodies)
+		for (const TSharedPtr<const UE::Mutable::Private::FPhysicsBody>& Entry : Program.ConstantPhysicsBodies)
 		{
-			PhysicsTotal += sizeof(mu::FPhysicsBody) + Entry->BoneIds.GetAllocatedSize() + Entry->Bodies.GetAllocatedSize() + Entry->BodiesCustomIds.GetAllocatedSize();
-			for (const mu::FPhysicsBodyAggregate& Body : Entry->Bodies)
+			PhysicsTotal += sizeof(UE::Mutable::Private::FPhysicsBody) + Entry->BoneIds.GetAllocatedSize() + Entry->Bodies.GetAllocatedSize() + Entry->BodiesCustomIds.GetAllocatedSize();
+			for (const UE::Mutable::Private::FPhysicsBodyAggregate& Body : Entry->Bodies)
 			{
 				PhysicsTotal += Body.Spheres.GetAllocatedSize() + Body.Boxes.GetAllocatedSize() + Body.Convex.GetAllocatedSize() + Body.Sphyls.GetAllocatedSize() + Body.TaperedCapsules.GetAllocatedSize();
 			}
 		}
 
 		int32 ParametersTotal = Program.Parameters.GetAllocatedSize();
-		for (const mu::FParameterDesc& Entry : Program.Parameters)
+		for (const UE::Mutable::Private::FParameterDesc& Entry : Program.Parameters)
 		{
 			ParametersTotal += Entry.Ranges.GetAllocatedSize() + Entry.PossibleValues.GetAllocatedSize();
-			for (const mu::FParameterDesc::FIntValueDesc& ValueDesc : Entry.PossibleValues)
+			for (const UE::Mutable::Private::FParameterDesc::FIntValueDesc& ValueDesc : Entry.PossibleValues)
 			{
 				ParametersTotal += ValueDesc.Name.GetAllocatedSize();
 			}
@@ -3352,7 +3352,7 @@ void UCustomizableObjectBulk::CookAdditionalFilesOverride(const TCHAR* PackageFi
 
 	UCustomizableObject* CustomizableObject = CastChecked<UCustomizableObject>(GetOutermostObject());
 
-	MutablePrivate::FMutableCachedPlatformData* PlatformData = CustomizableObject->GetPrivate()->CachedPlatformsData.Find(TargetPlatform->PlatformName());
+	UE::Mutable::Private::FMutableCachedPlatformData* PlatformData = CustomizableObject->GetPrivate()->CachedPlatformsData.Find(TargetPlatform->PlatformName());
 	if (!PlatformData)
 	{
 		UE_LOG(LogMutable, Warning, TEXT("CookAdditionalFilesOverride: Customizable Object [%s] is missing [%s] platform data."),
@@ -3364,7 +3364,7 @@ void UCustomizableObjectBulk::CookAdditionalFilesOverride(const TCHAR* PackageFi
 
 	const FString CookedBulkFileName = FString::Printf(TEXT("%s/%s"), *FPaths::GetPath(PackageFilename), *CustomizableObject->GetName());
 
-	const auto WriteFile = [WriteAdditionalFile, CookedBulkFileName](MutablePrivate::FFile& File, TArray64<uint8>& FileBulkData, uint32 FileIndex)
+	const auto WriteFile = [WriteAdditionalFile, CookedBulkFileName](UE::Mutable::Private::FFile& File, TArray64<uint8>& FileBulkData, uint32 FileIndex)
 		{
 			FString FileName = CookedBulkFileName + FString::Printf(TEXT("-%08x.mut"), File.Id);
 
@@ -3379,7 +3379,7 @@ void UCustomizableObjectBulk::CookAdditionalFilesOverride(const TCHAR* PackageFi
 		};
 
 	bool bDropData = true;
-	MutablePrivate::SerializeBulkDataFiles(*PlatformData, PlatformData->BulkDataFiles, WriteFile, bDropData);
+	UE::Mutable::Private::SerializeBulkDataFiles(*PlatformData, PlatformData->BulkDataFiles, WriteFile, bDropData);
 }
 #endif // WITH_EDITOR
 
diff --git a/Engine/Plugins/Mutable/Source/CustomizableObject/Private/MuCO/CustomizableObjectDGGUI.cpp b/Engine/Plugins/Mutable/Source/CustomizableObject/Private/MuCO/CustomizableObjectDGGUI.cpp
--- a/Engine/Plugins/Mutable/Source/CustomizableObject/Private/MuCO/CustomizableObjectDGGUI.cpp
+++ b/Engine/Plugins/Mutable/Source/CustomizableObject/Private/MuCO/CustomizableObjectDGGUI.cpp
@@ -26,7 +26,7 @@ void UDGGUI::OpenDGGUI(const int32 SlotID, UCustomizableObjectInstanceUsage* Sel
 		}
 
 		FSoftClassPath DGUIPath(DGGUIAssetPath);
-		if (UClass* DGUI = MutablePrivate::LoadClass<UDGGUI>(DGUIPath))
+		if (UClass* DGUI = UE::Mutable::Private::LoadClass<UDGGUI>(DGUIPath))
 		{
 			UDGGUI* WDGUI = CreateWidget<UDGGUI>(Player, DGUI);
 			if (WDGUI)
